Fix iOS Location Usage Description #jira UE-115155 #7866 #rb jack.porter

diff --git a/Engine/Source/Programs/UnrealBuildTool/Platform/IOS/UEDeployIOS.cs b/Engine/Source/Programs/UnrealBuildTool/Platform/IOS/UEDeployIOS.cs
--- a/Engine/Source/Programs/UnrealBuildTool/Platform/IOS/UEDeployIOS.cs
+++ b/Engine/Source/Programs/UnrealBuildTool/Platform/IOS/UEDeployIOS.cs
@@ -527,11 +527,16 @@ namespace UnrealBuildTool
 			Text.AppendLine("\t<key>ITSAppUsesNonExemptEncryption</key>");
 			Text.AppendLine("\t<false/>");
 			// add location services descriptions if used
-			Text.AppendLine("\t<key>NSLocationAlwaysAndWhenInUseUsageDescription</key>");
-			Text.AppendLine(string.Format("\t<string>{0}</string>", LocationAlwaysUsageDescription));
-			Text.AppendLine("\t<key>NSLocationWhenInUseUsageDescription</key>");
-			Text.AppendLine(string.Format("\t<string>{0}></string>", LocationWhenInUseDescription));
-
+			if (!string.IsNullOrWhiteSpace(LocationAlwaysUsageDescription))
+			{
+				Text.AppendLine("\t<key>NSLocationAlwaysAndWhenInUseUsageDescription</key>");
+				Text.AppendLine(string.Format("\t<string>{0}</string>", LocationAlwaysUsageDescription));
+			}
+			if (!string.IsNullOrWhiteSpace(LocationWhenInUseDescription))
+			{
+				Text.AppendLine("\t<key>NSLocationWhenInUseUsageDescription</key>");
+				Text.AppendLine(string.Format("\t<string>{0}</string>", LocationWhenInUseDescription));
+			}
 			// disable HTTPS requirement
 			if (bDisableHTTPS)
 			{
