UnrealBuildTool:

"-Rebuild" command line option for build targets, to perform a clean on a target before attempting to build it.

diff --git a/Engine/Source/Programs/UnrealBuildTool/Configuration/TargetDescriptor.cs b/Engine/Source/Programs/UnrealBuildTool/Configuration/TargetDescriptor.cs
--- a/Engine/Source/Programs/UnrealBuildTool/Configuration/TargetDescriptor.cs
+++ b/Engine/Source/Programs/UnrealBuildTool/Configuration/TargetDescriptor.cs
@@ -85,6 +85,12 @@ namespace UnrealBuildTool
 		[CommandLine("-Quiet")]
 		public bool bQuiet;
 
+		/// <summary>
+		/// Clean the target before trying to build it
+		/// </summary>
+		[CommandLine("-Rebuild")]
+		public bool bRebuild;
+
 		/// <summary>
 		/// Constructor
 		/// </summary>
diff --git a/Engine/Source/Programs/UnrealBuildTool/Modes/BuildMode.cs b/Engine/Source/Programs/UnrealBuildTool/Modes/BuildMode.cs
--- a/Engine/Source/Programs/UnrealBuildTool/Modes/BuildMode.cs
+++ b/Engine/Source/Programs/UnrealBuildTool/Modes/BuildMode.cs
@@ -172,6 +172,15 @@ namespace UnrealBuildTool
 					TargetDescriptors.RemoveAll(x => (x.Name == "ShaderCompileWorker" || x.Name == "LiveCodingConsole" || x.Name == "InterchangeWorker") && x.SpecificFilesToCompile.Count > 0);
 				}
 
+				// Clean any target that wanted to be cleaned before being rebuilt
+				if (TargetDescriptors.Any(D => D.bRebuild))
+				{
+					CleanMode CleanMode = new CleanMode();
+					CleanMode.bSkipPreBuildTargets = bSkipPreBuildTargets;
+					CleanMode.Clean(TargetDescriptors.Where(D => D.bRebuild).ToList(), BuildConfiguration);
+				}
+
+
 				// Handle remote builds
 				for (int Idx = 0; Idx < TargetDescriptors.Count; ++Idx)
 				{
diff --git a/Engine/Source/Programs/UnrealBuildTool/Modes/CleanMode.cs b/Engine/Source/Programs/UnrealBuildTool/Modes/CleanMode.cs
--- a/Engine/Source/Programs/UnrealBuildTool/Modes/CleanMode.cs
+++ b/Engine/Source/Programs/UnrealBuildTool/Modes/CleanMode.cs
@@ -47,6 +47,16 @@ namespace UnrealBuildTool
 
 			// Parse all the targets being built
 			List<TargetDescriptor> TargetDescriptors = TargetDescriptor.ParseCommandLine(Arguments, BuildConfiguration.bUsePrecompiled, bSkipRulesCompile);
+
+			Clean(TargetDescriptors, BuildConfiguration);
+
+			return 0;
+		}
+
+		public void Clean(List<TargetDescriptor> TargetDescriptors, BuildConfiguration BuildConfiguration)
+		{
+			using ITimelineEvent ScopeEvent = Timeline.ScopeEvent("CleanMode.Clean()");
+
 			if (TargetDescriptors.Count == 0)
 			{
 				throw new BuildException("No targets specified to clean");
@@ -240,8 +250,6 @@ namespace UnrealBuildTool
 					RemoteMac.Clean(TargetDescriptor);
 				}
 			}
-
-			return 0;
 		}
 	}
 }
