GenerateProjectFiles:

Group Automation projects in the same folder with their Game or Sample project.

diff --git a/Engine/Source/Programs/UnrealBuildTool/ProjectFiles/ProjectFileGenerator.cs b/Engine/Source/Programs/UnrealBuildTool/ProjectFiles/ProjectFileGenerator.cs
--- a/Engine/Source/Programs/UnrealBuildTool/ProjectFiles/ProjectFileGenerator.cs
+++ b/Engine/Source/Programs/UnrealBuildTool/ProjectFiles/ProjectFileGenerator.cs
@@ -419,9 +419,8 @@ namespace UnrealBuildTool
 		/// <summary>
 		/// Adds all .automation.csproj files to the solution.
 		/// </summary>
-		void AddAutomationModules(List<FileReference> UnrealProjectFiles, MasterProjectFolder ProgramsFolder)
+		void AddAutomationModules(List<FileReference> UnrealProjectFiles, MasterProjectFolder RootFolder, MasterProjectFolder ProgramsFolder)
 		{
-			MasterProjectFolder Folder = ProgramsFolder.AddSubFolder("Automation");
 			List<DirectoryReference> BuildFolders = new List<DirectoryReference>();
 			foreach (FileReference UnrealProjectFile in UnrealProjectFiles)
 			{
@@ -432,6 +431,9 @@ namespace UnrealBuildTool
 				}
 			}
 
+			MasterProjectFolder AutomationFolder = ProgramsFolder.AddSubFolder("Automation");
+			DirectoryReference SamplesDirectory = DirectoryReference.Combine(Unreal.RootDirectory, "Samples");
+
 			// Find all the automation modules .csproj files to add
 			List<FileReference> ModuleFiles = Rules.FindAllRulesSourceFiles(Rules.RulesFileType.AutomationModule, null, ForeignPlugins: null, AdditionalSearchPaths: BuildFolders);
 			foreach (FileReference ProjectFile in ModuleFiles)
@@ -443,12 +445,24 @@ namespace UnrealBuildTool
 					Project.ShouldBuildForAllSolutionTargets = false;//true;
 					AddExistingProjectFile(Project, bForceDevelopmentConfiguration: true);
 					AutomationProjectFiles.Add(Project);
-					Folder.ChildProjects.Add(Project);
 
 					if (!ProjectFile.IsUnderDirectory(Unreal.EngineDirectory))
 					{
 						FileReference PropsFile = new FileReference(ProjectFile.FullName + ".props");
 						CreateAutomationProjectPropsFile(PropsFile);
+
+						if (ProjectFile.IsUnderDirectory(SamplesDirectory))
+						{
+							RootFolder.AddSubFolder("Samples").ChildProjects.Add(Project);
+						}
+						else
+						{
+							RootFolder.AddSubFolder("Games").ChildProjects.Add(Project);
+						}
+					}
+					else
+					{
+						AutomationFolder.ChildProjects.Add(Project);
 					}
 				}
 			}
@@ -1019,7 +1033,7 @@ namespace UnrealBuildTool
 						ProgramsFolder.ChildProjects.Add(AutomationToolProject);
 
 					// Add automation.csproj files to the master project
-					AddAutomationModules(AllGameProjects, ProgramsFolder);
+					AddAutomationModules(AllGameProjects, RootFolder, ProgramsFolder);
 
 					// Add shared projects
 					AddSharedDotNetModules(ProgramsFolder);
