UBT: Added a new static analyzer type, 'Clang', which runs clang's static analyzer, regardless of the project. Also added a new option, StaticAnalyzerOutputType (-StaticAnalyzerOutputType on the command line) which takes either Text or Html as values, Text being default (and only option for non-clang analysers). Text simply writes out the analysis to stdout, whereas Html will write out an HTML page in the object output folder, that contains a navigable version of the code, including the analysis.

diff --git a/Engine/Source/Programs/UnrealBuildTool/Platform/Windows/UEBuildWindows.cs b/Engine/Source/Programs/UnrealBuildTool/Platform/Windows/UEBuildWindows.cs
--- a/Engine/Source/Programs/UnrealBuildTool/Platform/Windows/UEBuildWindows.cs
+++ b/Engine/Source/Programs/UnrealBuildTool/Platform/Windows/UEBuildWindows.cs
@@ -78,6 +78,32 @@ namespace UnrealBuildTool
 		/// Use PVS-Studio for static analysis
 		/// </summary>
 		PVSStudio,
+		
+		/// <summary>
+		/// Use clang for static analysis. This forces the compiler to clang.
+		/// </summary>
+		Clang,
+	}
+
+	/// <summary>
+	/// Output type for the static analyzer. This currently only works for the Clang static analyzer.
+	/// The Clang static analyzer can do either Text, which prints the analysis to stdout, or
+	/// html, where it writes out a navigable HTML page for each issue that it finds, per file.
+	/// The HTML is output in the same directory as the object fil that would otherwise have
+	/// been generated. 
+	/// All other analyzers default automatically to Text. 
+	/// </summary>
+	public enum WindowsStaticAnalyzerOutputType
+	{
+		/// <summary>
+		/// Output the analysis to stdout.
+		/// </summary>
+		Text,
+
+		/// <summary>
+		/// Output the analysis to an HTML file in the object folder.
+		/// </summary>
+		Html,
 	}
 
 	/// <summary>
@@ -190,6 +216,13 @@ namespace UnrealBuildTool
 		[CommandLine("-StaticAnalyzer")]
 		public WindowsStaticAnalyzer StaticAnalyzer = WindowsStaticAnalyzer.None;
 
+		/// <summary>
+		/// The output type to use for the static analyzer.
+		/// </summary>
+		[XmlConfigFile(Category = "WindowsPlatform")]
+		[CommandLine("-StaticAnalyzerOutputType")]
+		public WindowsStaticAnalyzerOutputType StaticAnalyzerOutputType = WindowsStaticAnalyzerOutputType.Text;
+
 		/// <summary>
 		/// Enables address sanitizer (ASan). Only supported for Visual Studio 2019 16.7.0 and up.
 		/// </summary>
@@ -499,6 +532,11 @@ namespace UnrealBuildTool
 		{
 			get { return Inner.StaticAnalyzer; }
 		}
+		
+		public WindowsStaticAnalyzerOutputType StaticAnalyzerOutputType
+		{
+			get { return Inner.StaticAnalyzerOutputType; }
+		}
 
 		public bool bEnableAddressSanitizer
 		{ 
@@ -839,6 +877,19 @@ namespace UnrealBuildTool
 			{				
 				Target.bCompileCEF3 = false;
 			}
+			
+			// If clang is selected for static analysis, switch compiler to clang and proceed
+			// as normal.
+			if (Target.WindowsPlatform.StaticAnalyzer == WindowsStaticAnalyzer.Clang)
+			{
+				Target.WindowsPlatform.Compiler = WindowsCompiler.Clang;
+				Target.WindowsPlatform.StaticAnalyzer = WindowsStaticAnalyzer.Default;
+			}
+			else if (Target.WindowsPlatform.StaticAnalyzer != WindowsStaticAnalyzer.None && 
+			         Target.WindowsPlatform.StaticAnalyzerOutputType != WindowsStaticAnalyzerOutputType.Text)
+			{
+				Log.TraceInformation("Defaulting static analyzer output type to text");
+			}
 
 			// Disable static analysis if not appropriate for the compiler.
 			if (Target.WindowsPlatform.StaticAnalyzer == WindowsStaticAnalyzer.Default)
diff --git a/Engine/Source/Programs/UnrealBuildTool/Platform/Windows/VCToolChain.cs b/Engine/Source/Programs/UnrealBuildTool/Platform/Windows/VCToolChain.cs
--- a/Engine/Source/Programs/UnrealBuildTool/Platform/Windows/VCToolChain.cs
+++ b/Engine/Source/Programs/UnrealBuildTool/Platform/Windows/VCToolChain.cs
@@ -236,7 +236,18 @@ namespace UnrealBuildTool
 					Arguments.Add("-Xclang -analyze");
 
 					// Make sure we get textual output and not XML.
-					Arguments.Add("-Xclang -analyzer-output=text");
+					if (Target.WindowsPlatform.StaticAnalyzerOutputType == WindowsStaticAnalyzerOutputType.Html)
+					{
+						// Write out a pretty web page with navigation to understand how the analysis was derived.
+						Arguments.Add("-Xclang -analyzer-output=html");
+						
+						// If writing to HTML, use the source filename as a basis for the report filename. 
+						Arguments.Add("-Xclang -analyzer-config -Xclang stable-report-filename=true");
+					}
+					else
+					{
+						Arguments.Add("-Xclang -analyzer-output=text");
+					}
 
 					// Make sure we check inside nested blocks (e.g. 'if ((foo = getchar()) == 0) {}')
 					Arguments.Add("-Xclang -analyzer-opt-analyze-nested-blocks");
