UnrealBuildTool:

Add new commandline option: -AlwaysRulesCompile, ensures that rules assemblies will be compiled during the run.

diff --git a/Engine/Source/Programs/UnrealBuildTool/Configuration/BuildConfiguration.cs b/Engine/Source/Programs/UnrealBuildTool/Configuration/BuildConfiguration.cs
--- a/Engine/Source/Programs/UnrealBuildTool/Configuration/BuildConfiguration.cs
+++ b/Engine/Source/Programs/UnrealBuildTool/Configuration/BuildConfiguration.cs
@@ -141,6 +141,12 @@ namespace UnrealBuildTool
 		[CommandLine("-SkipRulesCompile")]
 		public bool bSkipRulesCompile = false;
 
+		/// <summary>
+		/// Whether to force compiling rules assemblies, regardless of whether they are valid
+		/// </summary>
+		[CommandLine("-ForceRulesCompile")]
+		public bool bForceRulesCompile = false;
+
 		/// <summary>
 		/// Maximum recommended root path length.
 		/// </summary>
diff --git a/Engine/Source/Programs/UnrealBuildTool/Configuration/TargetDescriptor.cs b/Engine/Source/Programs/UnrealBuildTool/Configuration/TargetDescriptor.cs
--- a/Engine/Source/Programs/UnrealBuildTool/Configuration/TargetDescriptor.cs
+++ b/Engine/Source/Programs/UnrealBuildTool/Configuration/TargetDescriptor.cs
@@ -171,11 +171,12 @@ namespace UnrealBuildTool
 		/// <param name="Arguments">Command-line arguments</param>
 		/// <param name="bUsePrecompiled">Whether to use a precompiled engine distribution</param>
 		/// <param name="bSkipRulesCompile">Whether to skip compiling rules assemblies</param>
+		/// <param name="bForceRulesCompile">Whether to always compile all rules assemblies</param>
 		/// <returns>List of target descriptors</returns>
-		public static List<TargetDescriptor> ParseCommandLine(CommandLineArguments Arguments, bool bUsePrecompiled, bool bSkipRulesCompile)
+		public static List<TargetDescriptor> ParseCommandLine(CommandLineArguments Arguments, bool bUsePrecompiled, bool bSkipRulesCompile, bool bForceRulesCompile)
 		{
 			List<TargetDescriptor> TargetDescriptors = new List<TargetDescriptor>();
-			ParseCommandLine(Arguments, bUsePrecompiled, bSkipRulesCompile, TargetDescriptors);
+			ParseCommandLine(Arguments, bUsePrecompiled, bSkipRulesCompile, bForceRulesCompile, TargetDescriptors);
 			return TargetDescriptors;
 		}
 
@@ -185,8 +186,9 @@ namespace UnrealBuildTool
 		/// <param name="Arguments">Command-line arguments</param>
 		/// <param name="bUsePrecompiled">Whether to use a precompiled engine distribution</param>
 		/// <param name="bSkipRulesCompile">Whether to skip compiling rules assemblies</param>
+		/// <param name="bForceRulesCompile">Whether to always compile rules assemblies</param>
 		/// <param name="TargetDescriptors">Receives the list of parsed target descriptors</param>
-		public static void ParseCommandLine(CommandLineArguments Arguments, bool bUsePrecompiled, bool bSkipRulesCompile, List<TargetDescriptor> TargetDescriptors)
+		public static void ParseCommandLine(CommandLineArguments Arguments, bool bUsePrecompiled, bool bSkipRulesCompile, bool bForceRulesCompile, List<TargetDescriptor> TargetDescriptors)
 		{
 			List<string> TargetLists;
 			Arguments = Arguments.Remove("-TargetList=", out TargetLists);
@@ -206,7 +208,7 @@ namespace UnrealBuildTool
 						if(TrimLine.Length > 0 && TrimLine[0] != ';')
 						{
 							CommandLineArguments NewArguments = Arguments.Append(CommandLineArguments.Split(TrimLine));
-							ParseCommandLine(NewArguments, bUsePrecompiled, bSkipRulesCompile, TargetDescriptors);
+							ParseCommandLine(NewArguments, bUsePrecompiled, bSkipRulesCompile, bForceRulesCompile, TargetDescriptors);
 						}
 					}
 				}
@@ -214,13 +216,13 @@ namespace UnrealBuildTool
 				foreach(string Target in Targets)
 				{
 					CommandLineArguments NewArguments = Arguments.Append(CommandLineArguments.Split(Target));
-					ParseCommandLine(NewArguments, bUsePrecompiled, bSkipRulesCompile, TargetDescriptors);
+					ParseCommandLine(NewArguments, bUsePrecompiled, bSkipRulesCompile, bForceRulesCompile, TargetDescriptors);
 				}
 			}
 			else
 			{
 				// Otherwise just process the whole command line together
-				ParseSingleCommandLine(Arguments, bUsePrecompiled, bSkipRulesCompile, TargetDescriptors);
+				ParseSingleCommandLine(Arguments, bUsePrecompiled, bSkipRulesCompile, bForceRulesCompile, TargetDescriptors);
 			}
 		}
 
@@ -230,8 +232,9 @@ namespace UnrealBuildTool
 		/// <param name="Arguments">Command-line arguments</param>
 		/// <param name="bUsePrecompiled">Whether to use a precompiled engine distribution</param>
 		/// <param name="bSkipRulesCompile">Whether to skip compiling rules assemblies</param>
+		/// <param name="bForceRulesCompile">Whether to always compile all rules assemblies</param>
 		/// <param name="TargetDescriptors">List of target descriptors</param>
-		public static void ParseSingleCommandLine(CommandLineArguments Arguments, bool bUsePrecompiled, bool bSkipRulesCompile, List<TargetDescriptor> TargetDescriptors)
+		public static void ParseSingleCommandLine(CommandLineArguments Arguments, bool bUsePrecompiled, bool bSkipRulesCompile, bool bForceRulesCompile, List<TargetDescriptor> TargetDescriptors)
 		{
 			List<UnrealTargetPlatform> Platforms = new List<UnrealTargetPlatform>();
 			List<UnrealTargetConfiguration> Configurations = new List<UnrealTargetConfiguration>();
@@ -364,11 +367,11 @@ namespace UnrealBuildTool
 
 							if (ProjectFile == null)
 							{
-								TargetNames.Add(RulesCompiler.CreateEngineRulesAssembly(bUsePrecompiled, bSkipRulesCompile).GetTargetNameByType(TargetType, Platform, Configuration, Architecture, null));
+								TargetNames.Add(RulesCompiler.CreateEngineRulesAssembly(bUsePrecompiled, bSkipRulesCompile, bForceRulesCompile).GetTargetNameByType(TargetType, Platform, Configuration, Architecture, null));
 							}
 							else
 							{
-								TargetNames.Add(RulesCompiler.CreateProjectRulesAssembly(ProjectFile, bUsePrecompiled, bSkipRulesCompile).GetTargetNameByType(TargetType, Platform, Configuration, Architecture, ProjectFile));
+								TargetNames.Add(RulesCompiler.CreateProjectRulesAssembly(ProjectFile, bUsePrecompiled, bSkipRulesCompile, bForceRulesCompile).GetTargetNameByType(TargetType, Platform, Configuration, Architecture, ProjectFile));
 							}
 						}
 
diff --git a/Engine/Source/Programs/UnrealBuildTool/Configuration/UEBuildTarget.cs b/Engine/Source/Programs/UnrealBuildTool/Configuration/UEBuildTarget.cs
--- a/Engine/Source/Programs/UnrealBuildTool/Configuration/UEBuildTarget.cs
+++ b/Engine/Source/Programs/UnrealBuildTool/Configuration/UEBuildTarget.cs
@@ -672,9 +672,10 @@ namespace UnrealBuildTool
 		/// </summary>
 		/// <param name="Descriptor">Information about the target</param>
 		/// <param name="bSkipRulesCompile">Whether to skip compiling any rules assemblies</param>
+		/// <param name="bForceRulesCompile">Whether to always compile all rules assemblies</param>
 		/// <param name="bUsePrecompiled">Whether to use a precompiled engine build</param>
 		/// <returns>The build target object for the specified build rules source file</returns>
-		public static UEBuildTarget Create(TargetDescriptor Descriptor, bool bSkipRulesCompile, bool bUsePrecompiled)
+		public static UEBuildTarget Create(TargetDescriptor Descriptor, bool bSkipRulesCompile, bool bForceRulesCompile, bool bUsePrecompiled)
 		{
 			// make sure we are allowed to build this platform
 			if (!UEBuildPlatform.IsPlatformAvailable(Descriptor.Platform))
@@ -685,7 +686,7 @@ namespace UnrealBuildTool
 			RulesAssembly RulesAssembly;
 			using(Timeline.ScopeEvent("RulesCompiler.CreateTargetRulesAssembly()"))
 			{
-				RulesAssembly = RulesCompiler.CreateTargetRulesAssembly(Descriptor.ProjectFile, Descriptor.Name, bSkipRulesCompile, bUsePrecompiled, Descriptor.ForeignPlugin);
+				RulesAssembly = RulesCompiler.CreateTargetRulesAssembly(Descriptor.ProjectFile, Descriptor.Name, bSkipRulesCompile, bForceRulesCompile, bUsePrecompiled, Descriptor.ForeignPlugin);
 			}
 
 			TargetRules RulesObject;
diff --git a/Engine/Source/Programs/UnrealBuildTool/Modes/BuildMode.cs b/Engine/Source/Programs/UnrealBuildTool/Modes/BuildMode.cs
--- a/Engine/Source/Programs/UnrealBuildTool/Modes/BuildMode.cs
+++ b/Engine/Source/Programs/UnrealBuildTool/Modes/BuildMode.cs
@@ -163,7 +163,7 @@ namespace UnrealBuildTool
 				// Parse all the target descriptors
 				using(Timeline.ScopeEvent("TargetDescriptor.ParseCommandLine()"))
 				{
-					TargetDescriptors = TargetDescriptor.ParseCommandLine(Arguments, BuildConfiguration.bUsePrecompiled, BuildConfiguration.bSkipRulesCompile);
+					TargetDescriptors = TargetDescriptor.ParseCommandLine(Arguments, BuildConfiguration.bUsePrecompiled, BuildConfiguration.bSkipRulesCompile, BuildConfiguration.bForceRulesCompile);
 				}
 
 				// Hack for specific files compile; don't build the ShaderCompileWorker target that's added to the command line for generated project files
@@ -708,7 +708,7 @@ namespace UnrealBuildTool
 				UEBuildTarget Target;
 				using(Timeline.ScopeEvent("UEBuildTarget.Create()"))
 				{
-					Target = UEBuildTarget.Create(TargetDescriptor, BuildConfiguration.bSkipRulesCompile, BuildConfiguration.bUsePrecompiled);
+					Target = UEBuildTarget.Create(TargetDescriptor, BuildConfiguration.bSkipRulesCompile, BuildConfiguration.bForceRulesCompile, BuildConfiguration.bUsePrecompiled);
 				}
 
 				// Create the pre-build scripts
diff --git a/Engine/Source/Programs/UnrealBuildTool/Modes/CleanMode.cs b/Engine/Source/Programs/UnrealBuildTool/Modes/CleanMode.cs
--- a/Engine/Source/Programs/UnrealBuildTool/Modes/CleanMode.cs
+++ b/Engine/Source/Programs/UnrealBuildTool/Modes/CleanMode.cs
@@ -46,7 +46,7 @@ namespace UnrealBuildTool
 			Arguments.ApplyTo(BuildConfiguration);
 
 			// Parse all the targets being built
-			List<TargetDescriptor> TargetDescriptors = TargetDescriptor.ParseCommandLine(Arguments, BuildConfiguration.bUsePrecompiled, bSkipRulesCompile);
+			List<TargetDescriptor> TargetDescriptors = TargetDescriptor.ParseCommandLine(Arguments, BuildConfiguration.bUsePrecompiled, bSkipRulesCompile, BuildConfiguration.bForceRulesCompile);
 
 			Clean(TargetDescriptors, BuildConfiguration);
 
@@ -108,7 +108,7 @@ namespace UnrealBuildTool
 				TargetDescriptor TargetDescriptor = TargetDescriptors[Idx];
 
 				// Create the rules assembly
-				RulesAssembly RulesAssembly = RulesCompiler.CreateTargetRulesAssembly(TargetDescriptor.ProjectFile, TargetDescriptor.Name, bSkipRulesCompile, BuildConfiguration.bUsePrecompiled, TargetDescriptor.ForeignPlugin);
+				RulesAssembly RulesAssembly = RulesCompiler.CreateTargetRulesAssembly(TargetDescriptor.ProjectFile, TargetDescriptor.Name, bSkipRulesCompile, BuildConfiguration.bForceRulesCompile, BuildConfiguration.bUsePrecompiled, TargetDescriptor.ForeignPlugin);
 
 				// Create the rules object
 				ReadOnlyTargetRules Target = new ReadOnlyTargetRules(RulesAssembly.CreateTargetRules(TargetDescriptor.Name, TargetDescriptor.Platform, TargetDescriptor.Configuration, TargetDescriptor.Architecture, TargetDescriptor.ProjectFile, TargetDescriptor.AdditionalArguments));
diff --git a/Engine/Source/Programs/UnrealBuildTool/Modes/GenerateClangDatabase.cs b/Engine/Source/Programs/UnrealBuildTool/Modes/GenerateClangDatabase.cs
--- a/Engine/Source/Programs/UnrealBuildTool/Modes/GenerateClangDatabase.cs
+++ b/Engine/Source/Programs/UnrealBuildTool/Modes/GenerateClangDatabase.cs
@@ -50,7 +50,7 @@ namespace UnrealBuildTool
 			}
 
 			// Parse all the target descriptors
-			List<TargetDescriptor> TargetDescriptors = TargetDescriptor.ParseCommandLine(Arguments, BuildConfiguration.bUsePrecompiled, BuildConfiguration.bSkipRulesCompile);
+			List<TargetDescriptor> TargetDescriptors = TargetDescriptor.ParseCommandLine(Arguments, BuildConfiguration.bUsePrecompiled, BuildConfiguration.bSkipRulesCompile, BuildConfiguration.bForceRulesCompile);
 
 			// Generate the compile DB for each target
 			using (ISourceFileWorkingSet WorkingSet = new EmptySourceFileWorkingSet())
@@ -63,7 +63,7 @@ namespace UnrealBuildTool
 					TargetDescriptor.AdditionalArguments = TargetDescriptor.AdditionalArguments.Append(new string[] { "-NoPCH", "-DisableUnity" });
 
 					// Create a makefile for the target
-					UEBuildTarget Target = UEBuildTarget.Create(TargetDescriptor, BuildConfiguration.bSkipRulesCompile, BuildConfiguration.bUsePrecompiled);
+					UEBuildTarget Target = UEBuildTarget.Create(TargetDescriptor, BuildConfiguration.bSkipRulesCompile, BuildConfiguration.bForceRulesCompile, BuildConfiguration.bUsePrecompiled);
 
 					// Find the location of the compiler
 					VCEnvironment Environment = VCEnvironment.Create(WindowsCompiler.Clang, Target.Platform, Target.Rules.WindowsPlatform.Architecture, null, Target.Rules.WindowsPlatform.WindowsSdkVersion, null);
diff --git a/Engine/Source/Programs/UnrealBuildTool/Modes/JsonExportMode.cs b/Engine/Source/Programs/UnrealBuildTool/Modes/JsonExportMode.cs
--- a/Engine/Source/Programs/UnrealBuildTool/Modes/JsonExportMode.cs
+++ b/Engine/Source/Programs/UnrealBuildTool/Modes/JsonExportMode.cs
@@ -31,11 +31,11 @@ namespace UnrealBuildTool
 		{
 			Arguments.ApplyTo(this);
 
-			List<TargetDescriptor> TargetDescriptors = TargetDescriptor.ParseCommandLine(Arguments, false, false);
+			List<TargetDescriptor> TargetDescriptors = TargetDescriptor.ParseCommandLine(Arguments, false, false, false);
 			foreach(TargetDescriptor TargetDescriptor in TargetDescriptors)
 			{
 				// Create the target
-				UEBuildTarget Target = UEBuildTarget.Create(TargetDescriptor, false, false);
+				UEBuildTarget Target = UEBuildTarget.Create(TargetDescriptor, false, false, false);
 
 				// Get the output file
 				FileReference? OutputFile = TargetDescriptor.AdditionalArguments.GetFileReferenceOrDefault("-OutputFile=", null);
diff --git a/Engine/Source/Programs/UnrealBuildTool/Modes/QueryTargetsMode.cs b/Engine/Source/Programs/UnrealBuildTool/Modes/QueryTargetsMode.cs
--- a/Engine/Source/Programs/UnrealBuildTool/Modes/QueryTargetsMode.cs
+++ b/Engine/Source/Programs/UnrealBuildTool/Modes/QueryTargetsMode.cs
@@ -52,11 +52,11 @@ namespace UnrealBuildTool
 			RulesAssembly Assembly;
 			if (ProjectFile != null)
 			{
-				Assembly = RulesCompiler.CreateProjectRulesAssembly(ProjectFile, BuildConfiguration.bUsePrecompiled, BuildConfiguration.bSkipRulesCompile);
+				Assembly = RulesCompiler.CreateProjectRulesAssembly(ProjectFile, BuildConfiguration.bUsePrecompiled, BuildConfiguration.bSkipRulesCompile, BuildConfiguration.bForceRulesCompile);
 			}
 			else
 			{
-				Assembly = RulesCompiler.CreateEngineRulesAssembly(BuildConfiguration.bUsePrecompiled, BuildConfiguration.bSkipRulesCompile);
+				Assembly = RulesCompiler.CreateEngineRulesAssembly(BuildConfiguration.bUsePrecompiled, BuildConfiguration.bSkipRulesCompile, BuildConfiguration.bForceRulesCompile);
 			}
 
 			// Write information about these targets
diff --git a/Engine/Source/Programs/UnrealBuildTool/ProjectFiles/ProjectFileGenerator.cs b/Engine/Source/Programs/UnrealBuildTool/ProjectFiles/ProjectFileGenerator.cs
--- a/Engine/Source/Programs/UnrealBuildTool/ProjectFiles/ProjectFileGenerator.cs
+++ b/Engine/Source/Programs/UnrealBuildTool/ProjectFiles/ProjectFileGenerator.cs
@@ -1120,7 +1120,7 @@ namespace UnrealBuildTool
 				// Generate all the target info files for the editor
 				foreach (FileReference ProjectFile in AllGameProjects)
 				{
-					RulesAssembly RulesAssembly = RulesCompiler.CreateProjectRulesAssembly(ProjectFile, false, false);
+					RulesAssembly RulesAssembly = RulesCompiler.CreateProjectRulesAssembly(ProjectFile, false, false, false);
 					QueryTargetsMode.WriteTargetInfo(ProjectFile, RulesAssembly, QueryTargetsMode.GetDefaultOutputFile(ProjectFile), new CommandLineArguments(GetTargetArguments(Arguments)));
 				}
 			}
@@ -1885,7 +1885,7 @@ namespace UnrealBuildTool
 							TargetDescriptor TargetDesc = new TargetDescriptor(CurTarget.UnrealProjectFilePath, CurTarget.Name, IntellisensePlatform, UnrealTargetConfiguration.Development, DefaultArchitecture, new CommandLineArguments(NewArguments.ToArray()));
 
 							// Create the target
-							UEBuildTarget Target = UEBuildTarget.Create(TargetDesc, false, bUsePrecompiled);
+							UEBuildTarget Target = UEBuildTarget.Create(TargetDesc, false, false, bUsePrecompiled);
 							AddTargetForIntellisense(Target);
 
 							// Generate a compile environment for each module in the binary
@@ -2306,7 +2306,7 @@ namespace UnrealBuildTool
 					FileReference CheckProjectFile = AllGames.FirstOrDefault(x => TargetFilePath.IsUnderDirectory(x.Directory));
 					if (CheckProjectFile == null)
 					{
-						RulesAssembly = RulesCompiler.CreateEngineRulesAssembly(false, false);
+						RulesAssembly = RulesCompiler.CreateEngineRulesAssembly(false, false, false);
 					
 						// Record the Engine assembly, and any parent assemblies (varies e.g. Rules, ProgramRules, MarketplaceRules)
 						if (!RulesAssemblies.ContainsKey(RulesAssembly))
@@ -2320,7 +2320,7 @@ namespace UnrealBuildTool
 					}
 					else
 					{
-						RulesAssembly = RulesCompiler.CreateProjectRulesAssembly(CheckProjectFile, false, false);
+						RulesAssembly = RulesCompiler.CreateProjectRulesAssembly(CheckProjectFile, false, false, false);
 						// Recording of game rule assemblies happens after BaseFolder has been computed
 					}
 
diff --git a/Engine/Source/Programs/UnrealBuildTool/ProjectFiles/Rider/RiderProjectFile.cs b/Engine/Source/Programs/UnrealBuildTool/ProjectFiles/Rider/RiderProjectFile.cs
--- a/Engine/Source/Programs/UnrealBuildTool/ProjectFiles/Rider/RiderProjectFile.cs
+++ b/Engine/Source/Programs/UnrealBuildTool/ProjectFiles/Rider/RiderProjectFile.cs
@@ -91,7 +91,7 @@ namespace UnrealBuildTool
 							Platform, Configuration, DefaultArchitecture, Arguments);
 						try
 						{
-							UEBuildTarget BuildTarget = UEBuildTarget.Create(TargetDesc, false, false);
+							UEBuildTarget BuildTarget = UEBuildTarget.Create(TargetDesc, false, false, false);
 						
 							FileReference OutputFile = FileReference.Combine(TargetFolder, $"{ProjectName}.json");
 							FileToTarget.Add(Tuple.Create(OutputFile, BuildTarget));
diff --git a/Engine/Source/Programs/UnrealBuildTool/ProjectFiles/Rider/RiderProjectFileGenerator.cs b/Engine/Source/Programs/UnrealBuildTool/ProjectFiles/Rider/RiderProjectFileGenerator.cs
--- a/Engine/Source/Programs/UnrealBuildTool/ProjectFiles/Rider/RiderProjectFileGenerator.cs
+++ b/Engine/Source/Programs/UnrealBuildTool/ProjectFiles/Rider/RiderProjectFileGenerator.cs
@@ -228,11 +228,11 @@ namespace UnrealBuildTool
 						AllGames.FirstOrDefault(x => TargetFilePath.IsUnderDirectory(x.Directory));
 					if (CheckProjectFile == null)
 					{
-						RulesAssembly = RulesCompiler.CreateEngineRulesAssembly(false, false);
+						RulesAssembly = RulesCompiler.CreateEngineRulesAssembly(false, false, false);
 					}
 					else
 					{
-						RulesAssembly = RulesCompiler.CreateProjectRulesAssembly(CheckProjectFile, false, false);
+						RulesAssembly = RulesCompiler.CreateProjectRulesAssembly(CheckProjectFile, false, false, false);
 					}
 
 					// Create target rules for all of the platforms and configuration combinations that we want to enable support for.
diff --git a/Engine/Source/Programs/UnrealBuildTool/System/DynamicCompilation.cs b/Engine/Source/Programs/UnrealBuildTool/System/DynamicCompilation.cs
--- a/Engine/Source/Programs/UnrealBuildTool/System/DynamicCompilation.cs
+++ b/Engine/Source/Programs/UnrealBuildTool/System/DynamicCompilation.cs
@@ -286,14 +286,15 @@ namespace UnrealBuildTool
 		/// <param name="ReferencedAssembies"></param>
 		/// <param name="PreprocessorDefines"></param>
 		/// <param name="DoNotCompile"></param>
+		/// <param name="ForceCompile"></param>
 		/// <param name="TreatWarningsAsErrors"></param>
 		/// <returns>The assembly that was loaded</returns>
-		public static Assembly? CompileAndLoadAssembly(FileReference OutputAssemblyPath, HashSet<FileReference> SourceFileNames, List<string>? ReferencedAssembies = null, List<string>? PreprocessorDefines = null, bool DoNotCompile = false, bool TreatWarningsAsErrors = false)
+		public static Assembly? CompileAndLoadAssembly(FileReference OutputAssemblyPath, HashSet<FileReference> SourceFileNames, List<string>? ReferencedAssembies = null, List<string>? PreprocessorDefines = null, bool DoNotCompile = false, bool ForceCompile = false, bool TreatWarningsAsErrors = false)
 		{
 			// Check to see if the resulting assembly is compiled and up to date
 			FileReference AssemblyManifestFilePath = FileReference.Combine(OutputAssemblyPath.Directory, Path.GetFileNameWithoutExtension(OutputAssemblyPath.FullName) + "Manifest.json");
 
-			bool bNeedsCompilation = false;
+			bool bNeedsCompilation = ForceCompile;
 			if (!DoNotCompile)
 			{
 				bNeedsCompilation = RequiresCompilation(SourceFileNames, AssemblyManifestFilePath, OutputAssemblyPath);
diff --git a/Engine/Source/Programs/UnrealBuildTool/System/RulesAssembly.cs b/Engine/Source/Programs/UnrealBuildTool/System/RulesAssembly.cs
--- a/Engine/Source/Programs/UnrealBuildTool/System/RulesAssembly.cs
+++ b/Engine/Source/Programs/UnrealBuildTool/System/RulesAssembly.cs
@@ -111,8 +111,9 @@ namespace UnrealBuildTool
 		/// <param name="DefaultBuildSettings">Optional override for the default build settings version for modules created from this assembly.</param>
 		/// <param name="bReadOnly">Whether the modules and targets in this assembly are installed, and should be created with the bUsePrecompiled flag set</param> 
 		/// <param name="bSkipCompile">Whether to skip compiling this assembly</param>
+		/// <param name="bForceCompile">Whether to always compile this assembly</param>
 		/// <param name="Parent">The parent rules assembly</param>
-		internal RulesAssembly(RulesScope Scope, List<DirectoryReference> BaseDirs, IReadOnlyList<PluginInfo> Plugins, Dictionary<FileReference, ModuleRulesContext> ModuleFileToContext, List<FileReference> TargetFiles, FileReference AssemblyFileName, bool bContainsEngineModules, BuildSettingsVersion? DefaultBuildSettings, bool bReadOnly, bool bSkipCompile, RulesAssembly? Parent)
+		internal RulesAssembly(RulesScope Scope, List<DirectoryReference> BaseDirs, IReadOnlyList<PluginInfo> Plugins, Dictionary<FileReference, ModuleRulesContext> ModuleFileToContext, List<FileReference> TargetFiles, FileReference AssemblyFileName, bool bContainsEngineModules, BuildSettingsVersion? DefaultBuildSettings, bool bReadOnly, bool bSkipCompile, bool bForceCompile, RulesAssembly? Parent)
 		{
 			this.Scope = Scope;
 			this.BaseDirs = BaseDirs;
@@ -132,7 +133,7 @@ namespace UnrealBuildTool
 			if (AssemblySourceFiles.Count > 0)
 			{
 				PreprocessorDefines = GetPreprocessorDefinitions();
-				CompiledAssembly = DynamicCompilation.CompileAndLoadAssembly(AssemblyFileName, AssemblySourceFiles, PreprocessorDefines: PreprocessorDefines, DoNotCompile: bSkipCompile);
+				CompiledAssembly = DynamicCompilation.CompileAndLoadAssembly(AssemblyFileName, AssemblySourceFiles, PreprocessorDefines: PreprocessorDefines, DoNotCompile: bSkipCompile, ForceCompile: bForceCompile);
 			}
 
 			// Setup the module map
diff --git a/Engine/Source/Programs/UnrealBuildTool/System/RulesCompiler.cs b/Engine/Source/Programs/UnrealBuildTool/System/RulesCompiler.cs
--- a/Engine/Source/Programs/UnrealBuildTool/System/RulesCompiler.cs
+++ b/Engine/Source/Programs/UnrealBuildTool/System/RulesCompiler.cs
@@ -69,8 +69,9 @@ namespace UnrealBuildTool
 		/// </summary>
 		/// <param name="bUsePrecompiled">Whether to use a precompiled engine</param>
 		/// <param name="bSkipCompile">Whether to skip compilation for this assembly</param>
+		/// <param name="bForceCompile">Whether to always compile this assembly</param>
 		/// <returns>New rules assembly</returns>
-		public static RulesAssembly CreateEngineRulesAssembly(bool bUsePrecompiled, bool bSkipCompile)
+		public static RulesAssembly CreateEngineRulesAssembly(bool bUsePrecompiled, bool bSkipCompile, bool bForceCompile)
 		{
 			if (EngineRulesAssembly == null)
 			{
@@ -90,11 +91,11 @@ namespace UnrealBuildTool
 					}
 				}
 
-				EngineRulesAssembly = CreateEngineRulesAssemblyInternal(Unreal.GetExtensionDirs(Unreal.EngineDirectory), ProjectFileGenerator.EngineProjectFileNameBase, EnginePlugins, Unreal.IsEngineInstalled() || bUsePrecompiled, bSkipCompile, null);
+				EngineRulesAssembly = CreateEngineRulesAssemblyInternal(Unreal.GetExtensionDirs(Unreal.EngineDirectory), ProjectFileGenerator.EngineProjectFileNameBase, EnginePlugins, Unreal.IsEngineInstalled() || bUsePrecompiled, bSkipCompile, bForceCompile, null);
 
 				if (MarketplacePlugins.Count > 0)
 				{
-					EngineRulesAssembly = CreateMarketplaceRulesAssembly(MarketplacePlugins, bSkipCompile, EngineRulesAssembly);
+					EngineRulesAssembly = CreateMarketplaceRulesAssembly(MarketplacePlugins, bSkipCompile, bForceCompile, EngineRulesAssembly);
 				}
 			}
 			return EngineRulesAssembly;
@@ -108,9 +109,10 @@ namespace UnrealBuildTool
 		/// <param name="Plugins">List of plugins to include in this assembly</param>
 		/// <param name="bReadOnly">Whether the assembly should be marked as installed</param>
 		/// <param name="bSkipCompile">Whether to skip compilation for this assembly</param>
+		/// <param name="bForceCompile">Whether to always compile this assembly</param>
 		/// <param name="Parent">The parent rules assembly</param>
 		/// <returns>New rules assembly</returns>
-		private static RulesAssembly CreateEngineRulesAssemblyInternal(List<DirectoryReference> RootDirectories, string AssemblyPrefix, IReadOnlyList<PluginInfo> Plugins, bool bReadOnly, bool bSkipCompile, RulesAssembly? Parent)
+		private static RulesAssembly CreateEngineRulesAssemblyInternal(List<DirectoryReference> RootDirectories, string AssemblyPrefix, IReadOnlyList<PluginInfo> Plugins, bool bReadOnly, bool bSkipCompile, bool bForceCompile, RulesAssembly? Parent)
 		{
 			// Scope hierarchy
 			RulesScope Scope= new RulesScope("Engine", null);
@@ -144,7 +146,7 @@ namespace UnrealBuildTool
 			// Create the assembly
 			DirectoryReference AssemblyDir = RootDirectories[0];
 			FileReference EngineAssemblyFileName = FileReference.Combine(AssemblyDir, "Intermediate", "Build", "BuildRules", AssemblyPrefix + "Rules" + FrameworkAssemblyExtension);
-			RulesAssembly EngineAssembly = new RulesAssembly(Scope, RootDirectories, Plugins, ModuleFileToContext, new List<FileReference>(), EngineAssemblyFileName, bContainsEngineModules: true, DefaultBuildSettings: BuildSettingsVersion.Latest, bReadOnly: bReadOnly, bSkipCompile: bSkipCompile, Parent: Parent);
+			RulesAssembly EngineAssembly = new RulesAssembly(Scope, RootDirectories, Plugins, ModuleFileToContext, new List<FileReference>(), EngineAssemblyFileName, bContainsEngineModules: true, DefaultBuildSettings: BuildSettingsVersion.Latest, bReadOnly: bReadOnly, bSkipCompile: bSkipCompile, bForceCompile: bForceCompile, Parent: Parent);
 
 			List<FileReference> ProgramTargetFiles = new List<FileReference>();
 			Dictionary<FileReference, ModuleRulesContext> ProgramModuleFiles = new Dictionary<FileReference, ModuleRulesContext>();
@@ -171,7 +173,7 @@ namespace UnrealBuildTool
 
 			// Create a path to the assembly that we'll either load or compile
 			FileReference ProgramAssemblyFileName = FileReference.Combine(AssemblyDir, "Intermediate", "Build", "BuildRules", AssemblyPrefix + "ProgramRules" + FrameworkAssemblyExtension);
-			RulesAssembly ProgramAssembly = new RulesAssembly(ProgramsScope, RootDirectories, new List<PluginInfo>().AsReadOnly(), ProgramModuleFiles, ProgramTargetFiles, ProgramAssemblyFileName, bContainsEngineModules: false, DefaultBuildSettings: BuildSettingsVersion.Latest, bReadOnly: bReadOnly, bSkipCompile: bSkipCompile, Parent: EngineAssembly);
+			RulesAssembly ProgramAssembly = new RulesAssembly(ProgramsScope, RootDirectories, new List<PluginInfo>().AsReadOnly(), ProgramModuleFiles, ProgramTargetFiles, ProgramAssemblyFileName, bContainsEngineModules: false, DefaultBuildSettings: BuildSettingsVersion.Latest, bReadOnly: bReadOnly, bSkipCompile: bSkipCompile, bForceCompile: bForceCompile, Parent: EngineAssembly);
 
 			// Return the combined assembly
 			return ProgramAssembly;
@@ -182,9 +184,10 @@ namespace UnrealBuildTool
 		/// </summary>
 		/// <param name="Plugins">List of plugins to include in this assembly</param>
 		/// <param name="bSkipCompile">Whether to skip compilation for this assembly</param>
+		/// <param name="bForceCompile">Whether to always compile this assembly</param>
 		/// <param name="Parent">The parent rules assembly</param>
 		/// <returns>New rules assembly</returns>
-		private static RulesAssembly CreateMarketplaceRulesAssembly(IReadOnlyList<PluginInfo> Plugins, bool bSkipCompile, RulesAssembly Parent)
+		private static RulesAssembly CreateMarketplaceRulesAssembly(IReadOnlyList<PluginInfo> Plugins, bool bSkipCompile, bool bForceCompile, RulesAssembly Parent)
 		{
 			RulesScope MarketplaceScope = new RulesScope("Marketplace", Parent.Scope);
 
@@ -201,7 +204,7 @@ namespace UnrealBuildTool
 			if (ModuleFileToContext.Count > 0)
 			{
 				FileReference AssemblyFileName = FileReference.Combine(UnrealBuildTool.WritableEngineDirectory, "Intermediate", "Build", "BuildRules", "MarketplaceRules.dll");
-				Result = new RulesAssembly(MarketplaceScope, new List<DirectoryReference> { Unreal.EngineDirectory }, Plugins, ModuleFileToContext, new List<FileReference>(), AssemblyFileName, bContainsEngineModules: true, DefaultBuildSettings: BuildSettingsVersion.Latest, bReadOnly: false, bSkipCompile: bSkipCompile, Parent: Parent);
+				Result = new RulesAssembly(MarketplaceScope, new List<DirectoryReference> { Unreal.EngineDirectory }, Plugins, ModuleFileToContext, new List<FileReference>(), AssemblyFileName, bContainsEngineModules: true, DefaultBuildSettings: BuildSettingsVersion.Latest, bReadOnly: false, bSkipCompile: bSkipCompile, bForceCompile: bForceCompile, Parent: Parent);
 			}
 			return Result;
 		}
@@ -212,8 +215,9 @@ namespace UnrealBuildTool
 		/// <param name="ProjectFileName">The project file to create rules for. Null for the engine.</param>
 		/// <param name="bUsePrecompiled">Whether to use a precompiled engine</param>
 		/// <param name="bSkipCompile">Whether to skip compilation for this assembly</param>
+		/// <param name="bForceCompile">Whether to always compile this assembly</param>
 		/// <returns>New rules assembly</returns>
-		public static RulesAssembly CreateProjectRulesAssembly(FileReference ProjectFileName, bool bUsePrecompiled, bool bSkipCompile)
+		public static RulesAssembly CreateProjectRulesAssembly(FileReference ProjectFileName, bool bUsePrecompiled, bool bSkipCompile, bool bForceCompile)
 		{
 			// Check if there's an existing assembly for this project
 			RulesAssembly? ProjectRulesAssembly;
@@ -222,7 +226,7 @@ namespace UnrealBuildTool
 				ProjectDescriptor Project = ProjectDescriptor.FromFile(ProjectFileName);
 
 				// Create the parent assembly
-				RulesAssembly Parent = CreateEngineRulesAssembly(bUsePrecompiled, bSkipCompile);
+				RulesAssembly Parent = CreateEngineRulesAssembly(bUsePrecompiled, bSkipCompile, bForceCompile);
 
 				DirectoryReference MainProjectDirectory = ProjectFileName.Directory;
 				//DirectoryReference MainProjectSourceDirectory = DirectoryReference.Combine(MainProjectDirectory, "Source");
@@ -289,7 +293,7 @@ namespace UnrealBuildTool
 				}
 				else
 				{
-					ProjectRulesAssembly = new RulesAssembly(Scope, new List<DirectoryReference> { MainProjectDirectory }, ProjectPlugins, ModuleFiles, TargetFiles, AssemblyFileName, bContainsEngineModules: false, DefaultBuildSettings: null, bReadOnly: UnrealBuildTool.IsProjectInstalled(), bSkipCompile: bSkipCompile, Parent: Parent);
+					ProjectRulesAssembly = new RulesAssembly(Scope, new List<DirectoryReference> { MainProjectDirectory }, ProjectPlugins, ModuleFiles, TargetFiles, AssemblyFileName, bContainsEngineModules: false, DefaultBuildSettings: null, bReadOnly: UnrealBuildTool.IsProjectInstalled(), bSkipCompile: bSkipCompile, bForceCompile: bForceCompile, Parent: Parent);
 				}
 				LoadedAssemblyMap.Add(ProjectFileName, ProjectRulesAssembly);
 			}
@@ -301,10 +305,11 @@ namespace UnrealBuildTool
 		/// </summary>
 		/// <param name="PluginFileName">The plugin file to create rules for</param>
 		/// <param name="bSkipCompile">Whether to skip compilation for this assembly</param>
+		/// <param name="bForceCompile">Whether to always compile this assembly</param>
 		/// <param name="Parent">The parent rules assembly</param>
         /// <param name="bContainsEngineModules">Whether the plugin contains engine modules. Used to initialize the default value for ModuleRules.bTreatAsEngineModule.</param>
 		/// <returns>The new rules assembly</returns>
-        public static RulesAssembly CreatePluginRulesAssembly(FileReference PluginFileName, bool bSkipCompile, RulesAssembly Parent, bool bContainsEngineModules)
+        public static RulesAssembly CreatePluginRulesAssembly(FileReference PluginFileName, bool bSkipCompile, bool bForceCompile, RulesAssembly Parent, bool bContainsEngineModules)
 		{
 			// Check if there's an existing assembly for this project
 			RulesAssembly? PluginRulesAssembly;
@@ -331,7 +336,7 @@ namespace UnrealBuildTool
 
 				// Compile the assembly
 				FileReference AssemblyFileName = FileReference.Combine(PluginFileName.Directory, "Intermediate", "Build", "BuildRules", Path.GetFileNameWithoutExtension(PluginFileName.FullName) + "ModuleRules" + FrameworkAssemblyExtension);
-				PluginRulesAssembly = new RulesAssembly(Scope, new List<DirectoryReference> { PluginFileName.Directory }, ForeignPlugins, ModuleFiles, TargetFiles, AssemblyFileName, bContainsEngineModules, DefaultBuildSettings: null, bReadOnly: false, bSkipCompile: bSkipCompile, Parent: Parent);
+				PluginRulesAssembly = new RulesAssembly(Scope, new List<DirectoryReference> { PluginFileName.Directory }, ForeignPlugins, ModuleFiles, TargetFiles, AssemblyFileName, bContainsEngineModules, DefaultBuildSettings: null, bReadOnly: false, bSkipCompile: bSkipCompile, bForceCompile: bForceCompile, Parent: Parent);
 				LoadedAssemblyMap.Add(PluginFileName, PluginRulesAssembly);
 			}
 			return PluginRulesAssembly;
@@ -343,23 +348,24 @@ namespace UnrealBuildTool
 		/// <param name="ProjectFile">The project file being compiled</param>
 		/// <param name="TargetName">The target being built</param>
 		/// <param name="bSkipRulesCompile">Whether to skip compiling any rules assemblies</param>
+		/// <param name="bForceRulesCompile">Whether to always compile all rules assemblies</param>
 		/// <param name="bUsePrecompiled">Whether to use a precompiled engine build</param>
 		/// <param name="ForeignPlugin">Foreign plugin to be compiled</param>
 		/// <returns>The compiled rules assembly</returns>
-		public static RulesAssembly CreateTargetRulesAssembly(FileReference? ProjectFile, string TargetName, bool bSkipRulesCompile, bool bUsePrecompiled, FileReference? ForeignPlugin)
+		public static RulesAssembly CreateTargetRulesAssembly(FileReference? ProjectFile, string TargetName, bool bSkipRulesCompile, bool bForceRulesCompile, bool bUsePrecompiled, FileReference? ForeignPlugin)
 		{
 			RulesAssembly RulesAssembly;
 			if (ProjectFile != null)
 			{
-				RulesAssembly = CreateProjectRulesAssembly(ProjectFile, bUsePrecompiled, bSkipRulesCompile);
+				RulesAssembly = CreateProjectRulesAssembly(ProjectFile, bUsePrecompiled, bSkipRulesCompile, bForceRulesCompile);
 			}
 			else
 			{
-				RulesAssembly = CreateEngineRulesAssembly(bUsePrecompiled, bSkipRulesCompile);
+				RulesAssembly = CreateEngineRulesAssembly(bUsePrecompiled, bSkipRulesCompile, bForceRulesCompile);
 			}
 			if (ForeignPlugin != null)
 			{
-				RulesAssembly = CreatePluginRulesAssembly(ForeignPlugin, bSkipRulesCompile, RulesAssembly, true);
+				RulesAssembly = CreatePluginRulesAssembly(ForeignPlugin, bSkipRulesCompile, bForceRulesCompile, RulesAssembly, true);
 			}
 			return RulesAssembly;
 		}
diff --git a/Engine/Source/Programs/UnrealBuildTool/ToolChain/RemoteMac.cs b/Engine/Source/Programs/UnrealBuildTool/ToolChain/RemoteMac.cs
--- a/Engine/Source/Programs/UnrealBuildTool/ToolChain/RemoteMac.cs
+++ b/Engine/Source/Programs/UnrealBuildTool/ToolChain/RemoteMac.cs
@@ -421,7 +421,7 @@ namespace UnrealBuildTool
 		public bool Build(TargetDescriptor TargetDesc, FileReference RemoteLogFile, bool bSkipPreBuildTargets)
 		{
 			// Compile the rules assembly
-			RulesAssembly RulesAssembly = RulesCompiler.CreateTargetRulesAssembly(TargetDesc.ProjectFile, TargetDesc.Name, false, false, TargetDesc.ForeignPlugin);
+			RulesAssembly RulesAssembly = RulesCompiler.CreateTargetRulesAssembly(TargetDesc.ProjectFile, TargetDesc.Name, false, false, false, TargetDesc.ForeignPlugin);
 
 			// Create the target rules
 			TargetRules Rules = RulesAssembly.CreateTargetRules(TargetDesc.Name, TargetDesc.Platform, TargetDesc.Configuration, TargetDesc.Architecture, TargetDesc.ProjectFile, TargetDesc.AdditionalArguments);
