UBT: Add a mode for analyzing the dependencies between modules.

Example usage "-Mode=Analyze UnrealEditor Win64 Development".

The tool currently outputs the following files:
* A graph of dependencies between modules, with edges indicating unique outward references (eg. if module A depends on Core and BuildSettings, but Core depends on BuildSettings, an edge will only be shown from A -> Core).
* A graph showing shortest path from the target to each module.
* A text file listing all modules in the target, their references to other modules, the shortest path causing them to be instantiated, and a (very) rough estimate of their size taken by measuring the sum of all object files or output files that already exist on disk.
* A csv file containing the same information as above.

Graphs are output in GEFX format, and can be viewed using Gephi.

diff --git a/Engine/Source/Programs/UnrealBuildTool/Configuration/UEBuildPlugin.cs b/Engine/Source/Programs/UnrealBuildTool/Configuration/UEBuildPlugin.cs
--- a/Engine/Source/Programs/UnrealBuildTool/Configuration/UEBuildPlugin.cs
+++ b/Engine/Source/Programs/UnrealBuildTool/Configuration/UEBuildPlugin.cs
@@ -39,13 +39,20 @@ namespace UnrealBuildTool
 		/// </summary>
 		public bool bDescriptorReferencedExplicitly;
 
+		/// <summary>
+		/// Chain of references to this plugin
+		/// </summary>
+		public string ReferenceChain;
+
 		/// <summary>
 		/// Constructor
 		/// </summary>
 		/// <param name="Info">The static plugin information</param>
-		public UEBuildPlugin(PluginInfo Info)
+		/// <param name="ReferenceChain">Chain of references to this plugin</param>
+		public UEBuildPlugin(PluginInfo Info, string ReferenceChain)
 		{
 			this.Info = Info;
+			this.ReferenceChain = ReferenceChain;
 		}
 
 		/// <summary>
diff --git a/Engine/Source/Programs/UnrealBuildTool/Configuration/UEBuildTarget.cs b/Engine/Source/Programs/UnrealBuildTool/Configuration/UEBuildTarget.cs
--- a/Engine/Source/Programs/UnrealBuildTool/Configuration/UEBuildTarget.cs
+++ b/Engine/Source/Programs/UnrealBuildTool/Configuration/UEBuildTarget.cs
@@ -3278,7 +3278,7 @@ namespace UnrealBuildTool
 
 				// Create the new instance and add it to the cache
 				Log.TraceLog("Enabling plugin '{0}' (referenced via {1})", Reference.Name, ReferenceChain);
-				Instance = new UEBuildPlugin(Info);
+				Instance = new UEBuildPlugin(Info, ReferenceChain);
 				Instance.bDescriptorReferencedExplicitly = !Reference.bOptional;
 				NameToInstance.Add(Info.Name, Instance);
 
diff --git a/Engine/Source/Programs/UnrealBuildTool/Modes/AnalyzeMode.cs b/Engine/Source/Programs/UnrealBuildTool/Modes/AnalyzeMode.cs
new file mode 100644
--- /dev/null
+++ b/Engine/Source/Programs/UnrealBuildTool/Modes/AnalyzeMode.cs
@@ -0,0 +1,395 @@
+// Copyright Epic Games, Inc. All Rights Reserved.
+
+using EpicGames.Core;
+using System;
+using System.Collections.Generic;
+using System.IO;
+using System.Linq;
+using System.Text;
+using System.Xml;
+
+namespace UnrealBuildTool.Modes
+{
+	/// <summary>
+	/// Outputs information about the given target, including a module dependecy graph (in .gefx format and list of module references)
+	/// </summary>
+	[ToolMode("Analyze", ToolModeOptions.XmlConfig | ToolModeOptions.BuildPlatforms | ToolModeOptions.SingleInstance | ToolModeOptions.StartPrefetchingEngine | ToolModeOptions.ShowExecutionTime)]
+	class AnalyzeMode : ToolMode
+	{
+		/// <summary>
+		/// Execute the command
+		/// </summary>
+		/// <param name="Arguments">Command line arguments</param>
+		/// <returns>Exit code</returns>
+		public override int Execute(CommandLineArguments Arguments)
+		{
+			Arguments.ApplyTo(this);
+
+			// Create the build configuration object, and read the settings
+			BuildConfiguration BuildConfiguration = new BuildConfiguration();
+			XmlConfig.ApplyTo(BuildConfiguration);
+			Arguments.ApplyTo(BuildConfiguration);
+
+			// Parse all the target descriptors
+			List<TargetDescriptor> TargetDescriptors = TargetDescriptor.ParseCommandLine(Arguments, BuildConfiguration.bUsePrecompiled, BuildConfiguration.bSkipRulesCompile, BuildConfiguration.bForceRulesCompile);
+
+			// Generate the compile DB for each target
+			using (ISourceFileWorkingSet WorkingSet = new EmptySourceFileWorkingSet())
+			{
+				// Find the compile commands for each file in the target
+				Dictionary<FileReference, string> FileToCommand = new Dictionary<FileReference, string>();
+				foreach (TargetDescriptor TargetDescriptor in TargetDescriptors)
+				{
+					AnalyzeTarget(TargetDescriptor, BuildConfiguration);
+				}
+			}
+
+			return 0;
+		}
+
+		class ModuleInfo
+		{
+			public UEBuildModule Module;
+			public string Chain;
+			public HashSet<UEBuildModule> InwardRefs = new HashSet<UEBuildModule>();
+			public HashSet<UEBuildModule> UniqueInwardRefs = new HashSet<UEBuildModule>();
+			public HashSet<UEBuildModule> OutwardRefs = new HashSet<UEBuildModule>();
+			public HashSet<UEBuildModule> UniqueOutwardRefs = new HashSet<UEBuildModule>();
+
+			public List<FileReference> ObjectFiles = new List<FileReference>();
+			public long ObjSize = 0;
+			public List<FileReference> BinaryFiles = new List<FileReference>();
+			public long BinSize = 0;
+
+			public ModuleInfo(UEBuildModule Module, string Chain)
+			{
+				this.Module = Module;
+				this.Chain = Chain;
+			}
+		}
+
+		private void AnalyzeTarget(TargetDescriptor TargetDescriptor, BuildConfiguration BuildConfiguration)
+		{
+			// Create a makefile for the target
+			UEBuildTarget Target = UEBuildTarget.Create(TargetDescriptor, BuildConfiguration.bSkipRulesCompile, BuildConfiguration.bForceRulesCompile, BuildConfiguration.bUsePrecompiled);
+			DirectoryReference.CreateDirectory(Target.ReceiptFileName.Directory);
+
+			// Find the shortest path from the target to each module
+			Dictionary<UEBuildModule, ModuleInfo> ModuleToInfo = new Dictionary<UEBuildModule, ModuleInfo>();
+
+			List<string> RootModuleNames = new List<string>(Target.Rules.ExtraModuleNames);
+			if (Target.Rules.LaunchModuleName != null)
+			{
+				RootModuleNames.Add(Target.Rules.LaunchModuleName);
+			}
+
+			foreach (string RootModuleName in RootModuleNames)
+			{
+				UEBuildModule Module = Target.GetModuleByName(RootModuleName);
+				if (Module != null)
+				{
+					string Chain = $"target -> {RootModuleName}";
+					ModuleToInfo[Module] = new ModuleInfo(Module, Chain);
+				}
+			}
+
+			// Also enable all the plugin modules
+			foreach (UEBuildPlugin Plugin in Target.BuildPlugins!)
+			{
+				foreach (UEBuildModule Module in Plugin.Modules)
+				{
+					string Chain = $"{Plugin.ReferenceChain} -> {Module.Name}";
+					ModuleToInfo[Module] = new ModuleInfo(Module, Chain);
+				}
+			}
+
+			// Set of visited modules
+			HashSet<UEBuildModule> VisitedModules = new HashSet<UEBuildModule>();
+
+			// Recurse out to find new modules and the shortest path to each
+			List<UEBuildModule> SourceModules = new List<UEBuildModule>(ModuleToInfo.Keys);
+			while (SourceModules.Count > 0)
+			{
+				List<UEBuildModule> TargetModules = new List<UEBuildModule>();
+
+				foreach (UEBuildModule SourceModule in SourceModules)
+				{
+					int Idx = TargetModules.Count;
+					SourceModule.GetAllDependencyModules(TargetModules, VisitedModules, true, false, true);
+
+					for (; Idx < TargetModules.Count; Idx++)
+					{
+						UEBuildModule TargetModule = TargetModules[Idx];
+						string Chain = $"{ModuleToInfo[SourceModule].Chain} -> {TargetModule.Name}";
+						ModuleToInfo[TargetModule] = new ModuleInfo(TargetModule, Chain);
+					}
+				}
+
+				SourceModules = TargetModules;
+			}
+
+			// Find all the outward dependencies of each module
+			foreach ((UEBuildModule SourceModule, ModuleInfo SourceModuleInfo) in ModuleToInfo)
+			{
+				SourceModuleInfo.OutwardRefs.Add(SourceModule);
+				SourceModule.GetAllDependencyModules(new List<UEBuildModule>(), SourceModuleInfo.OutwardRefs, false, false, false);
+				SourceModuleInfo.OutwardRefs.Remove(SourceModule);
+			}
+
+			// Find the direct output dependencies of each module
+			foreach ((UEBuildModule SourceModule, ModuleInfo SourceModuleInfo) in ModuleToInfo)
+			{
+				SourceModuleInfo.UniqueOutwardRefs = new HashSet<UEBuildModule>(SourceModuleInfo.OutwardRefs);
+
+				foreach (UEBuildModule TargetModule in SourceModuleInfo.OutwardRefs)
+				{
+					HashSet<UEBuildModule> VisitedTargetModules = new HashSet<UEBuildModule>();
+					VisitedTargetModules.Add(SourceModule);
+
+					List<UEBuildModule> DependencyModules = new List<UEBuildModule>();
+					TargetModule.GetAllDependencyModules(DependencyModules, VisitedTargetModules, false, false, false);
+					DependencyModules.Remove(TargetModule);
+
+					SourceModuleInfo.UniqueOutwardRefs.ExceptWith(DependencyModules);
+				}
+			}
+
+			// Find the direct inward dependencies of each module
+			foreach ((UEBuildModule SourceModule, ModuleInfo SourceModuleInfo) in ModuleToInfo)
+			{
+				foreach (UEBuildModule TargetModule in SourceModuleInfo.OutwardRefs)
+				{
+					ModuleToInfo[TargetModule].InwardRefs.Add(SourceModule);
+				}
+				foreach (UEBuildModule TargetModule in SourceModuleInfo.UniqueOutwardRefs)
+				{
+					ModuleToInfo[TargetModule].UniqueInwardRefs.Add(SourceModule);
+				}
+			}
+
+			// Estimate the size of object files for each module
+			foreach ((UEBuildModule SourceModule, ModuleInfo SourceModuleInfo) in ModuleToInfo)
+			{
+				if (DirectoryReference.Exists(SourceModule.IntermediateDirectory))
+				{
+					foreach (FileReference IntermediateFile in DirectoryReference.EnumerateFiles(SourceModule.IntermediateDirectory, "*", SearchOption.AllDirectories))
+					{
+						if (IntermediateFile.HasExtension(".obj") || IntermediateFile.HasExtension(".o"))
+						{
+							SourceModuleInfo.ObjectFiles.Add(IntermediateFile);
+							SourceModuleInfo.ObjSize += IntermediateFile.ToFileInfo().Length;
+						}
+					}
+				}
+			}
+
+			HashSet<UEBuildModule> MissingModules = new HashSet<UEBuildModule>();
+			foreach (UEBuildBinary Binary in Target.Binaries)
+			{
+				long BinSize = 0;
+				foreach (FileReference OutputFilePath in Binary.OutputFilePaths)
+				{
+					FileInfo OutputFileInfo = OutputFilePath.ToFileInfo();
+					if (OutputFileInfo.Exists)
+					{
+						BinSize += OutputFileInfo.Length;
+					}
+				}
+				foreach (UEBuildModule Module in Binary.Modules)
+				{
+					ModuleInfo? ModuleInfo;
+					if (!ModuleToInfo.TryGetValue(Module, out ModuleInfo))
+					{
+						MissingModules.Add(Module);
+						continue;
+					}
+
+					ModuleInfo.BinaryFiles.AddRange(Binary.OutputFilePaths);
+					ModuleInfo.BinSize += BinSize;
+				}
+			}
+
+			// Warn about any missing modules
+			foreach (UEBuildModule MissingModule in MissingModules.OrderBy(x => x.Name))
+			{
+				Log.TraceWarning("Missing module '{0}'", MissingModule.Name);
+			}
+
+			// Generate the dependency graph between modules
+			FileReference DependencyGraphFile = Target.ReceiptFileName.ChangeExtension(".Dependencies.gexf");
+			Log.TraceInformation("Writing dependency graph to {0}...", DependencyGraphFile);
+			WriteDependencyGraph(Target, ModuleToInfo, DependencyGraphFile);
+
+			// Generate the dependency graph between modules
+			FileReference ShortestPathGraphFile = Target.ReceiptFileName.ChangeExtension(".ShortestPath.gexf");
+			Log.TraceInformation("Writing shortest-path graph to {0}...", ShortestPathGraphFile);
+			WriteShortestPathGraph(Target, ModuleToInfo, ShortestPathGraphFile);
+
+			// Write all the target stats as a text file
+			FileReference TextFile = Target.ReceiptFileName.ChangeExtension(".txt");
+			Log.TraceInformation("Writing module information to {0}", TextFile);
+			using (StreamWriter Writer = new StreamWriter(TextFile.FullName))
+			{
+				Writer.WriteLine("All modules in {0}, ordered by number of indirect references", Target.TargetName);
+
+				foreach (ModuleInfo ModuleInfo in ModuleToInfo.Values.OrderByDescending(x => x.InwardRefs.Count).ThenBy(x => x.BinSize))
+				{
+					Writer.WriteLine("");
+					Writer.WriteLine("Module:        \"{0}\"", ModuleInfo.Module.Name);
+					Writer.WriteLine("Shortest path: {0}", ModuleInfo.Chain);
+					WriteDependencyList(Writer, "Unique inward refs:     ", ModuleInfo.UniqueInwardRefs);
+					WriteDependencyList(Writer, "Unique outward refs:    ", ModuleInfo.UniqueOutwardRefs);
+					WriteDependencyList(Writer, "Recursive inward refs:  ", ModuleInfo.InwardRefs);
+					WriteDependencyList(Writer, "Recursive outward refs: ", ModuleInfo.OutwardRefs);
+					Writer.WriteLine("Object size:             {0:n0}kb", (ModuleInfo.ObjSize + 1023) / 1024);
+					Writer.WriteLine("Object files:            {0}", String.Join(", ", ModuleInfo.ObjectFiles.Select(x => x.GetFileName())));
+					Writer.WriteLine("Binary size:             {0:n0}kb", (ModuleInfo.BinSize + 1023) / 1024);
+					Writer.WriteLine("Binary files:            {0}", String.Join(", ", ModuleInfo.BinaryFiles.Select(x => x.GetFileName())));
+				}
+			}
+
+			// Write all the target stats as a CSV file
+			FileReference CsvFile = Target.ReceiptFileName.ChangeExtension(".csv");
+			Log.TraceInformation("Writing module information to {0}", CsvFile);
+			using (StreamWriter Writer = new StreamWriter(CsvFile.FullName))
+			{
+				List<string> Columns = new List<string>();
+				Columns.Add("Module");
+				Columns.Add("ShortestPath");
+				Columns.Add("NumUniqueInwardRefs");
+				Columns.Add("UniqueInwardRefs");
+				Columns.Add("NumRecursiveInwardRefs");
+				Columns.Add("RecursiveInwardRefs");
+				Columns.Add("NumUniqueOutwardRefs");
+				Columns.Add("UniqueOutwardRefs");
+				Columns.Add("NumRecursiveOutwardRefs");
+				Columns.Add("RecursiveOutwardRefs");
+				Columns.Add("ObjSize");
+				Columns.Add("ObjFiles");
+				Columns.Add("BinSize");
+				Columns.Add("BinFiles");
+				Writer.WriteLine(String.Join(",", Columns));
+
+				foreach (ModuleInfo ModuleInfo in ModuleToInfo.Values.OrderByDescending(x => x.InwardRefs.Count).ThenBy(x => x.BinSize))
+				{
+					Columns.Clear();
+					Columns.Add(ModuleInfo.Module.Name);
+					Columns.Add(ModuleInfo.Chain);
+					Columns.Add($"{ModuleInfo.UniqueInwardRefs.Count}");
+					Columns.Add($"\"{String.Join(", ", ModuleInfo.UniqueInwardRefs.Select(x => x.Name))}\"");
+					Columns.Add($"{ModuleInfo.InwardRefs.Count}");
+					Columns.Add($"\"{String.Join(", ", ModuleInfo.InwardRefs.Select(x => x.Name))}\"");
+					Columns.Add($"{ModuleInfo.UniqueOutwardRefs.Count}");
+					Columns.Add($"\"{String.Join(", ", ModuleInfo.UniqueOutwardRefs.Select(x => x.Name))}\"");
+					Columns.Add($"{ModuleInfo.OutwardRefs.Count}");
+					Columns.Add($"\"{String.Join(", ", ModuleInfo.OutwardRefs.Select(x => x.Name))}\"");
+					Columns.Add($"{ModuleInfo.ObjSize}");
+					Columns.Add($"\"{String.Join(", ", ModuleInfo.ObjectFiles.Select(x => x.GetFileName()))}\"");
+					Columns.Add($"{ModuleInfo.BinSize}");
+					Columns.Add($"\"{String.Join(", ", ModuleInfo.BinaryFiles.Select(x => x.GetFileName()))}\"");
+					Writer.WriteLine(String.Join(",", Columns));
+				}
+			}
+		}
+
+		private void WriteDependencyList(TextWriter Writer, string Prefix, HashSet<UEBuildModule> Modules)
+		{
+			if (Modules.Count == 0)
+			{
+				Writer.WriteLine("{0} 0", Prefix);
+			}
+			else
+			{
+				Writer.WriteLine("{0} {1} ({2})", Prefix, Modules.Count, String.Join(", ", Modules.Select(x => x.Name).OrderBy(x => x)));
+			}
+		}
+
+		private void WriteDependencyGraph(UEBuildTarget Target, Dictionary<UEBuildModule, ModuleInfo> ModuleToInfo, FileReference FileName)
+		{
+			List<GraphNode> Nodes = new List<GraphNode>();
+
+			Dictionary<UEBuildModule, GraphNode> ModuleToNode = new Dictionary<UEBuildModule, GraphNode>();
+			foreach (ModuleInfo ModuleInfo in ModuleToInfo.Values)
+			{
+				GraphNode Node = new GraphNode(ModuleInfo.Module.Name);
+
+				long Size;
+				if (Target.ShouldCompileMonolithic())
+				{
+					Size = ModuleInfo.ObjSize;
+				}
+				else
+				{
+					Size = ModuleInfo.BinSize;
+				}
+
+				Node.Size = 1.0f + (Size / (50.0f * 1024.0f * 1024.0f));
+				Nodes.Add(Node);
+				ModuleToNode[ModuleInfo.Module] = Node;
+			}
+
+			List<GraphEdge> Edges = new List<GraphEdge>();
+			foreach ((UEBuildModule SourceModule, ModuleInfo SourceModuleInfo) in ModuleToInfo)
+			{
+				GraphNode SourceNode = ModuleToNode[SourceModule];
+				foreach (UEBuildModule TargetModule in SourceModuleInfo.UniqueOutwardRefs)
+				{
+					ModuleInfo TargetModuleInfo = ModuleToInfo[TargetModule];
+
+					GraphNode? TargetNode;
+					if (ModuleToNode.TryGetValue(TargetModule, out TargetNode))
+					{
+						GraphEdge Edge = new GraphEdge(SourceNode, TargetNode);
+						Edge.Thickness = TargetModuleInfo.InwardRefs.Count;
+						Edges.Add(Edge);
+					}
+				}
+			}
+
+			GraphVisualization.WriteGraphFile(FileName, $"Module dependency graph for {Target.TargetName}", Nodes, Edges);
+		}
+
+		private void WriteShortestPathGraph(UEBuildTarget Target, Dictionary<UEBuildModule, ModuleInfo> ModuleToInfo, FileReference FileName)
+		{
+			Dictionary<string, GraphNode> NameToNode = new Dictionary<string, GraphNode>(StringComparer.Ordinal);
+
+			HashSet<(GraphNode, GraphNode)> EdgesSet = new HashSet<(GraphNode, GraphNode)>();
+			List<GraphEdge> Edges = new List<GraphEdge>();
+
+			foreach ((UEBuildModule Module, ModuleInfo ModuleInfo) in ModuleToInfo)
+			{
+				string[] Parts = ModuleInfo.Chain.Split(" -> ");
+
+				GraphNode? PrevNode = null;
+				foreach (string Part in Parts)
+				{
+					GraphNode? NextNode;
+					if (!NameToNode.TryGetValue(Part, out NextNode))
+					{
+						NextNode = new GraphNode(Part);
+						NameToNode[Part] = NextNode;
+					}
+					if (PrevNode != null && EdgesSet.Add((PrevNode, NextNode)))
+					{
+						GraphEdge Edge = new GraphEdge(PrevNode, NextNode);
+						Edges.Add(Edge);
+					}
+					PrevNode = NextNode;
+				}
+			}
+
+			GraphVisualization.WriteGraphFile(FileName, $"Module dependency graph for {Target.TargetName}", NameToNode.Values.ToList(), Edges);
+		}
+
+		private static HashSet<UEBuildModule> GetDirectDependencyModules(UEBuildModule Module)
+		{
+			HashSet<UEBuildModule> ReferencedModules = new HashSet<UEBuildModule>();
+			Module.GetAllDependencyModules(new List<UEBuildModule>(), ReferencedModules, true, false, false);
+
+			HashSet<UEBuildModule> Modules = new HashSet<UEBuildModule>(Module.GetDirectDependencyModules());
+			Modules.ExceptWith(ReferencedModules);
+			return Modules;
+		}
+	}
+}
+
diff --git a/Engine/Source/Programs/UnrealBuildTool/System/GraphVisualization.cs b/Engine/Source/Programs/UnrealBuildTool/System/GraphVisualization.cs
new file mode 100644
--- /dev/null
+++ b/Engine/Source/Programs/UnrealBuildTool/System/GraphVisualization.cs
@@ -0,0 +1,312 @@
+// Copyright Epic Games, Inc. All Rights Reserved.
+
+using EpicGames.Core;
+using System;
+using System.Collections.Generic;
+using System.Drawing;
+using System.Linq;
+using System.Text;
+using System.Xml;
+
+namespace UnrealBuildTool
+{
+	/// <summary>
+	/// A single "node" in a directed graph
+	/// </summary>
+	class GraphNode
+	{
+		public string Label;
+		public Color Color;
+		public float Size;
+		public Dictionary<string, object> Attributes = new Dictionary<string, object>(StringComparer.InvariantCultureIgnoreCase);
+
+		public GraphNode(string Label, float Size = 1.0f)
+			: this(Label, Color.Black, Size)
+		{
+		}
+
+		public GraphNode(string Label, Color Color, float Size = 1.0f)
+		{
+			this.Label = Label;
+			this.Color = Color;
+			this.Size = Size;
+		}
+	}
+
+	/// <summary>
+	/// Describes an edge in the directed graph
+	/// </summary>
+	class GraphEdge
+	{
+		public GraphNode Source;
+		public GraphNode Target;
+		public double Weight = 1.0f;
+		public Color Color = Color.FromArgb(64, 0, 0, 0);
+		public float Thickness = 0.1f;
+
+		public GraphEdge(GraphNode Source, GraphNode Target)
+		{
+			this.Source = Source;
+			this.Target = Target;
+		}
+	}
+
+	static class GraphVisualization
+	{
+		/// <summary>
+		/// Attribute for a graph node
+		/// </summary>
+		class GraphAttribute
+		{
+			/// Gexf ID for this attribute
+			public int Id;
+
+			/// Name of the attribute
+			public string Name;
+
+			/// Gexf type name
+			public string TypeName;
+
+			public GraphAttribute(int Id, string Name, string TypeName)
+			{
+				this.Id = Id;
+				this.Name = Name;
+				this.TypeName = TypeName;
+			}
+		}
+
+		/// <summary>
+		/// Writes a GEXF graph file for the specified graph nodes and edges
+		/// </summary>
+		/// <param name="Filename">The file name to write</param>
+		/// <param name="Description">The description to include in the graph file's metadata</param>
+		/// <param name="GraphNodes">List of all graph nodes.  Index order is important and must match with the individual node Id members!</param>
+		/// <param name="GraphEdges">List of all graph edges.  Index order is important and must match with the individual edge Id members!</param>
+		public static void WriteGraphFile(FileReference Filename, string Description, List<GraphNode> GraphNodes, List<GraphEdge> GraphEdges)
+		{
+			XmlWriterSettings Settings = new XmlWriterSettings();
+			Settings.Indent = true;
+			Settings.IndentChars = "    ";
+
+			// Figure out all of the custom attribute types we're dealing with
+			Dictionary<string, GraphAttribute> AllAttributes = new Dictionary<string, GraphAttribute>(StringComparer.InvariantCultureIgnoreCase);
+			foreach (GraphNode GraphNode in GraphNodes)
+			{
+				foreach ((string AttributeName, object AttributeValue) in GraphNode.Attributes)
+				{
+					string AttributeType = GetAttributeType(AttributeValue);
+
+					GraphAttribute? Attribute;
+					if (!AllAttributes.TryGetValue(AttributeName, out Attribute))
+					{
+						AllAttributes[AttributeName] = new GraphAttribute(AllAttributes.Count, AttributeName, AttributeType);
+					}
+					else if (!Attribute.TypeName.Equals(AttributeType))
+					{
+						throw new InvalidOperationException("Multiple graph nodes with the same attribute name but different types encountered!");
+					}
+				}
+			}
+
+			using (XmlWriter Writer = XmlWriter.Create(Filename.FullName, Settings))
+			{
+				// NOTE: The GEXF XML format is defined here:  http://gexf.net/1.2draft/gexf-12draft-primer.pdf
+
+				string GEXFNamespace = "http://www.gexf.net/1.2-draft";
+				string SchemaNamespace = "http://www.w3.org/2001/XMLSchema-instance";
+				string VizNamespace = "http://www.gexf.net/1.2draft/viz";
+
+				Writer.WriteStartElement("gexf", GEXFNamespace);
+				Writer.WriteAttributeString("xmlns", "xsi", null, SchemaNamespace);
+				Writer.WriteAttributeString("schemaLocation", SchemaNamespace, "http://www.gexf.net/1.2draft http://www.gexf.net/1.2draft/gexf.xsd");
+				Writer.WriteAttributeString("xmlns", "viz", null, VizNamespace);
+				Writer.WriteAttributeString("version", "1.2");
+
+				Writer.WriteStartElement("meta");
+				{
+					Writer.WriteAttributeString("creator", "UnrealBuildTool");
+					Writer.WriteAttributeString("description", Description);
+				}
+				Writer.WriteEndElement();   // meta
+
+				Dictionary<GraphNode, int> NodeToId = new Dictionary<GraphNode, int>();
+				{
+					Writer.WriteStartElement("graph");
+					{
+						Writer.WriteAttributeString("mode", "static");
+						Writer.WriteAttributeString("defaultedgetype", "directed");
+
+						if (AllAttributes.Count > 0)
+						{
+							Writer.WriteStartElement("attributes");
+							{
+								// @todo: Add support for edge attributes, not just node attributes
+								Writer.WriteAttributeString("class", "node");   // Node attributes, not edges!
+
+								foreach (var Attribute in AllAttributes.Values)
+								{
+									Writer.WriteStartElement("attribute");
+									{
+										Writer.WriteAttributeString("id", Attribute.Id.ToString());
+										Writer.WriteAttributeString("title", Attribute.Name);
+										Writer.WriteAttributeString("type", Attribute.TypeName);
+									}
+									Writer.WriteEndElement();   // attribute
+								}
+
+								// @todo: Add support for attribute type default values
+							}
+							Writer.WriteEndElement();   // attributes
+
+						}
+
+						Writer.WriteStartElement("nodes");
+						{
+							foreach (GraphNode GraphNode in GraphNodes)
+							{
+								Writer.WriteStartElement("node");
+								{
+									int Id = GetNodeId(GraphNode, NodeToId);
+									Writer.WriteAttributeString("id", Id.ToString());
+									Writer.WriteAttributeString("label", GraphNode.Label);
+
+									Writer.WriteStartElement("color", VizNamespace);
+									{
+										Writer.WriteAttributeString("r", GraphNode.Color.R.ToString());
+										Writer.WriteAttributeString("g", GraphNode.Color.G.ToString());
+										Writer.WriteAttributeString("b", GraphNode.Color.B.ToString());
+										Writer.WriteAttributeString("a", (GraphNode.Color.A / 255.0f).ToString());
+									}
+									Writer.WriteEndElement();   // viz:color
+
+									Writer.WriteStartElement("size", VizNamespace);
+									{
+										Writer.WriteAttributeString("value", GraphNode.Size.ToString());
+									}
+									Writer.WriteEndElement();   // viz:size
+
+									Writer.WriteStartElement("shape", VizNamespace);
+									{
+										// NOTE: Valid shapes are:  disc, square, triangle, diamond, image
+										Writer.WriteAttributeString("value", "disc");
+									}
+									Writer.WriteEndElement();   // viz:shape
+
+									if (GraphNode.Attributes.Count > 0)
+									{
+										Writer.WriteStartElement("attvalues");
+										{
+											foreach (var AttributeHashEntry in GraphNode.Attributes)
+											{
+												var AttributeName = AttributeHashEntry.Key;
+												var AttributeValue = AttributeHashEntry.Value;
+
+												var Attribute = AllAttributes[AttributeName];
+
+												Writer.WriteStartElement("attvalue");
+												{
+													Writer.WriteAttributeString("for", Attribute.Id.ToString());
+													Writer.WriteAttributeString("value", AttributeValue.ToString());
+												}
+												Writer.WriteEndElement();
+											}
+										}
+										Writer.WriteEndElement();   // attvalues
+									}
+								}
+								Writer.WriteEndElement();   // node
+							}
+						}
+						Writer.WriteEndElement();   // nodes
+
+
+						Writer.WriteStartElement("edges");
+						{
+							for(int EdgeId = 0; EdgeId < GraphEdges.Count; EdgeId++)
+							{
+								GraphEdge GraphEdge = GraphEdges[EdgeId];
+								Writer.WriteStartElement("edge");
+								{
+									Writer.WriteAttributeString("id", EdgeId.ToString());
+									Writer.WriteAttributeString("source", GetNodeId(GraphEdge.Source, NodeToId).ToString());
+									Writer.WriteAttributeString("target", GetNodeId(GraphEdge.Target, NodeToId).ToString());
+									Writer.WriteAttributeString("weight", GraphEdge.Weight.ToString());
+
+									Writer.WriteStartElement("color", VizNamespace);
+									{
+										Writer.WriteAttributeString("r", GraphEdge.Color.R.ToString());
+										Writer.WriteAttributeString("g", GraphEdge.Color.G.ToString());
+										Writer.WriteAttributeString("b", GraphEdge.Color.B.ToString());
+										Writer.WriteAttributeString("a", (GraphEdge.Color.A / 255.0f).ToString());
+									}
+									Writer.WriteEndElement();   // viz:color
+
+									Writer.WriteStartElement("thickness", VizNamespace);
+									{
+										Writer.WriteAttributeString("value", GraphEdge.Thickness.ToString());
+									}
+									Writer.WriteEndElement();   // viz:thickness
+
+									Writer.WriteStartElement("shape", VizNamespace);
+									{
+										// NOTE: Valid shapes are:  solid, dotted, dashed, double
+										Writer.WriteAttributeString("value", "solid");
+									}
+									Writer.WriteEndElement();   // viz:shape
+								}
+								Writer.WriteEndElement();   // edge
+							}
+						}
+						Writer.WriteEndElement();   // nodes
+					}
+
+					Writer.WriteEndElement();   // graph
+				}
+				Writer.WriteEndElement();   // gexf
+
+				Writer.Flush();
+			}
+		}
+
+		private static int GetNodeId(GraphNode Node, Dictionary<GraphNode, int> NodeToId)
+		{
+			int Id;
+			if (!NodeToId.TryGetValue(Node, out Id))
+			{
+				Id = NodeToId.Count;
+				NodeToId.Add(Node, Id);
+			}
+			return Id;
+		}
+
+		private static string GetAttributeType(object Value)
+		{
+			string AttributeTypeName;
+			if (Value.GetType() == typeof(int))
+			{
+				AttributeTypeName = "integer";
+			}
+			else if (Value.GetType() == typeof(float))
+			{
+				AttributeTypeName = "float";
+			}
+			else if (Value.GetType() == typeof(double))
+			{
+				AttributeTypeName = "double";
+			}
+			else if (Value.GetType() == typeof(string))
+			{
+				AttributeTypeName = "string";
+			}
+			else if (Value.GetType() == typeof(bool))
+			{
+				AttributeTypeName = "boolean";
+			}
+			else
+			{
+				throw new InvalidOperationException("Unsupported attribute data type encountered on graph node!");
+			}
+			return AttributeTypeName;
+		}
+	}
+}
