UBT: Fix WarningsAsErrors settings being ignored.

diff --git a/Engine/Source/Programs/UnrealBuildTool/Configuration/TargetRules.cs b/Engine/Source/Programs/UnrealBuildTool/Configuration/TargetRules.cs
--- a/Engine/Source/Programs/UnrealBuildTool/Configuration/TargetRules.cs
+++ b/Engine/Source/Programs/UnrealBuildTool/Configuration/TargetRules.cs
@@ -1016,7 +1016,7 @@ namespace UnrealBuildTool
 		/// </summary>
 		public WarningLevel DefaultWarningLevel
 		{
-			get => DefaultWarningLevelPrivate ?? (bWarningsAsErrors ? WarningLevel.Error : WarningLevel.Error);
+			get => DefaultWarningLevelPrivate ?? (bWarningsAsErrors ? WarningLevel.Error : WarningLevel.Warning);
 			set => DefaultWarningLevelPrivate = value;
 		}
 
