UBT: Fix exception trying to serialize overrides for warning levels to XML config cache.

diff --git a/Engine/Source/Programs/UnrealBuildTool/Configuration/ModuleRules.cs b/Engine/Source/Programs/UnrealBuildTool/Configuration/ModuleRules.cs
--- a/Engine/Source/Programs/UnrealBuildTool/Configuration/ModuleRules.cs
+++ b/Engine/Source/Programs/UnrealBuildTool/Configuration/ModuleRules.cs
@@ -17,6 +17,11 @@ namespace UnrealBuildTool
 	/// </summary>
 	public enum WarningLevel
 	{
+		/// <summary>
+		/// Use the default behavior
+		/// </summary>
+		Default,
+
 		/// <summary>
 		/// Do not display diagnostics
 		/// </summary>
@@ -735,10 +740,12 @@ namespace UnrealBuildTool
 		/// </summary>
 		public WarningLevel ShadowVariableWarningLevel
 		{
-			get { return ShadowVariableWarningLevelPrivate ?? ((DefaultBuildSettings >= BuildSettingsVersion.V2) ? WarningLevel.Error : Target.ShadowVariableWarningLevel); }
+			get { return (ShadowVariableWarningLevelPrivate == WarningLevel.Default)? ((DefaultBuildSettings >= BuildSettingsVersion.V2) ? WarningLevel.Error : Target.ShadowVariableWarningLevel) : ShadowVariableWarningLevelPrivate; }
 			set { ShadowVariableWarningLevelPrivate = value; }
 		}
-		private WarningLevel? ShadowVariableWarningLevelPrivate;
+
+		/// <inheritdoc cref="ShadowVariableWarningLevelPrivate"/>
+		private WarningLevel ShadowVariableWarningLevelPrivate;
 
 		/// <summary>
 		/// Enable warnings for shadowed variables
@@ -755,10 +762,12 @@ namespace UnrealBuildTool
 		/// </summary>
 		public WarningLevel UnsafeTypeCastWarningLevel
 		{
-			get { return UnsafeTypeCastWarningLevelPrivate ?? Target.UnsafeTypeCastWarningLevel; }
+			get { return (UnsafeTypeCastWarningLevelPrivate == WarningLevel.Default)? Target.UnsafeTypeCastWarningLevel : UnsafeTypeCastWarningLevelPrivate; }
 			set { UnsafeTypeCastWarningLevelPrivate = value; }
 		}
-		private WarningLevel? UnsafeTypeCastWarningLevelPrivate;
+
+		/// <inheritdoc cref="UnsafeTypeCastWarningLevel"/>
+		private WarningLevel UnsafeTypeCastWarningLevelPrivate;
 
 		/// <summary>
 		/// Enable warnings for using undefined identifiers in #if expressions
diff --git a/Engine/Source/Programs/UnrealBuildTool/Configuration/TargetRules.cs b/Engine/Source/Programs/UnrealBuildTool/Configuration/TargetRules.cs
--- a/Engine/Source/Programs/UnrealBuildTool/Configuration/TargetRules.cs
+++ b/Engine/Source/Programs/UnrealBuildTool/Configuration/TargetRules.cs
@@ -1016,24 +1016,26 @@ namespace UnrealBuildTool
 		/// </summary>
 		public WarningLevel DefaultWarningLevel
 		{
-			get => DefaultWarningLevelPrivate ?? (bWarningsAsErrors ? WarningLevel.Error : WarningLevel.Warning);
+			get => (DefaultWarningLevelPrivate == WarningLevel.Default)? (bWarningsAsErrors ? WarningLevel.Error : WarningLevel.Warning) : DefaultWarningLevelPrivate;
 			set => DefaultWarningLevelPrivate = value;
 		}
 
+		/// <inheritdoc cref="DefaultWarningLevel"/>
 		[XmlConfigFile(Category = "BuildConfiguration", Name = nameof(DefaultWarningLevel))]
-		private WarningLevel? DefaultWarningLevelPrivate;
+		private WarningLevel DefaultWarningLevelPrivate;
 
 		/// <summary>
 		/// Level to report deprecation warnings as errors
 		/// </summary>
 		public WarningLevel DeprecationWarningLevel
 		{
-			get => DeprecationWarningLevelPrivate ?? DefaultWarningLevel;
+			get => (DeprecationWarningLevelPrivate == WarningLevel.Default)? DefaultWarningLevel : DeprecationWarningLevelPrivate;
 			set => DeprecationWarningLevelPrivate = value;
 		}
 
+		/// <inheritdoc cref="DeprecationWarningLevel"/>
 		[XmlConfigFile(Category = "BuildConfiguration", Name = nameof(DeprecationWarningLevel))]
-		private WarningLevel? DeprecationWarningLevelPrivate;
+		private WarningLevel DeprecationWarningLevelPrivate;
 
 		/// <summary>
 		/// Forces shadow variable warnings to be treated as errors on platforms that support it.
