PR #8291: Fix not being able to generate VS2022 project (Contributed by [REDACTED_PERSON])

diff --git a/Engine/Source/Programs/UnrealBuildTool/Platform/Windows/UEBuildWindows.cs b/Engine/Source/Programs/UnrealBuildTool/Platform/Windows/UEBuildWindows.cs
--- a/Engine/Source/Programs/UnrealBuildTool/Platform/Windows/UEBuildWindows.cs
+++ b/Engine/Source/Programs/UnrealBuildTool/Platform/Windows/UEBuildWindows.cs
@@ -1144,7 +1144,7 @@ namespace UnrealBuildTool
 										{
 											continue;
 										}
-										else if (Compiler == WindowsCompiler.VisualStudio2019 && (Version < Version2019 || Version >= Version2022))
+										else if (Compiler == WindowsCompiler.VisualStudio2019 && Version < Version2019)
 										{
 											continue;
 										}
