UnrealBuildTool: Make WorkingSet & CandidateWorkingSet mutually exclusive.

Prevents a case where edits to source files would cause Makefile to be invalidated.

diff --git a/Engine/Source/Programs/UnrealBuildTool/System/Unity.cs b/Engine/Source/Programs/UnrealBuildTool/System/Unity.cs
--- a/Engine/Source/Programs/UnrealBuildTool/System/Unity.cs
+++ b/Engine/Source/Programs/UnrealBuildTool/System/Unity.cs
@@ -366,6 +366,7 @@ namespace UnrealBuildTool
 			HashSet<FileItem> CandidateAdaptiveFiles = new HashSet<FileItem>();
 			CandidateAdaptiveFiles.UnionWith(CPPFiles);
 			CandidateAdaptiveFiles.UnionWith(HeaderFiles);
+			CandidateAdaptiveFiles.ExceptWith(AdaptiveFiles);
 
 			foreach (FileItem File in CandidateAdaptiveFiles)
 			{
