UnrealBuildTool: Add support for passing /experimental:deterministic to the MSCV compiler (WindowsPlatform.bDeterministic). Disabled by default.

diff --git a/Engine/Source/Programs/UnrealBuildTool/Platform/HoloLens/HoloLensToolChain.cs b/Engine/Source/Programs/UnrealBuildTool/Platform/HoloLens/HoloLensToolChain.cs
--- a/Engine/Source/Programs/UnrealBuildTool/Platform/HoloLens/HoloLensToolChain.cs
+++ b/Engine/Source/Programs/UnrealBuildTool/Platform/HoloLens/HoloLensToolChain.cs
@@ -129,6 +129,15 @@ namespace UnrealBuildTool
 				VCToolChain.AddDefinition(Arguments, "_WINDLL");
 			}
 
+			// Experimental deterministic compile support
+			if (Target.WindowsPlatform.bDeterministic)
+			{
+				if (Target.WindowsPlatform.Compiler >= WindowsCompiler.VisualStudio2017)
+				{
+					Arguments.Add("/experimental:deterministic");
+				}
+			}
+
 			//
 			//	Debug
 			//
@@ -376,6 +385,12 @@ namespace UnrealBuildTool
 			// Don't embed the full PDB path; we want to be able to move binaries elsewhere. They will always be side by side.
 			Arguments.Add("/PDBALTPATH:%_PDB%");
 
+			// Experimental deterministic compile support
+			if (Target.WindowsPlatform.bDeterministic)
+			{
+				Arguments.Add("/experimental:deterministic");
+			}
+
 			//
 			//	Shipping & LTCG
 			//
diff --git a/Engine/Source/Programs/UnrealBuildTool/Platform/Windows/UEBuildWindows.cs b/Engine/Source/Programs/UnrealBuildTool/Platform/Windows/UEBuildWindows.cs
--- a/Engine/Source/Programs/UnrealBuildTool/Platform/Windows/UEBuildWindows.cs
+++ b/Engine/Source/Programs/UnrealBuildTool/Platform/Windows/UEBuildWindows.cs
@@ -349,6 +349,13 @@ namespace UnrealBuildTool
 		[XmlConfigFile(Category = "WindowsPlatform")]
 		public bool bShowIncludes = false;
 
+		/// <summary>
+		/// Set flags require for determinstic compiles (experimental)
+		/// </summary>
+		[CommandLine("-Deterministic")]
+		[XmlConfigFile(Category = "WindowsPlatform")]
+		public bool bDeterministic = false;
+
 		/// <summary>
 		/// Bundle a working version of dbghelp.dll with the application, and use this to generate minidumps. This works around a bug with the Windows 10 Fall Creators Update (1709)
 		/// where rich PE headers larger than a certain size would result in corrupt minidumps.
@@ -628,6 +635,11 @@ namespace UnrealBuildTool
 			get { return Inner.bShowIncludes; }
 		}
 
+		public bool bDeterministic
+		{
+			get { return Inner.bDeterministic; }
+		}
+
 		public string GetVisualStudioCompilerVersionName()
 		{
 			return Inner.GetVisualStudioCompilerVersionName();
diff --git a/Engine/Source/Programs/UnrealBuildTool/Platform/Windows/VCToolChain.cs b/Engine/Source/Programs/UnrealBuildTool/Platform/Windows/VCToolChain.cs
--- a/Engine/Source/Programs/UnrealBuildTool/Platform/Windows/VCToolChain.cs
+++ b/Engine/Source/Programs/UnrealBuildTool/Platform/Windows/VCToolChain.cs
@@ -436,6 +436,15 @@ namespace UnrealBuildTool
 				Arguments.Add("/Ob0");
 			}
 
+			// Experimental deterministic compile support
+			if (Target.WindowsPlatform.bDeterministic)
+			{
+				if (Target.WindowsPlatform.Compiler >= WindowsCompiler.VisualStudio2017)
+				{
+					Arguments.Add("/experimental:deterministic");
+				}
+			}
+
 			// Address sanitizer
 			if (Target.WindowsPlatform.bEnableAddressSanitizer)
 			{
@@ -1010,6 +1019,12 @@ namespace UnrealBuildTool
 			// Don't embed the full PDB path; we want to be able to move binaries elsewhere. They will always be side by side.
 			Arguments.Add("/PDBALTPATH:%_PDB%");
 
+			// Experimental deterministic compile support
+			if (Target.WindowsPlatform.bDeterministic)
+			{
+				Arguments.Add("/experimental:deterministic");
+			}
+
 			//
 			//	Shipping & LTCG
 			//
@@ -1066,7 +1081,6 @@ namespace UnrealBuildTool
 				Arguments.Add("/INCREMENTAL:NO");
 			}
 
-
 			// Add any extra options from the target
 			if (!string.IsNullOrEmpty(Target.WindowsPlatform.AdditionalLinkerOptions))
 			{
