UnrealBuildTool: Cap default parallel process count to physical core count.

Evidence suggests that throughput is not significantly reduced by using the number of physical cores, and that doing so reduces frequency of "C1060: compiler is out of heap space" errors.

diff --git a/Engine/Source/Programs/UnrealBuildTool/System/Utils.cs b/Engine/Source/Programs/UnrealBuildTool/System/Utils.cs
--- a/Engine/Source/Programs/UnrealBuildTool/System/Utils.cs
+++ b/Engine/Source/Programs/UnrealBuildTool/System/Utils.cs
@@ -1258,21 +1258,17 @@ namespace UnrealBuildTool
 
 			// The number of actions to execute in parallel is trying to keep the CPU busy enough in presence of I/O stalls.
 			int MaxActionsToExecuteInParallel;
-			if (NumPhysicalCores < NumLogicalCores && ProcessorCountMultiplier != 1.0)
+			if (ProcessorCountMultiplier != 1.0)
 			{
 				// The CPU has more logical cores than physical ones, aka uses hyper-threading. 
 				// Use multiplier if provided
 				MaxActionsToExecuteInParallel = (int)(NumPhysicalCores * ProcessorCountMultiplier);
 				Log.TraceInformation($"  Requested {ProcessorCountMultiplier} process count multiplier: limiting max parallel actions to {MaxActionsToExecuteInParallel}");
 			}
-			else if (NumPhysicalCores < NumLogicalCores)
-			{
-				// The CPU has more logical cores than physical ones, aka uses hyper-threading. 
-				MaxActionsToExecuteInParallel = NumLogicalCores;
-			}
-			// No hyper-threading. Only kicking off a task per CPU to keep machine responsive.
+			// kick off a task per physical core - evidence suggests that, in general, using more cores does not yield significantly better throughput
 			else
 			{
+				Log.TraceInformation($"  Limiting parallelism to one process per physical core ({NumPhysicalCores})");
 				MaxActionsToExecuteInParallel = NumPhysicalCores;
 			}
 
