Don't add all DotNet and Engine programs to -game specific generated projects by default

diff --git a/Engine/Source/Programs/UnrealBuildTool/ProjectFiles/ProjectFileGenerator.cs b/Engine/Source/Programs/UnrealBuildTool/ProjectFiles/ProjectFileGenerator.cs
--- a/Engine/Source/Programs/UnrealBuildTool/ProjectFiles/ProjectFileGenerator.cs
+++ b/Engine/Source/Programs/UnrealBuildTool/ProjectFiles/ProjectFileGenerator.cs
@@ -1181,6 +1181,7 @@ namespace UnrealBuildTool
 			}
 
 			bool bAlwaysIncludeEngineModules = false;
+			bool bAlwaysIncludeDotNetPrograms = false;
 			foreach (string CurArgument in Arguments)
 			{
 				if (CurArgument.StartsWith("-"))
@@ -1238,6 +1239,9 @@ namespace UnrealBuildTool
 
 							case "-GAME":
 								// Generates project files for a single game
+								bIncludeDotNetPrograms = false;
+								IncludeEnginePrograms = false;
+								bIncludeEngineSource = false;
 								bGeneratingGameProjectFiles = true;
 								break;
 
@@ -1279,7 +1283,7 @@ namespace UnrealBuildTool
 								break;
 
 							case "-DOTNET":
-								bIncludeDotNetPrograms = true;
+								bAlwaysIncludeDotNetPrograms = true;
 								break;
 
 							case "-NODOTNET":
@@ -1323,6 +1327,7 @@ namespace UnrealBuildTool
 				bIncludeTemplateFiles = false;
 				bIncludeConfigFiles = true;
 				IncludeEnginePrograms = bAlwaysIncludeEngineModules;
+				bIncludeDotNetPrograms = bIncludeDotNetPrograms || bAlwaysIncludeDotNetPrograms;
 			}
 			else
 			{
@@ -1501,7 +1506,7 @@ namespace UnrealBuildTool
 			List<string> UnsupportedPlatformNameStrings = Utils.MakeListOfUnsupportedPlatforms(SupportedPlatforms, bIncludeUnbuildablePlatforms: true);
 
 			// Locate all targets (*.Target.cs files)
-			List<FileReference> FoundTargetFiles = Rules.FindAllRulesSourceFiles(Rules.RulesFileType.Target, AllGameProjects.Select(x => x.Directory).ToList(), ForeignPlugins: null, AdditionalSearchPaths: null, bIncludeTempTargets: bIncludeTempTargets);
+			List<FileReference> FoundTargetFiles = Rules.FindAllRulesSourceFiles(Rules.RulesFileType.Target, AllGameProjects.Select(x => x.Directory).ToList(), ForeignPlugins: null, AdditionalSearchPaths: null, bIncludeEngine: IncludeEnginePrograms, bIncludeTempTargets: bIncludeTempTargets);
 			foreach (FileReference CurTargetFile in FoundTargetFiles)
 			{
 				string CleanTargetFileName = Utils.CleanDirectorySeparators(CurTargetFile.FullName);
diff --git a/Engine/Source/Programs/UnrealBuildTool/ProjectFiles/VisualStudioCode/VSCodeProjectFileGenerator.cs b/Engine/Source/Programs/UnrealBuildTool/ProjectFiles/VisualStudioCode/VSCodeProjectFileGenerator.cs
--- a/Engine/Source/Programs/UnrealBuildTool/ProjectFiles/VisualStudioCode/VSCodeProjectFileGenerator.cs
+++ b/Engine/Source/Programs/UnrealBuildTool/ProjectFiles/VisualStudioCode/VSCodeProjectFileGenerator.cs
@@ -1434,14 +1434,19 @@ namespace UnrealBuildTool
 				Args.Add("-projectfiles");
 				Args.Add("-vscode");
 
-				if (bForeignProject)
+				if (bGeneratingGameProjectFiles)
 				{
 					Args.Add("-project=" + MakeUnquotedPathString(OnlyGameProject, EPathType.Absolute));
 					Args.Add("-game");
+				}
+				if (bIncludeEngineSource)
+				{
 					Args.Add("-engine");
 				}
-				else
+				
+				if (bIncludeDotNetPrograms)
 				{
+					Args.Add("-dotnet");
 					PreLaunchTask = "UnrealBuildTool " + HostPlatform.ToString() + " Development Build";
 				}
 
