ProjectGeneration: Don't filter out engine by default for -game

#rnx

diff --git a/Engine/Source/Programs/UnrealBuildTool/ProjectFiles/ProjectFileGenerator.cs b/Engine/Source/Programs/UnrealBuildTool/ProjectFiles/ProjectFileGenerator.cs
--- a/Engine/Source/Programs/UnrealBuildTool/ProjectFiles/ProjectFileGenerator.cs
+++ b/Engine/Source/Programs/UnrealBuildTool/ProjectFiles/ProjectFileGenerator.cs
@@ -1241,7 +1241,6 @@ namespace UnrealBuildTool
 								// Generates project files for a single game
 								bIncludeDotNetPrograms = false;
 								IncludeEnginePrograms = false;
-								bIncludeEngineSource = false;
 								bGeneratingGameProjectFiles = true;
 								break;
 
diff --git a/Engine/Source/Programs/UnrealBuildTool/ProjectFiles/VisualStudioCode/VSCodeProjectFileGenerator.cs b/Engine/Source/Programs/UnrealBuildTool/ProjectFiles/VisualStudioCode/VSCodeProjectFileGenerator.cs
--- a/Engine/Source/Programs/UnrealBuildTool/ProjectFiles/VisualStudioCode/VSCodeProjectFileGenerator.cs
+++ b/Engine/Source/Programs/UnrealBuildTool/ProjectFiles/VisualStudioCode/VSCodeProjectFileGenerator.cs
@@ -1439,10 +1439,6 @@ namespace UnrealBuildTool
 					Args.Add("-project=" + MakeUnquotedPathString(OnlyGameProject, EPathType.Absolute));
 					Args.Add("-game");
 				}
-				if (bIncludeEngineSource)
-				{
-					Args.Add("-engine");
-				}
 				
 				if (bIncludeDotNetPrograms)
 				{
