Remove force disable of  bIncludeEngineSource when passing -code

#rnx

diff --git a/Engine/Source/Programs/UnrealBuildTool/ProjectFiles/ProjectFileGenerator.cs b/Engine/Source/Programs/UnrealBuildTool/ProjectFiles/ProjectFileGenerator.cs
--- a/Engine/Source/Programs/UnrealBuildTool/ProjectFiles/ProjectFileGenerator.cs
+++ b/Engine/Source/Programs/UnrealBuildTool/ProjectFiles/ProjectFileGenerator.cs
@@ -1241,7 +1241,6 @@ namespace UnrealBuildTool
 								// Generates project files for a single game
 								bIncludeDotNetPrograms = false;
 								IncludeEnginePrograms = false;
-								bIncludeEngineSource = false;
 								bGeneratingGameProjectFiles = true;
 								break;
 
