UnrealBuildTool: allow MemoryPerAction to be set per-target

Can be used to reduce the number of build actions that will be run in parallel by ParallelExecutor or TaskExecutor.

Example use, in MyProject.Target.cs:

public class MyProjectTarget : TargetRules
{
	public MyProjectTarget(TargetInfo Target) : base(Target)
	{
		Type = TargetType.Game;
		DefaultBuildSettings = BuildSettingsVersion.V2;

		ExtraModuleNames.AddRange( new string[] { "MyProject" } );

		MemoryPerActionGB = 4;
	}
}

diff --git a/Engine/Source/Programs/UnrealBuildTool/Configuration/TargetRules.cs b/Engine/Source/Programs/UnrealBuildTool/Configuration/TargetRules.cs
--- a/Engine/Source/Programs/UnrealBuildTool/Configuration/TargetRules.cs
+++ b/Engine/Source/Programs/UnrealBuildTool/Configuration/TargetRules.cs
@@ -1574,6 +1574,11 @@ namespace UnrealBuildTool
 		[CommandLine("-LinkerArguments=")]
 		public string? AdditionalLinkerArguments;
 
+		/// <summary>
+		/// Max amount of memory that each compile action may require. Used by ParallelExecutor and TaskExecutor to decide the maximum 
+		/// number of parallel actions to start at one time.
+		/// </summary>
+		public double MemoryPerActionGB = 0.0;
 
 		/// <summary>
 		/// List of modules to disable unity builds for
@@ -2842,6 +2847,11 @@ namespace UnrealBuildTool
 			get { return Inner.AdditionalLinkerArguments; }
 		}
 
+		public double MemoryPerActionGB
+		{
+			get { return Inner.MemoryPerActionGB; }
+		}
+
 		public string? GeneratedProjectName
 		{
 			get { return Inner.GeneratedProjectName; }
diff --git a/Engine/Source/Programs/UnrealBuildTool/Executors/Experimental/TaskExecutor.cs b/Engine/Source/Programs/UnrealBuildTool/Executors/Experimental/TaskExecutor.cs
--- a/Engine/Source/Programs/UnrealBuildTool/Executors/Experimental/TaskExecutor.cs
+++ b/Engine/Source/Programs/UnrealBuildTool/Executors/Experimental/TaskExecutor.cs
@@ -59,7 +59,13 @@ namespace UnrealBuildTool
 
 		public static int GetDefaultNumParallelProcesses()
 		{
-			return Utils.GetMaxActionsToExecuteInParallel(MaxProcessorCount, ProcessorCountMultiplier, Convert.ToInt64(MemoryPerActionBytes));
+			double MemoryPerActionBytesComputed = Math.Max(MemoryPerActionBytes, MemoryPerActionBytesOverride);
+			if (MemoryPerActionBytesComputed > MemoryPerActionBytes)
+			{
+				Log.TraceInformation($"Overriding MemoryPerAction with target-defined value of {MemoryPerActionBytesComputed / 1024 / 1024 / 1024} bytes");
+			}
+
+			return Utils.GetMaxActionsToExecuteInParallel(MaxProcessorCount, ProcessorCountMultiplier, Convert.ToInt64(MemoryPerActionBytesComputed));
 		}
 
 		/// <summary>
diff --git a/Engine/Source/Programs/UnrealBuildTool/Executors/ParallelExecutor.cs b/Engine/Source/Programs/UnrealBuildTool/Executors/ParallelExecutor.cs
--- a/Engine/Source/Programs/UnrealBuildTool/Executors/ParallelExecutor.cs
+++ b/Engine/Source/Programs/UnrealBuildTool/Executors/ParallelExecutor.cs
@@ -79,7 +79,13 @@ namespace UnrealBuildTool
 
 		public static int GetDefaultNumParallelProcesses()
 		{
-			return Utils.GetMaxActionsToExecuteInParallel(MaxProcessorCount, ProcessorCountMultiplier, Convert.ToInt64(MemoryPerActionBytes));
+			double MemoryPerActionBytesComputed = Math.Max(MemoryPerActionBytes, MemoryPerActionBytesOverride);
+			if (MemoryPerActionBytesComputed > MemoryPerActionBytes)
+			{
+				Log.TraceInformation($"Overriding MemoryPerAction with target-defined value of {MemoryPerActionBytesComputed / 1024 / 1024 / 1024} bytes");
+			}
+
+			return Utils.GetMaxActionsToExecuteInParallel(MaxProcessorCount, ProcessorCountMultiplier, Convert.ToInt64(MemoryPerActionBytesComputed));
 		}
 
 		/// <summary>
diff --git a/Engine/Source/Programs/UnrealBuildTool/Modes/BuildMode.cs b/Engine/Source/Programs/UnrealBuildTool/Modes/BuildMode.cs
--- a/Engine/Source/Programs/UnrealBuildTool/Modes/BuildMode.cs
+++ b/Engine/Source/Programs/UnrealBuildTool/Modes/BuildMode.cs
@@ -591,6 +591,8 @@ namespace UnrealBuildTool
 
 						OutputToolchainInfo(TargetDescriptors, Makefiles);
 
+						ActionExecutor.SetMemoryPerActionOverride(Makefiles.Select(x => x.MemoryPerActionGB).Max());
+
 						using(Timeline.ScopeEvent("ActionGraph.ExecuteActions()"))
 						{
 							ActionGraph.ExecuteActions(BuildConfiguration, MergedActionsToExecute);
@@ -727,6 +729,8 @@ namespace UnrealBuildTool
 					Makefile = Target.Build(BuildConfiguration, WorkingSet, TargetDescriptor.SpecificFilesToCompile);
 				}
 
+				Makefile.MemoryPerActionGB = Target.Rules.MemoryPerActionGB;
+
 				// Save the pre-build scripts onto the makefile
 				Makefile.PreBuildScripts = PreBuildScripts;
 
diff --git a/Engine/Source/Programs/UnrealBuildTool/System/ActionExecutor.cs b/Engine/Source/Programs/UnrealBuildTool/System/ActionExecutor.cs
--- a/Engine/Source/Programs/UnrealBuildTool/System/ActionExecutor.cs
+++ b/Engine/Source/Programs/UnrealBuildTool/System/ActionExecutor.cs
@@ -1,5 +1,6 @@
 // Copyright Epic Games, Inc. All Rights Reserved.
 
+using System;
 using System.Collections.Generic;
 
 namespace UnrealBuildTool
@@ -11,6 +12,22 @@ namespace UnrealBuildTool
 			get;
 		}
 
+		static protected double MemoryPerActionBytesOverride
+		{
+			get;
+			private set;
+		} = 0.0;
+
+		/// <summary>
+		/// Allow targets to override the expected amount of memory required for compiles, used to control the number
+		/// of parallel action processes.
+		/// </summary>
+		/// <param name="MemoryPerActionOverrideGB"></param>
+		public static void SetMemoryPerActionOverride(double MemoryPerActionOverrideGB)
+		{
+			MemoryPerActionBytesOverride = Math.Max(MemoryPerActionBytesOverride, MemoryPerActionOverrideGB * 1024 * 1024 * 1024);
+		}
+
 		public abstract bool ExecuteActions(List<LinkedAction> ActionsToExecute);
 	}
 
diff --git a/Engine/Source/Programs/UnrealBuildTool/System/TargetMakefile.cs b/Engine/Source/Programs/UnrealBuildTool/System/TargetMakefile.cs
--- a/Engine/Source/Programs/UnrealBuildTool/System/TargetMakefile.cs
+++ b/Engine/Source/Programs/UnrealBuildTool/System/TargetMakefile.cs
@@ -22,7 +22,7 @@ namespace UnrealBuildTool
 		/// <summary>
 		/// The version number to write
 		/// </summary>
-		public const int CurrentVersion = 24;
+		public const int CurrentVersion = 25;
 
 		/// <summary>
 		/// The time at which the makefile was created
@@ -175,6 +175,11 @@ namespace UnrealBuildTool
 		/// </summary>
 		public HashSet<FileItem> InternalDependencies = new HashSet<FileItem>();
 
+		/// <summary>
+		/// TargetRules-set memory estimate per action. Used to control the number of parallel actions that are spawned.
+		/// </summary>
+		public double MemoryPerActionGB = 0.0;
+
 		/// <summary>
 		/// Constructor
 		/// </summary>
@@ -250,6 +255,7 @@ namespace UnrealBuildTool
 			PluginFiles = Reader.ReadHashSet(() => Reader.ReadFileItem());
 			ExternalDependencies = Reader.ReadHashSet(() => Reader.ReadFileItem());
 			InternalDependencies = Reader.ReadHashSet(() => Reader.ReadFileItem());
+			MemoryPerActionGB = Reader.ReadDouble();
 		}
 
 		/// <summary>
@@ -286,6 +292,7 @@ namespace UnrealBuildTool
 			Writer.WriteHashSet(PluginFiles, x => Writer.WriteFileItem(x));
 			Writer.WriteHashSet(ExternalDependencies, x => Writer.WriteFileItem(x));
 			Writer.WriteHashSet(InternalDependencies, x => Writer.WriteFileItem(x));
+			Writer.WriteDouble(MemoryPerActionGB);
 		}
 
 		/// <summary>
