UnrealBuildTool: Split ActionGraph.IsActionOutdated() into two functions

IsActionGraph() is complicated: recursively evaluating a graph, modifying data at each node (ActionHistory.UpdateProducingAttributes) making it easy to introduce hard-to-see race conditions, and difficult to optimize.

This change splits the function into two parts:

IsIndividualActionOutdated(): finds nodes that are out of date, based on the node's own data - able to be run in parallel, with less cross-thread contention.

IsActionOutdatedDueToPrerequisites(): recursively traverse the graph, to determine if an action is out of date based on its prerequisites. Does not run in parallel, but is fast (processes ~6000 actions in 2-3ms).

For a clean, simple project editor build, total runtime is <100ms, compared to ~150ms previously.

diff --git a/Engine/Source/Programs/UnrealBuildTool/System/ActionGraph.cs b/Engine/Source/Programs/UnrealBuildTool/System/ActionGraph.cs
--- a/Engine/Source/Programs/UnrealBuildTool/System/ActionGraph.cs
+++ b/Engine/Source/Programs/UnrealBuildTool/System/ActionGraph.cs
@@ -9,6 +9,7 @@ using System.Linq;
 using System.Text;
 using System.Text.Json;
 using System.Text.Json.Serialization;
+using System.Threading;
 using System.Threading.Tasks;
 using UnrealBuildBase;
 
@@ -564,23 +565,29 @@ namespace UnrealBuildTool
 
 		/// <summary>
 		/// Determines whether an action is outdated based on the modification times for its prerequisite
-		/// and produced items.
+		/// and produced items, without considering the full set of prerequisites.
+		/// Writes to OutdatedActionDictionary iff the action is found to be outdated.
+		/// Safe to run in parallel, but only with different RootActions.
 		/// </summary>
 		/// <param name="RootAction">- The action being considered.</param>
 		/// <param name="OutdatedActionDictionary">-</param>
+		/// <param name="OutdatedActionLock"></param>
 		/// <param name="ActionHistory"></param>
 		/// <param name="CppDependencies"></param>
 		/// <param name="bIgnoreOutdatedImportLibraries"></param>
 		/// <returns>true if outdated</returns>
-		public static bool IsActionOutdated(LinkedAction RootAction, Dictionary<LinkedAction, bool> OutdatedActionDictionary, ActionHistory ActionHistory, CppDependencyCache CppDependencies, bool bIgnoreOutdatedImportLibraries)
+		public static void IsIndividualActionOutdated(LinkedAction RootAction, Dictionary<LinkedAction, bool> OutdatedActionDictionary, ReaderWriterLockSlim OutdatedActionLock, ActionHistory ActionHistory, CppDependencyCache CppDependencies, bool bIgnoreOutdatedImportLibraries)
 		{
 			// Only compute the outdated-ness for actions that don't aren't cached in the outdated action dictionary.
 			bool bIsOutdated = false;
-			lock(OutdatedActionDictionary)
 			{
-				if (OutdatedActionDictionary.TryGetValue(RootAction, out bIsOutdated))
+				// OutdatedActionDictionary may have already been populated for RootAction as part of a previously processed target
+				OutdatedActionLock.EnterReadLock();
+				bool bPresent = OutdatedActionDictionary.ContainsKey(RootAction);
+				OutdatedActionLock.ExitReadLock();
+				if (bPresent)
 				{
-					return bIsOutdated;
+					return;
 				}
 			}
 
@@ -597,7 +604,7 @@ namespace UnrealBuildTool
 						Log.TraceLog(
 							"{0}: Produced item \"{1}\" was produced by outdated attributes.\n  New attributes: {2}",
 							RootAction.StatusDescription,
-							Path.GetFileName(ProducedItem.AbsolutePath),
+							ProducedItem.AbsolutePath,
 							NewProducingAttributes
 							);
 					}
@@ -622,33 +629,12 @@ namespace UnrealBuildTool
 					Log.TraceLog(
 						"{0}: Produced item \"{1}\" doesn't exist.",
 						RootAction.StatusDescription,
-						Path.GetFileName(ProducedItem.AbsolutePath)
+						ProducedItem.AbsolutePath
 						);
 					bIsOutdated = true;
 				}
 			}
 
-			// Check if any of the prerequisite actions are out of date
-			if (!bIsOutdated)
-			{
-				foreach (LinkedAction PrerequisiteAction in RootAction.PrerequisiteActions)
-				{
-					if (IsActionOutdated(PrerequisiteAction, OutdatedActionDictionary, ActionHistory, CppDependencies, bIgnoreOutdatedImportLibraries))
-					{
-						// Only check for outdated import libraries if we were configured to do so.  Often, a changed import library
-						// won't affect a dependency unless a public header file was also changed, in which case we would be forced
-						// to recompile anyway.  This just allows for faster iteration when working on a subsystem in a DLL, as we
-						// won't have to wait for dependent targets to be relinked after each change.
-						if(!bIgnoreOutdatedImportLibraries || !IsImportLibraryDependency(RootAction, PrerequisiteAction))
-						{
-							Log.TraceLog("{0}: Prerequisite {1} is produced by outdated action.", RootAction.StatusDescription, PrerequisiteAction.StatusDescription);
-							bIsOutdated = true;
-							break;
-						}
-					}
-				}
-			} 
-
 			// Check if any prerequisite item has a newer timestamp than the last execution time of this action
 			if(!bIsOutdated)
 			{
@@ -664,7 +650,7 @@ namespace UnrealBuildTool
 							// Need to check for import libraries here too
 							if(!bIgnoreOutdatedImportLibraries || !IsImportLibraryDependency(RootAction, PrerequisiteItem))
 							{
-								Log.TraceLog("{0}: Prerequisite {1} is newer than the last execution of the action: {2} vs {3}", RootAction.StatusDescription, Path.GetFileName(PrerequisiteItem.AbsolutePath), PrerequisiteItem.LastWriteTimeUtc.ToLocalTime(), LastExecutionTimeUtc.LocalDateTime);
+								Log.TraceLog("{0}: Prerequisite {1} is newer than the last execution of the action: {2} vs {3}", RootAction.StatusDescription, PrerequisiteItem.AbsolutePath, PrerequisiteItem.LastWriteTimeUtc.ToLocalTime(), LastExecutionTimeUtc.LocalDateTime);
 								bIsOutdated = true;
 								break;
 							}
@@ -690,7 +676,7 @@ namespace UnrealBuildTool
 							Log.TraceLog(
 								"{0}: Dependency {1} is newer than the last execution of the action: {2} vs {3}",
 								RootAction.StatusDescription,
-								Path.GetFileName(DependencyFile.AbsolutePath),
+								DependencyFile.AbsolutePath,
 								DependencyFile.LastWriteTimeUtc.ToLocalTime(),
 								LastExecutionTimeUtc.LocalDateTime
 								);
@@ -701,14 +687,55 @@ namespace UnrealBuildTool
 				}
 			}
 
-			// Cache the outdated-ness of this action.
-			lock(OutdatedActionDictionary)
+			// if the action is known to be out of date, record that fact 
+			// We don't yet know that the action is up-to-date - to determine that requires traversal of the graph of prerequisites.
+			if (bIsOutdated)
 			{
-				if(!OutdatedActionDictionary.ContainsKey(RootAction))
+				OutdatedActionLock.EnterWriteLock();
+				OutdatedActionDictionary.Add(RootAction, bIsOutdated);
+				OutdatedActionLock.ExitWriteLock();
+			}
+		}
+
+		/// <summary>
+		/// Determines whether an action is outdated by examining the up-to-date state of all of its prerequisites, recursively.
+		/// Not thread safe. Typically very fast.
+		/// </summary>
+		/// <param name="RootAction">- The action being considered.</param>
+		/// <param name="OutdatedActionDictionary">-</param>
+		/// <param name="bIgnoreOutdatedImportLibraries"></param>
+		/// <returns>true if outdated</returns>
+		public static bool IsActionOutdatedDueToPrerequisites(LinkedAction RootAction, Dictionary<LinkedAction, bool> OutdatedActionDictionary, bool bIgnoreOutdatedImportLibraries)
+		{
+			// Only compute the outdated-ness for actions that aren't already cached in the outdated action dictionary.
+			if (OutdatedActionDictionary.TryGetValue(RootAction, out bool bIsOutdated))
+			{
+				return bIsOutdated;
+			}
+
+			// Check if any of the prerequisite actions are out of date
+			if (!bIsOutdated)
+			{
+				foreach (LinkedAction PrerequisiteAction in RootAction.PrerequisiteActions)
 				{
-					OutdatedActionDictionary.Add(RootAction, bIsOutdated);
+					if (IsActionOutdatedDueToPrerequisites(PrerequisiteAction, OutdatedActionDictionary, bIgnoreOutdatedImportLibraries))
+					{
+						// Only check for outdated import libraries if we were configured to do so.  Often, a changed import library
+						// won't affect a dependency unless a public header file was also changed, in which case we would be forced
+						// to recompile anyway.  This just allows for faster iteration when working on a subsystem in a DLL, as we
+						// won't have to wait for dependent targets to be relinked after each change.
+						if(!bIgnoreOutdatedImportLibraries || !IsImportLibraryDependency(RootAction, PrerequisiteAction))
+						{
+							Log.TraceLog("{0}: Prerequisite {1} is produced by outdated action.", RootAction.StatusDescription, PrerequisiteAction.StatusDescription);
+							bIsOutdated = true;
+							break;
+						}
+					}
 				}
-			}
+			} 
+
+			// Cache the outdated-ness of this action.
+			OutdatedActionDictionary.Add(RootAction, bIsOutdated);
 
 			return bIsOutdated;
 		}
@@ -771,11 +798,21 @@ namespace UnrealBuildTool
 				Parallel.ForEach(Dependencies, File => { CppDependencies.TryGetDependencies(File, out _); });
 			}
 
-			using(Timeline.ScopeEvent("Cache outdated actions"))
+			using (Timeline.ScopeEvent("Cache individual outdated actions"))
 			{
-				Parallel.ForEach(Actions, Action => IsActionOutdated(Action, OutdatedActions, ActionHistory, CppDependencies, bIgnoreOutdatedImportLibraries));
+				ReaderWriterLockSlim OutdatedActionsLock = new ReaderWriterLockSlim();
+				Parallel.ForEach(Actions, Action => IsIndividualActionOutdated(Action, OutdatedActions, OutdatedActionsLock, ActionHistory, CppDependencies, bIgnoreOutdatedImportLibraries));
+			}
+
+			using (Timeline.ScopeEvent("Cache outdated actions based on recursive prerequisites"))
+			{ 
+				foreach (var Action in Actions)
+				{
+					IsActionOutdatedDueToPrerequisites(Action, OutdatedActions, bIgnoreOutdatedImportLibraries);
+				}
 			}
 		}
+
 		/// <summary>
 		/// Deletes all the items produced by actions in the provided outdated action dictionary.
 		/// </summary>
