UnrealBuildTool: ActionGraph.cs clean up formatting, naming & some other minor changes

#trivial

diff --git a/Engine/Source/Programs/UnrealBuildTool/System/ActionGraph.cs b/Engine/Source/Programs/UnrealBuildTool/System/ActionGraph.cs
--- a/Engine/Source/Programs/UnrealBuildTool/System/ActionGraph.cs
+++ b/Engine/Source/Programs/UnrealBuildTool/System/ActionGraph.cs
@@ -4,6 +4,7 @@ using EpicGames.Core;
 using System;
 using System.Collections.Generic;
 using System.Diagnostics;
+using System.Globalization;
 using System.IO;
 using System.Linq;
 using System.Text;
@@ -15,13 +16,13 @@ using UnrealBuildBase;
 
 namespace UnrealBuildTool
 {
-	static class ActionGraph
+	internal static class ActionGraph
 	{
 		/// <summary>
 		/// Enum describing why an Action is in conflict with another Action
 		/// </summary>
 		[Flags]
-		internal enum ActionConflictReasonFlags : byte
+		private enum ActionConflictReasonFlags : byte
 		{
 			None = 0,
 			ActionType = 1 << 0,
@@ -56,7 +57,7 @@ namespace UnrealBuildTool
 			foreach (LinkedAction Action in Actions)
 			{
 				Action.PrerequisiteActions = new HashSet<LinkedAction>();
-				foreach(FileItem PrerequisiteItem in Action.PrerequisiteItems)
+				foreach (FileItem PrerequisiteItem in Action.PrerequisiteItems)
 				{
 					if (ItemToProducingAction.TryGetValue(PrerequisiteItem, out LinkedAction? PrerequisiteAction))
 					{
@@ -78,9 +79,9 @@ namespace UnrealBuildTool
 			bool bResult = true;
 
 			Dictionary<FileItem, IExternalAction> ItemToProducingAction = new Dictionary<FileItem, IExternalAction>();
-			foreach(IExternalAction Action in Actions)
+			foreach (IExternalAction Action in Actions)
 			{
-				foreach(FileItem ProducedItem in Action.ProducedItems)
+				foreach (FileItem ProducedItem in Action.ProducedItems)
 				{
 					if (ItemToProducingAction.TryGetValue(ProducedItem, out IExternalAction? ExistingAction))
 					{
@@ -93,45 +94,51 @@ namespace UnrealBuildTool
 				}
 			}
 
-			if(!bResult)
+			if (!bResult)
 			{
 				throw new BuildException("Action graph is invalid; unable to continue. See log for additional details.");
 			}
 		}
 
 		/// <summary>
-		/// Finds conflicts betwee two actions, and prints them to the log
+		/// Finds conflicts between two actions, and prints them to the log
 		/// </summary>
 		/// <param name="A">The first action</param>
 		/// <param name="B">The second action</param>
 		/// <returns>True if no conflicts were found, false otherwise.</returns>
-		public static bool CheckForConflicts(IExternalAction A, IExternalAction B)
+		private static bool CheckForConflicts(IExternalAction A, IExternalAction B)
 		{
 			ActionConflictReasonFlags Reason = ActionConflictReasonFlags.None;
 			if (A.ActionType != B.ActionType)
 			{
 				Reason |= ActionConflictReasonFlags.ActionType;
 			}
-			if (!Enumerable.SequenceEqual(A.PrerequisiteItems, B.PrerequisiteItems))
+
+			if (!A.PrerequisiteItems.SequenceEqual(B.PrerequisiteItems))
 			{
 				Reason |= ActionConflictReasonFlags.PrerequisiteItems;
 			}
-			if (!Enumerable.SequenceEqual(A.DeleteItems, B.DeleteItems))
+
+			if (!A.DeleteItems.SequenceEqual(B.DeleteItems))
 			{
 				Reason |= ActionConflictReasonFlags.DeleteItems;
 			}
+
 			if (A.DependencyListFile != B.DependencyListFile)
 			{
 				Reason |= ActionConflictReasonFlags.DependencyListFile;
 			}
+
 			if (A.WorkingDirectory != B.WorkingDirectory)
 			{
 				Reason |= ActionConflictReasonFlags.WorkingDirectory;
 			}
+
 			if (A.CommandPath != B.CommandPath)
 			{
 				Reason |= ActionConflictReasonFlags.CommandPath;
 			}
+
 			if (A.CommandArguments != B.CommandArguments)
 			{
 				Reason |= ActionConflictReasonFlags.CommandArguments;
@@ -142,58 +149,59 @@ namespace UnrealBuildTool
 				LogConflict(A, B, Reason);
 				return false;
 			}
+
 			return true;
 		}
 
-		internal class LogActionActionTypeConverter : JsonConverter<ActionType>
+		private class LogActionActionTypeConverter : JsonConverter<ActionType>
 		{
-			public override ActionType Read(ref Utf8JsonReader reader, Type typeToConvert, JsonSerializerOptions options)
+			public override ActionType Read(ref Utf8JsonReader Reader, Type TypeToConvert, JsonSerializerOptions Options)
 			{
 				throw new NotImplementedException();
 			}
 
-			public override void Write(Utf8JsonWriter writer, ActionType value, JsonSerializerOptions options)
+			public override void Write(Utf8JsonWriter Writer, ActionType Value, JsonSerializerOptions Options)
 			{
-				writer.WriteStringValue(value.ToString());
+				Writer.WriteStringValue(Value.ToString());
 			}
 		}
 
-		internal class LogActionFileItemConverter : JsonConverter<FileItem>
+		private class LogActionFileItemConverter : JsonConverter<FileItem>
 		{
-			public override FileItem Read(ref Utf8JsonReader reader, Type typeToConvert, JsonSerializerOptions options)
+			public override FileItem Read(ref Utf8JsonReader Reader, Type TypeToConvert, JsonSerializerOptions Options)
 			{
 				throw new NotImplementedException();
 			}
 
-			public override void Write(Utf8JsonWriter writer, FileItem value, JsonSerializerOptions options)
+			public override void Write(Utf8JsonWriter Writer, FileItem Value, JsonSerializerOptions Options)
 			{
-				writer.WriteStringValue(value.FullName);
+				Writer.WriteStringValue(Value.FullName);
 			}
 		}
 
-		internal class LogActionDirectoryReferenceConverter : JsonConverter<DirectoryReference>
+		private class LogActionDirectoryReferenceConverter : JsonConverter<DirectoryReference>
 		{
-			public override DirectoryReference Read(ref Utf8JsonReader reader, Type typeToConvert, JsonSerializerOptions options)
+			public override DirectoryReference Read(ref Utf8JsonReader Reader, Type TypeToConvert, JsonSerializerOptions Options)
 			{
 				throw new NotImplementedException();
 			}
 
-			public override void Write(Utf8JsonWriter writer, DirectoryReference value, JsonSerializerOptions options)
+			public override void Write(Utf8JsonWriter Writer, DirectoryReference Value, JsonSerializerOptions Options)
 			{
-				writer.WriteStringValue(value.FullName);
+				Writer.WriteStringValue(Value.FullName);
 			}
 		}
 
-		internal class LogActionFileReferenceConverter : JsonConverter<FileReference>
+		private class LogActionFileReferenceConverter : JsonConverter<FileReference>
 		{
-			public override FileReference Read(ref Utf8JsonReader reader, Type typeToConvert, JsonSerializerOptions options)
+			public override FileReference Read(ref Utf8JsonReader Reader, Type TypeToConvert, JsonSerializerOptions Options)
 			{
 				throw new NotImplementedException();
 			}
 
-			public override void Write(Utf8JsonWriter writer, FileReference value, JsonSerializerOptions options)
+			public override void Write(Utf8JsonWriter Writer, FileReference Value, JsonSerializerOptions Options)
 			{
-				writer.WriteStringValue(value.FullName);
+				Writer.WriteStringValue(Value.FullName);
 			}
 		}
 
@@ -238,7 +246,8 @@ namespace UnrealBuildTool
 		/// <summary>
 		/// Builds a list of actions that need to be executed to produce the specified output items.
 		/// </summary>
-		public static List<LinkedAction> GetActionsToExecute(List<LinkedAction> Actions, CppDependencyCache CppDependencies, ActionHistory History, bool bIgnoreOutdatedImportLibraries)
+		public static List<LinkedAction> GetActionsToExecute(List<LinkedAction> Actions,
+			CppDependencyCache CppDependencies, ActionHistory History, bool bIgnoreOutdatedImportLibraries)
 		{
 			using (Timeline.ScopeEvent("ActionGraph.GetActionsToExecute()"))
 			{
@@ -247,7 +256,7 @@ namespace UnrealBuildTool
 				GatherAllOutdatedActions(Actions, History, OutdatedActionDictionary, CppDependencies, bIgnoreOutdatedImportLibraries);
 
 				// Build a list of actions that are both needed for this target and outdated.
-				return Actions.Where(Action => Action.CommandPath != null && OutdatedActionDictionary[Action]).ToList();
+				return Actions.Where(Action => OutdatedActionDictionary[Action]).ToList();
 			}
 		}
 
@@ -273,13 +282,16 @@ namespace UnrealBuildTool
 							FailPaths.Add(PrerequisiteItem.Location);
 						}
 					}
+
 					foreach (FileItem ProducedItem in Action.ProducedItems)
 					{
 						if (ProducedItem.Location.FullName.Length >= MAX_PATH)
 						{
 							FailPaths.Add(ProducedItem.Location);
 						}
-						if (ProducedItem.Location.FullName.Length > Unreal.RootDirectory.FullName.Length + BuildConfiguration.MaxNestedPathLength && ProducedItem.Location.IsUnderDirectory(Unreal.RootDirectory))
+
+						if (ProducedItem.Location.FullName.Length > Unreal.RootDirectory.FullName.Length +
+						    BuildConfiguration.MaxNestedPathLength && ProducedItem.Location.IsUnderDirectory(Unreal.RootDirectory))
 						{
 							WarnPaths.Add(ProducedItem.Location);
 						}
@@ -289,24 +301,26 @@ namespace UnrealBuildTool
 				if (FailPaths.Count > 0)
 				{
 					StringBuilder Message = new StringBuilder();
-					Message.AppendFormat("The following output paths are longer than {0} characters. Please move the engine to a directory with a shorter path.", MAX_PATH);
+					Message.Append($"The following output paths are longer than {MAX_PATH} characters. Please move the engine to a directory with a shorter path.");
 					foreach (FileReference Path in FailPaths)
 					{
-						Message.AppendFormat("\n[{0} characters] {1}", Path.FullName.Length, Path);
+						Message.Append($"\n[{Path.FullName.Length.ToString()} characters] {Path}");
 					}
+
 					throw new BuildException(Message.ToString());
 				}
 
 				if (WarnPaths.Count > 0)
 				{
 					StringBuilder Message = new StringBuilder();
-					Message.AppendFormat("Detected paths more than {0} characters below UE root directory. This may cause portability issues due to the {1} character maximum path length on Windows:\n", BuildConfiguration.MaxNestedPathLength, MAX_PATH);
+					Message.Append($"Detected paths more than {BuildConfiguration.MaxNestedPathLength.ToString()} characters below UE root directory. This may cause portability issues due to the {MAX_PATH.ToString()} character maximum path length on Windows:\n");
 					foreach (FileReference Path in WarnPaths)
 					{
 						string RelativePath = Path.MakeRelativeTo(Unreal.RootDirectory);
-						Message.AppendFormat("\n[{0} characters] {1}", RelativePath.Length, RelativePath);
+						Message.Append($"\n[{RelativePath.Length.ToString()} characters] {RelativePath}");
 					}
-					Message.AppendFormat("\n\nConsider setting {0} = ... in module *.Build.cs files to use alternative names for intermediate paths.", nameof(ModuleRules.ShortName));
+
+					Message.Append($"\n\nConsider setting {nameof(ModuleRules.ShortName)} = ... in module *.Build.cs files to use alternative names for intermediate paths.");
 					Log.TraceWarning(Message.ToString());
 				}
 			}
@@ -317,7 +331,7 @@ namespace UnrealBuildTool
 		/// </summary>
 		public static void ExecuteActions(BuildConfiguration BuildConfiguration, List<LinkedAction> ActionsToExecute)
 		{
-			if(ActionsToExecute.Count == 0)
+			if (ActionsToExecute.Count == 0)
 			{
 				Log.TraceInformation("Target is up to date");
 			}
@@ -337,7 +351,7 @@ namespace UnrealBuildTool
 				{
 					Executor = new FASTBuild(BuildConfiguration.MaxParallelActions);
 				}
-				else if(BuildConfiguration.bAllowSNDBS && SNDBS.IsAvailable())
+				else if (BuildConfiguration.bAllowSNDBS && SNDBS.IsAvailable())
 				{
 					Executor = new SNDBS();
 				}
@@ -352,16 +366,17 @@ namespace UnrealBuildTool
 
 				// Execute the build
 				Stopwatch Timer = Stopwatch.StartNew();
-				if(!Executor.ExecuteActions(ActionsToExecute))
+				if (!Executor.ExecuteActions(ActionsToExecute))
 				{
 					throw new CompilationResultException(CompilationResult.OtherCompilationError);
 				}
-				Log.TraceInformation("Total time in {0} executor: {1:0.00} seconds", Executor.Name, Timer.Elapsed.TotalSeconds);
+
+				Log.TraceInformation($"Total time in {Executor.Name} executor: {Timer.Elapsed.TotalSeconds:0.00} seconds");
 
 				// Reset the file info for all the produced items
 				foreach (LinkedAction BuildAction in ActionsToExecute)
 				{
-					foreach(FileItem ProducedItem in BuildAction.ProducedItems)
+					foreach (FileItem ProducedItem in BuildAction.ProducedItems)
 					{
 						ProducedItem.ResetCachedInfo();
 					}
@@ -374,9 +389,9 @@ namespace UnrealBuildTool
 					{
 						foreach (FileItem Item in BuildAction.ProducedItems)
 						{
-							if(!Item.Exists)
+							if (!Item.Exists)
 							{
-								throw new BuildException("Failed to produce item: {0}", Item.AbsolutePath);
+								throw new BuildException($"Failed to produce item: {Item.AbsolutePath}");
 							}
 						}
 					}
@@ -390,13 +405,13 @@ namespace UnrealBuildTool
 		static void SortActionList(List<LinkedAction> Actions)
 		{
 			// Clear the current dependent count
-			foreach(LinkedAction Action in Actions)
+			foreach (LinkedAction Action in Actions)
 			{
 				Action.NumTotalDependentActions = 0;
 			}
 
 			// Increment all the dependencies
-			foreach(LinkedAction Action in Actions)
+			foreach (LinkedAction Action in Actions)
 			{
 				Action.IncrementDependentCount(new HashSet<LinkedAction>());
 			}
@@ -436,10 +451,10 @@ namespace UnrealBuildTool
 										CyclicActions.Add(Action, new List<LinkedAction>());
 									}
 
-									List<LinkedAction> CyclicPrereq = CyclicActions[Action];
-									if (!CyclicPrereq.Contains(ProducingAction))
+									List<LinkedAction> CyclicPrerequisite = CyclicActions[Action];
+									if (!CyclicPrerequisite.Contains(ProducingAction))
 									{
-										CyclicPrereq.Add(ProducingAction);
+										CyclicPrerequisite.Add(ProducingAction);
 									}
 								}
 							}
@@ -471,7 +486,7 @@ namespace UnrealBuildTool
 			{
 				// Find the index of each action
 				Dictionary<LinkedAction, int> ActionToIndex = new Dictionary<LinkedAction, int>();
-				for(int Idx = 0; Idx < Actions.Count; Idx++)
+				for (int Idx = 0; Idx < Actions.Count; Idx++)
 				{
 					ActionToIndex[Actions[Idx]] = Idx;
 				}
@@ -482,49 +497,54 @@ namespace UnrealBuildTool
 				{
 					if (!ActionIsNonCyclical.ContainsKey(Action))
 					{
-						CycleDescription += string.Format("Action #{0}: {1}\n", ActionToIndex[Action], Action.CommandPath);
-						CycleDescription += string.Format("\twith arguments: {0}\n", Action.CommandArguments);
+						CycleDescription += $"Action #{ActionToIndex[Action].ToString()}: {Action.CommandPath}\n";
+						CycleDescription += $"\twith arguments: {Action.CommandArguments}\n";
 						foreach (FileItem PrerequisiteItem in Action.PrerequisiteItems)
 						{
-							CycleDescription += string.Format("\tdepends on: {0}\n", PrerequisiteItem.AbsolutePath);
+							CycleDescription += $"\tdepends on: {PrerequisiteItem.AbsolutePath}\n";
 						}
+
 						foreach (FileItem ProducedItem in Action.ProducedItems)
 						{
-							CycleDescription += string.Format("\tproduces:   {0}\n", ProducedItem.AbsolutePath);
+							CycleDescription += $"\tproduces:   {ProducedItem.AbsolutePath}\n";
 						}
-						CycleDescription += string.Format("\tDepends on cyclic actions:\n");
+
+						CycleDescription += "\tDepends on cyclic actions:\n";
 						if (CyclicActions.ContainsKey(Action))
 						{
 							foreach (LinkedAction CyclicPrerequisiteAction in CyclicActions[Action])
 							{
 								if (CyclicActions.ContainsKey(CyclicPrerequisiteAction))
 								{
-									List<FileItem> CyclicProducedItems = CyclicPrerequisiteAction.ProducedItems.ToList();
+									List<FileItem> CyclicProducedItems =
+										CyclicPrerequisiteAction.ProducedItems.ToList();
 									if (CyclicProducedItems.Count == 1)
 									{
-										CycleDescription += string.Format("\t\t{0} (produces: {1})\n", ActionToIndex[CyclicPrerequisiteAction], CyclicProducedItems[0].AbsolutePath);
+										CycleDescription += $"\t\t{ActionToIndex[CyclicPrerequisiteAction].ToString()} (produces: {CyclicProducedItems[0].AbsolutePath})\n";
 									}
 									else
 									{
-										CycleDescription += string.Format("\t\t{0}\n", ActionToIndex[CyclicPrerequisiteAction]);
+										CycleDescription += $"\t\t{ActionToIndex[CyclicPrerequisiteAction].ToString()}\n";
 										foreach (FileItem CyclicProducedItem in CyclicProducedItems)
 										{
-											CycleDescription += string.Format("\t\t\tproduces:   {0}\n", CyclicProducedItem.AbsolutePath);
+											CycleDescription += $"\t\t\tproduces:   {CyclicProducedItem.AbsolutePath}\n";
 										}
 									}
 								}
 							}
+
 							CycleDescription += "\n";
 						}
 						else
 						{
-							CycleDescription += string.Format("\t\tNone?? Coding error!\n");
+							CycleDescription += "\t\tNone?? Coding error!\n";
 						}
+
 						CycleDescription += "\n\n";
 					}
 				}
 
-				throw new BuildException("Action graph contains cycle!\n\n{0}", CycleDescription);
+				throw new BuildException($"Action graph contains cycle!\n\n{CycleDescription}");
 			}
 		}
 
@@ -537,13 +557,14 @@ namespace UnrealBuildTool
 		public static List<LinkedAction> GatherPrerequisiteActions(List<LinkedAction> Actions, HashSet<FileItem> OutputItems)
 		{
 			HashSet<LinkedAction> PrerequisiteActions = new HashSet<LinkedAction>();
-			foreach(LinkedAction Action in Actions)
+			foreach (LinkedAction Action in Actions)
 			{
-				if(Action.ProducedItems.Any(x => OutputItems.Contains(x)))
+				if (Action.ProducedItems.Any(OutputItems.Contains))
 				{
 					GatherPrerequisiteActions(Action, PrerequisiteActions);
 				}
 			}
+
 			return PrerequisiteActions.ToList();
 		}
 
@@ -554,9 +575,9 @@ namespace UnrealBuildTool
 		/// <param name="PrerequisiteActions">Set of prerequisite actions</param>
 		private static void GatherPrerequisiteActions(LinkedAction Action, HashSet<LinkedAction> PrerequisiteActions)
 		{
-			if(PrerequisiteActions.Add(Action))
+			if (PrerequisiteActions.Add(Action))
 			{
-				foreach(LinkedAction PrerequisiteAction in Action.PrerequisiteActions)
+				foreach (LinkedAction PrerequisiteAction in Action.PrerequisiteActions)
 				{
 					GatherPrerequisiteActions(PrerequisiteAction, PrerequisiteActions);
 				}
@@ -576,7 +597,9 @@ namespace UnrealBuildTool
 		/// <param name="CppDependencies"></param>
 		/// <param name="bIgnoreOutdatedImportLibraries"></param>
 		/// <returns>true if outdated</returns>
-		public static void IsIndividualActionOutdated(LinkedAction RootAction, Dictionary<LinkedAction, bool> OutdatedActionDictionary, ReaderWriterLockSlim OutdatedActionLock, ActionHistory ActionHistory, CppDependencyCache CppDependencies, bool bIgnoreOutdatedImportLibraries)
+		private static void IsIndividualActionOutdated(LinkedAction RootAction,
+			Dictionary<LinkedAction, bool> OutdatedActionDictionary, ReaderWriterLockSlim OutdatedActionLock,
+			ActionHistory ActionHistory, CppDependencyCache CppDependencies, bool bIgnoreOutdatedImportLibraries)
 		{
 			// Only compute the outdated-ness for actions that don't aren't cached in the outdated action dictionary.
 			bool bIsOutdated = false;
@@ -596,17 +619,13 @@ namespace UnrealBuildTool
 			foreach (FileItem ProducedItem in RootAction.ProducedItems)
 			{
 				// Check if the command-line of the action previously used to produce the item is outdated.
-				string NewProducingAttributes = string.Format("{0} {1} (ver {2})", RootAction.CommandPath.FullName, RootAction.CommandArguments, RootAction.CommandVersion);
+				string NewProducingAttributes = $"{RootAction.CommandPath.FullName} {RootAction.CommandArguments} (ver {RootAction.CommandVersion})";
 				if (ActionHistory.UpdateProducingAttributes(ProducedItem, NewProducingAttributes))
 				{
-					if(ProducedItem.Exists)
+					if (ProducedItem.Exists)
 					{
-						Log.TraceLog(
-							"{0}: Produced item \"{1}\" was produced by outdated attributes.\n  New attributes: {2}",
-							RootAction.StatusDescription,
-							ProducedItem.AbsolutePath,
-							NewProducingAttributes
-							);
+						Log.TraceLog($"{RootAction.StatusDescription}: Produced item \"{ProducedItem.AbsolutePath}\" was produced by outdated attributes.");
+						Log.TraceLog($"  New attributes: {NewProducingAttributes}");
 					}
 
 					bIsOutdated = true;
@@ -615,7 +634,8 @@ namespace UnrealBuildTool
 				// If the produced file doesn't exist or has zero size, consider it outdated.  The zero size check is to detect cases
 				// where aborting an earlier compile produced invalid zero-sized obj files, but that may cause actions where that's
 				// legitimate output to always be considered outdated.
-				if (ProducedItem.Exists && (RootAction.ActionType != ActionType.Compile || ProducedItem.Length > 0 || (!ProducedItem.Location.HasExtension(".obj") && !ProducedItem.Location.HasExtension(".o"))))
+				if (ProducedItem.Exists && (RootAction.ActionType != ActionType.Compile || ProducedItem.Length > 0 ||
+				                            (!ProducedItem.Location.HasExtension(".obj") && !ProducedItem.Location.HasExtension(".o"))))
 				{
 					// Use the oldest produced item's time as the last execution time.
 					if (ProducedItem.LastWriteTimeUtc < LastExecutionTimeUtc)
@@ -626,17 +646,13 @@ namespace UnrealBuildTool
 				else
 				{
 					// If any of the produced items doesn't exist, the action is outdated.
-					Log.TraceLog(
-						"{0}: Produced item \"{1}\" doesn't exist.",
-						RootAction.StatusDescription,
-						ProducedItem.AbsolutePath
-						);
+					Log.TraceLog($"{RootAction.StatusDescription}: Produced item \"{ProducedItem.AbsolutePath}\" doesn't exist.");
 					bIsOutdated = true;
 				}
 			}
 
 			// Check if any prerequisite item has a newer timestamp than the last execution time of this action
-			if(!bIsOutdated)
+			if (!bIsOutdated)
 			{
 				foreach (FileItem PrerequisiteItem in RootAction.PrerequisiteItems)
 				{
@@ -648,9 +664,10 @@ namespace UnrealBuildTool
 						if (bPrerequisiteItemIsNewerThanLastExecution)
 						{
 							// Need to check for import libraries here too
-							if(!bIgnoreOutdatedImportLibraries || !IsImportLibraryDependency(RootAction, PrerequisiteItem))
+							if (!bIgnoreOutdatedImportLibraries || !IsImportLibraryDependency(RootAction, PrerequisiteItem))
 							{
-								Log.TraceLog("{0}: Prerequisite {1} is newer than the last execution of the action: {2} vs {3}", RootAction.StatusDescription, PrerequisiteItem.AbsolutePath, PrerequisiteItem.LastWriteTimeUtc.ToLocalTime(), LastExecutionTimeUtc.LocalDateTime);
+								Log.TraceLog($"{RootAction.StatusDescription}: Prerequisite {PrerequisiteItem.AbsolutePath} is newer than the last execution of the action: " +
+									$"{PrerequisiteItem.LastWriteTimeUtc.ToLocalTime().ToString(CultureInfo.CurrentCulture)} vs {LastExecutionTimeUtc.LocalDateTime.ToString(CultureInfo.CurrentCulture)}");
 								bIsOutdated = true;
 								break;
 							}
@@ -660,11 +677,11 @@ namespace UnrealBuildTool
 			}
 
 			// Check the dependency list
-			if(!bIsOutdated && RootAction.DependencyListFile != null)
+			if (!bIsOutdated && RootAction.DependencyListFile != null)
 			{
 				if (!CppDependencies.TryGetDependencies(RootAction.DependencyListFile, out List<FileItem>? DependencyFiles))
 				{
-					Log.TraceLog("{0}: Missing dependency list file \"{1}\"", RootAction.StatusDescription, RootAction.DependencyListFile);
+					Log.TraceLog($"{RootAction.StatusDescription}: Missing dependency list file \"{RootAction.DependencyListFile}\"");
 					bIsOutdated = true;
 				}
 				else
@@ -673,13 +690,8 @@ namespace UnrealBuildTool
 					{
 						if (!DependencyFile.Exists || DependencyFile.LastWriteTimeUtc > LastExecutionTimeUtc)
 						{
-							Log.TraceLog(
-								"{0}: Dependency {1} is newer than the last execution of the action: {2} vs {3}",
-								RootAction.StatusDescription,
-								DependencyFile.AbsolutePath,
-								DependencyFile.LastWriteTimeUtc.ToLocalTime(),
-								LastExecutionTimeUtc.LocalDateTime
-								);
+							Log.TraceLog($"{RootAction.StatusDescription}: Dependency {DependencyFile.AbsolutePath} is newer than the last execution of the action:" +
+								$"{DependencyFile.LastWriteTimeUtc.ToLocalTime().ToString(CultureInfo.CurrentCulture)} vs {LastExecutionTimeUtc.LocalDateTime.ToString(CultureInfo.CurrentCulture)}");
 							bIsOutdated = true;
 							break;
 						}
@@ -705,7 +717,7 @@ namespace UnrealBuildTool
 		/// <param name="OutdatedActionDictionary">-</param>
 		/// <param name="bIgnoreOutdatedImportLibraries"></param>
 		/// <returns>true if outdated</returns>
-		public static bool IsActionOutdatedDueToPrerequisites(LinkedAction RootAction, Dictionary<LinkedAction, bool> OutdatedActionDictionary, bool bIgnoreOutdatedImportLibraries)
+		private static bool IsActionOutdatedDueToPrerequisites(LinkedAction RootAction,	Dictionary<LinkedAction, bool> OutdatedActionDictionary, bool bIgnoreOutdatedImportLibraries)
 		{
 			// Only compute the outdated-ness for actions that aren't already cached in the outdated action dictionary.
 			if (OutdatedActionDictionary.TryGetValue(RootAction, out bool bIsOutdated))
@@ -714,25 +726,22 @@ namespace UnrealBuildTool
 			}
 
 			// Check if any of the prerequisite actions are out of date
-			if (!bIsOutdated)
+			foreach (LinkedAction PrerequisiteAction in RootAction.PrerequisiteActions)
 			{
-				foreach (LinkedAction PrerequisiteAction in RootAction.PrerequisiteActions)
+				if (IsActionOutdatedDueToPrerequisites(PrerequisiteAction, OutdatedActionDictionary, bIgnoreOutdatedImportLibraries))
 				{
-					if (IsActionOutdatedDueToPrerequisites(PrerequisiteAction, OutdatedActionDictionary, bIgnoreOutdatedImportLibraries))
+					// Only check for outdated import libraries if we were configured to do so.  Often, a changed import library
+					// won't affect a dependency unless a public header file was also changed, in which case we would be forced
+					// to recompile anyway.  This just allows for faster iteration when working on a subsystem in a DLL, as we
+					// won't have to wait for dependent targets to be relinked after each change.
+					if (!bIgnoreOutdatedImportLibraries || !IsImportLibraryDependency(RootAction, PrerequisiteAction))
 					{
-						// Only check for outdated import libraries if we were configured to do so.  Often, a changed import library
-						// won't affect a dependency unless a public header file was also changed, in which case we would be forced
-						// to recompile anyway.  This just allows for faster iteration when working on a subsystem in a DLL, as we
-						// won't have to wait for dependent targets to be relinked after each change.
-						if(!bIgnoreOutdatedImportLibraries || !IsImportLibraryDependency(RootAction, PrerequisiteAction))
-						{
-							Log.TraceLog("{0}: Prerequisite {1} is produced by outdated action.", RootAction.StatusDescription, PrerequisiteAction.StatusDescription);
-							bIsOutdated = true;
-							break;
-						}
+						Log.TraceLog($"{RootAction.StatusDescription}: Prerequisite {PrerequisiteAction.StatusDescription} is produced by outdated action.");
+						bIsOutdated = true;
+						break;
 					}
 				}
-			} 
+			}
 
 			// Cache the outdated-ness of this action.
 			OutdatedActionDictionary.Add(RootAction, bIsOutdated);
@@ -748,9 +757,9 @@ namespace UnrealBuildTool
 		/// <returns>True if the only dependency between two actions is for an import library</returns>
 		static bool IsImportLibraryDependency(LinkedAction RootAction, LinkedAction PrerequisiteAction)
 		{
-			if(PrerequisiteAction.bProducesImportLibrary)
+			if (PrerequisiteAction.bProducesImportLibrary)
 			{
-				return PrerequisiteAction.ProducedItems.All(x => x.Location.HasExtension(".lib") || !RootAction.PrerequisiteItems.Contains(x));
+				return PrerequisiteAction.ProducedItems.All(I => I.Location.HasExtension(".lib") || !RootAction.PrerequisiteItems.Contains(I));
 			}
 			else
 			{
@@ -766,16 +775,17 @@ namespace UnrealBuildTool
 		/// <returns>True if the only dependency between two actions is for an import library</returns>
 		static bool IsImportLibraryDependency(LinkedAction RootAction, FileItem PrerequisiteItem)
 		{
-			if(PrerequisiteItem.Location.HasExtension(".lib"))
+			if (PrerequisiteItem.Location.HasExtension(".lib"))
 			{
-				foreach(LinkedAction PrerequisiteAction in RootAction.PrerequisiteActions)
+				foreach (LinkedAction PrerequisiteAction in RootAction.PrerequisiteActions)
 				{
-					if(PrerequisiteAction.bProducesImportLibrary && PrerequisiteAction.ProducedItems.Contains(PrerequisiteItem))
+					if (PrerequisiteAction.bProducesImportLibrary && PrerequisiteAction.ProducedItems.Contains(PrerequisiteItem))
 					{
 						return true;
 					}
 				}
 			}
+
 			return false;
 		}
 
@@ -783,29 +793,34 @@ namespace UnrealBuildTool
 		/// Builds a dictionary containing the actions from AllActions that are outdated by calling
 		/// IsActionOutdated.
 		/// </summary>
-		public static void GatherAllOutdatedActions(IEnumerable<LinkedAction> Actions, ActionHistory ActionHistory, Dictionary<LinkedAction, bool> OutdatedActions, CppDependencyCache CppDependencies, bool bIgnoreOutdatedImportLibraries)
+		public static void GatherAllOutdatedActions(IReadOnlyList<LinkedAction> Actions, ActionHistory ActionHistory,
+			Dictionary<LinkedAction, bool> OutdatedActions, CppDependencyCache CppDependencies,
+			bool bIgnoreOutdatedImportLibraries)
 		{
-			using(Timeline.ScopeEvent("Prefetching include dependencies"))
+			using (Timeline.ScopeEvent("Prefetching include dependencies"))
 			{
 				List<FileItem> Dependencies = new List<FileItem>();
-				foreach(LinkedAction Action in Actions)
+				foreach (LinkedAction Action in Actions)
 				{
-					if(Action.DependencyListFile != null)
+					if (Action.DependencyListFile != null)
 					{
 						Dependencies.Add(Action.DependencyListFile);
 					}
 				}
+
 				Parallel.ForEach(Dependencies, File => { CppDependencies.TryGetDependencies(File, out _); });
 			}
 
 			using (Timeline.ScopeEvent("Cache individual outdated actions"))
 			{
 				ReaderWriterLockSlim OutdatedActionsLock = new ReaderWriterLockSlim();
-				Parallel.ForEach(Actions, Action => IsIndividualActionOutdated(Action, OutdatedActions, OutdatedActionsLock, ActionHistory, CppDependencies, bIgnoreOutdatedImportLibraries));
+				Parallel.ForEach(Actions,
+					Action => IsIndividualActionOutdated(Action, OutdatedActions, OutdatedActionsLock, 
+						ActionHistory, CppDependencies, bIgnoreOutdatedImportLibraries));
 			}
 
 			using (Timeline.ScopeEvent("Cache outdated actions based on recursive prerequisites"))
-			{ 
+			{
 				foreach (var Action in Actions)
 				{
 					IsActionOutdatedDueToPrerequisites(Action, OutdatedActions, bIgnoreOutdatedImportLibraries);
@@ -819,13 +834,13 @@ namespace UnrealBuildTool
 		/// <param name="OutdatedActions">List of outdated actions</param>
 		public static void DeleteOutdatedProducedItems(List<LinkedAction> OutdatedActions)
 		{
-			foreach(LinkedAction OutdatedAction in OutdatedActions)
+			foreach (LinkedAction OutdatedAction in OutdatedActions)
 			{
 				foreach (FileItem DeleteItem in OutdatedAction.DeleteItems)
 				{
 					if (DeleteItem.Exists)
 					{
-						Log.TraceLog("Deleting outdated item: {0}", DeleteItem.AbsolutePath);
+						Log.TraceLog($"Deleting outdated item: {DeleteItem.AbsolutePath}");
 						DeleteItem.Delete();
 					}
 				}
@@ -839,16 +854,17 @@ namespace UnrealBuildTool
 		public static void CreateDirectoriesForProducedItems(List<LinkedAction> OutdatedActions)
 		{
 			HashSet<DirectoryReference> OutputDirectories = new HashSet<DirectoryReference>();
-			foreach(LinkedAction OutdatedAction in OutdatedActions)
+			foreach (LinkedAction OutdatedAction in OutdatedActions)
 			{
-				foreach(FileItem ProducedItem in OutdatedAction.ProducedItems)
+				foreach (FileItem ProducedItem in OutdatedAction.ProducedItems)
 				{
 					OutputDirectories.Add(ProducedItem.Location.Directory);
 				}
 			}
-			foreach(DirectoryReference OutputDirectory in OutputDirectories)
+
+			foreach (DirectoryReference OutputDirectory in OutputDirectories)
 			{
-				if(!DirectoryReference.Exists(OutputDirectory))
+				if (!DirectoryReference.Exists(OutputDirectory))
 				{
 					DirectoryReference.CreateDirectory(OutputDirectory);
 				}
@@ -865,7 +881,7 @@ namespace UnrealBuildTool
 			JsonObject Object = JsonObject.Read(InputFile);
 
 			JsonObject EnvironmentObject = Object.GetObjectField("Environment");
-			foreach(string KeyName in EnvironmentObject.KeyNames)
+			foreach (string KeyName in EnvironmentObject.KeyNames)
 			{
 				Environment.SetEnvironmentVariable(KeyName, EnvironmentObject.GetStringField(KeyName));
 			}
@@ -875,6 +891,7 @@ namespace UnrealBuildTool
 			{
 				Actions.Add(Action.ImportJson(ActionObject));
 			}
+
 			return Actions;
 		}
 
@@ -883,7 +900,7 @@ namespace UnrealBuildTool
 		/// </summary>
 		/// <param name="Actions">The actions to write</param>
 		/// <param name="OutputFile">Output file to write the actions to</param>
-		public static void ExportJson(IEnumerable<LinkedAction> Actions, FileReference OutputFile)
+		public static void ExportJson(IReadOnlyList<LinkedAction> Actions, FileReference OutputFile)
 		{
 			DirectoryReference.CreateDirectory(OutputFile.Directory);
 			using JsonWriter Writer = new JsonWriter(OutputFile);
@@ -893,11 +910,13 @@ namespace UnrealBuildTool
 			foreach (object? Object in Environment.GetEnvironmentVariables())
 			{
 				System.Collections.DictionaryEntry Pair = (System.Collections.DictionaryEntry)Object!;
-				if (!UnrealBuildTool.InitialEnvironment!.Contains(Pair.Key) || (string)(UnrealBuildTool.InitialEnvironment![Pair.Key]!) != (string)(Pair.Value!))
+				if (!UnrealBuildTool.InitialEnvironment!.Contains(Pair.Key) ||
+				    (string)(UnrealBuildTool.InitialEnvironment[Pair.Key]!) != (string)(Pair.Value!))
 				{
-					Writer.WriteValue((string)Pair!.Key, (string)Pair.Value!);
+					Writer.WriteValue((string)Pair.Key, (string)Pair.Value!);
 				}
 			}
+
 			Writer.WriteObjectEnd();
 
 			Dictionary<LinkedAction, int> ActionToId = new Dictionary<LinkedAction, int>();
@@ -913,6 +932,7 @@ namespace UnrealBuildTool
 				Action.ExportJson(ActionToId, Writer);
 				Writer.WriteObjectEnd();
 			}
+
 			Writer.WriteArrayEnd();
 			Writer.WriteObjectEnd();
 		}
