Provide API to override output format of ISPC for platforms that should use an option other than --emit-obj

diff --git a/Engine/Source/Programs/UnrealBuildTool/ToolChain/UEToolChain.cs b/Engine/Source/Programs/UnrealBuildTool/ToolChain/UEToolChain.cs
--- a/Engine/Source/Programs/UnrealBuildTool/ToolChain/UEToolChain.cs
+++ b/Engine/Source/Programs/UnrealBuildTool/ToolChain/UEToolChain.cs
@@ -344,6 +344,33 @@ namespace UnrealBuildTool
 			return ISPCCompilerVersions[Platform];			
 		}
 
+		/// <summary>
+		/// Get object file format for ISPC.
+		/// </summary>
+		/// <param name="Platform">Which OS build platform is running on.</param>
+		/// <returns>Object file suffix</returns>
+		public virtual string GetISPCObjectFileFormat(UnrealTargetPlatform Platform)
+		{
+			string Format = "";
+
+			if (UEBuildPlatform.IsPlatformInGroup(Platform, UnrealPlatformGroup.Windows))
+			{
+				Format += "obj";
+			}
+			else if (UEBuildPlatform.IsPlatformInGroup(Platform, UnrealPlatformGroup.Unix) ||
+					Platform == UnrealTargetPlatform.Mac ||
+					Platform == UnrealTargetPlatform.Android)
+			{
+				Format += "obj";
+			}
+			else
+			{
+				Log.TraceWarning("Unsupported ISPC platform target!");
+			}
+
+			return Format;
+		}
+
 		/// <summary>
 		/// Get object file suffix for ISPC.
 		/// </summary>
@@ -427,6 +454,7 @@ namespace UnrealBuildTool
 				Arguments.Add(String.Format("--target-os={0}", GetISPCOSTarget(CompileEnvironment.Platform)));
 				Arguments.Add(String.Format("--arch={0}", GetISPCArchTarget(CompileEnvironment.Platform, null)));
 				Arguments.Add(String.Format("--target={0}", TargetString));
+				Arguments.Add(String.Format("--emit-{0}", GetISPCObjectFileFormat(CompileEnvironment.Platform)));
 
 				// PIC is needed for modular builds except on Windows
 				if ((CompileEnvironment.bIsBuildingDLL ||
@@ -600,6 +628,7 @@ namespace UnrealBuildTool
 				Arguments.Add(String.Format("--target-os=\"{0}\"", GetISPCOSTarget(CompileEnvironment.Platform)));
 				Arguments.Add(String.Format("--arch=\"{0}\"", GetISPCArchTarget(CompileEnvironment.Platform, null)));
 				Arguments.Add(String.Format("--target=\"{0}\"", TargetString));
+				Arguments.Add(String.Format("--emit-{0}", GetISPCObjectFileFormat(CompileEnvironment.Platform)));
 
 				if (CompileEnvironment.Configuration == CppConfiguration.Debug)
 				{
