UnrealBuildTool: Use --no-optional-locks when when running git status

Prevents interference with interactive git commands that may be running at the same time

diff --git a/Engine/Source/Programs/UnrealBuildTool/System/SourceFileWorkingSet.cs b/Engine/Source/Programs/UnrealBuildTool/System/SourceFileWorkingSet.cs
--- a/Engine/Source/Programs/UnrealBuildTool/System/SourceFileWorkingSet.cs
+++ b/Engine/Source/Programs/UnrealBuildTool/System/SourceFileWorkingSet.cs
@@ -116,7 +116,7 @@ namespace UnrealBuildTool
 
 			BackgroundProcess = new Process();
 			BackgroundProcess.StartInfo.FileName = GitPath;
-			BackgroundProcess.StartInfo.Arguments = "status --porcelain";
+			BackgroundProcess.StartInfo.Arguments = "--no-optional-locks status --porcelain";
 			BackgroundProcess.StartInfo.WorkingDirectory = RootDir.FullName;
 			BackgroundProcess.StartInfo.RedirectStandardOutput = true;
 			BackgroundProcess.StartInfo.RedirectStandardError = true;
