UnrealBuildTool: Change VS version detection a little

Check the major version number to simplify matching.

diff --git a/Engine/Source/Programs/UnrealBuildTool/Platform/Windows/UEBuildWindows.cs b/Engine/Source/Programs/UnrealBuildTool/Platform/Windows/UEBuildWindows.cs
--- a/Engine/Source/Programs/UnrealBuildTool/Platform/Windows/UEBuildWindows.cs
+++ b/Engine/Source/Programs/UnrealBuildTool/Platform/Windows/UEBuildWindows.cs
@@ -1110,79 +1110,90 @@ namespace UnrealBuildTool
 		public static List<VisualStudioInstallation> FindVisualStudioInstallations(WindowsCompiler Compiler)
 		{
 			List<VisualStudioInstallation>? Installations;
-			if(!CachedVisualStudioInstallations.TryGetValue(Compiler, out Installations))
+			if (CachedVisualStudioInstallations.TryGetValue(Compiler, out Installations))
 			{
-				Installations = new List<VisualStudioInstallation>();
-			    if (BuildHostPlatform.Current.Platform == UnrealTargetPlatform.Win64)
-			    {
-				    if(Compiler >= WindowsCompiler.VisualStudio2019)
-				    {
-						try
-						{
-							SetupConfiguration Setup = new SetupConfiguration();
-							IEnumSetupInstances Enumerator = Setup.EnumAllInstances();
+				return Installations;
+			}
 
-							ISetupInstance[] Instances = new ISetupInstance[1];
-							for(;;)
-							{
-								int NumFetched;
-								Enumerator.Next(1, Instances, out NumFetched);
+			if (Compiler != WindowsCompiler.VisualStudio2019 && Compiler != WindowsCompiler.VisualStudio2022)
+			{
+				throw new BuildException($"Unsupported compiler version ({Compiler})");
+			}
 
-								if(NumFetched == 0)
-								{
-									break;
-								}
+			Installations = new List<VisualStudioInstallation>();
 
-								ISetupInstance2 Instance = (ISetupInstance2)Instances[0];
-								if((Instance.GetState() & InstanceState.Local) == InstanceState.Local)
-								{
-									string VersionString = Instance.GetInstallationVersion();
-
-									VersionNumber? Version;
-									if (VersionNumber.TryParse(VersionString, out Version))
-									{
-										VersionNumber Version2019 = new VersionNumber(16);
-										VersionNumber Version2022 = new VersionNumber(17);
-										if (Compiler == WindowsCompiler.VisualStudio2022 && Version < Version2022)
-										{
-											continue;
-										}
-										else if (Compiler == WindowsCompiler.VisualStudio2019 && Version < Version2019)
-										{
-											continue;
-										}
-									}
-
-									ISetupInstanceCatalog? Catalog = Instance as ISetupInstanceCatalog;
-
-									VisualStudioInstallation Installation = new VisualStudioInstallation(new DirectoryReference(Instance.GetInstallationPath()));
-									if (Catalog != null && Catalog.IsPrerelease())
-									{
-										Installation.bPreview = true;
-									}
-
-									string ProductId = Instance.GetProduct().GetId();
-									if (ProductId.Equals("Microsoft.VisualStudio.Product.WDExpress", StringComparison.Ordinal))
-									{
-										Installation.bExpress = true;
-									}
-
-									Log.TraceLog("Found Visual Studio installation: {0} (Product={1}, Version={2})", Instance.GetInstallationPath(), ProductId, VersionString);
-									Installations.Add(Installation);
-								}
-							}
-						}
-						catch
-						{
-						}
+			if (BuildHostPlatform.Current.Platform != UnrealTargetPlatform.Win64)
+			{
+				CachedVisualStudioInstallations.Add(Compiler, Installations);
+				return Installations;
+			}
+
+			try
+			{
+				SetupConfiguration Setup = new SetupConfiguration();
+				IEnumSetupInstances Enumerator = Setup.EnumAllInstances();
+				ISetupInstance[] Instances = new ISetupInstance[1];
+				for (;;)
+				{
+					Enumerator.Next(1, Instances, out int NumFetched);
+
+					if (NumFetched == 0)
+					{
+						break;
 					}
-					else
-				    {
-					    throw new BuildException("Unsupported compiler version ({0})", Compiler);
-				    }
+
+					ISetupInstance2 Instance = (ISetupInstance2)Instances[0];
+					if ((Instance.GetState() & InstanceState.Local) != InstanceState.Local)
+					{
+						continue;
+					}
+
+					string VersionString = Instance.GetInstallationVersion();
+					string ProductId = Instance.GetProduct().GetId();
+
+					VersionNumber? Version;
+					if (!VersionNumber.TryParse(VersionString, out Version))
+					{
+						Log.TraceLog(
+							$"Skipping Visual Studio instance - could not parse version string: {Instance.GetInstallationPath()} (Product={ProductId}, Version={VersionString})");
+						continue;
+					}
+
+					int MajorVersionNumber = Version.GetComponent(0);
+
+					if (Compiler == WindowsCompiler.VisualStudio2019 && MajorVersionNumber != 16 ||
+					    Compiler == WindowsCompiler.VisualStudio2022 && MajorVersionNumber != 17)
+					{
+						Log.TraceLog(
+							$"Skipping Visual Studio instance - does not match requested compiler version: {Instance.GetInstallationPath()} (Product={ProductId}, Version={VersionString})");
+						continue;
+					}
+
+					ISetupInstanceCatalog? Catalog = Instance as ISetupInstanceCatalog;
+
+					VisualStudioInstallation Installation =
+						new VisualStudioInstallation(new DirectoryReference(Instance.GetInstallationPath()));
+					if (Catalog != null && Catalog.IsPrerelease())
+					{
+						Installation.bPreview = true;
+					}
+
+					if (ProductId.Equals("Microsoft.VisualStudio.Product.WDExpress", StringComparison.Ordinal))
+					{
+						Installation.bExpress = true;
+					}
+
+					Log.TraceLog(
+						$"Found Visual Studio installation: {Instance.GetInstallationPath()} (Product={ProductId}, Version={VersionString})");
+					Installations.Add(Installation);
 				}
-				CachedVisualStudioInstallations.Add(Compiler, Installations.OrderBy(x => !x.bExpress).ThenBy(x => !x.bPreview).ToList());
 			}
+			catch
+			{
+			}
+
+			CachedVisualStudioInstallations.Add(Compiler,
+				Installations.OrderBy(x => !x.bExpress).ThenBy(x => !x.bPreview).ToList());
 			return Installations;
 		}
 
@@ -1613,7 +1624,7 @@ namespace UnrealBuildTool
 			foreach (DirectoryReference InstallDir in InstallDirs2022)
 			{
 				FileReference MsBuildLocation = FileReference.Combine(InstallDir, "MSBuild", "Current", "Bin", "MSBuild.exe");
-				if (FileReference.Exists(MsBuildLocation))
+				if(FileReference.Exists(MsBuildLocation))
 				{
 					OutLocation = MsBuildLocation;
 					return true;
