UnrealBuildTool: decompose Pair in TargetMakefile.IsValidForSourceFiles()

#trivial

diff --git a/Engine/Source/Programs/UnrealBuildTool/System/TargetMakefile.cs b/Engine/Source/Programs/UnrealBuildTool/System/TargetMakefile.cs
--- a/Engine/Source/Programs/UnrealBuildTool/System/TargetMakefile.cs
+++ b/Engine/Source/Programs/UnrealBuildTool/System/TargetMakefile.cs
@@ -453,23 +453,22 @@ namespace UnrealBuildTool
 				ReadOnlyHashSet<string> ExcludedFolderNames = UEBuildPlatform.GetBuildPlatform(Platform).GetExcludedFolderNames();
 
 				// Check if any source files have been added or removed
-				foreach(KeyValuePair<DirectoryItem, FileItem[]> Pair in Makefile.DirectoryToSourceFiles)
+				foreach((DirectoryItem InputDirectory, FileItem[] Files) in Makefile.DirectoryToSourceFiles)
 				{
-					DirectoryItem InputDirectory = Pair.Key;
 					if(!InputDirectory.Exists || InputDirectory.LastWriteTimeUtc > Makefile.CreateTimeUtc)
 					{
 						FileItem[] SourceFiles = UEBuildModuleCPP.GetSourceFiles(InputDirectory);
-						if(SourceFiles.Length < Pair.Value.Length)
+						if(SourceFiles.Length < Files.Length)
 						{
 							ReasonNotLoaded = "source file removed";
 							return false;
 						}
-						else if(SourceFiles.Length > Pair.Value.Length)
+						else if(SourceFiles.Length > Files.Length)
 						{
 							ReasonNotLoaded = "source file added";
 							return false;
 						}
-						else if(SourceFiles.Intersect(Pair.Value).Count() != SourceFiles.Length)
+						else if(SourceFiles.Intersect(Files).Count() != SourceFiles.Length)
 						{
 							ReasonNotLoaded = "source file modified";
 							return false;
