UnrealBuildTool: add option bShowPerActionCompilationTimes to Parallel & Task executors

Disabled by default.

When enabled, adds wall time and CPU time numbers for each action e.g.
  [543/665] (Wall: 6.33s CPU: 8.19s) Compile AssetBundleData.cpp
  [544/665] (Wall: 3.77s CPU: 4.03s) Compile SavePackage2.cpp
  [545/665] (Wall: 3.91s CPU: 4.56s) Compile WeakObjectPtr.cpp

diff --git a/Engine/Source/Programs/UnrealBuildTool/Executors/Experimental/TaskExecutor.cs b/Engine/Source/Programs/UnrealBuildTool/Executors/Experimental/TaskExecutor.cs
--- a/Engine/Source/Programs/UnrealBuildTool/Executors/Experimental/TaskExecutor.cs
+++ b/Engine/Source/Programs/UnrealBuildTool/Executors/Experimental/TaskExecutor.cs
@@ -36,19 +36,25 @@ namespace UnrealBuildTool
 		/// Set to 0 to disable free memory checking.
 		/// </summary>
 		[XmlConfigFile]
-		private static readonly double MemoryPerActionBytes = 1.5 * 1024 * 1024 * 1024;
+		private static double MemoryPerActionBytes = 1.5 * 1024 * 1024 * 1024;
 
 		/// <summary>
 		/// When enabled, will stop compiling targets after a compile error occurs.
 		/// </summary>
 		[XmlConfigFile]
-		private static readonly bool bStopCompilationAfterErrors = false;
+		private static bool bStopCompilationAfterErrors = false;
 
 		/// <summary>
 		/// Whether to show compilation times along with worst offenders or not.
 		/// </summary>
 		[XmlConfigFile]
-		private static readonly bool bShowCompilationTimes = false;
+		private static bool bShowCompilationTimes = false;
+
+		/// <summary>
+		/// Whether to show compilation times for each executed action
+		/// </summary>
+		[XmlConfigFile]
+		private static bool bShowPerActionCompilationTimes = false;
 
 		/// <summary>
 		/// How many processes that will be executed in parallel
@@ -262,10 +268,15 @@ namespace UnrealBuildTool
 		{
 			List<string> LogLines = new List<string>();
 			int ExitCode = int.MaxValue;
+			TimeSpan ExecutionTime = TimeSpan.Zero;
+			TimeSpan ProcessorTime = TimeSpan.Zero;
 			if (ExecuteTask.Status == TaskStatus.RanToCompletion)
 			{
-				LogLines = ExecuteTask.Result.LogLines;
-				ExitCode = ExecuteTask.Result.ExitCode;
+				ExecuteResults ExecuteTaskResult = ExecuteTask.Result;
+				LogLines = ExecuteTaskResult.LogLines;
+				ExitCode = ExecuteTaskResult.ExitCode; 
+				ExecutionTime = ExecuteTaskResult.ExecutionTime;
+				ProcessorTime = ExecuteTaskResult.ProcessorTime;
 			}
 
 			// Write it to the log
@@ -292,7 +303,14 @@ namespace UnrealBuildTool
 					return;
 				}
 
-				Log.TraceInformation("[{0}/{1}] {2}", CompletedActions, TotalActions, Description);
+				string CompilationTimes = "";
+
+				if (bShowPerActionCompilationTimes)
+				{
+					CompilationTimes = $" (Wall: {ExecutionTime.TotalSeconds:0.00}s CPU: {ProcessorTime.TotalSeconds:0.00}s)";
+				}
+
+				Log.TraceInformation("[{0}/{1}]{2} {3}", CompletedActions, TotalActions, CompilationTimes, Description);
 				foreach (string Line in LogLines.Skip(Action.bShouldOutputStatusDescription ? 0 : 1))
 				{
 					Log.TraceInformation(Line);
diff --git a/Engine/Source/Programs/UnrealBuildTool/Executors/ParallelExecutor.cs b/Engine/Source/Programs/UnrealBuildTool/Executors/ParallelExecutor.cs
--- a/Engine/Source/Programs/UnrealBuildTool/Executors/ParallelExecutor.cs
+++ b/Engine/Source/Programs/UnrealBuildTool/Executors/ParallelExecutor.cs
@@ -72,6 +72,12 @@ namespace UnrealBuildTool
 		[XmlConfigFile]
 		bool bShowCompilationTimes = false;
 
+		/// <summary>
+		/// Whether to show compilation times for each executed action
+		/// </summary>
+		[XmlConfigFile]
+		bool bShowPerActionCompilationTimes = false;
+
 		/// <summary>
 		/// How many processes that will be executed in parallel
 		/// </summary>
@@ -248,7 +254,13 @@ namespace UnrealBuildTool
 											Description = $"{(CompletedAction.Inner.CommandDescription != null ? CompletedAction.Inner.CommandDescription : CompletedAction.Inner.CommandPath.GetFileNameWithoutExtension())} {CompletedAction.LogLines[0]}".Trim();
 										}
 
-										Log.TraceInformation("[{0}/{1}] {2}", NumCompletedActions, InputActions.Count, Description);
+										string CompilationTimes = "";
+										if (bShowPerActionCompilationTimes)
+										{
+											CompilationTimes = $" (Wall: {CompletedAction.ExecutionTime.TotalSeconds:0.00}s CPU: {CompletedAction.ProcessorTime.TotalSeconds:0.00}s)";
+										}
+											
+										Log.TraceInformation("[{0}/{1}]{2} {3}", NumCompletedActions, InputActions.Count, CompilationTimes, Description);
 										foreach (string Line in CompletedAction.LogLines.Skip(CompletedAction.Inner.bShouldOutputStatusDescription ? 0 : 1))
 										{
 											Log.TraceInformation(Line);
