UnrealBuildTool: Log BuildConfiguration.xml locations on tool startup

Log file will contain lines like the following:

[00:00:00.032] Log started at 9/28/2021 8:46:43 AM (2021-09-28T15:46:43Z)
[00:00:00.051] No config file at [REDACTED_USER_HOME]/Unreal Engine/UnrealBuildTool/BuildConfiguration.xml
[00:00:00.051] Configuration will be read from:
[00:00:00.051]   [REDACTED_USER_HOME]/Perforce/UE5_Main/Engine/Restricted/NotForLicensees/Programs/UnrealBuildTool/BuildConfiguration.xml
[00:00:00.051]   [REDACTED_USER_HOME]/Perforce/UE5_Main/Engine/Saved/UnrealBuildTool/BuildConfiguration.xml
[00:00:00.051]   [REDACTED_USER_HOME]/.config/Unreal Engine/UnrealBuildTool/BuildConfiguration.xml

#trivial

diff --git a/Engine/Source/Programs/UnrealBuildTool/System/XmlConfig.cs b/Engine/Source/Programs/UnrealBuildTool/System/XmlConfig.cs
--- a/Engine/Source/Programs/UnrealBuildTool/System/XmlConfig.cs
+++ b/Engine/Source/Programs/UnrealBuildTool/System/XmlConfig.cs
@@ -231,11 +231,16 @@ namespace UnrealBuildTool
  				{
  					InputFiles.Add(new InputFile(NotForLicenseesConfigLocation, "NotForLicensees"));
  				}
+                else
+                {
+	                Log.TraceLog($"No config file at {NotForLicenseesConfigLocation}");
+                }
 
 				// Check for the user config file under /Engine/Saved/UnrealBuildTool
 				FileReference UserConfigLocation = FileReference.Combine(Unreal.EngineDirectory, "Saved", "UnrealBuildTool", "BuildConfiguration.xml");
 				if(!FileReference.Exists(UserConfigLocation))
 				{
+					Log.TraceLog($"Creating default config file at {UserConfigLocation}");
 					CreateDefaultConfigFile(UserConfigLocation);
 				}
 				InputFiles.Add(new InputFile(UserConfigLocation, "User"));
@@ -248,6 +253,7 @@ namespace UnrealBuildTool
 				FileReference AppDataConfigLocation = FileReference.Combine(new DirectoryReference(AppDataFolder), "Unreal Engine", "UnrealBuildTool", "BuildConfiguration.xml");
 				if(!FileReference.Exists(AppDataConfigLocation))
 				{
+					Log.TraceLog($"Creating default config file at {AppDataConfigLocation}");
 					CreateDefaultConfigFile(AppDataConfigLocation);
 				}
 				InputFiles.Add(new InputFile(AppDataConfigLocation, "Global (AppData)"));
@@ -262,8 +268,18 @@ namespace UnrealBuildTool
 				{
 					InputFiles.Add(new InputFile(PersonalConfigLocation, "Global (Documents)"));
 				}
+				else
+				{
+					Log.TraceLog($"No config file at {PersonalConfigLocation}");
+				}
 			}
 
+			Log.TraceLog("Configuration will be read from:");
+			foreach (InputFile InputFile in InputFiles)
+			{
+				Log.TraceLog($"  {InputFile.Location.FullName}");
+			}
+			
 			return InputFiles;
 		}
 
