UnrealBuildTool: Suppress warning re: deprecation of enum arithmetic in C++20

C++20 deprecates some arithmetic operations on enums, resulting in warnings like the following:

warning C5054: operator '+': deprecated between enumerations of different types

Ref: http://www.open-std.org/jtc1/sc22/wg21/docs/papers/2018/p1120r0.html

The code is remains valid, but may become unsupported in a future C++ standard (C++23, C++26, or not at all - see for example discussion in http://www.open-std.org/jtc1/sc22/wg21/docs/papers/2020/p2139r2.html )

#trivial

diff --git a/Engine/Source/Programs/UnrealBuildTool/Platform/Windows/VCToolChain.cs b/Engine/Source/Programs/UnrealBuildTool/Platform/Windows/VCToolChain.cs
--- a/Engine/Source/Programs/UnrealBuildTool/Platform/Windows/VCToolChain.cs
+++ b/Engine/Source/Programs/UnrealBuildTool/Platform/Windows/VCToolChain.cs
@@ -793,6 +793,15 @@ namespace UnrealBuildTool
 					case CppStandardVersion.Cpp20:
 					case CppStandardVersion.Latest:
 						Arguments.Add("/std:c++latest");
+						
+						// warning C5054: operator ___: deprecated between enumerations of different types
+						// re: http://www.open-std.org/jtc1/sc22/wg21/docs/papers/2018/p1120r0.html
+						
+						// It seems unclear whether the deprecation will be enacted in C++23 or not
+						// e.g. http://www.open-std.org/jtc1/sc22/wg21/docs/papers/2020/p2139r2.html
+						// Until the path forward is clearer, it seems reasonable to leave things as they are.
+						Arguments.Add("/wd5054");
+						
 						break;
 					// Will be added when MSVC is feature-complete.
 					// https://docs.microsoft.com/en-us/cpp/build/reference/std-specify-language-standard-version?view=msvc-160
