UnrealBuildTool: Update source dependency json parsing for imported modules

diff --git a/Engine/Source/Programs/UnrealBuildTool/Modes/BuildMode.cs b/Engine/Source/Programs/UnrealBuildTool/Modes/BuildMode.cs
--- a/Engine/Source/Programs/UnrealBuildTool/Modes/BuildMode.cs
+++ b/Engine/Source/Programs/UnrealBuildTool/Modes/BuildMode.cs
@@ -414,10 +414,10 @@ namespace UnrealBuildTool
 									ModuleOutputs[ProducedModule] = CppModulesAction.CompiledModuleInterfaceFile;
 								}
 
-								List<string>? ImportedModules;
+								List<(string Name, string BMI)>? ImportedModules;
 								if (CppDependencies.TryGetImportedModules(ModuleDependencyAction.DependencyListFile!, out ImportedModules))
 								{
-									ModuleImports[CppModulesAction.CompiledModuleInterfaceFile] = ImportedModules;
+									ModuleImports[CppModulesAction.CompiledModuleInterfaceFile] = ImportedModules.Select(x => x.Name).ToList();
 								}
 
 								CompiledModuleInterfaces.Add(CppModulesAction.CompiledModuleInterfaceFile);
diff --git a/Engine/Source/Programs/UnrealBuildTool/System/CppDependencyCache.cs b/Engine/Source/Programs/UnrealBuildTool/System/CppDependencyCache.cs
--- a/Engine/Source/Programs/UnrealBuildTool/System/CppDependencyCache.cs
+++ b/Engine/Source/Programs/UnrealBuildTool/System/CppDependencyCache.cs
@@ -26,10 +26,10 @@ namespace UnrealBuildTool
 		{
 			public long LastWriteTimeUtc;
 			public string? ProducedModule;
-			public List<string>? ImportedModules;
+			public List<(string Name, string BMI)>? ImportedModules;
 			public List<FileItem> Files;
 
-			public DependencyInfo(long LastWriteTimeUtc, string? ProducedModule, List<string>? ImportedModules, List<FileItem> Files)
+			public DependencyInfo(long LastWriteTimeUtc, string? ProducedModule, List<(string, string)>? ImportedModules, List<FileItem> Files)
 			{
 				this.LastWriteTimeUtc = LastWriteTimeUtc;
 				this.ProducedModule = ProducedModule;
@@ -41,7 +41,10 @@ namespace UnrealBuildTool
 			{
 				long LastWriteTimeUtc = Reader.ReadLong();
 				string ProducedModule = Reader.ReadString();
-				List<string> ImportedModules = Reader.ReadList(() => Reader.ReadString());
+				List<(string, string)> ImportedModules = Reader.ReadList(() =>
+				{
+					return (Reader.ReadString(), Reader.ReadString());
+				});
 				List<FileItem> Files = Reader.ReadList(() => Reader.ReadCompactFileItem());
 				return new DependencyInfo(LastWriteTimeUtc, ProducedModule, ImportedModules, Files);
 			}
@@ -50,7 +53,11 @@ namespace UnrealBuildTool
 			{
 				Writer.WriteLong(LastWriteTimeUtc);
 				Writer.WriteString(ProducedModule);
-				Writer.WriteList(ImportedModules, Module => Writer.WriteString(Module));
+				Writer.WriteList(ImportedModules, (Module) =>
+				{
+					Writer.WriteString(Module.Name);
+					Writer.WriteString(Module.BMI);
+				});
 				Writer.WriteList<FileItem>(Files, File => Writer.WriteCompactFileItem(File));
 			}
 		}
@@ -197,7 +204,7 @@ namespace UnrealBuildTool
 		/// <param name="InputFile">The dependency file to query</param>
 		/// <param name="OutImportedModules">List of imported modules</param>
 		/// <returns>True if a list of imported modules was obtained</returns>
-		public bool TryGetImportedModules(FileItem InputFile, [NotNullWhen(true)] out List<string>? OutImportedModules)
+		public bool TryGetImportedModules(FileItem InputFile, [NotNullWhen(true)] out List<(string Name, string BMI)>? OutImportedModules)
 		{
 			DependencyInfo? Info;
 			if (TryGetDependencyInfo(InputFile, out Info))
@@ -251,6 +258,10 @@ namespace UnrealBuildTool
 			{
 				return TryGetDependencyInfoInternal(InputFile, out OutInfo);
 			}
+			catch (BuildException Ex)
+			{
+				throw Ex;
+			}
 			catch (Exception Ex)
 			{
 				Log.TraceLog("Unable to read {0}:\n{1}", InputFile, ExceptionUtils.FormatExceptionDetails(Ex));
@@ -426,8 +437,23 @@ namespace UnrealBuildTool
 			}
 			else if (InputFile.HasExtension(".json"))
 			{
+				// https://docs.microsoft.com/en-us/cpp/build/reference/sourcedependencies?view=msvc-160&viewFallbackFrom=vs-2019
+				
 				JsonObject Object = JsonObject.Read(InputFile.Location);
 
+				const string ExpectedVersion = "1.1";
+				if (!Object.TryGetStringField("Version", out string Version))
+				{
+					throw new BuildException(
+						$"Dependency file \"{InputFile.Location}\" does not have have a \"Version\" field.");
+				}
+
+				if (!String.Equals(Version, ExpectedVersion))
+				{
+					throw new BuildException(
+						$"Dependency file \"{InputFile.Location}\" version (\"{Version}\") is not supported. Expected version \"{ExpectedVersion}\"");
+				}
+
 				JsonObject Data;
 				if (!Object.TryGetObjectField("Data", out Data))
 				{
@@ -437,8 +463,22 @@ namespace UnrealBuildTool
 				string ProducedModule;
 				Data.TryGetStringField("ProvidedModule", out ProducedModule);
 
-				string[] ImportedModules;
-				Data.TryGetStringArrayField("ImportedModules", out ImportedModules);
+				List<(string Name, string BMI)>? ImportedModules = null;
+				if (Data.TryGetObjectArrayField("ImportedModules", out JsonObject[] ImportedModulesJson))
+				{
+					foreach (JsonObject ImportedModule in ImportedModulesJson)
+					{
+						ImportedModule.TryGetStringField("Name", out string Name);
+						ImportedModule.TryGetStringField("BMI", out string BMI);
+						
+						if (ImportedModules == null)
+						{
+							ImportedModules = new List<(string Name, string BMI)>();
+						}
+
+						ImportedModules.Add((Name, BMI));
+					}
+				}
 
 				string[] Includes;
 				Data.TryGetStringArrayField("Includes", out Includes);
