UBT: Fix generated documentation for UBT config classes not including properties, and producing warnings for private fields.

diff --git a/Engine/Source/Programs/UnrealBuildTool/System/RulesDocumentation.cs b/Engine/Source/Programs/UnrealBuildTool/System/RulesDocumentation.cs
--- a/Engine/Source/Programs/UnrealBuildTool/System/RulesDocumentation.cs
+++ b/Engine/Source/Programs/UnrealBuildTool/System/RulesDocumentation.cs
@@ -2,10 +2,12 @@
 
 using System;
 using System.Collections.Generic;
+using System.Diagnostics.CodeAnalysis;
 using System.IO;
 using System.Linq;
 using System.Reflection;
 using System.Text;
+using System.Text.RegularExpressions;
 using System.Threading.Tasks;
 using System.Xml;
 using EpicGames.Core;
@@ -15,6 +17,20 @@ namespace UnrealBuildTool
 {
 	static class RulesDocumentation
 	{
+		class SettingInfo
+		{
+			public string Name;
+			public Type Type;
+			public List<string> Description;
+
+			public SettingInfo(string Name, Type Type, List<string> Description)
+			{
+				this.Name = Name;
+				this.Type = Type;
+				this.Description = Description;
+			}
+		}
+
 		public static void WriteDocumentation(Type RulesType, FileReference OutputFile)
 		{
 			// Get the path to the XML documentation
@@ -28,26 +44,54 @@ namespace UnrealBuildTool
 			XmlDocument InputDocumentation = new XmlDocument();
 			InputDocumentation.Load(InputDocumentationFile.FullName);
 
-			// Filter the properties into read-only and read/write lists
-			List<FieldInfo> ReadOnlyFields = new List<FieldInfo>();
-			List<FieldInfo> ReadWriteFields = new List<FieldInfo>();
-			foreach(FieldInfo Field in RulesType.GetFields(BindingFlags.Instance | BindingFlags.SetProperty | BindingFlags.Public))
+			// Filter the settings into read-only and read/write lists
+			List<SettingInfo> ReadOnlySettings = new List<SettingInfo>();
+			List<SettingInfo> ReadWriteSettings = new List<SettingInfo>();
+
+			// First read the fields
+			foreach (FieldInfo Field in RulesType.GetFields(BindingFlags.Instance | BindingFlags.SetProperty | BindingFlags.Public))
 			{
 				if(!Field.FieldType.IsClass || !Field.FieldType.Name.EndsWith("TargetRules"))
 				{
-					if(Field.IsInitOnly)
+					List<string>? Lines;
+					if (TryGetXmlComment(InputDocumentation, Field, out Lines))
 					{
-						ReadOnlyFields.Add(Field);
+						SettingInfo Setting = new SettingInfo(Field.Name, Field.FieldType, Lines);
+						if (Field.IsInitOnly)
+						{
+							ReadOnlySettings.Add(Setting);
+						}
+						else
+						{
+							ReadWriteSettings.Add(Setting);
+						}
 					}
-					else
+				}
+			}
+
+			// ...then read all the properties
+			foreach (PropertyInfo Property in RulesType.GetProperties(BindingFlags.Instance | BindingFlags.SetProperty | BindingFlags.Public))
+			{
+				if (!Property.PropertyType.IsClass || !Property.PropertyType.Name.EndsWith("TargetRules"))
+				{
+					List<string>? Lines;
+					if (TryGetXmlComment(InputDocumentation, Property, out Lines))
 					{
-						ReadWriteFields.Add(Field);
+						SettingInfo Setting = new SettingInfo(Property.Name, Property.PropertyType, Lines);
+						if (Property.SetMethod == null)
+						{
+							ReadOnlySettings.Add(Setting);
+						}
+						else
+						{
+							ReadWriteSettings.Add(Setting);
+						}
 					}
 				}
 			}
 
 			// Make sure the output file is writable
-			if(FileReference.Exists(OutputFile))
+			if (FileReference.Exists(OutputFile))
 			{
 				FileReference.MakeWriteable(OutputFile);
 			}
@@ -59,11 +103,11 @@ namespace UnrealBuildTool
 			// Generate the documentation file
 			if (OutputFile.HasExtension(".udn"))
 			{
-				WriteDocumentationUDN(OutputFile, ReadOnlyFields, ReadWriteFields, InputDocumentation);
+				WriteDocumentationUDN(OutputFile, ReadOnlySettings, ReadWriteSettings);
 			}
 			else if (OutputFile.HasExtension(".html"))
 			{
-				WriteDocumentationHTML(OutputFile, ReadOnlyFields, ReadWriteFields, InputDocumentation);
+				WriteDocumentationHTML(OutputFile, ReadOnlySettings, ReadWriteSettings);
 			}
 			else
 			{
@@ -74,7 +118,7 @@ namespace UnrealBuildTool
 			Log.TraceInformation("Written documentation to {0}.", OutputFile);
 		}
 
-		static void WriteDocumentationUDN(FileReference OutputFile, List<FieldInfo> ReadOnlyFields, List<FieldInfo> ReadWriteFields, XmlDocument InputDocumentation)
+		static void WriteDocumentationUDN(FileReference OutputFile, List<SettingInfo> ReadOnlySettings, List<SettingInfo> ReadWriteSettings)
 		{
 			// Generate the UDN documentation file
 			using (StreamWriter Writer = new StreamWriter(OutputFile.FullName))
@@ -85,97 +129,161 @@ namespace UnrealBuildTool
 				Writer.WriteLine("Description: This is a procedurally generated markdown page.");
 				Writer.WriteLine("Version: {0}.{1}", ReadOnlyBuildVersion.Current.MajorVersion, ReadOnlyBuildVersion.Current.MinorVersion);
 				Writer.WriteLine("");
-				if (ReadOnlyFields.Count > 0)
+				if (ReadOnlySettings.Count > 0)
 				{
 					Writer.WriteLine("### Read-Only Properties");
 					Writer.WriteLine();
-					foreach (FieldInfo Field in ReadOnlyFields)
+					foreach (SettingInfo Field in ReadOnlySettings)
 					{
-						WriteFieldUDN(InputDocumentation, Field, Writer);
+						WriteFieldUDN(Field, Writer);
 					}
 					Writer.WriteLine();
 				}
-				if (ReadWriteFields.Count > 0)
+				if (ReadWriteSettings.Count > 0)
 				{
 					Writer.WriteLine("### Read/Write Properties");
-					foreach (FieldInfo Field in ReadWriteFields)
+					foreach (SettingInfo Field in ReadWriteSettings)
 					{
-						WriteFieldUDN(InputDocumentation, Field, Writer);
+						WriteFieldUDN(Field, Writer);
 					}
 					Writer.WriteLine("");
 				}
 			}
 		}
 
-		static void WriteFieldUDN(XmlDocument InputDocumentation, FieldInfo Field, TextWriter Writer)
+		/// <summary>
+		/// Gets the XML comment for a particular field
+		/// </summary>
+		/// <param name="Documentation">The XML documentation</param>
+		/// <param name="Field">The member to search for</param>
+		/// <param name="Lines">Receives the description for the requested field</param>
+		/// <returns>True if a comment was found for the field</returns>
+		public static bool TryGetXmlComment(XmlDocument Documentation, FieldInfo Field, [NotNullWhen(true)] out List<string>? Lines)
+		{
+			return TryGetXmlComment(Documentation, $"F:{Field.DeclaringType}.{Field.Name}", out Lines);
+		}
+
+		/// <summary>
+		/// Gets the XML comment for a particular field
+		/// </summary>
+		/// <param name="Documentation">The XML documentation</param>
+		/// <param name="Property">The member to search for</param>
+		/// <param name="Lines">Receives the description for the requested field</param>
+		/// <returns>True if a comment was found for the field</returns>
+		public static bool TryGetXmlComment(XmlDocument Documentation, PropertyInfo Property, [NotNullWhen(true)] out List<string>? Lines)
 		{
-			XmlNode Node = InputDocumentation.SelectSingleNode(String.Format("//member[@name='F:{0}.{1}']/summary", Field.DeclaringType!.FullName, Field.Name));
+			return TryGetXmlComment(Documentation, $"P:{Property.DeclaringType}.{Property.Name}", out Lines);
+		}
+
+		/// <summary>
+		/// Gets the XML comment for a particular field
+		/// </summary>
+		/// <param name="Documentation">The XML documentation</param>
+		/// <param name="MemberName">The member to search for</param>
+		/// <param name="Lines">Receives the description for the requested field</param>
+		/// <returns>True if a comment was found for the field</returns>
+		public static bool TryGetXmlComment(XmlDocument Documentation, string MemberName, [NotNullWhen(true)] out List<string>? Lines)
+		{
+			HashSet<string> VisitedProperties = new HashSet<string>(StringComparer.Ordinal);
+
+			XmlNode? Node = null;
+			for (; ; )
+			{
+				Node = Documentation.SelectSingleNode($"//member[@name='{MemberName}']/summary");
+				if (Node == null)
+				{
+					break;
+				}
+
+				XmlNode InheritNode = Documentation.SelectSingleNode($"//member[@name='{MemberName}']/inheritdoc/@cref");
+				if (InheritNode == null || !VisitedProperties.Add(InheritNode.InnerText))
+				{
+					break;
+				}
+
+				MemberName = InheritNode.InnerText;
+			}
+
 			if (Node != null)
 			{
 				// Reflow the comments into paragraphs, assuming that each paragraph will be separated by a blank line
-				List<string> Lines = new List<string>(Node.InnerText.Trim().Split('\n').Select(x => x.Trim()));
-				for (int Idx = Lines.Count - 1; Idx > 0; Idx--)
+				List<string> NewLines = new List<string>(Node.InnerText.Trim().Split('\n').Select(x => x.Trim()));
+				for (int Idx = NewLines.Count - 1; Idx > 0; Idx--)
 				{
-					if (Lines[Idx - 1].Length > 0 && !Lines[Idx].StartsWith("*") && !Lines[Idx].StartsWith("-"))
+					if (NewLines[Idx - 1].Length > 0 && !NewLines[Idx].StartsWith("*") && !NewLines[Idx].StartsWith("-"))
 					{
-						Lines[Idx - 1] += " " + Lines[Idx];
-						Lines.RemoveAt(Idx);
+						NewLines[Idx - 1] += " " + NewLines[Idx];
+						NewLines.RemoveAt(Idx);
 					}
 				}
 
-				// Write the values of the enum
-				/*				if(Field.FieldType.IsEnum)
+				if (NewLines.Count > 0)
+				{
+					Lines = NewLines;
+					return true;
+				}
+			}
+
+			Log.TraceWarning("Missing XML comment for {0}", Regex.Replace(MemberName, @"^[A-Z]:", ""));
+			Lines = null;
+			return false;
+		}
+
+
+		static void WriteFieldUDN(SettingInfo Field, TextWriter Writer)
+		{
+			// Write the values of the enum
+			/*				if(Field.FieldType.IsEnum)
+							{
+								Lines.Add("Valid values are:");
+								foreach(string Value in Enum.GetNames(Field.FieldType))
 								{
-									Lines.Add("Valid values are:");
-									foreach(string Value in Enum.GetNames(Field.FieldType))
-									{
-										Lines.Add(String.Format("* {0}.{1}", Field.FieldType.Name, Value));
-									}
+									Lines.Add(String.Format("* {0}.{1}", Field.FieldType.Name, Value));
 								}
-				*/
-				// Write the result to the .udn file
-				if (Lines.Count > 0)
+							}
+			*/
+
+			List<string> Lines = Field.Description;
+
+			// Write the result to the .udn file
+			Writer.WriteLine("$ {0} ({1}): {2}", Field.Name, GetPrettyTypeName(Field.Type), Lines[0]);
+			for (int Idx = 1; Idx < Lines.Count; Idx++)
+			{
+				if (Lines[Idx].StartsWith("*") || Lines[Idx].StartsWith("-"))
 				{
-					Writer.WriteLine("$ {0} ({1}): {2}", Field.Name, GetPrettyTypeName(Field.FieldType), Lines[0]);
-					for (int Idx = 1; Idx < Lines.Count; Idx++)
-					{
-						if (Lines[Idx].StartsWith("*") || Lines[Idx].StartsWith("-"))
-						{
-							Writer.WriteLine("        * {0}", Lines[Idx].Substring(1).TrimStart());
-						}
-						else
-						{
-							Writer.WriteLine("    * {0}", Lines[Idx]);
-						}
-					}
-					Writer.WriteLine();
+					Writer.WriteLine("        * {0}", Lines[Idx].Substring(1).TrimStart());
+				}
+				else
+				{
+					Writer.WriteLine("    * {0}", Lines[Idx]);
 				}
 			}
+			Writer.WriteLine();
 		}
 
-		static void WriteDocumentationHTML(FileReference OutputFile, List<FieldInfo> ReadOnlyFields, List<FieldInfo> ReadWriteFields, XmlDocument InputDocumentation)
+		static void WriteDocumentationHTML(FileReference OutputFile, List<SettingInfo> ReadOnlySettings, List<SettingInfo> ReadWriteSettings)
 		{
 			using (StreamWriter Writer = new StreamWriter(OutputFile.FullName))
 			{
 				Writer.WriteLine("<html>");
 				Writer.WriteLine("  <body>");
-				if (ReadOnlyFields.Count > 0)
+				if (ReadOnlySettings.Count > 0)
 				{
 					Writer.WriteLine("    <h2>Read-Only Properties</h2>");
 					Writer.WriteLine("    <dl>");
-					foreach (FieldInfo Field in ReadOnlyFields)
+					foreach (SettingInfo Setting in ReadOnlySettings)
 					{
-						WriteFieldHTML(InputDocumentation, Field, Writer);
+						WriteFieldHTML(Setting, Writer);
 					}
 					Writer.WriteLine("    </dl>");
 				}
-				if (ReadWriteFields.Count > 0)
+				if (ReadWriteSettings.Count > 0)
 				{
 					Writer.WriteLine("    <h2>Read/Write Properties</h2>");
 					Writer.WriteLine("    <dl>");
-					foreach (FieldInfo Field in ReadWriteFields)
+					foreach (SettingInfo Setting in ReadWriteSettings)
 					{
-						WriteFieldHTML(InputDocumentation, Field, Writer);
+						WriteFieldHTML(Setting, Writer);
 					}
 					Writer.WriteLine("    </dl>");
 				}
@@ -184,62 +292,49 @@ namespace UnrealBuildTool
 			}
 		}
 
-		static void WriteFieldHTML(XmlDocument InputDocumentation, FieldInfo Field, TextWriter Writer)
+		static void WriteFieldHTML(SettingInfo Setting, TextWriter Writer)
 		{
-			XmlNode Node = InputDocumentation.SelectSingleNode(String.Format("//member[@name='F:{0}.{1}']/summary", Field.DeclaringType!.FullName, Field.Name));
-			if(Node != null)
+
+			// Write the values of the enum
+			/*				if(Field.FieldType.IsEnum)
+							{
+								Lines.Add("Valid values are:");
+								foreach(string Value in Enum.GetNames(Field.FieldType))
+								{
+									Lines.Add(String.Format("* {0}.{1}", Field.FieldType.Name, Value));
+								}
+							}
+			*/
+			// Write the result to the HTML file
+			List<string> Lines = Setting.Description;
+			if (Lines.Count > 0)
 			{
-				// Reflow the comments into paragraphs, assuming that each paragraph will be separated by a blank line
-				List<string> Lines = new List<string>(Node.InnerText.Trim().Split('\n').Select(x => x.Trim()));
-				for(int Idx = Lines.Count - 1; Idx > 0; Idx--)
-				{
-					if(Lines[Idx - 1].Length > 0 && !Lines[Idx].StartsWith("*") && !Lines[Idx].StartsWith("-"))
-					{
-						Lines[Idx - 1] += " " + Lines[Idx];
-						Lines.RemoveAt(Idx);
-					}
-				}
+				Writer.WriteLine("      <dt>{0} ({1})</dt>", Setting.Name, GetPrettyTypeName(Setting.Type));
 
-				// Write the values of the enum
-/*				if(Field.FieldType.IsEnum)
+				if(Lines.Count == 1)
 				{
-					Lines.Add("Valid values are:");
-					foreach(string Value in Enum.GetNames(Field.FieldType))
-					{
-						Lines.Add(String.Format("* {0}.{1}", Field.FieldType.Name, Value));
-					}
+					Writer.WriteLine("      <dd>{0}</dd>", Lines[0]);
 				}
-*/
-				// Write the result to the HTML file
-				if(Lines.Count > 0)
+				else
 				{
-					Writer.WriteLine("      <dt>{0} ({1})</dt>", Field.Name, GetPrettyTypeName(Field.FieldType));
-
-					if(Lines.Count == 1)
-					{
-						Writer.WriteLine("      <dd>{0}</dd>", Lines[0]);
-					}
-					else
+					Writer.WriteLine("      <dd>");
+					for(int Idx = 0; Idx < Lines.Count; Idx++)
 					{
-						Writer.WriteLine("      <dd>");
-						for(int Idx = 0; Idx < Lines.Count; Idx++)
+						if(Lines[Idx].StartsWith("*") || Lines[Idx].StartsWith("-"))
 						{
-							if(Lines[Idx].StartsWith("*") || Lines[Idx].StartsWith("-"))
-							{
-								Writer.WriteLine("        <ul>");
-								for(; Idx < Lines.Count && (Lines[Idx].StartsWith("*") || Lines[Idx].StartsWith("-")); Idx++)
-								{
-									Writer.WriteLine("          <li>{0}</li>", Lines[Idx].Substring(1).TrimStart());
-								}
-								Writer.WriteLine("        </ul>");
-							}
-							else
+							Writer.WriteLine("        <ul>");
+							for(; Idx < Lines.Count && (Lines[Idx].StartsWith("*") || Lines[Idx].StartsWith("-")); Idx++)
 							{
-								Writer.WriteLine("        {0}", Lines[Idx]);
+								Writer.WriteLine("          <li>{0}</li>", Lines[Idx].Substring(1).TrimStart());
 							}
+							Writer.WriteLine("        </ul>");
+						}
+						else
+						{
+							Writer.WriteLine("        {0}", Lines[Idx]);
 						}
-						Writer.WriteLine("      </dd>");
 					}
+					Writer.WriteLine("      </dd>");
 				}
 			}
 		}
diff --git a/Engine/Source/Programs/UnrealBuildTool/System/XmlConfig.cs b/Engine/Source/Programs/UnrealBuildTool/System/XmlConfig.cs
--- a/Engine/Source/Programs/UnrealBuildTool/System/XmlConfig.cs
+++ b/Engine/Source/Programs/UnrealBuildTool/System/XmlConfig.cs
@@ -787,37 +787,6 @@ namespace UnrealBuildTool
 			Log.TraceInformation("Written documentation to {0}.", OutputFile);
 		}
 
-		/// <summary>
-		/// Gets the XML comment for a particular field
-		/// </summary>
-		/// <param name="Documentation">The XML documentation</param>
-		/// <param name="Field">The field to search for</param>
-		/// <param name="Lines">Receives the description for the requested field</param>
-		/// <returns>True if a comment was found for the field</returns>
-		private static bool TryGetXmlComment(XmlDocument Documentation, FieldInfo Field, [NotNullWhen(true)] out List<string>? Lines)
-		{
-			XmlNode Node = Documentation.SelectSingleNode(String.Format("//member[@name='F:{0}.{1}']/summary", Field.DeclaringType!.FullName, Field.Name));
-			if (Node == null)
-			{
-				Lines = null;
-				return false;
-			}
-			else
-			{
-				// Reflow the comments into paragraphs, assuming that each paragraph will be separated by a blank line
-				Lines = new List<string>(Node.InnerText.Trim().Split('\n').Select(x => x.Trim()));
-				for (int Idx = Lines.Count - 1; Idx > 0; Idx--)
-				{
-					if (Lines[Idx - 1].Length > 0 && !Lines[Idx].StartsWith("*") && !Lines[Idx].StartsWith("-"))
-					{
-						Lines[Idx - 1] += " " + Lines[Idx];
-						Lines.RemoveAt(Idx);
-					}
-				}
-				return true;
-			}
-		}
-
 		/// <summary>
 		/// Writes out documentation in UDN format
 		/// </summary>
@@ -846,7 +815,7 @@ namespace UnrealBuildTool
 					{
 						// Get the XML comment for this field
 						List<string>? Lines;
-						if(!TryGetXmlComment(InputDocumentation, FieldPair.Value, out Lines) || Lines.Count == 0)
+						if(!RulesDocumentation.TryGetXmlComment(InputDocumentation, FieldPair.Value, out Lines) || Lines.Count == 0)
 						{
 							Log.TraceWarning("Missing XML comment for {0}", FieldPair.Value.Name);
 							continue;
@@ -895,9 +864,8 @@ namespace UnrealBuildTool
 					{
 						// Get the XML comment for this field
 						List<string>? Lines;
-						if (!TryGetXmlComment(InputDocumentation, FieldPair.Value, out Lines) || Lines.Count == 0)
+						if (!RulesDocumentation.TryGetXmlComment(InputDocumentation, FieldPair.Value, out Lines) || Lines.Count == 0)
 						{
-							Log.TraceWarning("Missing XML comment for {0}", FieldPair.Value.Name);
 							continue;
 						}
 
