UBT: Fix warnings generating documentation.

diff --git a/Engine/Source/Programs/UnrealBuildTool/System/RulesDocumentation.cs b/Engine/Source/Programs/UnrealBuildTool/System/RulesDocumentation.cs
--- a/Engine/Source/Programs/UnrealBuildTool/System/RulesDocumentation.cs
+++ b/Engine/Source/Programs/UnrealBuildTool/System/RulesDocumentation.cs
@@ -190,7 +190,7 @@ namespace UnrealBuildTool
 			for (; ; )
 			{
 				Node = Documentation.SelectSingleNode($"//member[@name='{MemberName}']/summary");
-				if (Node == null)
+				if (Node != null)
 				{
 					break;
 				}
diff --git a/Engine/Source/Programs/UnrealBuildTool/System/XmlConfig.cs b/Engine/Source/Programs/UnrealBuildTool/System/XmlConfig.cs
--- a/Engine/Source/Programs/UnrealBuildTool/System/XmlConfig.cs
+++ b/Engine/Source/Programs/UnrealBuildTool/System/XmlConfig.cs
@@ -817,7 +817,6 @@ namespace UnrealBuildTool
 						List<string>? Lines;
 						if(!RulesDocumentation.TryGetXmlComment(InputDocumentation, FieldPair.Value, out Lines) || Lines.Count == 0)
 						{
-							Log.TraceWarning("Missing XML comment for {0}", FieldPair.Value.Name);
 							continue;
 						}
 
