UnrealBuildTool: Handle different module dependency formats correctly

Output from /sourceDependencies isn't quite the same as from /sourceDependencies:directives

diff --git a/Engine/Source/Programs/UnrealBuildTool/System/CppDependencyCache.cs b/Engine/Source/Programs/UnrealBuildTool/System/CppDependencyCache.cs
--- a/Engine/Source/Programs/UnrealBuildTool/System/CppDependencyCache.cs
+++ b/Engine/Source/Programs/UnrealBuildTool/System/CppDependencyCache.cs
@@ -464,19 +464,31 @@ namespace UnrealBuildTool
 				Data.TryGetStringField("ProvidedModule", out ProducedModule);
 
 				List<(string Name, string BMI)>? ImportedModules = null;
-				if (Data.TryGetObjectArrayField("ImportedModules", out JsonObject[] ImportedModulesJson))
+				if (InputFile.HasExtension(".md.json"))
 				{
-					foreach (JsonObject ImportedModule in ImportedModulesJson)
+					Data.TryGetStringArrayField("ImportedModules", out string[] ImportedModuleArray);
+					if (ImportedModuleArray.Length > 0)
 					{
-						ImportedModule.TryGetStringField("Name", out string Name);
-						ImportedModule.TryGetStringField("BMI", out string BMI);
-						
-						if (ImportedModules == null)
+						ImportedModules =
+							new List<(string Name, string BMI)>(ImportedModuleArray.ConvertAll(x => (x, "")));
+					}
+				}
+				else
+				{
+					if (Data.TryGetObjectArrayField("ImportedModules", out JsonObject[] ImportedModulesJson))
+					{
+						foreach (JsonObject ImportedModule in ImportedModulesJson)
 						{
-							ImportedModules = new List<(string Name, string BMI)>();
-						}
+							ImportedModule.TryGetStringField("Name", out string Name);
+							ImportedModule.TryGetStringField("BMI", out string BMI);
 
-						ImportedModules.Add((Name, BMI));
+							if (ImportedModules == null)
+							{
+								ImportedModules = new List<(string Name, string BMI)>();
+							}
+
+							ImportedModules.Add((Name, BMI));
+						}
 					}
 				}
 
@@ -492,7 +504,7 @@ namespace UnrealBuildTool
 					}
 				}
 
-				return new DependencyInfo(InputFile.LastWriteTimeUtc.Ticks, ProducedModule, ImportedModules?.ToList(), Files);
+				return new DependencyInfo(InputFile.LastWriteTimeUtc.Ticks, ProducedModule, ImportedModules, Files);
 			}
 			else
 			{
