UnrealBuildTool: Avoid having /sourceDependencies and /sourceDependencies:directives on the same command line

Can result in an invalid .md.json file that lacks a list of ImportedModules.

diff --git a/Engine/Source/Programs/UnrealBuildTool/Platform/Windows/VCCompileAction.cs b/Engine/Source/Programs/UnrealBuildTool/Platform/Windows/VCCompileAction.cs
--- a/Engine/Source/Programs/UnrealBuildTool/Platform/Windows/VCCompileAction.cs
+++ b/Engine/Source/Programs/UnrealBuildTool/Platform/Windows/VCCompileAction.cs
@@ -438,7 +438,8 @@ namespace UnrealBuildTool
 				VCToolChain.AddObjectFile(Arguments, ObjectFile, CompilerType, PreprocessedFile != null);
 			}
 
-			if (DependencyListFile != null && DependencyListFile.HasExtension(".json"))
+			// A better way to express this? .json is used as output for /sourceDependencies), but .md.json is used as output for /sourceDependencies:directives)
+			if (DependencyListFile != null && DependencyListFile.HasExtension(".json") && !DependencyListFile.HasExtension(".md.json"))
 			{
 				VCToolChain.AddSourceDependenciesFile(Arguments, DependencyListFile, CompilerType, PreprocessedFile != null);
 			}
