UnrealBuildTool: Add support for deprecating an xml config option

Example code:

	[XmlConfigFile(Deprecated = true, NewAttributeName = "PrimaryProjectName")]
	protected string MasterProjectName = "UE5";

Example output:

WARNING: Deprecated setting found in "[REDACTED_USER_HOME]/.config/Unreal Engine/UnrealBuildTool/BuildConfiguration.xml":
WARNING: The setting "MasterProjectName" is deprecated. Support for this setting will be removed in a future version of Unreal Engine.
WARNING: Use "PrimaryProjectName" in place of "MasterProjectName"
The value provided for "MasterProjectName" will be applied to "PrimaryProjectName"

diff --git a/Engine/Source/Programs/UnrealBuildTool/System/XmlConfig.cs b/Engine/Source/Programs/UnrealBuildTool/System/XmlConfig.cs
--- a/Engine/Source/Programs/UnrealBuildTool/System/XmlConfig.cs
+++ b/Engine/Source/Programs/UnrealBuildTool/System/XmlConfig.cs
@@ -128,7 +128,8 @@ namespace UnrealBuildTool
 					}
 
 					// Read all the XML files and validate them against the schema
-					Dictionary<Type, Dictionary<FieldInfo, object>> TypeToValues = new Dictionary<Type, Dictionary<FieldInfo, object>>();
+					Dictionary<Type, Dictionary<FieldInfo, XmlConfigData.ValueInfo>> TypeToValues =
+						new Dictionary<Type, Dictionary<FieldInfo, XmlConfigData.ValueInfo>>();
 					foreach(FileReference InputFile in InputFiles)
 					{
 						if(!TryReadFile(InputFile, CategoryToFields, TypeToValues, Schema))
@@ -141,20 +142,22 @@ namespace UnrealBuildTool
 					DirectoryReference.CreateDirectory(CacheFile.Directory);
 
 					// Create the new cache
-					Values = new XmlConfigData(InputFiles, TypeToValues.ToDictionary(x => x.Key, x => x.Value.ToArray()));
+					Values = new XmlConfigData(InputFiles, TypeToValues.ToDictionary(
+							x => x.Key, 
+							x => x.Value.Select(x => x.Value).ToArray()));
 					Values.Write(CacheFile);
 				}
 			}
 
 			// Apply all the static field values
-			foreach(KeyValuePair<Type, KeyValuePair<FieldInfo, object>[]> TypeValuesPair in Values.TypeToValues)
+			foreach(KeyValuePair<Type, XmlConfigData.ValueInfo[]> TypeValuesPair in Values.TypeToValues)
 			{
-				foreach(KeyValuePair<FieldInfo, object> FieldValuePair in TypeValuesPair.Value)
+				foreach(XmlConfigData.ValueInfo FieldValue in TypeValuesPair.Value)
 				{
-					if(FieldValuePair.Key.IsStatic)
+					if(FieldValue.FieldInfo.IsStatic)
 					{
-						object Value = InstanceValue(FieldValuePair.Value, FieldValuePair.Key.FieldType);
-						FieldValuePair.Key.SetValue(null, Value);
+						object Value = InstanceValue(FieldValue.Value, FieldValue.FieldInfo.FieldType);
+						FieldValue.FieldInfo.SetValue(null, Value);
 					}
 				}
 			}
@@ -305,16 +308,38 @@ namespace UnrealBuildTool
 		public static void ApplyTo(object TargetObject)
 		{
 			for(Type? TargetType = TargetObject.GetType(); TargetType != null; TargetType = TargetType.BaseType)
-			{
-				KeyValuePair<FieldInfo, object>[]? FieldValues;
+			{ 
+				XmlConfigData.ValueInfo[]? FieldValues;
 				if(Values!.TypeToValues.TryGetValue(TargetType, out FieldValues))
 				{
-					foreach(KeyValuePair<FieldInfo, object> FieldValuePair in FieldValues)
+					foreach(XmlConfigData.ValueInfo FieldValue in FieldValues)
 					{
-						if(!FieldValuePair.Key.IsStatic)
+						if (!FieldValue.FieldInfo.IsStatic)
 						{
-							object ValueInstance = InstanceValue(FieldValuePair.Value, FieldValuePair.Key.FieldType);
-							FieldValuePair.Key.SetValue(TargetObject, ValueInstance);
+							FieldInfo FieldInfoToWrite = FieldValue.FieldInfo;
+							
+							// Check if setting has been deprecated
+							if (FieldValue.XmlConfigAttribute.Deprecated)
+							{
+								Log.TraceWarning($"Deprecated setting found in \"{FieldValue.SourceFile}\":");
+								Log.TraceWarning(
+									$"The setting \"{FieldValue.FieldInfo.Name}\" is deprecated. Support for this setting will be removed in a future version of Unreal Engine.");
+
+								if (FieldValue.XmlConfigAttribute.NewAttributeName != null)
+								{
+									Log.TraceWarning(
+										$"Use \"{FieldValue.XmlConfigAttribute.NewAttributeName}\" in place of \"{FieldValue.FieldInfo.Name}\"");
+									Log.TraceInformation(
+										$"The value provided for deprecated setting \"{FieldValue.FieldInfo.Name}\" will be applied to \"{FieldValue.XmlConfigAttribute.NewAttributeName}\"");
+
+									// Find the new field to which the setting should be actually applied
+									FieldInfoToWrite = TargetType.GetRuntimeFields()
+										.First(x => x.Name == FieldValue.XmlConfigAttribute.NewAttributeName);
+								}
+							}
+
+							object ValueInstance = InstanceValue(FieldValue.Value, FieldValue.FieldInfo.FieldType);
+							FieldInfoToWrite.SetValue(TargetObject, ValueInstance);
 						}
 					}
 				}
@@ -349,7 +374,7 @@ namespace UnrealBuildTool
 		public static bool TryGetValue(Type TargetType, string Name, [NotNullWhen(true)] out object? Value)
 		{
 			// Find all the config values for this type
-			KeyValuePair<FieldInfo, object>[]? FieldValues;
+			XmlConfigData.ValueInfo[]? FieldValues;
 			if(!Values!.TypeToValues.TryGetValue(TargetType, out FieldValues))
 			{
 				Value = null;
@@ -357,11 +382,11 @@ namespace UnrealBuildTool
 			}
 
 			// Find the value with the matching name
-			foreach(KeyValuePair<FieldInfo, object> FieldPair in FieldValues)
+			foreach(XmlConfigData.ValueInfo FieldValue in FieldValues)
 			{
-				if(FieldPair.Key.Name == Name)
+				if(FieldValue.FieldInfo.Name == Name)
 				{
-					Value = FieldPair.Value;
+					Value = FieldValue.Value;
 					return true;
 				}
 			}
@@ -592,7 +617,8 @@ namespace UnrealBuildTool
 		/// <param name="TypeToValues">Map of types to fields and their associated values</param>
 		/// <param name="Schema">Schema to validate against</param>
 		/// <returns>True if the file was read successfully</returns>
-		static bool TryReadFile(FileReference Location, Dictionary<string, Dictionary<string, FieldInfo>> CategoryToFields, Dictionary<Type, Dictionary<FieldInfo, object>> TypeToValues, XmlSchema Schema)
+		static bool TryReadFile(FileReference Location, Dictionary<string, Dictionary<string, FieldInfo>> CategoryToFields,
+			Dictionary<Type, Dictionary<FieldInfo, XmlConfigData.ValueInfo>> TypeToValues, XmlSchema Schema)
 		{
 			// Read the XML file, and validate it against the schema
 			XmlConfigFile? ConfigFile;
@@ -610,10 +636,9 @@ namespace UnrealBuildTool
 					foreach(XmlElement KeyElement in CategoryElement.ChildNodes.OfType<XmlElement>())
 					{
 						FieldInfo? Field;
+						object Value;
 						if(NameToField.TryGetValue(KeyElement.Name, out Field))
 						{
-							// Parse the corresponding value
-							object Value;
 							if (Field.FieldType == typeof(string[]))
 							{
 								Value = KeyElement.ChildNodes.OfType<XmlElement>().Where(x => x.Name == "Item").Select(x => x.InnerText).ToArray();
@@ -628,13 +653,18 @@ namespace UnrealBuildTool
 							}
 
 							// Add it to the set of values for the type containing this field
-							Dictionary<FieldInfo, object>? FieldToValue;
+							Dictionary<FieldInfo, XmlConfigData.ValueInfo>? FieldToValue;
 							if(!TypeToValues.TryGetValue(Field.DeclaringType!, out FieldToValue))
 							{
-								FieldToValue = new Dictionary<FieldInfo, object>();
+								FieldToValue = new Dictionary<FieldInfo, XmlConfigData.ValueInfo>();
 								TypeToValues.Add(Field.DeclaringType!, FieldToValue);
 							}
-							FieldToValue[Field] = Value;
+							
+							// Parse the corresponding value
+							XmlConfigData.ValueInfo FieldValue = new XmlConfigData.ValueInfo(Field, Value, Location,
+								Field.GetCustomAttribute<XmlConfigFileAttribute>()!);
+							
+							FieldToValue[Field] = FieldValue;
 						}
 					}
 				}
diff --git a/Engine/Source/Programs/UnrealBuildTool/System/XmlConfigData.cs b/Engine/Source/Programs/UnrealBuildTool/System/XmlConfigData.cs
--- a/Engine/Source/Programs/UnrealBuildTool/System/XmlConfigData.cs
+++ b/Engine/Source/Programs/UnrealBuildTool/System/XmlConfigData.cs
@@ -21,24 +21,40 @@ namespace UnrealBuildTool
 		/// <summary>
 		/// The current cache serialization version
 		/// </summary>
-		const int SerializationVersion = 1;
+		const int SerializationVersion = 2;
 
 		/// <summary>
 		/// List of input files. Stored to allow checking cache validity.
 		/// </summary>
 		public FileReference[] InputFiles;
 
+		public class ValueInfo
+		{
+			public FieldInfo FieldInfo;
+			public object Value;
+			public FileReference SourceFile;
+			public XmlConfigFileAttribute XmlConfigAttribute;
+
+			public ValueInfo(FieldInfo FieldInfo, object Value, FileReference SourceFile, XmlConfigFileAttribute XmlConfigAttribute)
+			{
+				this.FieldInfo = FieldInfo;
+				this.Value = Value;
+				this.SourceFile = SourceFile;
+				this.XmlConfigAttribute = XmlConfigAttribute;
+			}
+		}
+		
 		/// <summary>
 		/// Stores a mapping from type -> field -> value, with all the config values for configurable fields.
 		/// </summary>
-		public Dictionary<Type, KeyValuePair<FieldInfo, object>[]> TypeToValues;
+		public Dictionary<Type, ValueInfo[]> TypeToValues;
 
 		/// <summary>
 		/// Constructor
 		/// </summary>
 		/// <param name="InputFiles"></param>
 		/// <param name="TypeToValues"></param>
-		public XmlConfigData(FileReference[] InputFiles, Dictionary<Type, KeyValuePair<FieldInfo, object>[]> TypeToValues)
+		public XmlConfigData(FileReference[] InputFiles, Dictionary<Type, ValueInfo[]> TypeToValues)
 		{
 			this.InputFiles = InputFiles;
 			this.TypeToValues = TypeToValues;
@@ -75,7 +91,7 @@ namespace UnrealBuildTool
 
 				// Read the types
 				int NumTypes = Reader.ReadInt32();
-				Dictionary<Type, KeyValuePair<FieldInfo, object>[]> TypeToValues = new Dictionary<Type, KeyValuePair<FieldInfo, object>[]>(NumTypes);
+				Dictionary<Type, ValueInfo[]> TypeToValues = new Dictionary<Type, ValueInfo[]>(NumTypes);
 				for(int TypeIdx = 0; TypeIdx < NumTypes; TypeIdx++)
 				{
 					// Read the type name
@@ -90,14 +106,15 @@ namespace UnrealBuildTool
 					}
 
 					// Read all the values
-					KeyValuePair<FieldInfo, object>[] Values = new KeyValuePair<FieldInfo, object>[Reader.ReadInt32()];
+					ValueInfo[] Values = new ValueInfo[Reader.ReadInt32()];
 					for(int ValueIdx = 0; ValueIdx < Values.Length; ValueIdx++)
 					{
 						string FieldName = Reader.ReadString();
 
 						// Find the matching field on the output type
 						FieldInfo? Field = Type.GetField(FieldName, BindingFlags.Public | BindingFlags.NonPublic | BindingFlags.Instance | BindingFlags.Static);
-						if(Field == null || Field.GetCustomAttribute<XmlConfigFileAttribute>() == null)
+						XmlConfigFileAttribute? XmlConfigAttribute = Field?.GetCustomAttribute<XmlConfigFileAttribute>();
+						if(Field == null || XmlConfigAttribute == null)
 						{
 							Data = null;
 							return false;
@@ -105,7 +122,11 @@ namespace UnrealBuildTool
 
 						// Try to parse the value and add it to the output array
 						object Value = Reader.ReadObject(Field.FieldType)!;
-						Values[ValueIdx] = new KeyValuePair<FieldInfo, object>(Field, Value);
+						
+						// Read the path of the config file that provided this setting
+						FileReference SourceFile = Reader.ReadFileReference();
+
+						Values[ValueIdx] = new ValueInfo(Field, Value, SourceFile, XmlConfigAttribute);
 					}
 
 					// Add it to the type map
@@ -134,14 +155,15 @@ namespace UnrealBuildTool
 
 				// Write all the categories
 				Writer.Write(TypeToValues.Count);
-				foreach(KeyValuePair<Type, KeyValuePair<FieldInfo, object>[]> TypePair in TypeToValues)
+				foreach(KeyValuePair<Type, ValueInfo[]> TypePair in TypeToValues)
 				{
 					Writer.Write(TypePair.Key.Name);
 					Writer.Write(TypePair.Value.Length);
-					foreach(KeyValuePair<FieldInfo, object> FieldPair in TypePair.Value)
+					foreach(ValueInfo FieldPair in TypePair.Value)
 					{
-						Writer.Write(FieldPair.Key.Name);
-						Writer.Write(FieldPair.Key.FieldType, FieldPair.Value);
+						Writer.Write(FieldPair.FieldInfo.Name);
+						Writer.Write(FieldPair.FieldInfo.FieldType, FieldPair.Value);
+						Writer.Write(FieldPair.SourceFile);
 					}
 				}
 			}
diff --git a/Engine/Source/Programs/UnrealBuildTool/System/XmlConfigFileAttribute.cs b/Engine/Source/Programs/UnrealBuildTool/System/XmlConfigFileAttribute.cs
--- a/Engine/Source/Programs/UnrealBuildTool/System/XmlConfigFileAttribute.cs
+++ b/Engine/Source/Programs/UnrealBuildTool/System/XmlConfigFileAttribute.cs
@@ -23,5 +23,18 @@ namespace UnrealBuildTool
 		/// Name of the key to read. Optional; defaults to the field name.
 		/// </summary>
 		public string? Name = null;
+
+		/// <summary>
+		/// Use this field to indicate that the XML attribute has been deprecated, and that a warning should
+		/// be shown to the user if it is used.
+		/// A deprecated field should also be marked with the [Obsolete] attribute.
+		/// </summary>
+		public bool Deprecated = false;
+		
+		/// <summary>
+		/// If the attribute has been deprecated because it has been renamed, this field can be used to apply the
+		/// value used for this field to another.
+		/// </summary>
+		public string? NewAttributeName = null;
 	}
 }
