UnrealBuildTool: Show StatusDescription in debugger for LinkedAction

#trivial

diff --git a/Engine/Source/Programs/UnrealBuildTool/System/Action.cs b/Engine/Source/Programs/UnrealBuildTool/System/Action.cs
--- a/Engine/Source/Programs/UnrealBuildTool/System/Action.cs
+++ b/Engine/Source/Programs/UnrealBuildTool/System/Action.cs
@@ -2,6 +2,7 @@
 
 using System;
 using System.Collections.Generic;
+using System.Diagnostics;
 using System.IO;
 using System.Linq;
 using System.Reflection;
@@ -482,6 +483,7 @@ namespace UnrealBuildTool
 	/// <summary>
 	/// Information about an action queued to be executed
 	/// </summary>
+	[DebuggerDisplay("{StatusDescription}")]
 	class LinkedAction : IExternalAction
 	{
 		/// <summary>
