Add support for platform aliases to UnrealTargetPlatform, for use with deprecating or renaming an existing platform.

diff --git a/Engine/Source/Programs/UnrealBuildTool/Configuration/UEBuildTarget.cs b/Engine/Source/Programs/UnrealBuildTool/Configuration/UEBuildTarget.cs
--- a/Engine/Source/Programs/UnrealBuildTool/Configuration/UEBuildTarget.cs
+++ b/Engine/Source/Programs/UnrealBuildTool/Configuration/UEBuildTarget.cs
@@ -82,6 +82,19 @@ namespace UnrealBuildTool
 			return new UnrealTargetPlatform(GetUniqueStringRegistry().FindOrAddByName(Name));
 		}
 
+		/// <summary>
+		/// Return the single instance of the Group with this name
+		/// </summary>
+		/// <param name="Alias"></param>
+		/// /// <param name="Original"></param>
+		/// <returns></returns>
+		static private UnrealTargetPlatform AddAliasByName(string Alias, UnrealTargetPlatform Original)
+		{
+			GetUniqueStringRegistry().FindOrAddAlias(Alias, Original.ToString());
+			return Original;
+		}
+
+
 		/// <summary>
 		/// 
 		/// </summary>
@@ -150,6 +163,12 @@ namespace UnrealBuildTool
 				return true;
 			}
 
+			if (GetUniqueStringRegistry().HasAlias(Name))
+			{
+				Platform.Id = GetUniqueStringRegistry().FindExistingAlias(Name);
+				return true;
+			}
+
 			Platform.Id = -1;
 			return false;
 		}
@@ -166,6 +185,12 @@ namespace UnrealBuildTool
 				return new UnrealTargetPlatform(Name);
 			}
 
+			if (GetUniqueStringRegistry().HasAlias(Name))
+			{
+				int Id = GetUniqueStringRegistry().FindExistingAlias(Name);
+				return new UnrealTargetPlatform(Id);
+			}
+
 			throw new BuildException(string.Format("The platform name {0} is not a valid platform name. Valid names are ({1})", Name, 
 				string.Join(",", GetUniqueStringRegistry().GetStringNames())));
 		}
diff --git a/Engine/Source/Programs/UnrealBuildTool/System/UniqueStringRegistry.cs b/Engine/Source/Programs/UnrealBuildTool/System/UniqueStringRegistry.cs
--- a/Engine/Source/Programs/UnrealBuildTool/System/UniqueStringRegistry.cs
+++ b/Engine/Source/Programs/UnrealBuildTool/System/UniqueStringRegistry.cs
@@ -17,6 +17,9 @@ namespace UnrealBuildTool
 		// holds a mapping of string name to single instance
 		private Dictionary<string, int> StringToInstanceMap = new Dictionary<string, int>(StringComparer.OrdinalIgnoreCase);
 
+		// holds a mapping of string aliases to original string name
+		private Dictionary<string, string> AliasToStringMap = new Dictionary<string, string>(StringComparer.OrdinalIgnoreCase);
+
 		public UniqueStringRegistry()
 		{
 		}
@@ -65,5 +68,51 @@ namespace UnrealBuildTool
 		{
 			return StringToInstanceMap.First(x => x.Value == Id).Key;
 		}
+
+
+		public bool HasAlias(string Alias)
+		{
+			return AliasToStringMap.ContainsKey(Alias);
+		}
+
+		public int FindExistingAlias(string Alias)
+		{
+			string? Name;
+			if (!AliasToStringMap.TryGetValue(Alias, out Name) || Name == null)
+			{
+				throw new BuildException($"Alias {Alias} not found");
+			}
+
+			return FindOrAddByName(Name);
+		}
+
+		public int FindOrAddAlias(string Alias, string OriginalName)
+		{
+			string? Name;
+			if (!AliasToStringMap.TryGetValue(Alias, out Name))
+			{
+				lock (LockObject)
+				{
+					// copy over the dictionary and add, so that other threads can keep reading out of the old one until it's time
+					Dictionary<string, string> NewAliasToStringMap = new Dictionary<string, string>(AliasToStringMap, StringComparer.OrdinalIgnoreCase);
+
+					// add a new alias
+					NewAliasToStringMap[Alias] = OriginalName;
+
+					// replace the instance map
+					AliasToStringMap = NewAliasToStringMap;
+				}
+				Name = OriginalName;
+			}
+			else
+			{
+				if (Name != OriginalName || Name == null)
+				{
+					throw new BuildException($"{Alias} is already an alias for {Name}. Can't associate with {OriginalName}");
+				}
+			}
+
+			return FindOrAddByName(Name);
+		}
 	}
 }
\ No newline at end of file
