Fix assertion when generating Xcode project files for a single game target
#rnx

diff --git a/Engine/Source/Programs/UnrealBuildTool/ProjectFiles/Xcode/XcodeProject.cs b/Engine/Source/Programs/UnrealBuildTool/ProjectFiles/Xcode/XcodeProject.cs
--- a/Engine/Source/Programs/UnrealBuildTool/ProjectFiles/Xcode/XcodeProject.cs
+++ b/Engine/Source/Programs/UnrealBuildTool/ProjectFiles/Xcode/XcodeProject.cs
@@ -1563,12 +1563,13 @@ namespace UnrealBuildTool
 							UEBuildPlatform BuildPlatform;
 							if (UEBuildPlatform.TryGetBuildPlatform(Platform, out BuildPlatform) && (BuildPlatform.HasRequiredSDKsInstalled() == SDKStatus.Valid))
 							{
-								// Now go through all of the target types for this project
-								if (ProjectTargets.Count == 0)
+								// Check we have targets (Expected to be no Engine targets when generating for a single .uproject)
+								if (ProjectTargets.Count == 0 && BaseDir != Unreal.EngineDirectory)
 								{
 									throw new BuildException("Expecting at least one ProjectTarget to be associated with project '{0}' in the TargetProjects list ", ProjectFilePath);
 								}
 
+								// Now go through all of the target types for this project
 								foreach (ProjectTarget ProjectTarget in ProjectTargets)
 								{
 									if (MSBuildProjectFile.IsValidProjectPlatformAndConfiguration(ProjectTarget, Platform, Configuration, PlatformProjectGenerators))
