UnrealBuildTool: Exclude more C# projects from generated solution based on unsupported platform directories

C# projects found in Engine/Platforms/PLATFORM/... will not be included in the generated solution if UBT has decided that PLATFORM is not supported.

Additionally, provides extra logging of which directories and projects have been excluded from the solution.

diff --git a/Engine/Source/Programs/UnrealBuildTool/ProjectFiles/ProjectFileGenerator.cs b/Engine/Source/Programs/UnrealBuildTool/ProjectFiles/ProjectFileGenerator.cs
--- a/Engine/Source/Programs/UnrealBuildTool/ProjectFiles/ProjectFileGenerator.cs
+++ b/Engine/Source/Programs/UnrealBuildTool/ProjectFiles/ProjectFileGenerator.cs
@@ -528,7 +528,17 @@ namespace UnrealBuildTool
 			ProjectDirs.Add(EngineExtras);
 			DiscoverCSharpProgramProjectsRecursively(EngineExtras, FoundProjects);
 
-			List<DirectoryReference> AllEngineDirectories = Unreal.GetExtensionDirs(Unreal.EngineDirectory, "Source/Programs");
+			List<DirectoryReference> AllEngineDirectories =
+				Unreal.GetExtensionDirs(Unreal.EngineDirectory, "Source/Programs").Where(x =>
+				{
+					if (x.ContainsAnyNames(UnsupportedPlatformNames, Unreal.EngineDirectory))
+					{
+						Log.TraceLog($"Skipping any C# project files in \"{x}\" due to unsupported platform");
+						return false;
+					}
+					return true;
+				}).ToList();
+			
 			ProjectDirs = ProjectDirs.Union(AllEngineDirectories).ToList();
 			foreach (DirectoryReference EngineDir in AllEngineDirectories)
 			{
@@ -562,6 +572,10 @@ namespace UnrealBuildTool
 							AddExistingProjectFile(Project, bForceDevelopmentConfiguration: false);
 							ProgramsFolder.ChildProjects.Add(Project);
 						}
+						else
+						{
+							Log.TraceLog($"Skipping C# project \"{FoundProject}\" due to unsupported platform");
+						}
 						break;
 					}
 				}
