Do not generate target/platform combinations where the project requires modular builds and the platform requires monolithic

diff --git a/Engine/Source/Programs/UnrealBuildTool/ProjectFiles/Xcode/XcodeProject.cs b/Engine/Source/Programs/UnrealBuildTool/ProjectFiles/Xcode/XcodeProject.cs
--- a/Engine/Source/Programs/UnrealBuildTool/ProjectFiles/Xcode/XcodeProject.cs
+++ b/Engine/Source/Programs/UnrealBuildTool/ProjectFiles/Xcode/XcodeProject.cs
@@ -29,35 +29,35 @@ namespace UnrealBuildTool
 			FileRefGuid = XcodeProjectFileGenerator.MakeXcodeGuid();
 		}
 
-		/// <summary>
-		/// File Guid for use in Xcode project
-		/// </summary>
-		public string FileGuid
-		{
-			get;
-			private set;
-		}
-
-		public void ReplaceGuids(string NewFileGuid, string NewFileRefGuid)
-		{
-			FileGuid = NewFileGuid;
-			FileRefGuid = NewFileRefGuid;
-		}
+	/// <summary>
+	/// File Guid for use in Xcode project
+	/// </summary>
+	public string FileGuid
+	{
+		get;
+		private set;
+	}
 
-		/// <summary>
-		/// File reference Guid for use in Xcode project
-		/// </summary>
-		public string FileRefGuid
-		{
-			get;
-			private set;
-		}
+	public void ReplaceGuids(string NewFileGuid, string NewFileRefGuid)
+	{
+		FileGuid = NewFileGuid;
+		FileRefGuid = NewFileRefGuid;
 	}
 
 	/// <summary>
-	/// Represents a group of files shown in Xcode's project navigator as a folder
+	/// File reference Guid for use in Xcode project
 	/// </summary>
-	class XcodeFileGroup
+	public string FileRefGuid
+	{
+		get;
+		private set;
+	}
+	}
+
+		/// <summary>
+		/// Represents a group of files shown in Xcode's project navigator as a folder
+		/// </summary>
+		class XcodeFileGroup
 	{
 		public XcodeFileGroup(string InName, string InPath, bool InIsReference)
 		{
@@ -180,10 +180,10 @@ namespace UnrealBuildTool
 			// @todo Mac: Handle more categories
 			switch (Extension)
 			{
-				case ".framework":
-					return "Frameworks";
-				default:
-					return "Sources";
+			case ".framework":
+				return "Frameworks";
+			default:
+				return "Sources";
 			}
 		}
 
@@ -195,27 +195,27 @@ namespace UnrealBuildTool
 			// @todo Mac: Handle more file types
 			switch (Extension)
 			{
-				case ".c":
-				case ".m":
-					return "sourcecode.c.objc";
-				case ".cc":
-				case ".cpp":
-				case ".mm":
-					return "sourcecode.cpp.objcpp";
-				case ".h":
-				case ".inl":
-				case ".pch":
-					return "sourcecode.c.h";
-				case ".framework":
-					return "wrapper.framework";
-				case ".plist":
-					return "text.plist.xml";
-				case ".png":
-					return "image.png";
-				case ".icns":
-					return "image.icns";
-				default:
-					return "file.text";
+			case ".c":
+			case ".m":
+				return "sourcecode.c.objc";
+			case ".cc":
+			case ".cpp":
+			case ".mm":
+				return "sourcecode.cpp.objcpp";
+			case ".h":
+			case ".inl":
+			case ".pch":
+				return "sourcecode.c.h";
+			case ".framework":
+				return "wrapper.framework";
+			case ".plist":
+				return "text.plist.xml";
+			case ".png":
+				return "image.png";
+			case ".icns":
+				return "image.icns";
+			default:
+				return "file.text";
 			}
 		}
 
@@ -233,7 +233,7 @@ namespace UnrealBuildTool
 
 			if (IsSourceCode(FileExtension))
 			{
-				foreach (string PlatformName in UnrealTargetPlatform.GetValidPlatformNames())
+				foreach(string PlatformName in UnrealTargetPlatform.GetValidPlatformNames())
 				{
 					string AltName = PlatformName == "Win64" ? "windows" : PlatformName.ToLower();
 					if ((SourceFile.Reference.FullName.ToLower().Contains("/" + PlatformName.ToLower() + "/") || SourceFile.Reference.FullName.ToLower().Contains("/" + AltName + "/"))
@@ -323,9 +323,9 @@ namespace UnrealBuildTool
 		/// </summary>
 		private void GenerateSectionsWithSourceFiles(StringBuilder PBXBuildFileSection, StringBuilder PBXFileReferenceSection, StringBuilder PBXSourcesBuildPhaseSection, string TargetAppGuid, string TargetName, bool bIsAppBundle)
 		{
-			SourceFiles.Sort((x, y) => { return x.Reference.FullName.CompareTo(y.Reference.FullName); });
+			SourceFiles.Sort((x, y) = > { return x.Reference.FullName.CompareTo(y.Reference.FullName); });
 
-			foreach (SourceFile CurSourceFile in SourceFiles)
+			foreach(SourceFile CurSourceFile in SourceFiles)
 			{
 				XcodeSourceFile SourceFile = CurSourceFile as XcodeSourceFile;
 				string FileName = SourceFile.Reference.GetFileName();
@@ -364,7 +364,7 @@ namespace UnrealBuildTool
 		}
 
 		private void GenerateSectionsWithExtensions(StringBuilder PBXBuildFileSection, StringBuilder PBXFileReferenceSection, StringBuilder PBXCopyFilesBuildPhaseSection, StringBuilder PBXResourcesBuildPhaseSection,
-													List<XcodeExtensionInfo> AllExtensions, FileReference UProjectPath, List<XcodeBuildConfig> BuildConfigs)
+			List<XcodeExtensionInfo> AllExtensions, FileReference UProjectPath, List<XcodeBuildConfig> BuildConfigs)
 		{
 			if (UProjectPath != null)
 			{
@@ -373,7 +373,7 @@ namespace UnrealBuildTool
 
 				if (Directory.Exists(ProjectExtensionsDir))
 				{
-					foreach (DirectoryInfo DI in new System.IO.DirectoryInfo(ProjectExtensionsDir).EnumerateDirectories())
+					foreach(DirectoryInfo DI in new System.IO.DirectoryInfo(ProjectExtensionsDir).EnumerateDirectories())
 					{
 						Console.WriteLine("  Project {0} has Extension {1}!", UProjectPath, DI);
 
@@ -407,7 +407,7 @@ namespace UnrealBuildTool
 						if (Directory.Exists(Path.Combine(DI.FullName, "Resources")))
 						{
 							DirectoryInfo ResourceDir = new System.IO.DirectoryInfo(Path.Combine(DI.FullName, "Resources"));
-							foreach (FileSystemInfo FSI in ResourceDir.EnumerateFileSystemInfos())
+							foreach(FileSystemInfo FSI in ResourceDir.EnumerateFileSystemInfos())
 							{
 								if (FSI.Name.StartsWith("."))
 								{
@@ -438,7 +438,7 @@ namespace UnrealBuildTool
 
 						StringBuilder ConfigSection = new StringBuilder();
 						// copy over the configs from the general project to the extension
-						foreach (XcodeBuildConfig Configuration in BuildConfigs)
+						foreach(XcodeBuildConfig Configuration in BuildConfigs)
 						{
 							string ConfigGuid = XcodeProjectFileGenerator.MakeXcodeGuid();
 							string ConfigName = Configuration.DisplayName;
@@ -492,12 +492,12 @@ namespace UnrealBuildTool
 				Content.Append("\t\t\tisa = PBXGroup;" + ProjectFileGenerator.NewLine);
 				Content.Append("\t\t\tchildren = (" + ProjectFileGenerator.NewLine);
 
-				foreach (XcodeFileGroup ChildGroup in Group.Children.Values)
+				foreach(XcodeFileGroup ChildGroup in Group.Children.Values)
 				{
 					Content.Append(string.Format("\t\t\t\t{0} /* {1} */,{2}", ChildGroup.GroupGuid, ChildGroup.GroupName, ProjectFileGenerator.NewLine));
 				}
 
-				foreach (XcodeSourceFile File in Group.Files)
+				foreach(XcodeSourceFile File in Group.Files)
 				{
 					Content.Append(string.Format("\t\t\t\t{0} /* {1} */,{2}", File.FileRefGuid, File.Reference.GetFileName(), ProjectFileGenerator.NewLine));
 				}
@@ -508,7 +508,7 @@ namespace UnrealBuildTool
 				Content.Append("\t\t\tsourceTree = \"<absolute>\";" + ProjectFileGenerator.NewLine);
 				Content.Append("\t\t};" + ProjectFileGenerator.NewLine);
 
-				foreach (XcodeFileGroup ChildGroup in Group.Children.Values)
+				foreach(XcodeFileGroup ChildGroup in Group.Children.Values)
 				{
 					AppendGroup(ChildGroup, Content);
 				}
@@ -545,7 +545,7 @@ namespace UnrealBuildTool
 
 		private XcodeFileGroup FindRootFileGroup(Dictionary<string, XcodeFileGroup> GroupsDict)
 		{
-			foreach (XcodeFileGroup Group in GroupsDict.Values)
+			foreach(XcodeFileGroup Group in GroupsDict.Values)
 			{
 				if (Group.Children.Count > 1 || Group.Files.Count > 0)
 				{
@@ -587,7 +587,7 @@ namespace UnrealBuildTool
 			{
 				return;
 			}
-			
+
 			Content.Append("/* Begin PBXGroup section */" + ProjectFileGenerator.NewLine);
 
 			// Main group
@@ -595,12 +595,12 @@ namespace UnrealBuildTool
 			Content.Append("\t\t\tisa = PBXGroup;" + ProjectFileGenerator.NewLine);
 			Content.Append("\t\t\tchildren = (" + ProjectFileGenerator.NewLine);
 
-			foreach (XcodeFileGroup Group in RootGroup.Children.Values)
+			foreach(XcodeFileGroup Group in RootGroup.Children.Values)
 			{
 				Content.Append(string.Format("\t\t\t\t{0} /* {1} */,{2}", Group.GroupGuid, Group.GroupName, ProjectFileGenerator.NewLine));
 			}
 
-			foreach (XcodeSourceFile File in RootGroup.Files)
+			foreach(XcodeSourceFile File in RootGroup.Files)
 			{
 				Content.Append(string.Format("\t\t\t\t{0} /* {1} */,{2}", File.FileRefGuid, File.Reference.GetFileName(), ProjectFileGenerator.NewLine));
 			}
@@ -611,7 +611,7 @@ namespace UnrealBuildTool
 			Content.Append("\t\t};" + ProjectFileGenerator.NewLine);
 
 			// Sources groups
-			foreach (XcodeFileGroup Group in RootGroup.Children.Values)
+			foreach(XcodeFileGroup Group in RootGroup.Children.Values)
 			{
 				AppendGroup(Group, Content);
 			}
@@ -621,7 +621,7 @@ namespace UnrealBuildTool
 			Content.Append("\t\t\tisa = PBXGroup;" + ProjectFileGenerator.NewLine);
 			Content.Append("\t\t\tchildren = (" + ProjectFileGenerator.NewLine);
 			Content.Append(string.Format("\t\t\t\t{0} /* {1} */,{2}", TargetAppGuid, TargetName, ProjectFileGenerator.NewLine));
-			foreach (XcodeExtensionInfo EI in AllExtensions)
+			foreach(XcodeExtensionInfo EI in AllExtensions)
 			{
 				Content.Append(string.Format("\t\t\t\t{0} /* {1} */,{2}", EI.ProductGuid, EI.Name, ProjectFileGenerator.NewLine));
 			}
@@ -642,8 +642,8 @@ namespace UnrealBuildTool
 
 			Content.Append("\t\t" + TargetGuid + " /* " + TargetName + " */ = {" + ProjectFileGenerator.NewLine);
 			Content.Append("\t\t\tisa = PBXLegacyTarget;" + ProjectFileGenerator.NewLine);
-			Content.Append("\t\t\tbuildArgumentsString = \"$(ACTION) $(UE_BUILD_TARGET_NAME) $(PLATFORM_NAME) $(UE_BUILD_TARGET_CONFIG)" 
-				+ (UProjectPath == null ? "" : " \\\"" + UProjectPath.FullName + "\\\"") 
+			Content.Append("\t\t\tbuildArgumentsString = \"$(ACTION) $(UE_BUILD_TARGET_NAME) $(PLATFORM_NAME) $(UE_BUILD_TARGET_CONFIG)"
+				+ (UProjectPath == null ? "" : " \\\"" + UProjectPath.FullName + "\\\"")
 				+ "\";" + ProjectFileGenerator.NewLine);
 			Content.Append("\t\t\tbuildConfigurationList = "  + TargetBuildConfigGuid + " /* Build configuration list for PBXLegacyTarget \"" + TargetName + "\" */;" + ProjectFileGenerator.NewLine);
 			Content.Append("\t\t\tbuildPhases = (" + ProjectFileGenerator.NewLine);
@@ -660,8 +660,8 @@ namespace UnrealBuildTool
 			Content.Append("/* End PBXLegacyTarget section */" + ProjectFileGenerator.NewLine + ProjectFileGenerator.NewLine);
 		}
 
-		private void AppendRunTargetSection(StringBuilder Content, string TargetName, string TargetGuid, string TargetBuildConfigGuid, string TargetDependencyGuid, 
-				string TargetAppGuid, string CopyExtensionsBuildPhaseGuid, string ShellScriptSectionGuid, List<XcodeExtensionInfo> AllExtensions, bool bIsAppBundle)
+		private void AppendRunTargetSection(StringBuilder Content, string TargetName, string TargetGuid, string TargetBuildConfigGuid, string TargetDependencyGuid,
+			string TargetAppGuid, string CopyExtensionsBuildPhaseGuid, string ShellScriptSectionGuid, List<XcodeExtensionInfo> AllExtensions, bool bIsAppBundle)
 		{
 			List<string> DependencyGuids = new List<string>();
 			// depends on the Run target if we want one
@@ -670,7 +670,7 @@ namespace UnrealBuildTool
 				DependencyGuids.Add(TargetDependencyGuid);
 			}
 			// make sure extensions get built
-			foreach (XcodeExtensionInfo EI in AllExtensions)
+			foreach(XcodeExtensionInfo EI in AllExtensions)
 			{
 				DependencyGuids.Add(EI.TargetDependencyGuid);
 			}
@@ -700,7 +700,7 @@ namespace UnrealBuildTool
 			Content.Append("\t\t\tbuildPhases = (" + ProjectFileGenerator.NewLine);
 			if (BuildPhases != null)
 			{
-				foreach (KeyValuePair<string, string> BuildPhasePair in BuildPhases)
+				foreach(KeyValuePair<string, string> BuildPhasePair in BuildPhases)
 				{
 					Content.Append("\t\t\t\t" + BuildPhasePair.Key + " /* " + BuildPhasePair.Value + " */, " + ProjectFileGenerator.NewLine);
 				}
@@ -709,8 +709,8 @@ namespace UnrealBuildTool
 			Content.Append("\t\t\tdependencies = (" + ProjectFileGenerator.NewLine);
 			if (TargetDependencyGuids != null)
 			{
-				foreach (string DependencyGuid in TargetDependencyGuids)
-			{
+				foreach(string DependencyGuid in TargetDependencyGuids)
+				{
 					Content.Append("\t\t\t\t" + DependencyGuid + " /* PBXTargetDependency */," + ProjectFileGenerator.NewLine);
 				}
 			}
@@ -755,14 +755,14 @@ namespace UnrealBuildTool
 			{
 				return;
 			}
-			
+
 			// @todo: look also in Project/Build/Frameworks directory!
 			ProjectDescriptor Project = ProjectDescriptor.FromFile(UProjectPath);
 			List<PluginInfo> AvailablePlugins = Plugins.ReadAvailablePlugins(UnrealBuildTool.EngineDirectory, DirectoryReference.FromFile(UProjectPath), Project.AdditionalPluginDirectories);
 
 			// look in each plugin for frameworks
 			// @todo: Cache this kind of things since every target will re-do this work!
-			foreach (PluginInfo PI in AvailablePlugins)
+			foreach(PluginInfo PI in AvailablePlugins)
 			{
 				if (!Plugins.IsPluginEnabledForTarget(PI, Project, UnrealTargetPlatform.IOS, UnrealTargetConfiguration.Development, TargetRules.TargetType.Game))
 				{
@@ -778,7 +778,7 @@ namespace UnrealBuildTool
 				if (DirectoryReference.Exists(FrameworkDir))
 				{
 					// look at each zip
-					foreach (FileInfo FI in new System.IO.DirectoryInfo(FrameworkDir.FullName).EnumerateFiles("*.framework.zip"))
+					foreach(FileInfo FI in new System.IO.DirectoryInfo(FrameworkDir.FullName).EnumerateFiles("*.framework.zip"))
 					{
 						//string Guid = XcodeProjectFileGenerator.MakeXcodeGuid();
 						//string RefGuid = XcodeProjectFileGenerator.MakeXcodeGuid();
@@ -796,12 +796,12 @@ namespace UnrealBuildTool
 			string ShellScript = "set -e\\n\\nIFS=$'\\\\n'\\n\\n" +
 				"if [ $PLATFORM_NAME = iphoneos ] || [ $PLATFORM_NAME = tvos ]; then \\n" +
 				"\\tFRAMEWORK_DIR=$TARGET_BUILD_DIR/$EXECUTABLE_FOLDER_PATH/Frameworks\\n" +
-				FrameworkScript.ToString() + 
+				FrameworkScript.ToString() +
 				// and now code sign anything that has been unzipped above
 				"\\tfor FRAMEWORK in ${FRAMEWORK_DIR}/*.framework; do\\n" +
-					"\\t\\t[ -d \\\"${FRAMEWORK}\\\" ] || continue\\n" +
-					"\\t\\techo Codesigning ${FRAMEWORK}\\n" +
-					"\\t\\tcodesign --force --sign ${EXPANDED_CODE_SIGN_IDENTITY} --verbose --preserve-metadata=identifier,entitlements,flags --timestamp=none \\\"${FRAMEWORK}\\\"\\n" +
+				"\\t\\t[ -d \\\"${FRAMEWORK}\\\" ] || continue\\n" +
+				"\\t\\techo Codesigning ${FRAMEWORK}\\n" +
+				"\\t\\tcodesign --force --sign ${EXPANDED_CODE_SIGN_IDENTITY} --verbose --preserve-metadata=identifier,entitlements,flags --timestamp=none \\\"${FRAMEWORK}\\\"\\n" +
 				"\\tdone\\n" +
 				"fi\\n";
 
@@ -825,7 +825,7 @@ namespace UnrealBuildTool
 
 		private void AppendExtensionTargetSections(StringBuilder ProjectFileContent, List<XcodeExtensionInfo> AllExtensions)
 		{
-			foreach (XcodeExtensionInfo EI in AllExtensions)
+			foreach(XcodeExtensionInfo EI in AllExtensions)
 			{
 				Dictionary<string, string> BuildPhases = new Dictionary<string, string>();
 				BuildPhases.Add(EI.ResourceBuildPhaseGuid, "Resources");
@@ -843,8 +843,8 @@ namespace UnrealBuildTool
 			Content.Append("\t\t\tattributes = {" + ProjectFileGenerator.NewLine);
 			Content.Append("\t\t\t\tLastUpgradeCheck = 2000;" + ProjectFileGenerator.NewLine);
 			Content.Append("\t\t\t\tORGANIZATIONNAME = \"Epic Games, Inc.\";" + ProjectFileGenerator.NewLine);
-            Content.Append("\t\t\t\tTargetAttributes = {" + ProjectFileGenerator.NewLine);
-            Content.Append("\t\t\t\t\t" + TargetGuid + " = {" + ProjectFileGenerator.NewLine);
+			Content.Append("\t\t\t\tTargetAttributes = {" + ProjectFileGenerator.NewLine);
+			Content.Append("\t\t\t\t\t" + TargetGuid + " = {" + ProjectFileGenerator.NewLine);
 
 			bool bAutomaticSigning = false;
 			if (InstalledPlatformInfo.IsValidPlatform(UnrealTargetPlatform.IOS, EProjectType.Code))
@@ -870,9 +870,9 @@ namespace UnrealBuildTool
 			{
 				Content.Append("\t\t\t\t\t\tProvisioningStyle = Manual;" + ProjectFileGenerator.NewLine);
 			}
-            Content.Append("\t\t\t\t\t};" + ProjectFileGenerator.NewLine);
-            Content.Append("\t\t\t\t};" + ProjectFileGenerator.NewLine);
-            Content.Append("\t\t\t};" + ProjectFileGenerator.NewLine);
+			Content.Append("\t\t\t\t\t};" + ProjectFileGenerator.NewLine);
+			Content.Append("\t\t\t\t};" + ProjectFileGenerator.NewLine);
+			Content.Append("\t\t\t};" + ProjectFileGenerator.NewLine);
 			Content.Append("\t\t\tbuildConfigurationList = " + ProjectBuildConfigGuid + " /* Build configuration list for PBXProject \"" + TargetName + "\" */;" + ProjectFileGenerator.NewLine);
 			Content.Append("\t\t\tcompatibilityVersion = \"Xcode 8.0\";" + ProjectFileGenerator.NewLine);
 			Content.Append("\t\t\tdevelopmentRegion = English;" + ProjectFileGenerator.NewLine);
@@ -888,7 +888,7 @@ namespace UnrealBuildTool
 			Content.Append("\t\t\t" + TargetGuid + " /* " + TargetName + " */," + ProjectFileGenerator.NewLine);
 			Content.Append("\t\t\t" + BuildTargetGuid + " /* " + BuildTargetName + " */," + ProjectFileGenerator.NewLine);
 			Content.Append("\t\t\t" + IndexTargetGuid + " /* " + IndexTargetName + " */," + ProjectFileGenerator.NewLine);
-			foreach (XcodeExtensionInfo EI in AllExtensions)
+			foreach(XcodeExtensionInfo EI in AllExtensions)
 			{
 				Content.Append("\t\t\t" + EI.TargetGuid + " /* " + EI.Name + " */," + ProjectFileGenerator.NewLine);
 			}
@@ -929,7 +929,7 @@ namespace UnrealBuildTool
 			Content.Append("\t\t\tbuildSettings = {" + ProjectFileGenerator.NewLine);
 
 			Content.Append("\t\t\t\tGCC_PREPROCESSOR_DEFINITIONS = (" + ProjectFileGenerator.NewLine);
-			foreach (string Definition in IntelliSensePreprocessorDefinitions)
+			foreach(string Definition in IntelliSensePreprocessorDefinitions)
 			{
 				Content.Append("\t\t\t\t\t\"" + Definition.Replace("\"", "").Replace("\\", "") + "\"," + ProjectFileGenerator.NewLine);
 			}
@@ -938,7 +938,7 @@ namespace UnrealBuildTool
 			Content.Append("\t\t\t\t);" + ProjectFileGenerator.NewLine);
 
 			Content.Append("\t\t\t\tHEADER_SEARCH_PATHS = (" + ProjectFileGenerator.NewLine);
-			foreach (string SearchPath in IntelliSenseSystemIncludeSearchPaths)
+			foreach(string SearchPath in IntelliSenseSystemIncludeSearchPaths)
 			{
 				string Path = SearchPath.Contains(" ") ? "\\\"" + SearchPath + "\\\"" : SearchPath;
 				Content.Append("\t\t\t\t\t\"" + Path + "\"," + ProjectFileGenerator.NewLine);
@@ -946,7 +946,7 @@ namespace UnrealBuildTool
 			Content.Append("\t\t\t\t);" + ProjectFileGenerator.NewLine);
 
 			Content.Append("\t\t\t\tUSER_HEADER_SEARCH_PATHS = (" + ProjectFileGenerator.NewLine);
-			foreach (string SearchPath in IntelliSenseIncludeSearchPaths)
+			foreach(string SearchPath in IntelliSenseIncludeSearchPaths)
 			{
 				string Path = SearchPath.Contains(" ") ? "\\\"" + SearchPath + "\\\"" : SearchPath;
 				Content.Append("\t\t\t\t\t\"" + Path + "\"," + ProjectFileGenerator.NewLine);
@@ -983,7 +983,7 @@ namespace UnrealBuildTool
 		IEnumerable<string> GetSupportedMacArchitectures(XcodeBuildConfig Config, FileReference InProjectFile)
 		{
 			// All architectures supported
-			IEnumerable<string> AllArchitectures = new[] { MacExports.IntelArchitecture, MacExports.AppleArchitecture};			
+			IEnumerable<string> AllArchitectures = new[] { MacExports.IntelArchitecture, MacExports.AppleArchitecture};
 
 			// Add a way on the command line of forcing a project file with all architectures (there isn't a good way to let this be
 			// set and checked where we can access it).
@@ -1108,17 +1108,17 @@ namespace UnrealBuildTool
 				string SupportedPlatforms = bSupportMac ? "macosx" : "";
 
 				bool bAutomaticSigning = false;
-                string UUID_IOS = "";
-                string UUID_TVOS = "";
-                string TEAM_IOS = "";
-                string TEAM_TVOS = "";
-                string IOS_CERT = "iPhone Developer";
-                string TVOS_CERT = "iPhone Developer";
+				string UUID_IOS = "";
+				string UUID_TVOS = "";
+				string TEAM_IOS = "";
+				string TEAM_TVOS = "";
+				string IOS_CERT = "iPhone Developer";
+				string TVOS_CERT = "iPhone Developer";
 				string IOS_BUNDLE = "";
 				string TVOS_BUNDLE = "";
 				if (bSupportIOS && InstalledPlatformInfo.IsValidPlatform(UnrealTargetPlatform.IOS, EProjectType.Code))
-                {
- 					IOSPlatform IOSPlatform = ((IOSPlatform)UEBuildPlatform.GetBuildPlatform(UnrealTargetPlatform.IOS));
+				{
+					IOSPlatform IOSPlatform = ((IOSPlatform)UEBuildPlatform.GetBuildPlatform(UnrealTargetPlatform.IOS));
 					IOSProjectSettings ProjectSettings = IOSPlatform.ReadProjectSettings(ProjectFile);
 					IOSProvisioningData ProvisioningData = IOSPlatform.ReadProvisioningData(ProjectSettings, bForDistribution);
 					IOSRunTimeVersion = ProjectSettings.RuntimeVersion;
@@ -1130,9 +1130,9 @@ namespace UnrealBuildTool
 						UUID_IOS = ProvisioningData.MobileProvisionUUID;
 						IOS_CERT = ProvisioningData.SigningCertificate;
 					}
-                    TEAM_IOS = ProvisioningData.TeamUUID;
+					TEAM_IOS = ProvisioningData.TeamUUID;
 					IOS_BUNDLE = ProjectSettings.BundleIdentifier;
-                }
+				}
 
 				if (bSupportTVOS && InstalledPlatformInfo.IsValidPlatform(UnrealTargetPlatform.TVOS, EProjectType.Code))
 				{
@@ -1147,9 +1147,9 @@ namespace UnrealBuildTool
 						UUID_TVOS = ProvisioningData.MobileProvisionUUID;
 						TVOS_CERT = ProvisioningData.SigningCertificate;
 					}
-                    TEAM_TVOS = ProvisioningData.TeamUUID;
+					TEAM_TVOS = ProvisioningData.TeamUUID;
 					TVOS_BUNDLE = ProjectSettings.BundleIdentifier;
-                }
+				}
 
 				Content.Append("\t\t\t\tSUPPORTED_PLATFORMS = \"" + SupportedPlatforms.Trim() + "\";" + ProjectFileGenerator.NewLine);
 				if (bAutomaticSigning)
@@ -1162,7 +1162,7 @@ namespace UnrealBuildTool
 					Content.Append("\t\t\t\tIPHONEOS_DEPLOYMENT_TARGET = " + IOSRunTimeVersion + ";" + ProjectFileGenerator.NewLine);
 					Content.Append("\t\t\t\t\"PRODUCT_NAME[sdk=iphoneos*]\" = \"" + TargetName + "\";" + ProjectFileGenerator.NewLine); // @todo: change to Path.GetFileName(Config.IOSExecutablePath) when we stop using payload
 					Content.Append("\t\t\t\t\"TARGETED_DEVICE_FAMILY[sdk=iphoneos*]\" = \"" + IOSRunTimeDevices + "\";" + ProjectFileGenerator.NewLine);
-                    Content.Append("\t\t\t\t\"SDKROOT[sdk=iphoneos]\" = iphoneos;" + ProjectFileGenerator.NewLine);
+					Content.Append("\t\t\t\t\"SDKROOT[sdk=iphoneos]\" = iphoneos;" + ProjectFileGenerator.NewLine);
 					if (!string.IsNullOrEmpty(TEAM_IOS))
 					{
 						Content.Append("\t\t\t\t\"DEVELOPMENT_TEAM[sdk=iphoneos*]\" = " + TEAM_IOS + ";" + ProjectFileGenerator.NewLine);
@@ -1177,13 +1177,13 @@ namespace UnrealBuildTool
 						Content.Append("\t\t\t\t\"PRODUCT_BUNDLE_IDENTIFIER[sdk=iphoneos*]\" = " + IOS_BUNDLE + ";" + ProjectFileGenerator.NewLine);
 					}
 				}
-                if (TVOSRunTimeVersion != null)
+				if (TVOSRunTimeVersion != null)
 				{
 					Content.Append("\t\t\t\t\"VALID_ARCHS[sdk=appletvos*]\" = \"" + string.Join(" ", this.SupportedIOSArchitectures) + "\";" + ProjectFileGenerator.NewLine);
 					Content.Append("\t\t\t\tTVOS_DEPLOYMENT_TARGET = " + TVOSRunTimeVersion + ";" + ProjectFileGenerator.NewLine);
 					Content.Append("\t\t\t\t\"PRODUCT_NAME[sdk=appletvos*]\" = \"" + TargetName + "\";" + ProjectFileGenerator.NewLine); // @todo: change to Path.GetFileName(Config.TVOSExecutablePath) when we stop using payload
 					Content.Append("\t\t\t\t\"TARGETED_DEVICE_FAMILY[sdk=appletvos*]\" = \"" + TVOSRunTimeDevices + "\";" + ProjectFileGenerator.NewLine);
-                    Content.Append("\t\t\t\t\"SDKROOT[sdk=appletvos]\" = appletvos;" + ProjectFileGenerator.NewLine);
+					Content.Append("\t\t\t\t\"SDKROOT[sdk=appletvos]\" = appletvos;" + ProjectFileGenerator.NewLine);
 					if (!string.IsNullOrEmpty(TEAM_TVOS))
 					{
 						Content.Append("\t\t\t\t\"DEVELOPMENT_TEAM[sdk=appletvos*]\" = " + TEAM_TVOS + ";" + ProjectFileGenerator.NewLine);
@@ -1293,7 +1293,7 @@ namespace UnrealBuildTool
 					MacInfoPlistPath = UEDir + "/Engine/Intermediate/Mac/" + MacExecutableFileName + "-Info.plist";
 					IOSEntitlementPath = "";
 					TVOSEntitlementPath = "";
-					}
+				}
 				else if (ProjectFile != null)
 				{
 					IOSInfoPlistPath = GamePath + "/Intermediate/IOS/" + Config.BuildTarget + "-Info.plist";
@@ -1385,7 +1385,7 @@ namespace UnrealBuildTool
 								ReceiptFilename = TargetReceipt.GetDefaultPath(ProjectPath, GameName, UnrealTargetPlatform.IOS, Config.BuildConfig, "");
 							}
 							Directory.CreateDirectory(Path.GetDirectoryName(IOSInfoPlistPath));
-                            bool bSupportPortrait, bSupportLandscape, bSkipIcons;
+							bool bSupportPortrait, bSupportLandscape, bSkipIcons;
 							TargetReceipt Receipt;
 							TargetReceipt.TryRead(ReceiptFilename, out Receipt);
 							VersionNumber SdkVersion = UEDeployIOS.GetSdkVersion(Receipt);
@@ -1403,22 +1403,22 @@ namespace UnrealBuildTool
 			Content.Append("\t\t\t\tMACOSX_DEPLOYMENT_TARGET = " + MacToolChain.Settings.MacOSVersion + ";" + ProjectFileGenerator.NewLine);
 			Content.Append("\t\t\t\tINFOPLIST_OUTPUT_FORMAT = xml;" + ProjectFileGenerator.NewLine);
 			Content.Append("\t\t\t\tCOMBINE_HIDPI_IMAGES = YES;" + ProjectFileGenerator.NewLine);
-			
-            //#jira UE-50382 Xcode Address Sanitizer feature does not work on iOS
-            // address sanitizer dylib loader depends on the SDKROOT parameter. For macosx or default (missing, translated as macosx), the path is incorrect for iphone/appletv
-            if (XcodeProjectFileGenerator.bGeneratingRunIOSProject)
-            {
-                Content.Append("\t\t\t\tSDKROOT = iphoneos;" + ProjectFileGenerator.NewLine);
-            }
-            else if (XcodeProjectFileGenerator.bGeneratingRunTVOSProject)
-            {
-                Content.Append("\t\t\t\tSDKROOT = appletvos;" + ProjectFileGenerator.NewLine);
-            }
-            else
-            {
-                Content.Append("\t\t\t\tSDKROOT = macosx;" + ProjectFileGenerator.NewLine);
-            }
-            Content.Append("\t\t\t\tGCC_PRECOMPILE_PREFIX_HEADER = YES;" + ProjectFileGenerator.NewLine);
+
+			//#jira UE-50382 Xcode Address Sanitizer feature does not work on iOS
+			// address sanitizer dylib loader depends on the SDKROOT parameter. For macosx or default (missing, translated as macosx), the path is incorrect for iphone/appletv
+			if (XcodeProjectFileGenerator.bGeneratingRunIOSProject)
+			{
+				Content.Append("\t\t\t\tSDKROOT = iphoneos;" + ProjectFileGenerator.NewLine);
+			}
+			else if (XcodeProjectFileGenerator.bGeneratingRunTVOSProject)
+			{
+				Content.Append("\t\t\t\tSDKROOT = appletvos;" + ProjectFileGenerator.NewLine);
+			}
+			else
+			{
+				Content.Append("\t\t\t\tSDKROOT = macosx;" + ProjectFileGenerator.NewLine);
+			}
+			Content.Append("\t\t\t\tGCC_PRECOMPILE_PREFIX_HEADER = YES;" + ProjectFileGenerator.NewLine);
 			Content.Append("\t\t\t\tGCC_PREFIX_HEADER = \"" + UEDir + "/Engine/Source/Editor/UnrealEd/Public/UnrealEd.h\";" + ProjectFileGenerator.NewLine);
 			Content.Append("\t\t\t};" + ProjectFileGenerator.NewLine);
 			Content.Append("\t\t\tname = \"" + Config.DisplayName + "\";" + ProjectFileGenerator.NewLine);
@@ -1476,32 +1476,32 @@ namespace UnrealBuildTool
 
 
 		private void AppendXCBuildConfigurationSection(StringBuilder Content, Dictionary<string, XcodeBuildConfig> ProjectBuildConfigs, Dictionary<string, XcodeBuildConfig> TargetBuildConfigs,
-				Dictionary<string, XcodeBuildConfig> BuildTargetBuildConfigs, Dictionary<string, XcodeBuildConfig> IndexTargetBuildConfigs, FileReference GameProjectPath,
-				List<XcodeExtensionInfo> AllExtensions)
+			Dictionary<string, XcodeBuildConfig> BuildTargetBuildConfigs, Dictionary<string, XcodeBuildConfig> IndexTargetBuildConfigs, FileReference GameProjectPath,
+			List<XcodeExtensionInfo> AllExtensions)
 		{
 			Content.Append("/* Begin XCBuildConfiguration section */" + ProjectFileGenerator.NewLine);
 
-			foreach (KeyValuePair<string, XcodeBuildConfig> Config in ProjectBuildConfigs)
+			foreach(KeyValuePair<string, XcodeBuildConfig> Config in ProjectBuildConfigs)
 			{
 				AppendProjectBuildConfiguration(Content, Config.Value.DisplayName, Config.Key);
 			}
 
-			foreach (KeyValuePair<string, XcodeBuildConfig> Config in TargetBuildConfigs)
+			foreach(KeyValuePair<string, XcodeBuildConfig> Config in TargetBuildConfigs)
 			{
 				AppendNativeTargetBuildConfiguration(Content, Config.Value, Config.Key, GameProjectPath);
 			}
 
-			foreach (KeyValuePair<string, XcodeBuildConfig> Config in BuildTargetBuildConfigs)
+			foreach(KeyValuePair<string, XcodeBuildConfig> Config in BuildTargetBuildConfigs)
 			{
 				AppendLegacyTargetBuildConfiguration(Content, Config.Value, Config.Key, GameProjectPath);
 			}
 
-			foreach (KeyValuePair<string, XcodeBuildConfig> Config in IndexTargetBuildConfigs)
+			foreach(KeyValuePair<string, XcodeBuildConfig> Config in IndexTargetBuildConfigs)
 			{
 				AppendNativeTargetBuildConfiguration(Content, Config.Value, Config.Key, GameProjectPath);
 			}
 
-			foreach (XcodeExtensionInfo EI in AllExtensions)
+			foreach(XcodeExtensionInfo EI in AllExtensions)
 			{
 				Content.Append(EI.ConfigurationContents);
 			}
@@ -1514,7 +1514,7 @@ namespace UnrealBuildTool
 			Content.Append("\t\t" + ConfigListGuid + " /* Build configuration list for " + TypeName + " \"" + TargetName + "\" */ = {" + ProjectFileGenerator.NewLine);
 			Content.Append("\t\t\tisa = XCConfigurationList;" + ProjectFileGenerator.NewLine);
 			Content.Append("\t\t\tbuildConfigurations = (" + ProjectFileGenerator.NewLine);
-			foreach (KeyValuePair<string, XcodeBuildConfig> Config in BuildConfigs)
+			foreach(KeyValuePair<string, XcodeBuildConfig> Config in BuildConfigs)
 			{
 				Content.Append("\t\t\t\t" + Config.Key + " /* \"" + Config.Value.DisplayName + "\" */," + ProjectFileGenerator.NewLine);
 			}
@@ -1537,13 +1537,13 @@ namespace UnrealBuildTool
 			AppendXCConfigurationList(Content, "PBXNativeTarget", TargetName, TargetConfigListGuid, TargetBuildConfigs);
 			AppendXCConfigurationList(Content, "PBXNativeTarget", IndexTargetName, IndexTargetConfigListGuid, IndexTargetBuildConfigs);
 
-			foreach (XcodeExtensionInfo EI in AllExtensions)
+			foreach(XcodeExtensionInfo EI in AllExtensions)
 			{
 				AppendXCConfigurationList(Content, "PBXNativeTarget", EI.Name, EI.ConfigListGuid, EI.AllConfigs);
-		}
+			}
 
 			Content.Append("/* End XCConfigurationList section */" + ProjectFileGenerator.NewLine);
-			}
+		}
 
 		private List<XcodeBuildConfig> GetSupportedBuildConfigs(List<UnrealTargetPlatform> Platforms, List<UnrealTargetConfiguration> Configurations, PlatformProjectGeneratorCollection PlatformProjectGenerators)
 		{
@@ -1551,11 +1551,11 @@ namespace UnrealBuildTool
 
 			//string ProjectName = ProjectFilePath.GetFileNameWithoutExtension();
 
-			foreach (UnrealTargetConfiguration Configuration in Configurations)
+			foreach(UnrealTargetConfiguration Configuration in Configurations)
 			{
 				if (InstalledPlatformInfo.IsValidConfiguration(Configuration, EProjectType.Code))
 				{
-					foreach (UnrealTargetPlatform Platform in Platforms)
+					foreach(UnrealTargetPlatform Platform in Platforms)
 					{
 						if (InstalledPlatformInfo.IsValidPlatform(Platform, EProjectType.Code) && (Platform == UnrealTargetPlatform.Mac || Platform == UnrealTargetPlatform.IOS || Platform == UnrealTargetPlatform.TVOS)) // @todo support other platforms
 						{
@@ -1568,64 +1568,71 @@ namespace UnrealBuildTool
 									throw new BuildException("Expecting at least one ProjectTarget to be associated with project '{0}' in the TargetProjects list ", ProjectFilePath);
 								}
 
-								foreach (ProjectTarget ProjectTarget in ProjectTargets)
+								foreach(ProjectTarget ProjectTarget in ProjectTargets)
 								{
 									if (MSBuildProjectFile.IsValidProjectPlatformAndConfiguration(ProjectTarget, Platform, Configuration, PlatformProjectGenerators))
 									{
-										// Figure out if this is a monolithic build
-										bool bShouldCompileMonolithic = BuildPlatform.ShouldCompileMonolithicBinary(Platform);
-										bShouldCompileMonolithic |= (ProjectTarget.CreateRulesDelegate(Platform, Configuration).LinkType == TargetLinkType.Monolithic);
-
-										string ConfigName = Configuration.ToString();
-										if (ProjectTarget.TargetRules.Type != TargetType.Game && ProjectTarget.TargetRules.Type != TargetType.Program)
+										try
 										{
-											ConfigName += " " + ProjectTarget.TargetRules.Type.ToString();
-										}
+											// Figure out if this is a monolithic build
+											bool bShouldCompileMonolithic = BuildPlatform.ShouldCompileMonolithicBinary(Platform);
+											bShouldCompileMonolithic |= (ProjectTarget.CreateRulesDelegate(Platform, Configuration).LinkType == TargetLinkType.Monolithic);
 
-										if (BuildConfigs.Where(Config => Config.DisplayName == ConfigName).ToList().Count == 0)
-										{
-											string TargetName = ProjectTarget.TargetFilePath.GetFileNameWithoutAnyExtensions();
+											string ConfigName = Configuration.ToString();
+											if (ProjectTarget.TargetRules.Type != TargetType.Game && ProjectTarget.TargetRules.Type != TargetType.Program)
+											{
+												ConfigName += " " + ProjectTarget.TargetRules.Type.ToString();
+											}
 
-											// Get the output directory
-											DirectoryReference RootDirectory = UnrealBuildTool.EngineDirectory;
-											if ((ProjectTarget.TargetRules.Type == TargetType.Game || ProjectTarget.TargetRules.Type == TargetType.Client || ProjectTarget.TargetRules.Type == TargetType.Server) && bShouldCompileMonolithic)
+											if (BuildConfigs.Where(Config = > Config.DisplayName == ConfigName).ToList().Count == 0)
 											{
-												if(ProjectTarget.UnrealProjectFilePath != null)
+												string TargetName = ProjectTarget.TargetFilePath.GetFileNameWithoutAnyExtensions();
+
+												// Get the output directory
+												DirectoryReference RootDirectory = UnrealBuildTool.EngineDirectory;
+												if ((ProjectTarget.TargetRules.Type == TargetType.Game || ProjectTarget.TargetRules.Type == TargetType.Client || ProjectTarget.TargetRules.Type == TargetType.Server) && bShouldCompileMonolithic)
+												{
+													if (ProjectTarget.UnrealProjectFilePath != null)
+													{
+														RootDirectory = ProjectTarget.UnrealProjectFilePath.Directory;
+													}
+												}
+
+												if (ProjectTarget.TargetRules.Type == TargetType.Program && ProjectTarget.UnrealProjectFilePath != null)
 												{
 													RootDirectory = ProjectTarget.UnrealProjectFilePath.Directory;
 												}
-											}
 
-											if(ProjectTarget.TargetRules.Type == TargetType.Program && ProjectTarget.UnrealProjectFilePath != null)
-											{
-												RootDirectory = ProjectTarget.UnrealProjectFilePath.Directory;
-											}
+												// Get the output directory
+												DirectoryReference OutputDirectory = DirectoryReference.Combine(RootDirectory, "Binaries");
 
-											// Get the output directory
-											DirectoryReference OutputDirectory = DirectoryReference.Combine(RootDirectory, "Binaries");
+												string ExeName = TargetName;
+												if (!bShouldCompileMonolithic && ProjectTarget.TargetRules.Type != TargetType.Program)
+												{
+													// Figure out what the compiled binary will be called so that we can point the IDE to the correct file
+													ExeName = "Unreal" + ProjectTarget.TargetRules.Type.ToString();
+												}
 
-											string ExeName = TargetName;
-											if (!bShouldCompileMonolithic && ProjectTarget.TargetRules.Type != TargetType.Program)
-											{
-												// Figure out what the compiled binary will be called so that we can point the IDE to the correct file
-												ExeName = "Unreal" + ProjectTarget.TargetRules.Type.ToString();
+												if (BuildPlatform.Platform == UnrealTargetPlatform.Mac)
+												{
+													string MacExecutableName = MakeExecutableFileName(ExeName, UnrealTargetPlatform.Mac, Configuration, ProjectTarget.TargetRules.Architecture, ProjectTarget.TargetRules.UndecoratedConfiguration);
+													string IOSExecutableName = MacExecutableName.Replace("-Mac-", "-IOS-");
+													string TVOSExecutableName = MacExecutableName.Replace("-Mac-", "-TVOS-");
+													BuildConfigs.Add(new XcodeBuildConfig(ConfigName, TargetName, FileReference.Combine(OutputDirectory, "Mac", MacExecutableName), FileReference.Combine(OutputDirectory, "IOS", IOSExecutableName), FileReference.Combine(OutputDirectory, "TVOS", TVOSExecutableName), ProjectTarget, Configuration));
+												}
+												else if (BuildPlatform.Platform == UnrealTargetPlatform.IOS || BuildPlatform.Platform == UnrealTargetPlatform.TVOS)
+												{
+													string IOSExecutableName = MakeExecutableFileName(ExeName, UnrealTargetPlatform.IOS, Configuration, ProjectTarget.TargetRules.Architecture, ProjectTarget.TargetRules.UndecoratedConfiguration);
+													string TVOSExecutableName = IOSExecutableName.Replace("-IOS-", "-TVOS-");
+													//string MacExecutableName = IOSExecutableName.Replace("-IOS-", "-Mac-");
+													BuildConfigs.Add(new XcodeBuildConfig(ConfigName, TargetName, FileReference.Combine(OutputDirectory, "Mac", IOSExecutableName), FileReference.Combine(OutputDirectory, "IOS", IOSExecutableName), FileReference.Combine(OutputDirectory, "TVOS", TVOSExecutableName), ProjectTarget, Configuration));
+												}
 											}
-
-                                            if (BuildPlatform.Platform == UnrealTargetPlatform.Mac)
-                                            {
-                                                string MacExecutableName = MakeExecutableFileName(ExeName, UnrealTargetPlatform.Mac, Configuration, ProjectTarget.TargetRules.Architecture, ProjectTarget.TargetRules.UndecoratedConfiguration);
-                                                string IOSExecutableName = MacExecutableName.Replace("-Mac-", "-IOS-");
-                                                string TVOSExecutableName = MacExecutableName.Replace("-Mac-", "-TVOS-");
-                                                BuildConfigs.Add(new XcodeBuildConfig(ConfigName, TargetName, FileReference.Combine(OutputDirectory, "Mac", MacExecutableName), FileReference.Combine(OutputDirectory, "IOS", IOSExecutableName), FileReference.Combine(OutputDirectory, "TVOS", TVOSExecutableName), ProjectTarget, Configuration));
-                                            }
-                                            else if (BuildPlatform.Platform == UnrealTargetPlatform.IOS || BuildPlatform.Platform == UnrealTargetPlatform.TVOS)
-                                            {
-                                                string IOSExecutableName = MakeExecutableFileName(ExeName, UnrealTargetPlatform.IOS, Configuration, ProjectTarget.TargetRules.Architecture, ProjectTarget.TargetRules.UndecoratedConfiguration);
-                                                string TVOSExecutableName = IOSExecutableName.Replace("-IOS-", "-TVOS-");
-                                                //string MacExecutableName = IOSExecutableName.Replace("-IOS-", "-Mac-");
-                                                BuildConfigs.Add(new XcodeBuildConfig(ConfigName, TargetName, FileReference.Combine(OutputDirectory, "Mac", IOSExecutableName), FileReference.Combine(OutputDirectory, "IOS", IOSExecutableName), FileReference.Combine(OutputDirectory, "TVOS", TVOSExecutableName), ProjectTarget, Configuration));
-                                            }
-                                        }
+										}
+										catch (BuildException)
+										{
+											// skip target for this platform
+										}
 									}
 								}
 							}
@@ -1637,7 +1644,7 @@ namespace UnrealBuildTool
 			return BuildConfigs;
 		}
 
-        private static string MakeExecutableFileName(string BinaryName, UnrealTargetPlatform Platform, UnrealTargetConfiguration Configuration, string Architecture, UnrealTargetConfiguration UndecoratedConfiguration)
+		private static string MakeExecutableFileName(string BinaryName, UnrealTargetPlatform Platform, UnrealTargetConfiguration Configuration, string Architecture, UnrealTargetConfiguration UndecoratedConfiguration)
 		{
 			StringBuilder Result = new StringBuilder();
 
@@ -1649,12 +1656,12 @@ namespace UnrealBuildTool
 			}
 
 			UEBuildPlatform BuildPlatform = UEBuildPlatform.GetBuildPlatform(Platform);
-			if(BuildPlatform.RequiresArchitectureSuffix())
+			if (BuildPlatform.RequiresArchitectureSuffix())
 			{
 				Result.Append(Architecture);
 			}
 
-            return Result.ToString();
+			return Result.ToString();
 		}
 
 		private FileReference GetUserSchemeManagementFilePath()
@@ -1674,7 +1681,7 @@ namespace UnrealBuildTool
 
 		private void WriteSchemeFile(string TargetName, string TargetGuid, string BuildTargetGuid, string IndexTargetGuid, bool bHasEditorConfiguration, string GameProjectPath)
 		{
-		
+
 			FileReference SchemeFilePath = GetProjectSchemeFilePathForTarget(TargetName);
 
 			DirectoryReference.CreateDirectory(SchemeFilePath.Directory);
@@ -1888,7 +1895,7 @@ namespace UnrealBuildTool
 
 		public static IEnumerable<UnrealTargetConfiguration> GetSupportedConfigurations()
 		{
-			return new UnrealTargetConfiguration[] { 
+			return new UnrealTargetConfiguration[]{
 				UnrealTargetConfiguration.Debug,
 				UnrealTargetConfiguration.DebugGame,
 				UnrealTargetConfiguration.Development,
@@ -1904,11 +1911,11 @@ namespace UnrealBuildTool
 
 		public bool CanBuildProjectLocally()
 		{
-			foreach (ProjectTarget ProjectTarget in ProjectTargets)
+			foreach(ProjectTarget ProjectTarget in ProjectTargets)
 			{
-				foreach (UnrealTargetPlatform Platform in GetSupportedPlatforms())
+				foreach(UnrealTargetPlatform Platform in GetSupportedPlatforms())
 				{
-					foreach (UnrealTargetConfiguration Config in GetSupportedConfigurations())
+					foreach(UnrealTargetConfiguration Config in GetSupportedConfigurations())
 					{
 						if (MSBuildProjectFile.IsValidProjectPlatformAndConfiguration(ProjectTarget, Platform, Config, null))
 						{
@@ -1957,7 +1964,7 @@ namespace UnrealBuildTool
 			FileReference GameProjectPath = null;
 			foreach(ProjectTarget Target in ProjectTargets)
 			{
-				if(Target.UnrealProjectFilePath != null)
+				if (Target.UnrealProjectFilePath != null)
 				{
 					GameProjectPath = Target.UnrealProjectFilePath;
 					break;
@@ -1971,7 +1978,7 @@ namespace UnrealBuildTool
 			Dictionary<string, XcodeBuildConfig> TargetBuildConfigs = new Dictionary<string, XcodeBuildConfig>();
 			Dictionary<string, XcodeBuildConfig> BuildTargetBuildConfigs = new Dictionary<string, XcodeBuildConfig>();
 			Dictionary<string, XcodeBuildConfig> IndexTargetBuildConfigs = new Dictionary<string, XcodeBuildConfig>();
-			foreach (XcodeBuildConfig Config in BuildConfigs)
+			foreach(XcodeBuildConfig Config in BuildConfigs)
 			{
 				ProjectBuildConfigs[XcodeProjectFileGenerator.MakeXcodeGuid()] = Config;
 				TargetBuildConfigs[XcodeProjectFileGenerator.MakeXcodeGuid()] = Config;
@@ -2017,7 +2024,7 @@ namespace UnrealBuildTool
 			{
 				AppendTargetDependencySection(ProjectFileContent, BuildTargetName, BuildTargetGuid, TargetDependencyGuid, TargetProxyGuid);
 			}
-			foreach (XcodeExtensionInfo EI in AllExtensions)
+			foreach(XcodeExtensionInfo EI in AllExtensions)
 			{
 				AppendContainerItemProxySection(ProjectFileContent, EI.Name, EI.TargetGuid, EI.TargetProxyGuid, ProjectGuid);
 				AppendTargetDependencySection(ProjectFileContent, EI.Name, EI.TargetGuid, EI.TargetDependencyGuid, EI.TargetProxyGuid);
@@ -2061,7 +2068,7 @@ namespace UnrealBuildTool
 				{
 					DirectoryReference.Delete(SchemeDir, true);
 				}
-			}				
+			}
 
 			return bSuccess;
 		}
