Add UseVerse support to UBT.
See review https://p4-swarm.epicgames.net/reviews/16249378

diff --git a/Engine/Source/Programs/UnrealBuildTool/Configuration/ModuleRules.cs b/Engine/Source/Programs/UnrealBuildTool/Configuration/ModuleRules.cs
--- a/Engine/Source/Programs/UnrealBuildTool/Configuration/ModuleRules.cs
+++ b/Engine/Source/Programs/UnrealBuildTool/Configuration/ModuleRules.cs
@@ -1212,6 +1212,29 @@ namespace UnrealBuildTool
 			}
 		}
 
+		/// <summary>
+		/// Setup this module for verse support (based on the settings in UEBuildConfiguration)
+		/// </summary>
+		public void SetupModuleVerseSupport(ReadOnlyTargetRules Target)
+		{
+			bool bUseVerse = (Target.bUseVerse == true);
+			if (bUseVerse)
+			{
+				PublicDefinitions.Add("WITH_VERSE=1");
+
+				PrivateDependencyModuleNames.AddRange(
+					new string[] {
+						"Solaris",
+						"VerseCore",
+					}
+					);
+			}
+			else
+			{
+				PublicDefinitions.Add("WITH_VERSE=0");
+			}
+		}
+
 		/// <summary>
 		/// Setup this module for physics support (based on the settings in UEBuildConfiguration)
 		/// </summary>
diff --git a/Engine/Source/Programs/UnrealBuildTool/Configuration/TargetRules.cs b/Engine/Source/Programs/UnrealBuildTool/Configuration/TargetRules.cs
--- a/Engine/Source/Programs/UnrealBuildTool/Configuration/TargetRules.cs
+++ b/Engine/Source/Programs/UnrealBuildTool/Configuration/TargetRules.cs
@@ -390,6 +390,14 @@ namespace UnrealBuildTool
 		[RequiresUniqueBuildEnvironment]
 		public bool bEnableMeshEditor = false;
 
+		/// <summary>
+		/// Whether to use the verse script interface.
+		/// </summary>
+		[RequiresUniqueBuildEnvironment]
+		[CommandLine("-NoUseVerse", Value = "false")]
+		[CommandLine("-UseVerse", Value = "true")]
+		public bool bUseVerse = false;
+
 		/// <summary>
 		/// Whether to compile the Chaos physics plugin.
 		/// </summary>
@@ -2036,6 +2044,11 @@ namespace UnrealBuildTool
 			get { return Inner.bEnableMeshEditor; }
 		}
 
+		public bool bUseVerse
+		{
+			get { return Inner.bUseVerse; }
+		}
+
 		public bool bCompileChaos
 		{
 			get { return Inner.bCompileChaos; }
