Move to LLVM 11.0.1

diff --git a/Engine/Build/BatchFiles/Linux/SetupToolchain.sh b/Engine/Build/BatchFiles/Linux/SetupToolchain.sh
--- a/Engine/Build/BatchFiles/Linux/SetupToolchain.sh
+++ b/Engine/Build/BatchFiles/Linux/SetupToolchain.sh
@@ -4,7 +4,7 @@
 # put ourselves into Engine directory (two up from location of this script)
 pushd "`dirname "$0"`/../../.." > /dev/null
 
-TOOLCHAIN_VERSION=v17_clang-10.0.1-centos7
+TOOLCHAIN_VERSION=v18_clang-11.0.1-centos7
 TOOLCHAIN_ARCHIVE=$TOOLCHAIN_VERSION.tar.gz
 
 TOOLCHAIN_URL=http://cdn.unrealengine.com/Toolchain_Linux/native-linux-$TOOLCHAIN_ARCHIVE
diff --git a/Engine/Source/Programs/UnrealBuildTool/Platform/Linux/LinuxToolChain.cs b/Engine/Source/Programs/UnrealBuildTool/Platform/Linux/LinuxToolChain.cs
--- a/Engine/Source/Programs/UnrealBuildTool/Platform/Linux/LinuxToolChain.cs
+++ b/Engine/Source/Programs/UnrealBuildTool/Platform/Linux/LinuxToolChain.cs
@@ -188,10 +188,10 @@ namespace UnrealBuildTool
 				throw new BuildException("Unable to build: no compatible clang version found. Please run Setup.sh");
 			}
 			// prevent unknown clangs since the build is likely to fail on too old or too new compilers
-			else if ((CompilerVersionMajor * 10 + CompilerVersionMinor) > 100 || (CompilerVersionMajor * 10 + CompilerVersionMinor) < 60)
+			else if ((CompilerVersionMajor * 10 + CompilerVersionMinor) > 110 || (CompilerVersionMajor * 10 + CompilerVersionMinor) < 60)
 			{
 				throw new BuildException(
-					string.Format("This version of the Unreal Engine can only be compiled with clang 10.0, 9.0, 8.0, 7.0 and 6.0. clang {0} may not build it - please use a different version.",
+					string.Format("This version of the Unreal Engine can only be compiled with clang 11.0, 10.0, 9.0, 8.0, 7.0 and 6.0. clang {0} may not build it - please use a different version.",
 						CompilerVersionString)
 					);
 			}
