Unshelved from pending changelist '15218728':

PR #7717: HL2 asset bundling (Contributed by [REDACTED_PERSON])

diff --git a/Engine/Source/Programs/AutomationTool/HoloLens/HoloLensPlatform.Automation.cs b/Engine/Source/Programs/AutomationTool/HoloLens/HoloLensPlatform.Automation.cs
--- a/Engine/Source/Programs/AutomationTool/HoloLens/HoloLensPlatform.Automation.cs
+++ b/Engine/Source/Programs/AutomationTool/HoloLens/HoloLensPlatform.Automation.cs
@@ -720,6 +720,9 @@ namespace HoloLens.Automation
 			var AppXRecipeBuiltFiles = new StringBuilder();
 
 			string MapFilename = Path.Combine(SC.StageDirectory.FullName, OutputName + ".pkgmap");
+			AppXRecipeBuiltFiles.AppendLine("[ResourceMetadata]");
+			AppXRecipeBuiltFiles.AppendLine("\"ResourceId\"\t\"UnrealAssets\"");
+			AppXRecipeBuiltFiles.AppendLine("");
 			AppXRecipeBuiltFiles.AppendLine(@"[Files]");
 
 			string OutputAppX = Path.Combine(SC.StageDirectory.FullName, OutputName + Extension);
@@ -750,11 +753,11 @@ namespace HoloLens.Automation
 				}
 			}
 
-			AppXRecipeBuiltFiles.AppendLine(String.Format("\"{0}\"\t\"{1}\"", Path.Combine(SC.StageDirectory.FullName, "AppxManifest_assets.xml"), "AppxManifest.xml"));
+			//AppXRecipeBuiltFiles.AppendLine(String.Format("\"{0}\"\t\"{1}\"", Path.Combine(SC.StageDirectory.FullName, "AppxManifest_assets.xml"), "AppxManifest.xml"));
 
 			File.WriteAllText(MapFilename, AppXRecipeBuiltFiles.ToString(), Encoding.UTF8);
 
-			string MakeAppXCommandLine = String.Format(@"pack /o /f ""{0}"" /p ""{1}""", MapFilename, OutputAppX);
+			string MakeAppXCommandLine = String.Format(@"pack /r /o /f ""{0}"" /p ""{1}"" /m ""{2}""", MapFilename, OutputAppX, Path.Combine(SC.StageDirectory.FullName, "AppxManifest_assets.xml"));
 			RunAndLog(CmdEnv, MakeAppXPath.FullName, MakeAppXCommandLine, null, 0, null, ERunOptions.None);
 			SignPackage(Params, SC, OutputAppX);
 		}
@@ -1046,7 +1049,7 @@ namespace HoloLens.Automation
 			string OutputAppX = Path.Combine(SC.StageDirectory.FullName, OutputNameBase + Extension);
 			{
 				OutputAppX += "bundle";
-				bool SeparateAssetPackaging = false;
+				bool SeparateAssetPackaging = true;
 				bool bStartInVR = false;
 
 				//auto update
@@ -1057,7 +1060,7 @@ namespace HoloLens.Automation
 				ConfigHierarchy PlatformEngineConfig = null;
 				if (Params.EngineConfigs.TryGetValue(PlatformType, out PlatformEngineConfig))
 				{
-					PlatformEngineConfig.GetBool("/Script/HoloLensPlatformEditor.HoloLensTargetSettings", "bUseAssetPackage", out SeparateAssetPackaging);
+					//PlatformEngineConfig.GetBool("/Script/HoloLensPlatformEditor.HoloLensTargetSettings", "bUseAssetPackage", out SeparateAssetPackaging);
 					
 					// Get auto update vars
 					PlatformEngineConfig.GetBool("/Script/HoloLensPlatformEditor.HoloLensTargetSettings", "bShouldCreateAppInstaller", out bShouldCreateAppInstaller);
diff --git a/Engine/Source/Programs/UnrealBuildTool/Platform/HoloLens/HoloLensDeploy.cs b/Engine/Source/Programs/UnrealBuildTool/Platform/HoloLens/HoloLensDeploy.cs
--- a/Engine/Source/Programs/UnrealBuildTool/Platform/HoloLens/HoloLensDeploy.cs
+++ b/Engine/Source/Programs/UnrealBuildTool/Platform/HoloLens/HoloLensDeploy.cs
@@ -17,8 +17,9 @@ namespace UnrealBuildTool
 	/// </summary>
 	class HoloLensDeploy : UEBuildDeploy
 	{
-		//private FileReference MakeAppXPath;
-		//private FileReference SignToolPath;
+		private FileReference MakeAppXPath;
+		private FileReference SignToolPath;
+		private string Extension = ".appx";
 
 		/// <summary>
 		/// Utility function to delete a file
@@ -206,8 +207,8 @@ namespace UnrealBuildTool
 				ConfigDirRef = new DirectoryReference(UnrealBuildTool.GetRemoteIniPath());
 			}
 
-			//MakeAppXPath = HoloLensExports.GetWindowsSdkToolPath("makeappx.exe");
-			//SignToolPath = HoloLensExports.GetWindowsSdkToolPath("signtool.exe");
+			MakeAppXPath = HoloLensExports.GetWindowsSdkToolPath("makeappx.exe");
+			SignToolPath = HoloLensExports.GetWindowsSdkToolPath("signtool.exe");
 
 			return true;
 		}
@@ -217,7 +218,7 @@ namespace UnrealBuildTool
 			string OutputName = String.Format("{0}_{1}_{2}_{3}", Receipt.TargetName, Receipt.Platform, Receipt.Configuration, WindowsExports.GetArchitectureSubpath(Architecture));
 			string IntermediateDirectory = Path.Combine(Receipt.ProjectDir != null ? Receipt.ProjectDir.FullName : UnrealBuildTool.EngineDirectory.FullName, "Intermediate", "Deploy", WindowsExports.GetArchitectureSubpath(Architecture));
 			string OutputDirectory = Receipt.Launch.Directory.FullName;
-			string OutputAppX = Path.Combine(OutputDirectory, OutputName + ".appx");
+			string OutputAppX = Path.Combine(OutputDirectory, OutputName + Extension);
 			string SigningCertificate = @"Build\HoloLens\SigningCertificate.pfx";
 			string SigningCertificatePath = Path.Combine(Receipt.ProjectDir != null ? Receipt.ProjectDir.FullName : UnrealBuildTool.EngineDirectory.FullName, SigningCertificate);
 
@@ -425,7 +426,32 @@ namespace UnrealBuildTool
 			}
 			File.WriteAllText(MapFilename, AppXRecipeBuiltFiles.ToString(), Encoding.UTF8);
 
-			NewReceipt.BuildProducts.Add(new BuildProduct(new FileReference(MapFilename), BuildProductType.MapFile));
+			string MakeAppXCommandLine = String.Format("pack /o /f \"{0}\" /p \"{1}\"", MapFilename, OutputAppX);
+
+			var StartInfo = new ProcessStartInfo(MakeAppXPath.FullName, MakeAppXCommandLine);
+			StartInfo.UseShellExecute = false;
+			StartInfo.CreateNoWindow = true;
+			var ExitCode = Utils.RunLocalProcessAndPrintfOutput(StartInfo);
+			if (ExitCode < 0)
+			{
+				throw new BuildException("Failed to generate AppX file.  See log for details.");
+			}
+
+			if (File.Exists(SigningCertificatePath))
+			{
+				string SignToolCommandLine = String.Format("sign /a /f \"{0}\" /fd SHA256 \"{1}\"", SigningCertificatePath, OutputAppX);
+				StartInfo = new ProcessStartInfo(SignToolPath.FullName, SignToolCommandLine);
+				StartInfo.UseShellExecute = false;
+				StartInfo.CreateNoWindow = true;
+				ExitCode = Utils.RunLocalProcessAndPrintfOutput(StartInfo);
+				if (ExitCode < 0)
+				{
+					throw new BuildException("Failed to generate AppX file.  See log for details.");
+				}
+			}
+
+			Log.TraceInformation("AppX successfully packaged to \'{0}{1}\'", OutputName, Extension);
+			NewReceipt.BuildProducts.Add(new BuildProduct(new FileReference(OutputAppX), BuildProductType.Package));
 		}
 
 
