Update to new LLVM 11.0.1 toolchain

diff --git a/Engine/Source/Programs/UnrealBuildTool/Platform/Linux/LinuxPlatformSDK.cs b/Engine/Source/Programs/UnrealBuildTool/Platform/Linux/LinuxPlatformSDK.cs
--- a/Engine/Source/Programs/UnrealBuildTool/Platform/Linux/LinuxPlatformSDK.cs
+++ b/Engine/Source/Programs/UnrealBuildTool/Platform/Linux/LinuxPlatformSDK.cs
@@ -14,14 +14,14 @@ namespace UnrealBuildTool
 	{
 		public override string GetMainVersion()
 		{
-			return "v17_clang-10.0.1-centos7";
+			return "v18_clang-11.0.1-centos7";
 		}
 
 		public override void GetValidVersionRange(out string MinVersion, out string MaxVersion)
 		{
 			// all that matters is the number after the v, according to TryConvertVersionToInt()
 			MinVersion = "v10_clang-5.0.0-centos7";
-			MaxVersion = "v17_clang-10.0.1-centos7";
+			MaxVersion = "v18_clang-11.0.1-centos7";
 		}
 
 		public override void GetValidSoftwareVersionRange(out string MinVersion, out string MaxVersion)
