Restore -gdwarf-2 flag while we sort out buildfarm issues.

#rnx

diff --git a/Engine/Source/Programs/UnrealBuildTool/Platform/IOS/IOSToolChain.cs b/Engine/Source/Programs/UnrealBuildTool/Platform/IOS/IOSToolChain.cs
--- a/Engine/Source/Programs/UnrealBuildTool/Platform/IOS/IOSToolChain.cs
+++ b/Engine/Source/Programs/UnrealBuildTool/Platform/IOS/IOSToolChain.cs
@@ -347,7 +347,7 @@ namespace UnrealBuildTool
 			// Create DWARF format debug info if wanted,
 			if (CompileEnvironment.bCreateDebugInfo)
 			{
-				Result += " -g";
+				Result += " -gdwarf-2";
 			}
 
 			// Add additional frameworks so that their headers can be found
diff --git a/Engine/Source/Programs/UnrealBuildTool/Platform/Mac/MacToolChain.cs b/Engine/Source/Programs/UnrealBuildTool/Platform/Mac/MacToolChain.cs
--- a/Engine/Source/Programs/UnrealBuildTool/Platform/Mac/MacToolChain.cs
+++ b/Engine/Source/Programs/UnrealBuildTool/Platform/Mac/MacToolChain.cs
@@ -287,7 +287,7 @@ namespace UnrealBuildTool
 			// Create DWARF format debug info if wanted,
 			if (CompileEnvironment.bCreateDebugInfo)
 			{
-				Result += " -g";
+				Result += " -gdwarf-2";
 			}
 
 			return Result;
