Better fix for project file generation trying to generate IOS targets for modular build projects

diff --git a/Engine/Source/Programs/UnrealBuildTool/ProjectFiles/ProjectFileGenerator.cs b/Engine/Source/Programs/UnrealBuildTool/ProjectFiles/ProjectFileGenerator.cs
--- a/Engine/Source/Programs/UnrealBuildTool/ProjectFiles/ProjectFileGenerator.cs
+++ b/Engine/Source/Programs/UnrealBuildTool/ProjectFiles/ProjectFileGenerator.cs
@@ -2421,14 +2421,17 @@ namespace UnrealBuildTool
 					}
 
 					ProjectTarget ProjectTarget = new ProjectTarget()
-					{
-						TargetRules = TargetRulesObject,
-						TargetFilePath = TargetFilePath,
-						ProjectFilePath = ProjectFilePath,
-						UnrealProjectFilePath = CheckProjectFile,
-						SupportedPlatforms = TargetRulesObject.GetSupportedPlatforms().Where(x => UEBuildPlatform.TryGetBuildPlatform(x, out _)).ToArray(),
-						CreateRulesDelegate = (Platform, Configuration) => RulesAssembly.CreateTargetRules(TargetName, Platform, Configuration, "", CheckProjectFile, new CommandLineArguments(GetTargetArguments(Arguments)))
-					};
+						{
+							TargetRules = TargetRulesObject,
+							TargetFilePath = TargetFilePath,
+							ProjectFilePath = ProjectFilePath,
+							UnrealProjectFilePath = CheckProjectFile,
+							SupportedPlatforms = TargetRulesObject.GetSupportedPlatforms().Where(
+													x => UEBuildPlatform.TryGetBuildPlatform(x, out _) && 
+													(TargetRulesObject.LinkType != TargetLinkType.Modular || !UEBuildPlatform.PlatformRequiresMonolithicBuilds(x, TargetRulesObject.Configuration))
+													).ToArray(),
+							CreateRulesDelegate = (Platform, Configuration) => RulesAssembly.CreateTargetRules(TargetName, Platform, Configuration, "", CheckProjectFile, new CommandLineArguments(GetTargetArguments(Arguments)))
+                        };
 
 					ProjectFile.ProjectTargets.Add(ProjectTarget);
 
diff --git a/Engine/Source/Programs/UnrealBuildTool/ProjectFiles/Xcode/XcodeProject.cs b/Engine/Source/Programs/UnrealBuildTool/ProjectFiles/Xcode/XcodeProject.cs
--- a/Engine/Source/Programs/UnrealBuildTool/ProjectFiles/Xcode/XcodeProject.cs
+++ b/Engine/Source/Programs/UnrealBuildTool/ProjectFiles/Xcode/XcodeProject.cs
@@ -1099,7 +1099,7 @@ namespace UnrealBuildTool
 			else
 			{
 				bool bIsUnrealGame = TargetName.Equals("UnrealGame", StringComparison.InvariantCultureIgnoreCase);
-				bool bIsUE4Client = TargetName.Equals("UnrealClient", StringComparison.InvariantCultureIgnoreCase);
+				bool bIsUnrealClient = TargetName.Equals("UnrealClient", StringComparison.InvariantCultureIgnoreCase);
 				DirectoryReference GameDir = ProjectFile != null ? ProjectFile.Directory : null;
 				string GamePath = GameDir != null ? ConvertPath(GameDir.FullName) : null;
 
@@ -1206,7 +1206,7 @@ namespace UnrealBuildTool
 					Content.Append("\t\t\t\t\"SDKROOT[sdk=macosx]\" = macosx;" + ProjectFileGenerator.NewLine);
 				}
 
-				if (bIsUnrealGame || bIsUE4Client)
+				if (bIsUnrealGame || bIsUnrealClient)
 				{
 					if (IOSRunTimeVersion != null)
 					{
@@ -1268,7 +1268,7 @@ namespace UnrealBuildTool
 			if (!bMacOnly)
 			{
 				bool bIsUnrealGame = Config.BuildTarget.Equals("UnrealGame", StringComparison.InvariantCultureIgnoreCase);
-				bool bIsUE4Client = Config.BuildTarget.Equals("UnrealClient", StringComparison.InvariantCultureIgnoreCase);
+				bool bIsUnrealClient = Config.BuildTarget.Equals("UnrealClient", StringComparison.InvariantCultureIgnoreCase);
 
 				DirectoryReference GameDir = ProjectFile != null ? ProjectFile.Directory : null;
 				string GamePath = GameDir != null ? ConvertPath(GameDir.FullName) : null;
@@ -1286,7 +1286,7 @@ namespace UnrealBuildTool
 					IOSEntitlementPath = "";
 					TVOSEntitlementPath = "";
 				}
-				else if (bIsUE4Client)
+				else if (bIsUnrealClient)
 				{
 					IOSInfoPlistPath = UEDir + "/Engine/Intermediate/IOS/UnrealGame-Info.plist";
 					TVOSInfoPlistPath = UEDir + "/Engine/Intermediate/TVOS/UnrealGame-Info.plist";
@@ -1572,66 +1572,62 @@ namespace UnrealBuildTool
 								{
 									if (MSBuildProjectFile.IsValidProjectPlatformAndConfiguration(ProjectTarget, Platform, Configuration, PlatformProjectGenerators))
 									{
-										try
+										// Figure out if this is a monolithic build
+										bool bShouldCompileMonolithic = BuildPlatform.ShouldCompileMonolithicBinary(Platform);
+										bShouldCompileMonolithic |= (ProjectTarget.CreateRulesDelegate(Platform, Configuration).LinkType == TargetLinkType.Monolithic);
+
+										string ConfigName = Configuration.ToString();
+										if (ProjectTarget.TargetRules.Type != TargetType.Game && ProjectTarget.TargetRules.Type != TargetType.Program)
 										{
-											// Figure out if this is a monolithic build
-											bool bShouldCompileMonolithic = BuildPlatform.ShouldCompileMonolithicBinary(Platform);
-											bShouldCompileMonolithic |= (ProjectTarget.CreateRulesDelegate(Platform, Configuration).LinkType == TargetLinkType.Monolithic);
+											ConfigName += " " + ProjectTarget.TargetRules.Type.ToString();
+										}
 
-											string ConfigName = Configuration.ToString();
-											if (ProjectTarget.TargetRules.Type != TargetType.Game && ProjectTarget.TargetRules.Type != TargetType.Program)
-											{
-												ConfigName += " " + ProjectTarget.TargetRules.Type.ToString();
-											}
+										if (BuildConfigs.Where(Config => Config.DisplayName == ConfigName).ToList().Count == 0)
+										{
+											string TargetName = ProjectTarget.TargetFilePath.GetFileNameWithoutAnyExtensions();
 
-											if (BuildConfigs.Where(Config => Config.DisplayName == ConfigName).ToList().Count == 0)
+											// Get the output directory
+											DirectoryReference RootDirectory = UnrealBuildTool.EngineDirectory;
+											if ((ProjectTarget.TargetRules.Type == TargetType.Game || ProjectTarget.TargetRules.Type == TargetType.Client || ProjectTarget.TargetRules.Type == TargetType.Server) && bShouldCompileMonolithic)
 											{
-												string TargetName = ProjectTarget.TargetFilePath.GetFileNameWithoutAnyExtensions();
-
-												// Get the output directory
-												DirectoryReference RootDirectory = UnrealBuildTool.EngineDirectory;
-												if ((ProjectTarget.TargetRules.Type == TargetType.Game || ProjectTarget.TargetRules.Type == TargetType.Client || ProjectTarget.TargetRules.Type == TargetType.Server) && bShouldCompileMonolithic)
-												{
-													if (ProjectTarget.UnrealProjectFilePath != null)
-													{
-														RootDirectory = ProjectTarget.UnrealProjectFilePath.Directory;
-													}
-												}
-
-												if (ProjectTarget.TargetRules.Type == TargetType.Program && ProjectTarget.UnrealProjectFilePath != null)
+												if (ProjectTarget.UnrealProjectFilePath != null)
 												{
 													RootDirectory = ProjectTarget.UnrealProjectFilePath.Directory;
 												}
+											}
+
+											if (ProjectTarget.TargetRules.Type == TargetType.Program && ProjectTarget.UnrealProjectFilePath != null)
+											{
+												RootDirectory = ProjectTarget.UnrealProjectFilePath.Directory;
+											}
 
-												// Get the output directory
-												DirectoryReference OutputDirectory = DirectoryReference.Combine(RootDirectory, "Binaries");
+											// Get the output directory
+											DirectoryReference OutputDirectory = DirectoryReference.Combine(RootDirectory, "Binaries");
 
-												string ExeName = TargetName;
-												if (!bShouldCompileMonolithic && ProjectTarget.TargetRules.Type != TargetType.Program)
+											string ExeName = TargetName;
+											if (!bShouldCompileMonolithic && ProjectTarget.TargetRules.Type != TargetType.Program)
+											{
+												// Figure out what the compiled binary will be called so that we can point the IDE to the correct file
+												if (ProjectTarget.TargetRules.Type != TargetType.Game)
 												{
-													// Figure out what the compiled binary will be called so that we can point the IDE to the correct file
 													ExeName = "Unreal" + ProjectTarget.TargetRules.Type.ToString();
 												}
+											}
 
-												if (BuildPlatform.Platform == UnrealTargetPlatform.Mac)
-												{
-													string MacExecutableName = MakeExecutableFileName(ExeName, UnrealTargetPlatform.Mac, Configuration, ProjectTarget.TargetRules.Architecture, ProjectTarget.TargetRules.UndecoratedConfiguration);
-													string IOSExecutableName = MacExecutableName.Replace("-Mac-", "-IOS-");
-													string TVOSExecutableName = MacExecutableName.Replace("-Mac-", "-TVOS-");
-													BuildConfigs.Add(new XcodeBuildConfig(ConfigName, TargetName, FileReference.Combine(OutputDirectory, "Mac", MacExecutableName), FileReference.Combine(OutputDirectory, "IOS", IOSExecutableName), FileReference.Combine(OutputDirectory, "TVOS", TVOSExecutableName), ProjectTarget, Configuration));
-												}
-												else if (BuildPlatform.Platform == UnrealTargetPlatform.IOS || BuildPlatform.Platform == UnrealTargetPlatform.TVOS)
-												{
-													string IOSExecutableName = MakeExecutableFileName(ExeName, UnrealTargetPlatform.IOS, Configuration, ProjectTarget.TargetRules.Architecture, ProjectTarget.TargetRules.UndecoratedConfiguration);
-													string TVOSExecutableName = IOSExecutableName.Replace("-IOS-", "-TVOS-");
-													//string MacExecutableName = IOSExecutableName.Replace("-IOS-", "-Mac-");
-													BuildConfigs.Add(new XcodeBuildConfig(ConfigName, TargetName, FileReference.Combine(OutputDirectory, "Mac", IOSExecutableName), FileReference.Combine(OutputDirectory, "IOS", IOSExecutableName), FileReference.Combine(OutputDirectory, "TVOS", TVOSExecutableName), ProjectTarget, Configuration));
-												}
+											if (BuildPlatform.Platform == UnrealTargetPlatform.Mac)
+											{
+												string MacExecutableName = MakeExecutableFileName(ExeName, UnrealTargetPlatform.Mac, Configuration, ProjectTarget.TargetRules.Architecture, ProjectTarget.TargetRules.UndecoratedConfiguration);
+												string IOSExecutableName = MacExecutableName.Replace("-Mac-", "-IOS-");
+												string TVOSExecutableName = MacExecutableName.Replace("-Mac-", "-TVOS-");
+												BuildConfigs.Add(new XcodeBuildConfig(ConfigName, TargetName, FileReference.Combine(OutputDirectory, "Mac", MacExecutableName), FileReference.Combine(OutputDirectory, "IOS", IOSExecutableName), FileReference.Combine(OutputDirectory, "TVOS", TVOSExecutableName), ProjectTarget, Configuration));
+											}
+											else if (BuildPlatform.Platform == UnrealTargetPlatform.IOS || BuildPlatform.Platform == UnrealTargetPlatform.TVOS)
+											{
+												string IOSExecutableName = MakeExecutableFileName(ExeName, UnrealTargetPlatform.IOS, Configuration, ProjectTarget.TargetRules.Architecture, ProjectTarget.TargetRules.UndecoratedConfiguration);
+												string TVOSExecutableName = IOSExecutableName.Replace("-IOS-", "-TVOS-");
+												//string MacExecutableName = IOSExecutableName.Replace("-IOS-", "-Mac-");
+												BuildConfigs.Add(new XcodeBuildConfig(ConfigName, TargetName, FileReference.Combine(OutputDirectory, "Mac", IOSExecutableName), FileReference.Combine(OutputDirectory, "IOS", IOSExecutableName), FileReference.Combine(OutputDirectory, "TVOS", TVOSExecutableName), ProjectTarget, Configuration));
 											}
-										}
-										catch (BuildException)
-										{
-											// skip target for this platform
 										}
 									}
 								}
@@ -2073,4 +2069,4 @@ namespace UnrealBuildTool
 			return bSuccess;
 		}
 	}
-}
\ No newline at end of file
+}
