PR #7957: Fix for [REDACTED_PROFILE_URL]/RiderSourceCodeAccess/issues/8 (Contributed by [REDACTED_PERSON])

diff --git a/Engine/Source/Programs/UnrealBuildTool/ProjectFiles/Rider/RiderProjectFile.cs b/Engine/Source/Programs/UnrealBuildTool/ProjectFiles/Rider/RiderProjectFile.cs
--- a/Engine/Source/Programs/UnrealBuildTool/ProjectFiles/Rider/RiderProjectFile.cs
+++ b/Engine/Source/Programs/UnrealBuildTool/ProjectFiles/Rider/RiderProjectFile.cs
@@ -526,9 +526,12 @@ namespace UnrealBuildTool
 				var Version = GetLinuxToolchainVersionFromFullString(PlatformSdkVersionString);
 
 				var InternalSdkPath = UEBuildPlatform.GetSDK(UnrealTargetPlatform.Linux).GetInternalSDKPath();
-				Writer.WriteValue(Path.Combine(InternalSdkPath, "include"));
-				Writer.WriteValue(Path.Combine(InternalSdkPath, "usr/include"));
-				Writer.WriteValue(Path.Combine(InternalSdkPath, "lib/clang/" + Version + "/include/"));
+				if (InternalSdkPath != null)
+				{
+					Writer.WriteValue(Path.Combine(InternalSdkPath, "include"));
+					Writer.WriteValue(Path.Combine(InternalSdkPath, "usr/include"));
+					Writer.WriteValue(Path.Combine(InternalSdkPath, "lib/clang/" + Version + "/include/"));
+				}
 			}
 			
 			Writer.WriteArrayEnd();
