made FindResourceBinaryFile virtual so derived manifest generators can override where manifest resources are found.

#rnx

diff --git a/Engine/Source/Programs/UnrealBuildTool/Platform/Windows/VCManifestGenerator.cs b/Engine/Source/Programs/UnrealBuildTool/Platform/Windows/VCManifestGenerator.cs
--- a/Engine/Source/Programs/UnrealBuildTool/Platform/Windows/VCManifestGenerator.cs
+++ b/Engine/Source/Programs/UnrealBuildTool/Platform/Windows/VCManifestGenerator.cs
@@ -277,7 +277,7 @@ namespace UnrealBuildTool
 			return DefaultValue;
 		}
 
-		protected bool FindResourceBinaryFile( out string SourcePath, string ResourceFileName, bool AllowEngineFallback = true)
+		protected virtual bool FindResourceBinaryFile( out string SourcePath, string ResourceFileName, bool AllowEngineFallback = true)
 		{
 			// look in project normal Build location
 			SourcePath = Path.Combine(ProjectPath, "Build", Platform.ToString(), BuildResourceSubPath);
