Stripped PLATFORM_ANDROID_VULKAN and PLATFORM_ANDROIDGL4
Removed old GPU Architecture code from android tool chain

diff --git a/Engine/Source/Programs/AutomationTool/Android/AndroidPlatform.Automation.cs b/Engine/Source/Programs/AutomationTool/Android/AndroidPlatform.Automation.cs
--- a/Engine/Source/Programs/AutomationTool/Android/AndroidPlatform.Automation.cs
+++ b/Engine/Source/Programs/AutomationTool/Android/AndroidPlatform.Automation.cs
@@ -256,7 +256,7 @@ public class AndroidPlatform : Platform
 		return Path.Combine(Path.GetDirectoryName(Params.GetProjectExeForPlatform(UnrealTargetPlatform.Android).ToString()), DecoratedExeName) + ".so";
 	}
 
-	private static string GetFinalApkName(ProjectParams Params, string DecoratedExeName, bool bRenameUnrealGame, string Architecture, string GPUArchitecture)
+	private static string GetFinalApkName(ProjectParams Params, string DecoratedExeName, bool bRenameUnrealGame, string Architecture)
 	{
 		string ProjectDir = Path.Combine(Path.GetDirectoryName(Path.GetFullPath(Params.RawProjectPath.FullName)), "Binaries/Android");
 
@@ -266,7 +266,7 @@ public class AndroidPlatform : Platform
 		}
 
 		// Apk's go to project location, not necessarily where the .so is (content only packages need to output to their directory)
-		string ApkName = Path.Combine(ProjectDir, DecoratedExeName) + Architecture + GPUArchitecture + ".apk";
+		string ApkName = Path.Combine(ProjectDir, DecoratedExeName) + Architecture + ".apk";
 
 		// if the source binary was UnrealGame, handle using it or switching to project name
 		if (Path.GetFileNameWithoutExtension(Params.GetProjectExeForPlatform(UnrealTargetPlatform.Android).ToString()) == "UnrealGame")
@@ -329,7 +329,7 @@ public class AndroidPlatform : Platform
 	#pragma warning restore CS0414
 
 
-	private static string GetFinalSymbolizedSODirectory(string ApkName, DeploymentContext SC, string Architecture, string GPUArchitecture)
+	private static string GetFinalSymbolizedSODirectory(string ApkName, DeploymentContext SC, string Architecture)
 	{
 		string PackageVersion = GetPackageInfo(ApkName, SC, true);
 		if (PackageVersion == null || PackageVersion.Length == 0)
@@ -337,7 +337,7 @@ public class AndroidPlatform : Platform
 			throw new AutomationException(ExitCode.Error_FailureGettingPackageInfo, "Failed to get package version from " + ApkName);
 		}
 
-		return SC.ShortProjectName + "_Symbols_v" + PackageVersion + "/" + SC.ShortProjectName + Architecture + GPUArchitecture;
+		return SC.ShortProjectName + "_Symbols_v" + PackageVersion + "/" + SC.ShortProjectName + Architecture;
 	}
 
 	private static string GetFinalObbName(string ApkName, DeploymentContext SC, bool bUseAppType = true)
@@ -498,7 +498,7 @@ public class AndroidPlatform : Platform
 		Uninstall,
 		Symbolize,
 	};
-	private static string GetFinalBatchName(string ApkName, DeploymentContext SC, string Architecture, string GPUArchitecture, bool bNoOBBInstall, EBatchType BatchType, UnrealTargetPlatform Target)
+	private static string GetFinalBatchName(string ApkName, DeploymentContext SC, string Architecture, bool bNoOBBInstall, EBatchType BatchType, UnrealTargetPlatform Target)
 	{
 		string Extension = ".bat";
 		if (Target == UnrealTargetPlatform.Linux || Target == UnrealTargetPlatform.LinuxArm64)
@@ -674,7 +674,6 @@ public class AndroidPlatform : Platform
 
 		IAndroidToolChain ToolChain = AndroidExports.CreateToolChain(Params.RawProjectPath);
 		var Architectures = ToolChain.GetAllArchitectures();
-		var GPUArchitectures = ToolChain.GetAllGPUArchitectures();
 		bool bMakeSeparateApks = UnrealBuildTool.AndroidExports.ShouldMakeSeparateApks();
 		bool bBuildWithHiddenSymbolVisibility = BuildWithHiddenSymbolVisibility(SC);
 		bool bSaveSymbols = GetSaveSymbols(SC);
@@ -684,7 +683,7 @@ public class AndroidPlatform : Platform
 		var Deploy = AndroidExports.CreateDeploymentHandler(Params.RawProjectPath, Params.ForcePackageData);
 		bool bPackageDataInsideApk = Deploy.GetPackageDataInsideApk();
 
-		string BaseApkName = GetFinalApkName(Params, SC.StageExecutables[0], true, "", "");
+		string BaseApkName = GetFinalApkName(Params, SC.StageExecutables[0], true, "");
 		LogInformation("BaseApkName = {0}", BaseApkName);
 
 		// Create main OBB with entire contents of staging dir. This
@@ -889,216 +888,213 @@ public class AndroidPlatform : Platform
 
 		foreach (string Architecture in Architectures)
 		{
-			foreach (string GPUArchitecture in GPUArchitectures)
+			string ApkName = GetFinalApkName(Params, SC.StageExecutables[0], true, bMakeSeparateApks ? Architecture : "");
+			string ApkBareName = GetFinalApkName(Params, SC.StageExecutables[0], true, "");
+			bool bHaveAPK = !bEnableBundle;     // do not have a standard APK if bundle enabled
+			if (!SC.IsCodeBasedProject)
 			{
-				string ApkName = GetFinalApkName(Params, SC.StageExecutables[0], true, bMakeSeparateApks ? Architecture : "", bMakeSeparateApks ? GPUArchitecture : "");
-				string ApkBareName = GetFinalApkName(Params, SC.StageExecutables[0], true, "", "");
-				bool bHaveAPK = !bEnableBundle;     // do not have a standard APK if bundle enabled
-				if (!SC.IsCodeBasedProject)
+				string UE4SOName = GetFinalApkName(Params, SC.StageExecutables[0], false, bMakeSeparateApks ? Architecture : "");
+                UE4SOName = UE4SOName.Replace(".apk", ".so");
+                if (FileExists_NoExceptions(UE4SOName) == false)
 				{
-					string UE4SOName = GetFinalApkName(Params, SC.StageExecutables[0], false, bMakeSeparateApks ? Architecture : "", bMakeSeparateApks ? GPUArchitecture : "");
-                    UE4SOName = UE4SOName.Replace(".apk", ".so");
-                    if (FileExists_NoExceptions(UE4SOName) == false)
-					{
-						LogInformation("Failed to find game .so " + UE4SOName);
-                        throw new AutomationException(ExitCode.Error_MissingExecutable, "Stage Failed. Could not find .so {0}. You may need to build the UE project with your target configuration and platform.", UE4SOName);
-					}
+					LogInformation("Failed to find game .so " + UE4SOName);
+                    throw new AutomationException(ExitCode.Error_MissingExecutable, "Stage Failed. Could not find .so {0}. You may need to build the UE project with your target configuration and platform.", UE4SOName);
 				}
+			}
 				
-				TargetReceipt Receipt = SC.StageTargets[0].Receipt;
+			TargetReceipt Receipt = SC.StageTargets[0].Receipt;
 				              
-				// when we make an embedded executable, all we do is output to libUnreal.so - we don't need to make an APK at all
-				// however, we still let package go through to make the .obb file
-				string CookFlavor = SC.FinalCookPlatform.IndexOf("_") > 0 ? SC.FinalCookPlatform.Substring(SC.FinalCookPlatform.IndexOf("_")) : "";
-				if (!Params.Prebuilt)
+			// when we make an embedded executable, all we do is output to libUnreal.so - we don't need to make an APK at all
+			// however, we still let package go through to make the .obb file
+			string CookFlavor = SC.FinalCookPlatform.IndexOf("_") > 0 ? SC.FinalCookPlatform.Substring(SC.FinalCookPlatform.IndexOf("_")) : "";
+			if (!Params.Prebuilt)
+			{
+				string SOName = GetSONameWithoutArchitecture(Params, SC.StageExecutables[0]);
+				bool bShouldCompileAsDll = Receipt.HasValueForAdditionalProperty("CompileAsDll", "true");
+				if (bShouldCompileAsDll)
 				{
-					string SOName = GetSONameWithoutArchitecture(Params, SC.StageExecutables[0]);
-					bool bShouldCompileAsDll = Receipt.HasValueForAdditionalProperty("CompileAsDll", "true");
-					if (bShouldCompileAsDll)
-					{
-						// MakeApk
-						SOName = Receipt.BuildProducts[0].Path.FullName;
+					// MakeApk
+					SOName = Receipt.BuildProducts[0].Path.FullName;
 
-						// saving package info, which will allow 
-						TargetType Type = TargetType.Game;
-						if (CookFlavor.EndsWith("Client"))
-						{
-							Type = TargetType.Client;
-						}
-						else if (CookFlavor.EndsWith("Server"))
-						{
-							Type = TargetType.Server;
-						}
-						LogInformation("SavePackageInfo");
-						Deploy.SavePackageInfo(Params.ShortProjectName, SC.ProjectRoot.FullName, Type, true);
+					// saving package info, which will allow 
+					TargetType Type = TargetType.Game;
+					if (CookFlavor.EndsWith("Client"))
+					{
+						Type = TargetType.Client;
+					}
+					else if (CookFlavor.EndsWith("Server"))
+					{
+						Type = TargetType.Server;
 					}
-					Deploy.PrepForUATPackageOrDeploy(Params.RawProjectPath, Params.ShortProjectName, SC.ProjectRoot, SOName, SC.LocalRoot + "/Engine", Params.Distribution, CookFlavor, SC.StageTargets[0].Receipt.Configuration, false, bShouldCompileAsDll);
+					LogInformation("SavePackageInfo");
+					Deploy.SavePackageInfo(Params.ShortProjectName, SC.ProjectRoot.FullName, Type, true);
 				}
+				Deploy.PrepForUATPackageOrDeploy(Params.RawProjectPath, Params.ShortProjectName, SC.ProjectRoot, SOName, SC.LocalRoot + "/Engine", Params.Distribution, CookFlavor, SC.StageTargets[0].Receipt.Configuration, false, bShouldCompileAsDll);
+			}
 
-			    // Create APK specific OBB in case we have a detached OBB.
-			    string DeviceObbName = "";
-			    string ObbName = "";
-				string DevicePatchName = "";
-				string PatchName = "";
-				string DeviceOverflow1Name = "";
-				string Overflow1Name = "";
-				string DeviceOverflow2Name = "";
-				string Overflow2Name = "";
-				if (!bPackageDataInsideApk)
-			    {
-				    DeviceObbName = GetDeviceObbName(ApkName, SC);
-				    ObbName = GetFinalObbName(ApkName, SC);
-					CopyFile(LocalObbName, ObbName);
+			// Create APK specific OBB in case we have a detached OBB.
+			string DeviceObbName = "";
+			string ObbName = "";
+			string DevicePatchName = "";
+			string PatchName = "";
+			string DeviceOverflow1Name = "";
+			string Overflow1Name = "";
+			string DeviceOverflow2Name = "";
+			string Overflow2Name = "";
+			if (!bPackageDataInsideApk)
+			{
+				DeviceObbName = GetDeviceObbName(ApkName, SC);
+				ObbName = GetFinalObbName(ApkName, SC);
+				CopyFile(LocalObbName, ObbName);
 
-					// apply store version to OBB to make it unique for PlayStore upload
-					UpdateObbStoreVersion(ObbName);
+				// apply store version to OBB to make it unique for PlayStore upload
+				UpdateObbStoreVersion(ObbName);
 
-					if (File.Exists(LocalPatchName))
-					{
-						DevicePatchName = GetDevicePatchName(ApkName, SC);
-						PatchName = GetFinalPatchName(ApkName, SC);
-						CopyFile(LocalPatchName, PatchName);
+				if (File.Exists(LocalPatchName))
+				{
+					DevicePatchName = GetDevicePatchName(ApkName, SC);
+					PatchName = GetFinalPatchName(ApkName, SC);
+					CopyFile(LocalPatchName, PatchName);
 
-						// apply store version to OBB to make it unique for PlayStore upload
-						UpdateObbStoreVersion(PatchName);
-					}
+					// apply store version to OBB to make it unique for PlayStore upload
+					UpdateObbStoreVersion(PatchName);
+				}
 
-					if (File.Exists(LocalOverflow1Name))
-					{
-						DeviceOverflow1Name = GetDeviceOverflowName(ApkName, SC, 1);
-						Overflow1Name = GetFinalOverflowName(ApkName, SC, 1);
-						CopyFile(LocalOverflow1Name, Overflow1Name);
+				if (File.Exists(LocalOverflow1Name))
+				{
+					DeviceOverflow1Name = GetDeviceOverflowName(ApkName, SC, 1);
+					Overflow1Name = GetFinalOverflowName(ApkName, SC, 1);
+					CopyFile(LocalOverflow1Name, Overflow1Name);
 
-						// apply store version to OBB to make it unique for PlayStore upload
-						UpdateObbStoreVersion(Overflow1Name);
-					}
+					// apply store version to OBB to make it unique for PlayStore upload
+					UpdateObbStoreVersion(Overflow1Name);
+				}
 
-					if (File.Exists(LocalOverflow2Name))
-					{
-						DeviceOverflow2Name = GetDeviceOverflowName(ApkName, SC, 2);
-						Overflow2Name = GetFinalOverflowName(ApkName, SC, 2);
-						CopyFile(LocalOverflow2Name, Overflow2Name);
+				if (File.Exists(LocalOverflow2Name))
+				{
+					DeviceOverflow2Name = GetDeviceOverflowName(ApkName, SC, 2);
+					Overflow2Name = GetFinalOverflowName(ApkName, SC, 2);
+					CopyFile(LocalOverflow2Name, Overflow2Name);
 
-						// apply store version to OBB to make it unique for PlayStore upload
-						UpdateObbStoreVersion(Overflow2Name);
-					}
+					// apply store version to OBB to make it unique for PlayStore upload
+					UpdateObbStoreVersion(Overflow2Name);
 				}
+			}
 
-				// check for optional universal apk
-				string APKDirectory = Path.GetDirectoryName(ApkName);
-				string APKNameWithoutExtension = Path.GetFileNameWithoutExtension(ApkName);
-				string APKBareNameWithoutExtension = Path.GetFileNameWithoutExtension(ApkBareName);
-				string UniversalApkName = Path.Combine(APKDirectory, APKNameWithoutExtension + "_universal.apk");
-				bool bHaveUniversal = false;
-				if (bEnableBundle && bEnableUniversalAPK)
+			// check for optional universal apk
+			string APKDirectory = Path.GetDirectoryName(ApkName);
+			string APKNameWithoutExtension = Path.GetFileNameWithoutExtension(ApkName);
+			string APKBareNameWithoutExtension = Path.GetFileNameWithoutExtension(ApkBareName);
+			string UniversalApkName = Path.Combine(APKDirectory, APKNameWithoutExtension + "_universal.apk");
+			bool bHaveUniversal = false;
+			if (bEnableBundle && bEnableUniversalAPK)
+			{
+				if (FileExists(UniversalApkName))
+				{
+					bHaveUniversal = true;
+				}
+				else
 				{
+					UniversalApkName = Path.Combine(APKDirectory, APKBareNameWithoutExtension + "_universal.apk");
 					if (FileExists(UniversalApkName))
 					{
 						bHaveUniversal = true;
 					}
-					else
-					{
-						UniversalApkName = Path.Combine(APKDirectory, APKBareNameWithoutExtension + "_universal.apk");
-						if (FileExists(UniversalApkName))
-						{
-							bHaveUniversal = true;
-						}
-					}
 				}
+			}
 
-				//figure out which platforms we need to create install files for
-				bool bNeedsPCInstall = false;
-				bool bNeedsMacInstall = false;
-				bool bNeedsLinuxInstall = false;
-				GetPlatformInstallOptions(SC, out bNeedsPCInstall, out bNeedsMacInstall, out bNeedsLinuxInstall);
+			//figure out which platforms we need to create install files for
+			bool bNeedsPCInstall = false;
+			bool bNeedsMacInstall = false;
+			bool bNeedsLinuxInstall = false;
+			GetPlatformInstallOptions(SC, out bNeedsPCInstall, out bNeedsMacInstall, out bNeedsLinuxInstall);
 
-				//helper delegate to prevent code duplication but allow us access to all the local variables we need
-				var CreateInstallFilesAction = new Action<UnrealTargetPlatform>(Target =>
+			//helper delegate to prevent code duplication but allow us access to all the local variables we need
+			var CreateInstallFilesAction = new Action<UnrealTargetPlatform>(Target =>
+			{
+				bool bIsPC = (Target == UnrealTargetPlatform.Win64);
+				string LineEnding = bIsPC ? "\r\n" : "\n";
+				// Write install batch file(s).
+				string PackageName = GetPackageInfo(ApkName, SC, false);
+				string BatchName = GetFinalBatchName(ApkName, SC, bMakeSeparateApks ? Architecture : "", false, EBatchType.Install, Target);
+				string[] BatchLines = GenerateInstallBatchFile(bPackageDataInsideApk, PackageName, ApkName, Params, ObbName, DeviceObbName, false, PatchName, DevicePatchName, false, 
+					Overflow1Name, DeviceOverflow1Name, false, Overflow2Name, DeviceOverflow2Name, false, bIsPC, Params.Distribution, TargetSDKVersion > 22, bDisablePerfHarden);
+				if (bHaveAPK)
 				{
-					bool bIsPC = (Target == UnrealTargetPlatform.Win64);
-					string LineEnding = bIsPC ? "\r\n" : "\n";
-					// Write install batch file(s).
-					string PackageName = GetPackageInfo(ApkName, SC, false);
-					string BatchName = GetFinalBatchName(ApkName, SC, bMakeSeparateApks ? Architecture : "", bMakeSeparateApks ? GPUArchitecture : "", false, EBatchType.Install, Target);
-					string[] BatchLines = GenerateInstallBatchFile(bPackageDataInsideApk, PackageName, ApkName, Params, ObbName, DeviceObbName, false, PatchName, DevicePatchName, false, 
+					// make a batch file that can be used to install the .apk and .obb files
+					File.WriteAllText(BatchName, string.Join(LineEnding, BatchLines) + LineEnding);
+				}
+				// make a batch file that can be used to uninstall the .apk and .obb files
+				string UninstallBatchName = GetFinalBatchName(ApkName, SC, bMakeSeparateApks ? Architecture : "", false, EBatchType.Uninstall, Target);
+				BatchLines = GenerateUninstallBatchFile(bPackageDataInsideApk, PackageName, ApkName, Params, bIsPC);
+				if (bHaveAPK || bHaveUniversal)
+				{
+					File.WriteAllText(UninstallBatchName, string.Join(LineEnding, BatchLines) + LineEnding);
+				}
+
+				string UniversalBatchName = "";
+				if (bHaveUniversal)
+				{
+					UniversalBatchName = GetFinalBatchName(UniversalApkName, SC, "", false, EBatchType.Install, Target);
+					// make a batch file that can be used to install the .apk
+					string[] UniversalBatchLines = GenerateInstallBatchFile(bPackageDataInsideApk, PackageName, UniversalApkName, Params, ObbName, DeviceObbName, false, PatchName, DevicePatchName, false,
 						Overflow1Name, DeviceOverflow1Name, false, Overflow2Name, DeviceOverflow2Name, false, bIsPC, Params.Distribution, TargetSDKVersion > 22, bDisablePerfHarden);
+					File.WriteAllText(UniversalBatchName, string.Join(LineEnding, UniversalBatchLines) + LineEnding);
+				}
+
+				string SymbolizeBatchName = GetFinalBatchName(ApkName, SC, Architecture, false, EBatchType.Symbolize, Target);
+				if(bBuildWithHiddenSymbolVisibility || bSaveSymbols)
+				{
+					BatchLines = GenerateSymbolizeBatchFile(Params, PackageName, ApkName, SC, Architecture, bIsPC);
+					File.WriteAllText(SymbolizeBatchName, string.Join(LineEnding, BatchLines) + LineEnding);
+				}
+
+				if (!RuntimePlatform.IsWindows)
+				{
 					if (bHaveAPK)
 					{
-						// make a batch file that can be used to install the .apk and .obb files
-						File.WriteAllText(BatchName, string.Join(LineEnding, BatchLines) + LineEnding);
+						CommandUtils.FixUnixFilePermissions(BatchName);
 					}
-					// make a batch file that can be used to uninstall the .apk and .obb files
-					string UninstallBatchName = GetFinalBatchName(ApkName, SC, bMakeSeparateApks ? Architecture : "", bMakeSeparateApks ? GPUArchitecture : "", false, EBatchType.Uninstall, Target);
-					BatchLines = GenerateUninstallBatchFile(bPackageDataInsideApk, PackageName, ApkName, Params, bIsPC);
 					if (bHaveAPK || bHaveUniversal)
 					{
-						File.WriteAllText(UninstallBatchName, string.Join(LineEnding, BatchLines) + LineEnding);
+						CommandUtils.FixUnixFilePermissions(UninstallBatchName);
 					}
-
-					string UniversalBatchName = "";
 					if (bHaveUniversal)
 					{
-						UniversalBatchName = GetFinalBatchName(UniversalApkName, SC, "", "", false, EBatchType.Install, Target);
-						// make a batch file that can be used to install the .apk
-						string[] UniversalBatchLines = GenerateInstallBatchFile(bPackageDataInsideApk, PackageName, UniversalApkName, Params, ObbName, DeviceObbName, false, PatchName, DevicePatchName, false,
-							Overflow1Name, DeviceOverflow1Name, false, Overflow2Name, DeviceOverflow2Name, false, bIsPC, Params.Distribution, TargetSDKVersion > 22, bDisablePerfHarden);
-						File.WriteAllText(UniversalBatchName, string.Join(LineEnding, UniversalBatchLines) + LineEnding);
-					}
-
-					string SymbolizeBatchName = GetFinalBatchName(ApkName, SC, Architecture, GPUArchitecture, false, EBatchType.Symbolize, Target);
-					if(bBuildWithHiddenSymbolVisibility || bSaveSymbols)
-					{
-						BatchLines = GenerateSymbolizeBatchFile(Params, PackageName, ApkName, SC, Architecture, GPUArchitecture, bIsPC);
-						File.WriteAllText(SymbolizeBatchName, string.Join(LineEnding, BatchLines) + LineEnding);
+						CommandUtils.FixUnixFilePermissions(UniversalBatchName);
 					}
-
-					if (!RuntimePlatform.IsWindows)
+					if (bBuildWithHiddenSymbolVisibility || bSaveSymbols)
 					{
-						if (bHaveAPK)
-						{
-							CommandUtils.FixUnixFilePermissions(BatchName);
-						}
-						if (bHaveAPK || bHaveUniversal)
-						{
-							CommandUtils.FixUnixFilePermissions(UninstallBatchName);
-						}
-						if (bHaveUniversal)
-						{
-							CommandUtils.FixUnixFilePermissions(UniversalBatchName);
-						}
-						if (bBuildWithHiddenSymbolVisibility || bSaveSymbols)
-						{
-							CommandUtils.FixUnixFilePermissions(SymbolizeBatchName);
-						}
-						//if(File.Exists(NoInstallBatchName)) 
-						//{
-						//    CommandUtils.FixUnixFilePermissions(NoInstallBatchName);
-						//}
+						CommandUtils.FixUnixFilePermissions(SymbolizeBatchName);
 					}
-				});
-
-				if (bNeedsPCInstall)
-				{
-					CreateInstallFilesAction.Invoke(UnrealTargetPlatform.Win64);
-				}
-				if (bNeedsMacInstall)
-				{
-					CreateInstallFilesAction.Invoke(UnrealTargetPlatform.Mac);
-				}
-				if (bNeedsLinuxInstall)
-				{
-					CreateInstallFilesAction.Invoke(UnrealTargetPlatform.Linux);
+					//if(File.Exists(NoInstallBatchName)) 
+					//{
+					//    CommandUtils.FixUnixFilePermissions(NoInstallBatchName);
+					//}
 				}
+			});
 
-				// If we aren't packaging data in the APK then lets write out a bat file to also let us test without the OBB
-				// on the device.
-				//String NoInstallBatchName = GetFinalBatchName(ApkName, Params, bMakeSeparateApks ? Architecture : "", bMakeSeparateApks ? GPUArchitecture : "", true, false);
-				// if(!bPackageDataInsideApk)
-				//{
-				//    BatchLines = GenerateInstallBatchFile(bPackageDataInsideApk, PackageName, ApkName, Params, ObbName, DeviceObbName, true);
-				//    File.WriteAllLines(NoInstallBatchName, BatchLines);
-				//}
+			if (bNeedsPCInstall)
+			{
+				CreateInstallFilesAction.Invoke(UnrealTargetPlatform.Win64);
+			}
+			if (bNeedsMacInstall)
+			{
+				CreateInstallFilesAction.Invoke(UnrealTargetPlatform.Mac);
+			}
+			if (bNeedsLinuxInstall)
+			{
+				CreateInstallFilesAction.Invoke(UnrealTargetPlatform.Linux);
 			}
+
+			// If we aren't packaging data in the APK then lets write out a bat file to also let us test without the OBB
+			// on the device.
+			//String NoInstallBatchName = GetFinalBatchName(ApkName, Params, bMakeSeparateApks ? Architecture : "", bMakeSeparateApks ? GPUArchitecture : "", true, false);
+			// if(!bPackageDataInsideApk)
+			//{
+			//    BatchLines = GenerateInstallBatchFile(bPackageDataInsideApk, PackageName, ApkName, Params, ObbName, DeviceObbName, true);
+			//    File.WriteAllLines(NoInstallBatchName, BatchLines);
+			//}
 		}
 
 		PrintRunTime();
@@ -1307,7 +1303,7 @@ public class AndroidPlatform : Platform
 		return BatchLines;
 	}
 
-	private string[] GenerateSymbolizeBatchFile(ProjectParams Params, string PackageName, string ApkName, DeploymentContext SC, string Architecture, string GPUArchitecture, bool bIsPC)
+	private string[] GenerateSymbolizeBatchFile(ProjectParams Params, string PackageName, string ApkName, DeploymentContext SC, string Architecture, bool bIsPC)
 	{
 		string[] BatchLines = null;
 
@@ -1323,7 +1319,7 @@ public class AndroidPlatform : Platform
 				"cd \"`dirname \"$0\"`\"",
 				"NDKSTACK=",
 				"if [ \"$ANDROID_NDK_ROOT\" != \"\" ]; then NDKSTACK=$%ANDROID_NDK_ROOT/ndk-stack; else ADB=" + Environment.GetEnvironmentVariable("ANDROID_NDK_ROOT") + "/ndk-stack; fi",
-				"$NDKSTACK -sym " + GetFinalSymbolizedSODirectory(ApkName, SC, Architecture, GPUArchitecture) + " -dump \"%1\" > " + Params.ShortProjectName + "_SymbolizedCallStackOutput.txt",
+				"$NDKSTACK -sym " + GetFinalSymbolizedSODirectory(ApkName, SC, Architecture) + " -dump \"%1\" > " + Params.ShortProjectName + "_SymbolizedCallStackOutput.txt",
 				"exit 0",
 				};
 		}
@@ -1338,7 +1334,7 @@ public class AndroidPlatform : Platform
 						"if \"%ANDROID_NDK_ROOT%\"==\"\" set NDK_ROOT=\""+Environment.GetEnvironmentVariable("ANDROID_NDK_ROOT")+"\"",
 						"set NDKSTACK=%NDK_ROOT%\\ndk-stack.cmd",
 						"",
-						"%NDKSTACK% -sym "+GetFinalSymbolizedSODirectory(ApkName, SC, Architecture, GPUArchitecture)+" -dump \"%1\" > "+ Params.ShortProjectName+"_SymbolizedCallStackOutput.txt",
+						"%NDKSTACK% -sym "+GetFinalSymbolizedSODirectory(ApkName, SC, Architecture)+" -dump \"%1\" > "+ Params.ShortProjectName+"_SymbolizedCallStackOutput.txt",
 						"",
 						"goto:eof",
 						"",
@@ -1362,197 +1358,193 @@ public class AndroidPlatform : Platform
 		UnrealTargetConfiguration TargetConfiguration = SC.StageTargetConfigurations[0];
 		IAndroidToolChain ToolChain = AndroidExports.CreateToolChain(Params.RawProjectPath);
 		var Architectures = ToolChain.GetAllArchitectures();
-		var GPUArchitectures = ToolChain.GetAllGPUArchitectures();
 		bool bMakeSeparateApks = UnrealBuildTool.AndroidExports.ShouldMakeSeparateApks();
 		bool bPackageDataInsideApk = UnrealBuildTool.AndroidExports.CreateDeploymentHandler(Params.RawProjectPath, Params.ForcePackageData).GetPackageDataInsideApk();
 
 		List<string> AddedObbFiles = new List<string>();
 		foreach (string Architecture in Architectures)
 		{
-			foreach (string GPUArchitecture in GPUArchitectures)
+			string ApkBareName = GetFinalApkName(Params, SC.StageExecutables[0], true, "");
+			string ApkName = GetFinalApkName(Params, SC.StageExecutables[0], true, bMakeSeparateApks ? Architecture : "");
+			bool bHaveAPK = FileExists(ApkName);
+			string ObbName = GetFinalObbName(ApkName, SC);
+			string PatchName = GetFinalPatchName(ApkName, SC);
+			string Overflow1Name = GetFinalOverflowName(ApkName, SC, 1);
+			string Overflow2Name = GetFinalOverflowName(ApkName, SC, 2);
+			bool bBuildWithHiddenSymbolVisibility = BuildWithHiddenSymbolVisibility(SC);
+			bool bSaveSymbols = GetSaveSymbols(SC);
+			//string NoOBBBatchName = GetFinalBatchName(ApkName, Params, bMakeSeparateApks ? Architecture : "", bMakeSeparateApks ? GPUArchitecture : "", true, false);
+
+			string APKDirectory = Path.GetDirectoryName(ApkName);
+			string APKNameWithoutExtension = Path.GetFileNameWithoutExtension(ApkName);
+			string APKBareNameWithoutExtension = Path.GetFileNameWithoutExtension(ApkBareName);
+
+			bool bHaveAAB = false;
+			bool bHaveUniversal = false;
+
+			// copy optional app bundle if exists
+			string AppBundleName = Path.Combine(APKDirectory, APKNameWithoutExtension + ".aab");
+			if (FileExists(AppBundleName))
 			{
-				string ApkBareName = GetFinalApkName(Params, SC.StageExecutables[0], true, "", "");
-				string ApkName = GetFinalApkName(Params, SC.StageExecutables[0], true, bMakeSeparateApks ? Architecture : "", bMakeSeparateApks ? GPUArchitecture : "");
-				bool bHaveAPK = FileExists(ApkName);
-				string ObbName = GetFinalObbName(ApkName, SC);
-				string PatchName = GetFinalPatchName(ApkName, SC);
-				string Overflow1Name = GetFinalOverflowName(ApkName, SC, 1);
-				string Overflow2Name = GetFinalOverflowName(ApkName, SC, 2);
-				bool bBuildWithHiddenSymbolVisibility = BuildWithHiddenSymbolVisibility(SC);
-				bool bSaveSymbols = GetSaveSymbols(SC);
-				//string NoOBBBatchName = GetFinalBatchName(ApkName, Params, bMakeSeparateApks ? Architecture : "", bMakeSeparateApks ? GPUArchitecture : "", true, false);
-
-				string APKDirectory = Path.GetDirectoryName(ApkName);
-				string APKNameWithoutExtension = Path.GetFileNameWithoutExtension(ApkName);
-				string APKBareNameWithoutExtension = Path.GetFileNameWithoutExtension(ApkBareName);
-
-				bool bHaveAAB = false;
-				bool bHaveUniversal = false;
-
-				// copy optional app bundle if exists
-				string AppBundleName = Path.Combine(APKDirectory, APKNameWithoutExtension + ".aab");
+				bHaveAAB = true;
+				SC.ArchiveFiles(APKDirectory, Path.GetFileName(AppBundleName));
+			}
+			else
+			{
+				AppBundleName = Path.Combine(APKDirectory, APKBareNameWithoutExtension + ".aab");
 				if (FileExists(AppBundleName))
 				{
 					bHaveAAB = true;
 					SC.ArchiveFiles(APKDirectory, Path.GetFileName(AppBundleName));
 				}
-				else
-				{
-					AppBundleName = Path.Combine(APKDirectory, APKBareNameWithoutExtension + ".aab");
-					if (FileExists(AppBundleName))
-					{
-						bHaveAAB = true;
-						SC.ArchiveFiles(APKDirectory, Path.GetFileName(AppBundleName));
-					}
-				}
+			}
 
-				// copy optional apks (zip of split apks) if exists
-				string APKSName = Path.Combine(APKDirectory, APKNameWithoutExtension + ".apks");
+			// copy optional apks (zip of split apks) if exists
+			string APKSName = Path.Combine(APKDirectory, APKNameWithoutExtension + ".apks");
+			if (FileExists(APKSName))
+			{
+				SC.ArchiveFiles(APKDirectory, Path.GetFileName(APKSName));
+			}
+			else
+			{
+				APKSName = Path.Combine(APKDirectory, APKBareNameWithoutExtension + ".apks");
 				if (FileExists(APKSName))
 				{
 					SC.ArchiveFiles(APKDirectory, Path.GetFileName(APKSName));
 				}
-				else
-				{
-					APKSName = Path.Combine(APKDirectory, APKBareNameWithoutExtension + ".apks");
-					if (FileExists(APKSName))
-					{
-						SC.ArchiveFiles(APKDirectory, Path.GetFileName(APKSName));
-					}
-				}
+			}
 
-				// copy optional universal apk if exists
-				string UniversalApkName = Path.Combine(APKDirectory, APKNameWithoutExtension + "_universal.apk");
+			// copy optional universal apk if exists
+			string UniversalApkName = Path.Combine(APKDirectory, APKNameWithoutExtension + "_universal.apk");
+			if (FileExists(UniversalApkName))
+			{
+				bHaveUniversal = true;
+				SC.ArchiveFiles(APKDirectory, Path.GetFileName(UniversalApkName));
+			}
+			else
+			{
+				UniversalApkName = Path.Combine(APKDirectory, APKBareNameWithoutExtension + "_universal.apk");
 				if (FileExists(UniversalApkName))
 				{
 					bHaveUniversal = true;
 					SC.ArchiveFiles(APKDirectory, Path.GetFileName(UniversalApkName));
 				}
-				else
-				{
-					UniversalApkName = Path.Combine(APKDirectory, APKBareNameWithoutExtension + "_universal.apk");
-					if (FileExists(UniversalApkName))
-					{
-						bHaveUniversal = true;
-						SC.ArchiveFiles(APKDirectory, Path.GetFileName(UniversalApkName));
-					}
-				}
-
-				// verify the files exist
-				if (!FileExists(ApkName))
-				{
-					// still valid if we found an AAB
-					if (!bHaveAAB)
-					{
-						throw new AutomationException(ExitCode.Error_AppNotFound, "ARCHIVE FAILED - {0} was not found", ApkName);
-					}
-				}
-				else
-				{
-					SC.ArchiveFiles(Path.GetDirectoryName(ApkName), Path.GetFileName(ApkName));
-				}
+			}
 
-				if (!bPackageDataInsideApk && !FileExists(ObbName))
+			// verify the files exist
+			if (!FileExists(ApkName))
+			{
+				// still valid if we found an AAB
+				if (!bHaveAAB)
 				{
-                    throw new AutomationException(ExitCode.Error_ObbNotFound, "ARCHIVE FAILED - {0} was not found", ObbName);
+					throw new AutomationException(ExitCode.Error_AppNotFound, "ARCHIVE FAILED - {0} was not found", ApkName);
 				}
+			}
+			else
+			{
+				SC.ArchiveFiles(Path.GetDirectoryName(ApkName), Path.GetFileName(ApkName));
+			}
 
-				if (bBuildWithHiddenSymbolVisibility || bSaveSymbols)
-				{
-					string SymbolizedSODirectory = GetFinalSymbolizedSODirectory(ApkName, SC, Architecture, GPUArchitecture);
-					string SymbolizedSOPath = Path.Combine(Path.Combine(Path.GetDirectoryName(ApkName), SymbolizedSODirectory), "libUnreal.so");
-					if (!FileExists(SymbolizedSOPath))
-					{
-						throw new AutomationException(ExitCode.Error_SymbolizedSONotFound, "ARCHIVE FAILED - {0} was not found", SymbolizedSOPath);
-					}
-
-					// Add symbolized .so directory
-					SC.ArchiveFiles(Path.GetDirectoryName(SymbolizedSOPath), Path.GetFileName(SymbolizedSOPath), true, null, SymbolizedSODirectory);
-				}
+			if (!bPackageDataInsideApk && !FileExists(ObbName))
+			{
+                throw new AutomationException(ExitCode.Error_ObbNotFound, "ARCHIVE FAILED - {0} was not found", ObbName);
+			}
 
-				if (!bPackageDataInsideApk)
+			if (bBuildWithHiddenSymbolVisibility || bSaveSymbols)
+			{
+				string SymbolizedSODirectory = GetFinalSymbolizedSODirectory(ApkName, SC, Architecture);
+				string SymbolizedSOPath = Path.Combine(Path.Combine(Path.GetDirectoryName(ApkName), SymbolizedSODirectory), "libUnreal.so");
+				if (!FileExists(SymbolizedSOPath))
 				{
-					// only add if not already in archive list
-					if (!AddedObbFiles.Contains(ObbName))
-					{
-						AddedObbFiles.Add(ObbName);
-
-						SC.ArchiveFiles(Path.GetDirectoryName(ObbName), Path.GetFileName(ObbName));
-						if (FileExists(PatchName))
-						{
-							SC.ArchiveFiles(Path.GetDirectoryName(PatchName), Path.GetFileName(PatchName));
-						}
-						if (FileExists(Overflow1Name))
-						{
-							SC.ArchiveFiles(Path.GetDirectoryName(Overflow1Name), Path.GetFileName(Overflow1Name));
-						}
-						if (FileExists(Overflow2Name))
-						{
-							SC.ArchiveFiles(Path.GetDirectoryName(Overflow2Name), Path.GetFileName(Overflow2Name));
-						}
-					}
+					throw new AutomationException(ExitCode.Error_SymbolizedSONotFound, "ARCHIVE FAILED - {0} was not found", SymbolizedSOPath);
 				}
 
-				// copy optional unprotected APK if exists
-				string UnprotectedApkName = Path.Combine(APKDirectory, "unprotected_" + APKNameWithoutExtension + ".apk");
-				if (FileExists(UnprotectedApkName))
-				{
-					SC.ArchiveFiles(APKDirectory, Path.GetFileName(UnprotectedApkName));
-				}
+				// Add symbolized .so directory
+				SC.ArchiveFiles(Path.GetDirectoryName(SymbolizedSOPath), Path.GetFileName(SymbolizedSOPath), true, null, SymbolizedSODirectory);
+			}
 
-				// copy optional logs directory if exists
-				string LogsDirName = Path.Combine(APKDirectory, APKNameWithoutExtension + ".logs");
-				if (DirectoryExists(LogsDirName))
+			if (!bPackageDataInsideApk)
+			{
+				// only add if not already in archive list
+				if (!AddedObbFiles.Contains(ObbName))
 				{
-					SC.ArchiveFiles(LogsDirName);
-				}
+					AddedObbFiles.Add(ObbName);
 
-				bool bNeedsPCInstall = false;
-				bool bNeedsMacInstall = false;
-				bool bNeedsLinuxInstall = false;
-				GetPlatformInstallOptions(SC, out bNeedsPCInstall, out bNeedsMacInstall, out bNeedsLinuxInstall);
-
-				//helper delegate to prevent code duplication but allow us access to all the local variables we need
-				var CreateBatchFilesAndArchiveAction = new Action<UnrealTargetPlatform>(Target =>
-				{
-					if (bHaveAPK)
+					SC.ArchiveFiles(Path.GetDirectoryName(ObbName), Path.GetFileName(ObbName));
+					if (FileExists(PatchName))
 					{
-						string BatchName = GetFinalBatchName(ApkName, SC, bMakeSeparateApks ? Architecture : "", bMakeSeparateApks ? GPUArchitecture : "", false, EBatchType.Install, Target);
-						SC.ArchiveFiles(Path.GetDirectoryName(BatchName), Path.GetFileName(BatchName));
+						SC.ArchiveFiles(Path.GetDirectoryName(PatchName), Path.GetFileName(PatchName));
 					}
-					if (bHaveAPK || bHaveUniversal)
-					{
-						string UninstallBatchName = GetFinalBatchName(ApkName, SC, bMakeSeparateApks ? Architecture : "", bMakeSeparateApks ? GPUArchitecture : "", false, EBatchType.Uninstall, Target);
-						SC.ArchiveFiles(Path.GetDirectoryName(UninstallBatchName), Path.GetFileName(UninstallBatchName));
-					}
-					if (bHaveUniversal)
+					if (FileExists(Overflow1Name))
 					{
-						string UniversalBatchName = GetFinalBatchName(UniversalApkName, SC, "", "", false, EBatchType.Install, Target);
-						SC.ArchiveFiles(Path.GetDirectoryName(UniversalBatchName), Path.GetFileName(UniversalBatchName));
+						SC.ArchiveFiles(Path.GetDirectoryName(Overflow1Name), Path.GetFileName(Overflow1Name));
 					}
-
-					if (bBuildWithHiddenSymbolVisibility || bSaveSymbols)
+					if (FileExists(Overflow2Name))
 					{
-						string SymbolizeBatchName = GetFinalBatchName(ApkName, SC, Architecture, GPUArchitecture, false, EBatchType.Symbolize, Target);
-						SC.ArchiveFiles(Path.GetDirectoryName(SymbolizeBatchName), Path.GetFileName(SymbolizeBatchName));
+						SC.ArchiveFiles(Path.GetDirectoryName(Overflow2Name), Path.GetFileName(Overflow2Name));
 					}
-					//SC.ArchiveFiles(Path.GetDirectoryName(NoOBBBatchName), Path.GetFileName(NoOBBBatchName));
 				}
-				);
+			}
+
+			// copy optional unprotected APK if exists
+			string UnprotectedApkName = Path.Combine(APKDirectory, "unprotected_" + APKNameWithoutExtension + ".apk");
+			if (FileExists(UnprotectedApkName))
+			{
+				SC.ArchiveFiles(APKDirectory, Path.GetFileName(UnprotectedApkName));
+			}
+
+			// copy optional logs directory if exists
+			string LogsDirName = Path.Combine(APKDirectory, APKNameWithoutExtension + ".logs");
+			if (DirectoryExists(LogsDirName))
+			{
+				SC.ArchiveFiles(LogsDirName);
+			}
+
+			bool bNeedsPCInstall = false;
+			bool bNeedsMacInstall = false;
+			bool bNeedsLinuxInstall = false;
+			GetPlatformInstallOptions(SC, out bNeedsPCInstall, out bNeedsMacInstall, out bNeedsLinuxInstall);
 
-				//it's possible we will need both PC and Mac/Linux install files, do both
-				if (bNeedsPCInstall)
+			//helper delegate to prevent code duplication but allow us access to all the local variables we need
+			var CreateBatchFilesAndArchiveAction = new Action<UnrealTargetPlatform>(Target =>
+			{
+				if (bHaveAPK)
 				{
-					CreateBatchFilesAndArchiveAction(UnrealTargetPlatform.Win64);
+					string BatchName = GetFinalBatchName(ApkName, SC, bMakeSeparateApks ? Architecture : "", false, EBatchType.Install, Target);
+					SC.ArchiveFiles(Path.GetDirectoryName(BatchName), Path.GetFileName(BatchName));
 				}
-				if (bNeedsMacInstall)
+				if (bHaveAPK || bHaveUniversal)
 				{
-					CreateBatchFilesAndArchiveAction(UnrealTargetPlatform.Mac);
+					string UninstallBatchName = GetFinalBatchName(ApkName, SC, bMakeSeparateApks ? Architecture : "", false, EBatchType.Uninstall, Target);
+					SC.ArchiveFiles(Path.GetDirectoryName(UninstallBatchName), Path.GetFileName(UninstallBatchName));
 				}
-				if (bNeedsLinuxInstall)
+				if (bHaveUniversal)
+				{
+					string UniversalBatchName = GetFinalBatchName(UniversalApkName, SC, "", false, EBatchType.Install, Target);
+					SC.ArchiveFiles(Path.GetDirectoryName(UniversalBatchName), Path.GetFileName(UniversalBatchName));
+				}
+
+				if (bBuildWithHiddenSymbolVisibility || bSaveSymbols)
 				{
-					CreateBatchFilesAndArchiveAction(UnrealTargetPlatform.Linux);
+					string SymbolizeBatchName = GetFinalBatchName(ApkName, SC, Architecture, false, EBatchType.Symbolize, Target);
+					SC.ArchiveFiles(Path.GetDirectoryName(SymbolizeBatchName), Path.GetFileName(SymbolizeBatchName));
 				}
+				//SC.ArchiveFiles(Path.GetDirectoryName(NoOBBBatchName), Path.GetFileName(NoOBBBatchName));
+			}
+			);
+
+			//it's possible we will need both PC and Mac/Linux install files, do both
+			if (bNeedsPCInstall)
+			{
+				CreateBatchFilesAndArchiveAction(UnrealTargetPlatform.Win64);
+			}
+			if (bNeedsMacInstall)
+			{
+				CreateBatchFilesAndArchiveAction(UnrealTargetPlatform.Mac);
+			}
+			if (bNeedsLinuxInstall)
+			{
+				CreateBatchFilesAndArchiveAction(UnrealTargetPlatform.Linux);
 			}
 		}
 	}
@@ -1776,9 +1768,7 @@ public class AndroidPlatform : Platform
 		foreach (var DeviceName in Params.DeviceNames)
         {
             string DeviceArchitecture = GetBestDeviceArchitecture(Params, DeviceName);
-            string GPUArchitecture = GetBestGPUArchitecture(Params, DeviceName);
-
-            string ApkName = GetFinalApkName(Params, SC.StageExecutables[0], true, DeviceArchitecture, GPUArchitecture);
+            string ApkName = GetFinalApkName(Params, SC.StageExecutables[0], true, DeviceArchitecture);
 
             // make sure APK is up to date (this is fast if so)
             var Deploy = AndroidExports.CreateDeploymentHandler(Params.RawProjectPath, Params.ForcePackageData);
@@ -2623,11 +2613,6 @@ public class AndroidPlatform : Platform
 		return DeviceArch;
 	}
 
-	private string GetBestGPUArchitecture(ProjectParams Params, string DeviceName)
-	{
-		return "";
-	}
-
 	public override IProcessResult RunClient(ERunOptions ClientRunFlags, string ClientApp, string ClientCmdLine, ProjectParams Params)
 	{
 		IProcessResult Result = null;
@@ -2643,25 +2628,15 @@ public class AndroidPlatform : Platform
 
 			//get the package name and save that
 			string DeviceArchitecture = GetBestDeviceArchitecture(Params, DeviceName);
-			string GPUArchitecture = GetBestGPUArchitecture(Params, DeviceName);
 
 			//strip off the device, GPU architecture and extension (.so)
 			int DashIndex = ClientApp.LastIndexOf("-");
 			if (DashIndex >= 0)
 			{
 				ClientApp = ClientApp.Substring(0, DashIndex);
-
-				if (GPUArchitecture.Length > 0)
-				{
-					DashIndex = ClientApp.LastIndexOf("-");
-					if (DashIndex >= 0)
-					{
-						ClientApp = ClientApp.Substring(0, DashIndex);
-					}
-				}
 			}
 
-			string ApkName = GetFinalApkName(Params, Path.GetFileNameWithoutExtension(ClientApp), true, DeviceArchitecture, GPUArchitecture);
+			string ApkName = GetFinalApkName(Params, Path.GetFileNameWithoutExtension(ClientApp), true, DeviceArchitecture);
 
 			if (!File.Exists(ApkName))
 			{
diff --git a/Engine/Source/Programs/UnrealBuildTool/Platform/Android/AndroidExports.cs b/Engine/Source/Programs/UnrealBuildTool/Platform/Android/AndroidExports.cs
--- a/Engine/Source/Programs/UnrealBuildTool/Platform/Android/AndroidExports.cs
+++ b/Engine/Source/Programs/UnrealBuildTool/Platform/Android/AndroidExports.cs
@@ -23,15 +23,9 @@ namespace UnrealBuildTool
 		List<string> GetAllArchitectures();
 
 		/// <summary>
-		/// Finds the list of supported GPU architectures
+		/// Returns the Android NDK Level
 		/// </summary>
-		/// <returns>The targeted GPU architectures</returns>
-		List<string> GetAllGPUArchitectures();
-
-		/// <summary>
-		/// Finds the list of supported GPU architectures
-		/// </summary>
-		/// <returns>The targeted GPU architectures</returns>
+		/// <returns>The NDK Level</returns>
 		int GetNdkApiLevelInt(int MinNDK);
 	}
 
diff --git a/Engine/Source/Programs/UnrealBuildTool/Platform/Android/AndroidToolChain.cs b/Engine/Source/Programs/UnrealBuildTool/Platform/Android/AndroidToolChain.cs
--- a/Engine/Source/Programs/UnrealBuildTool/Platform/Android/AndroidToolChain.cs
+++ b/Engine/Source/Programs/UnrealBuildTool/Platform/Android/AndroidToolChain.cs
@@ -59,11 +59,6 @@ namespace UnrealBuildTool
 			"-x64"
 		};
 
-		public static readonly string[] AllGpuSuffixes =
-		{
-			"",
-		};
-
 		// sh0rt names for the above suffixes
 		public static readonly Dictionary<string, string> ShortArchNames = new Dictionary<string, string>()
 		{
@@ -112,10 +107,6 @@ namespace UnrealBuildTool
 
 		// the list of architectures we will compile for
 		protected List<string> Arches = null;
-		// the list of GPU architectures we will compile for
-		protected List<string> GPUArchitectures = null;
-		// a list of all architecture+GPUArchitecture names (-armv7-es2, etc)
-		protected List<string> AllComboNames = null;
 
 		private AndroidToolChainOptions Options;
 
@@ -400,14 +391,6 @@ namespace UnrealBuildTool
 					Arches.Add("-arm64");
 				}
 			}
-
-			// For android just set the GPUArchitecture to an empty string
-			GPUArchitectures = new List<string>();
-			GPUArchitectures.Add("");
-
-			AllComboNames = (from Arch in Arches
-							 from GPUArch in GPUArchitectures
-							 select Arch + GPUArch).ToList();
 		}
 
 		static public string GetGLESVersion(bool bBuildForES31)
@@ -467,16 +450,6 @@ namespace UnrealBuildTool
 			return Arches;
 		}
 
-		public List<string> GetAllGPUArchitectures()
-		{
-			if (GPUArchitectures == null)
-			{
-				ParseArchitectures();
-			}
-
-			return GPUArchitectures;
-		}
-
 		public int GetNdkApiLevelInt(int MinNdk = 21)
 		{
 			string NDKVersion = GetNdkApiLevel();
@@ -1074,7 +1047,7 @@ namespace UnrealBuildTool
 			return false;
 		}
 
-		bool ShouldSkipLib(string FullLib, string Arch, string GPUArchitecture)
+		bool ShouldSkipLib(string FullLib, string Arch)
 		{
 			// strip any absolute path
 			string Lib = Path.GetFileNameWithoutExtension(FullLib);
@@ -1109,9 +1082,9 @@ namespace UnrealBuildTool
 			}
 
 			// if another architecture is in the filename, reject it
-			foreach (string ComboName in AllComboNames)
+			foreach (string ComboName in Arches)
 			{
-				if (ComboName != Arch + GPUArchitecture)
+				if (ComboName != Arch)
 				{
 					if (Lib.EndsWith(ComboName))
 					{
@@ -1383,211 +1356,199 @@ namespace UnrealBuildTool
 					continue;
 				}
 
-				foreach (string GPUArchitecture in GPUArchitectures)
+				// which toolchain to use
+				string Arguments = GetCLArguments_Global(CompileEnvironment, Arch) + BaseArguments;
+
+				switch (Arch)
 				{
-					// which toolchain to use
-					string Arguments = GetCLArguments_Global(CompileEnvironment, Arch) + BaseArguments;
+					case "-x64": Arguments += " -DPLATFORM_64BITS=1 -DPLATFORM_ANDROID_X64=1 -DPLATFORM_USED_NDK_VERSION_INTEGER=" + NDKApiLevel64Int.ToString(); break;
+					case "-arm64": Arguments += " -DPLATFORM_64BITS=1 -DPLATFORM_ANDROID_ARM64=1 -DPLATFORM_USED_NDK_VERSION_INTEGER=" + NDKApiLevel64Int.ToString(); break;
+					default: throw new BuildException($"Unknown arch {Arch}");
+				}
 
-					switch (Arch)
+				if(CompileEnvironment.bCompileISPC)
+				{
+					Arguments += " -DINTEL_ISPC=1";
+				}
+
+				// which PCH file to include
+				string PCHArguments = "";
+				if (CompileEnvironment.PrecompiledHeaderAction == PrecompiledHeaderAction.Include)
+				{
+					// add the platform-specific PCH reference
+					PCHArguments += string.Format(" -include-pch \"{0}\"", InlineArchName(BasePCHName, Arch) + PCHExtension);
+
+					// Add the precompiled header file's path to the include path so Clang can find it.
+					// This needs to be before the other include paths to ensure Clang uses it instead of the source header file.
+					PCHArguments += string.Format(" -include \"{0}\"", BasePCHName);
+				}
+
+				// Add include paths to the argument list (filtered by architecture)
+				foreach (DirectoryReference IncludePath in CompileEnvironment.SystemIncludePaths)
+				{
+					if (IsDirectoryForArch(IncludePath.FullName, Arch))
 					{
-						case "-x64": Arguments += " -DPLATFORM_64BITS=1 -DPLATFORM_ANDROID_X64=1 -DPLATFORM_USED_NDK_VERSION_INTEGER=" + NDKApiLevel64Int.ToString(); break;
-						case "-arm64": Arguments += " -DPLATFORM_64BITS=1 -DPLATFORM_ANDROID_ARM64=1 -DPLATFORM_USED_NDK_VERSION_INTEGER=" + NDKApiLevel64Int.ToString(); break;
-						default: throw new BuildException($"Unknown arch {Arch}");
+						Arguments += string.Format(" -I\"{0}\"", IncludePath);
 					}
-
-					if(CompileEnvironment.bCompileISPC)
+				}
+				foreach (DirectoryReference IncludePath in CompileEnvironment.UserIncludePaths)
+				{
+					if (IsDirectoryForArch(IncludePath.FullName, Arch))
 					{
-						Arguments += " -DINTEL_ISPC=1";
+						Arguments += string.Format(" -I\"{0}\"", IncludePath);
 					}
+				}
 
-					Arguments += " -DPLATFORM_ANDROIDGL4=" + ((GPUArchitecture == "-gl4") ? "1" : "0");
-
-					// which PCH file to include
-					string PCHArguments = "";
-					if (CompileEnvironment.PrecompiledHeaderAction == PrecompiledHeaderAction.Include)
-					{
-						// add the platform-specific PCH reference
-						PCHArguments += string.Format(" -include-pch \"{0}\"", InlineArchName(BasePCHName, Arch, GPUArchitecture) + PCHExtension);
+				foreach (FileItem SourceFile in InputFiles)
+				{
+					Action CompileAction = Graph.CreateAction(ActionType.Compile);
+					CompileAction.PrerequisiteItems.AddRange(CompileEnvironment.ForceIncludeFiles);
+					CompileAction.PrerequisiteItems.AddRange(CompileEnvironment.AdditionalPrerequisites);
 
-						// Add the precompiled header file's path to the include path so Clang can find it.
-						// This needs to be before the other include paths to ensure Clang uses it instead of the source header file.
-						PCHArguments += string.Format(" -include \"{0}\"", BasePCHName);
-					}
+					string FileArguments = "";
+					bool bIsPlainCFile = Path.GetExtension(SourceFile.AbsolutePath).ToUpperInvariant() == ".C";
+					bool bDisableShadowWarning = false;
 
-					// Add include paths to the argument list (filtered by architecture)
-					foreach (DirectoryReference IncludePath in CompileEnvironment.SystemIncludePaths)
+					// Add C or C++ specific compiler arguments.
+					if (bIsPlainCFile)
 					{
-						if (IsDirectoryForArch(IncludePath.FullName, Arch))
+						FileArguments += GetCompileArguments_C();
+
+						// remove shadow variable warnings for externally included files
+						if (!SourceFile.Location.IsUnderDirectory(Unreal.RootDirectory))
 						{
-							Arguments += string.Format(" -I\"{0}\"", IncludePath);
+							bDisableShadowWarning = true;
 						}
 					}
-					foreach (DirectoryReference IncludePath in CompileEnvironment.UserIncludePaths)
+					else if (CompileEnvironment.PrecompiledHeaderAction == PrecompiledHeaderAction.Create)
 					{
-						if (IsDirectoryForArch(IncludePath.FullName, Arch))
-						{
-							Arguments += string.Format(" -I\"{0}\"", IncludePath);
-						}
+						FileArguments += GetCompileArguments_PCH(CompileEnvironment);
 					}
-
-					foreach (FileItem SourceFile in InputFiles)
+					else
 					{
-						Action CompileAction = Graph.CreateAction(ActionType.Compile);
-						CompileAction.PrerequisiteItems.AddRange(CompileEnvironment.ForceIncludeFiles);
-						CompileAction.PrerequisiteItems.AddRange(CompileEnvironment.AdditionalPrerequisites);
+						FileArguments += GetCompileArguments_CPP(CompileEnvironment);
 
-						string FileArguments = "";
-						bool bIsPlainCFile = Path.GetExtension(SourceFile.AbsolutePath).ToUpperInvariant() == ".C";
-						bool bDisableShadowWarning = false;
+						// only use PCH for .cpp files
+						FileArguments += PCHArguments;
+					}
 
-						// Add C or C++ specific compiler arguments.
-						if (bIsPlainCFile)
-						{
-							FileArguments += GetCompileArguments_C();
+					foreach (FileItem ForceIncludeFile in CompileEnvironment.ForceIncludeFiles)
+					{
+						FileArguments += string.Format(" -include \"{0}\"", ForceIncludeFile.Location);
+					}
 
-							// remove shadow variable warnings for externally included files
-							if (!SourceFile.Location.IsUnderDirectory(Unreal.RootDirectory))
-							{
-								bDisableShadowWarning = true;
-							}
-						}
-						else if (CompileEnvironment.PrecompiledHeaderAction == PrecompiledHeaderAction.Create)
-						{
-							FileArguments += GetCompileArguments_PCH(CompileEnvironment);
-						}
-						else
-						{
-							FileArguments += GetCompileArguments_CPP(CompileEnvironment);
+					// Add the C++ source file and its included files to the prerequisite item list.
+					CompileAction.PrerequisiteItems.Add(SourceFile);
 
-							// only use PCH for .cpp files
-							FileArguments += PCHArguments;
-						}
+					if (CompileEnvironment.PrecompiledHeaderAction == PrecompiledHeaderAction.Create && !bIsPlainCFile)
+					{
+						// Add the precompiled header file to the produced item list.
+						FileItem PrecompiledHeaderFile = FileItem.GetItemByFileReference(
+							FileReference.Combine(
+								OutputDir,
+								Path.GetFileName(InlineArchName(SourceFile.AbsolutePath, Arch) + PCHExtension)
+								)
+							);
 
-						foreach (FileItem ForceIncludeFile in CompileEnvironment.ForceIncludeFiles)
+						CompileAction.ProducedItems.Add(PrecompiledHeaderFile);
+						Result.PrecompiledHeaderFile = PrecompiledHeaderFile;
+
+						// Add the parameters needed to compile the precompiled header file to the command-line.
+						FileArguments += string.Format(" -o \"{0}\"", PrecompiledHeaderFile.AbsolutePath);
+					}
+					else
+					{
+						if (CompileEnvironment.PrecompiledHeaderAction == PrecompiledHeaderAction.Include)
 						{
-							FileArguments += string.Format(" -include \"{0}\"", ForceIncludeFile.Location);
+							FileItem ArchPrecompiledHeaderFile = FileItem.GetItemByPath(InlineArchName(BasePCHName, Arch) + PCHExtension);
+							CompileAction.PrerequisiteItems.Add(ArchPrecompiledHeaderFile);
 						}
 
-						// Add the C++ source file and its included files to the prerequisite item list.
-						CompileAction.PrerequisiteItems.Add(SourceFile);
-
-						if (CompileEnvironment.PrecompiledHeaderAction == PrecompiledHeaderAction.Create && !bIsPlainCFile)
+						string ObjectFileExtension;
+						if(CompileEnvironment.AdditionalArguments != null && CompileEnvironment.AdditionalArguments.Contains("-emit-llvm"))
 						{
-							// Add the precompiled header file to the produced item list.
-							FileItem PrecompiledHeaderFile = FileItem.GetItemByFileReference(
-								FileReference.Combine(
-									OutputDir,
-									Path.GetFileName(InlineArchName(SourceFile.AbsolutePath, Arch, GPUArchitecture) + PCHExtension)
-									)
-								);
-
-							CompileAction.ProducedItems.Add(PrecompiledHeaderFile);
-							Result.PrecompiledHeaderFile = PrecompiledHeaderFile;
-
-							// Add the parameters needed to compile the precompiled header file to the command-line.
-							FileArguments += string.Format(" -o \"{0}\"", PrecompiledHeaderFile.AbsolutePath);
+							ObjectFileExtension = ".bc";
 						}
 						else
 						{
-							if (CompileEnvironment.PrecompiledHeaderAction == PrecompiledHeaderAction.Include)
-							{
-								FileItem ArchPrecompiledHeaderFile = FileItem.GetItemByPath(InlineArchName(BasePCHName, Arch, GPUArchitecture) + PCHExtension);
-								CompileAction.PrerequisiteItems.Add(ArchPrecompiledHeaderFile);
-							}
-
-							string ObjectFileExtension;
-							if(CompileEnvironment.AdditionalArguments != null && CompileEnvironment.AdditionalArguments.Contains("-emit-llvm"))
-							{
-								ObjectFileExtension = ".bc";
-							}
-							else
-							{
-								ObjectFileExtension = ".o";
-							}
+							ObjectFileExtension = ".o";
+						}
 
-							// Add the object file to the produced item list.
-							FileItem ObjectFile = FileItem.GetItemByFileReference(
-								FileReference.Combine(
-									OutputDir,
-									InlineArchName(Path.GetFileName(SourceFile.AbsolutePath) + ObjectFileExtension, Arch, GPUArchitecture, true)
-									)
-								);
-							CompileAction.ProducedItems.Add(ObjectFile);
-							Result.ObjectFiles.Add(ObjectFile);
+						// Add the object file to the produced item list.
+						FileItem ObjectFile = FileItem.GetItemByFileReference(
+							FileReference.Combine(
+								OutputDir,
+								InlineArchName(Path.GetFileName(SourceFile.AbsolutePath) + ObjectFileExtension, Arch, true)
+								)
+							);
+						CompileAction.ProducedItems.Add(ObjectFile);
+						Result.ObjectFiles.Add(ObjectFile);
 
-							FileArguments += string.Format(" -o \"{0}\"", ObjectFile.AbsolutePath);
-						}
+						FileArguments += string.Format(" -o \"{0}\"", ObjectFile.AbsolutePath);
+					}
 
-						// Add the source file path to the command-line.
-						FileArguments += string.Format(" \"{0}\"", SourceFile.AbsolutePath);
+					// Add the source file path to the command-line.
+					FileArguments += string.Format(" \"{0}\"", SourceFile.AbsolutePath);
 
-						// Generate the included header dependency list
-						if(CompileEnvironment.bGenerateDependenciesFile)
-						{
-							FileItem DependencyListFile = FileItem.GetItemByFileReference(FileReference.Combine(OutputDir, InlineArchName(Path.GetFileName(SourceFile.AbsolutePath) + ".d", Arch, GPUArchitecture, true)));
-							FileArguments += string.Format(" -MD -MF\"{0}\"", DependencyListFile.AbsolutePath.Replace('\\', '/'));
-							CompileAction.DependencyListFile = DependencyListFile;
-							CompileAction.ProducedItems.Add(DependencyListFile);
-						}
+					// Generate the included header dependency list
+					if(CompileEnvironment.bGenerateDependenciesFile)
+					{
+						FileItem DependencyListFile = FileItem.GetItemByFileReference(FileReference.Combine(OutputDir, InlineArchName(Path.GetFileName(SourceFile.AbsolutePath) + ".d", Arch, true)));
+						FileArguments += string.Format(" -MD -MF\"{0}\"", DependencyListFile.AbsolutePath.Replace('\\', '/'));
+						CompileAction.DependencyListFile = DependencyListFile;
+						CompileAction.ProducedItems.Add(DependencyListFile);
+					}
 
-						// Build a full argument list
-						string AllArguments = Arguments + FileArguments + CompileEnvironment.AdditionalArguments;
+					// Build a full argument list
+					string AllArguments = Arguments + FileArguments + CompileEnvironment.AdditionalArguments;
 
-						if (SourceFile.AbsolutePath.Equals(NativeGluePath))
-						{
-							// Remove visibility settings for android native glue. Since it doesn't decorate with visibility attributes.
-							//AllArguments = AllArguments.Replace("-fvisibility-ms-compat -fvisibility-inlines-hidden", "");
-							AllArguments = AllArguments.Replace("-fvisibility=hidden -fvisibility-inlines-hidden", "");
-						}
+					if (SourceFile.AbsolutePath.Equals(NativeGluePath))
+					{
+						// Remove visibility settings for android native glue. Since it doesn't decorate with visibility attributes.
+						//AllArguments = AllArguments.Replace("-fvisibility-ms-compat -fvisibility-inlines-hidden", "");
+						AllArguments = AllArguments.Replace("-fvisibility=hidden -fvisibility-inlines-hidden", "");
+					}
 
-						AllArguments = Utils.ExpandVariables(AllArguments);
-						AllArguments = AllArguments.Replace("\\", "/");
+					AllArguments = Utils.ExpandVariables(AllArguments);
+					AllArguments = AllArguments.Replace("\\", "/");
 
-						// Remove shadow warning for this file if requested
-						if (bDisableShadowWarning)
+					// Remove shadow warning for this file if requested
+					if (bDisableShadowWarning)
+					{
+						int WarningIndex = AllArguments.IndexOf(" -Wshadow");
+						if (WarningIndex > 0)
 						{
-							int WarningIndex = AllArguments.IndexOf(" -Wshadow");
-							if (WarningIndex > 0)
-							{
-								AllArguments = AllArguments.Remove(WarningIndex, 9);
-							}
+							AllArguments = AllArguments.Remove(WarningIndex, 9);
 						}
+					}
 
-						// Create the response file
-						FileReference ResponseFileName = CompileAction.ProducedItems[0].Location + ".rsp";
-						FileItem ResponseFileItem = Graph.CreateIntermediateTextFile(ResponseFileName, new List<string> { AllArguments }, StringComparison.InvariantCultureIgnoreCase);
-						string ResponseArgument = string.Format("@\"{0}\"", ResponseFileName);
+					// Create the response file
+					FileReference ResponseFileName = CompileAction.ProducedItems[0].Location + ".rsp";
+					FileItem ResponseFileItem = Graph.CreateIntermediateTextFile(ResponseFileName, new List<string> { AllArguments }, StringComparison.InvariantCultureIgnoreCase);
+					string ResponseArgument = string.Format("@\"{0}\"", ResponseFileName);
 
-						CompileAction.WorkingDirectory = UnrealBuildTool.EngineSourceDirectory;
-						if(bExecuteCompilerThroughShell)
-						{
-							SetupActionToExecuteCompilerThroughShell(ref CompileAction, ClangPath, ResponseArgument, "Compile");
-						}
-						else
-						{
-							CompileAction.CommandPath = new FileReference(ClangPath);
-							CompileAction.CommandArguments = ResponseArgument;
-						}
-						CompileAction.PrerequisiteItems.Add(ResponseFileItem);
-						CompileAction.CommandVersion = AndroidClangBuild;
+					CompileAction.WorkingDirectory = UnrealBuildTool.EngineSourceDirectory;
+					if(bExecuteCompilerThroughShell)
+					{
+						SetupActionToExecuteCompilerThroughShell(ref CompileAction, ClangPath, ResponseArgument, "Compile");
+					}
+					else
+					{
+						CompileAction.CommandPath = new FileReference(ClangPath);
+						CompileAction.CommandArguments = ResponseArgument;
+					}
+					CompileAction.PrerequisiteItems.Add(ResponseFileItem);
+					CompileAction.CommandVersion = AndroidClangBuild;
 
-						if (GPUArchitecture.Length > 0)
-						{
-							CompileAction.StatusDescription = string.Format("{0} [{1}-{2}]", Path.GetFileName(SourceFile.AbsolutePath), Arch.Replace("-", ""), GPUArchitecture.Replace("-", ""));
-						}
-						else
-						{
-							CompileAction.StatusDescription = string.Format("{0} [{1}]", Path.GetFileName(SourceFile.AbsolutePath), Arch.Replace("-", ""));
-						}
+					CompileAction.StatusDescription = string.Format("{0} [{1}]", Path.GetFileName(SourceFile.AbsolutePath), Arch.Replace("-", ""));
 
-						// VC++ always outputs the source file name being compiled, so we don't need to emit this ourselves
-						CompileAction.bShouldOutputStatusDescription = true;
+					// VC++ always outputs the source file name being compiled, so we don't need to emit this ourselves
+					CompileAction.bShouldOutputStatusDescription = true;
 
-						// Don't farm out creation of pre-compiled headers as it is the critical path task.
-						CompileAction.bCanExecuteRemotely =
-							CompileEnvironment.PrecompiledHeaderAction != PrecompiledHeaderAction.Create ||
-							CompileEnvironment.bAllowRemotelyCompiledPCHs;
-					}
+					// Don't farm out creation of pre-compiled headers as it is the critical path task.
+					CompileAction.bCanExecuteRemotely =
+						CompileEnvironment.PrecompiledHeaderAction != PrecompiledHeaderAction.Create ||
+						CompileEnvironment.bAllowRemotelyCompiledPCHs;
 				}
 			}
 
@@ -1599,16 +1560,14 @@ namespace UnrealBuildTool
 			return null;
 		}
 
-		static public string InlineArchName(string Pathname, string Arch, string GPUArchitecture, bool bUseShortNames = false)
+		static public string InlineArchName(string Pathname, string Arch, bool bUseShortNames = false)
 		{
 			string FinalArch = Arch;
-			string FinalGPUArch = GPUArchitecture;
 			if (bUseShortNames)
 			{
 				FinalArch = ShortArchNames[FinalArch];
-				FinalGPUArch = ShortArchNames[FinalGPUArch];
 			}
-			return Path.Combine(Path.GetDirectoryName(Pathname), Path.GetFileNameWithoutExtension(Pathname) + FinalArch + FinalGPUArch + Path.GetExtension(Pathname));
+			return Path.Combine(Path.GetDirectoryName(Pathname), Path.GetFileNameWithoutExtension(Pathname) + FinalArch + Path.GetExtension(Pathname));
 		}
 
 		public string RemoveArchName(string Pathname)
@@ -1616,17 +1575,14 @@ namespace UnrealBuildTool
 			// remove all architecture names
 			foreach (string Arch in GetAllArchitectures())
 			{
-				foreach (string GPUArchitecture in GetAllGPUArchitectures())
-				{
-					Pathname = Path.Combine(Path.GetDirectoryName(Pathname), Path.GetFileName(Pathname).Replace(Arch + GPUArchitecture, ""));
-				}
+				Pathname = Path.Combine(Path.GetDirectoryName(Pathname), Path.GetFileName(Pathname).Replace(Arch, ""));
 			}
 			return Pathname;
 		}
 
-		static public DirectoryReference InlineArchIncludeFolder(DirectoryReference PathRef, string Arch, string GPUArchitecture)
+		static public DirectoryReference InlineArchIncludeFolder(DirectoryReference PathRef, string Arch)
 		{
-			return DirectoryReference.Combine(PathRef, "include", Arch.Replace("-", "") + GPUArchitecture);
+			return DirectoryReference.Combine(PathRef, "include", Arch.Replace("-", ""));
 		}
 
 		public override CPPOutput GenerateISPCHeaders(CppCompileEnvironment CompileEnvironment, List<FileItem> InputFiles, DirectoryReference OutputDir, IActionGraphBuilder Graph)
@@ -1645,126 +1601,123 @@ namespace UnrealBuildTool
 
 			foreach (string Arch in Arches)
 			{
-				foreach (string GPUArchitecture in GPUArchitectures)
-				{
-					List<string> CompileTargets = GetISPCCompileTargets(CompileEnvironment.Platform, Arch);
+				List<string> CompileTargets = GetISPCCompileTargets(CompileEnvironment.Platform, Arch);
 
-					CompileEnvironment.UserIncludePaths.Add(InlineArchIncludeFolder(OutputDir, Arch, GPUArchitecture));
+				CompileEnvironment.UserIncludePaths.Add(InlineArchIncludeFolder(OutputDir, Arch));
 
-					foreach (FileItem ISPCFile in InputFiles)
+				foreach (FileItem ISPCFile in InputFiles)
+				{
+					Action CompileAction = Graph.CreateAction(ActionType.Compile);
+					CompileAction.CommandDescription = "Compile";
+					CompileAction.WorkingDirectory = UnrealBuildTool.EngineSourceDirectory;
+					CompileAction.CommandPath = new FileReference(GetISPCHostCompilerPath(BuildHostPlatform.Current.Platform));
+					CompileAction.StatusDescription = Path.GetFileName(ISPCFile.AbsolutePath);
+					CompileAction.CommandVersion = GetISPCHostCompilerVersion(BuildHostPlatform.Current.Platform).ToString();
+
+					// Disable remote execution to workaround mismatched case on XGE
+					CompileAction.bCanExecuteRemotely = false;
+
+					List<string> Arguments = new List<string>();
+
+					// Add the ISPC obj file as a prerequisite of the action.
+					Arguments.Add(String.Format(" \"{0}\"", ISPCFile.AbsolutePath));
+
+					// Add the ISPC h file to the produced item list.
+					FileItem ISPCIncludeHeaderFile = FileItem.GetItemByFileReference(
+						FileReference.Combine(
+							InlineArchIncludeFolder(OutputDir, Arch),
+							Path.GetFileName(ISPCFile.AbsolutePath) + ".generated.dummy.h"
+							)
+						);
+
+					// Add the ISPC file to be compiled.
+					Arguments.Add(String.Format("-h \"{0}\"", ISPCIncludeHeaderFile));
+
+					// Build target string. No comma on last
+					string TargetString = "";
+					foreach (string Target in CompileTargets)
 					{
-						Action CompileAction = Graph.CreateAction(ActionType.Compile);
-						CompileAction.CommandDescription = "Compile";
-						CompileAction.WorkingDirectory = UnrealBuildTool.EngineSourceDirectory;
-						CompileAction.CommandPath = new FileReference(GetISPCHostCompilerPath(BuildHostPlatform.Current.Platform));
-						CompileAction.StatusDescription = Path.GetFileName(ISPCFile.AbsolutePath);
-						CompileAction.CommandVersion = GetISPCHostCompilerVersion(BuildHostPlatform.Current.Platform).ToString();
-
-						// Disable remote execution to workaround mismatched case on XGE
-						CompileAction.bCanExecuteRemotely = false;
-
-						List<string> Arguments = new List<string>();
-
-						// Add the ISPC obj file as a prerequisite of the action.
-						Arguments.Add(String.Format(" \"{0}\"", ISPCFile.AbsolutePath));
-
-						// Add the ISPC h file to the produced item list.
-						FileItem ISPCIncludeHeaderFile = FileItem.GetItemByFileReference(
-							FileReference.Combine(
-								InlineArchIncludeFolder(OutputDir, Arch, GPUArchitecture),
-								Path.GetFileName(ISPCFile.AbsolutePath) + ".generated.dummy.h"
-								)
-							);
-
-						// Add the ISPC file to be compiled.
-						Arguments.Add(String.Format("-h \"{0}\"", ISPCIncludeHeaderFile));
-
-						// Build target string. No comma on last
-						string TargetString = "";
-						foreach (string Target in CompileTargets)
+						if (Target == CompileTargets.Last())
 						{
-							if (Target == CompileTargets.Last())
-							{
-								TargetString += Target;
-							}
-							else
-							{
-								TargetString += Target + ",";
-							}
+							TargetString += Target;
 						}
-
-						// Build target triplet
-						Arguments.Add(String.Format("--target-os=\"{0}\"", GetISPCOSTarget(CompileEnvironment.Platform)));
-						Arguments.Add(String.Format("--arch=\"{0}\"", GetISPCArchTarget(CompileEnvironment.Platform, Arch)));
-						Arguments.Add(String.Format("--target=\"{0}\"", TargetString));
-
-						Arguments.Add("--pic");
-
-						// Include paths. Don't use AddIncludePath() here, since it uses the full path and exceeds the max command line length.
-						foreach (DirectoryReference IncludePath in CompileEnvironment.UserIncludePaths)
+						else
 						{
-							Arguments.Add(String.Format("-I\"{0}\"", IncludePath));
+							TargetString += Target + ",";
 						}
+					}
 
-						// System include paths.
-						foreach (DirectoryReference SystemIncludePath in CompileEnvironment.SystemIncludePaths)
-						{
-							Arguments.Add(String.Format("-I\"{0}\"", SystemIncludePath));
-						}
+					// Build target triplet
+					Arguments.Add(String.Format("--target-os=\"{0}\"", GetISPCOSTarget(CompileEnvironment.Platform)));
+					Arguments.Add(String.Format("--arch=\"{0}\"", GetISPCArchTarget(CompileEnvironment.Platform, Arch)));
+					Arguments.Add(String.Format("--target=\"{0}\"", TargetString));
 
-						// Generate the included header dependency list
-						if (CompileEnvironment.bGenerateDependenciesFile)
-						{
-							FileItem DependencyListFile = FileItem.GetItemByFileReference(FileReference.Combine(OutputDir, InlineArchName(Path.GetFileName(ISPCFile.AbsolutePath) + ".d", Arch, GPUArchitecture, true)));
-							Arguments.Add(String.Format("-M -MF \"{0}\"", DependencyListFile.AbsolutePath.Replace('\\', '/')));
-							CompileAction.DependencyListFile = DependencyListFile;
-							CompileAction.ProducedItems.Add(DependencyListFile);
-						}
+					Arguments.Add("--pic");
 
-						CompileAction.ProducedItems.Add(ISPCIncludeHeaderFile);
+					// Include paths. Don't use AddIncludePath() here, since it uses the full path and exceeds the max command line length.
+					foreach (DirectoryReference IncludePath in CompileEnvironment.UserIncludePaths)
+					{
+						Arguments.Add(String.Format("-I\"{0}\"", IncludePath));
+					}
 
-						CompileAction.CommandArguments = String.Join(" ", Arguments);
+					// System include paths.
+					foreach (DirectoryReference SystemIncludePath in CompileEnvironment.SystemIncludePaths)
+					{
+						Arguments.Add(String.Format("-I\"{0}\"", SystemIncludePath));
+					}
 
-						// Add the source file and its included files to the prerequisite item list.
-						CompileAction.PrerequisiteItems.Add(ISPCFile);
-						CompileAction.StatusDescription = string.Format("{0} [{1}]", Path.GetFileName(ISPCFile.AbsolutePath), Arch.Replace("-", ""));
+					// Generate the included header dependency list
+					if (CompileEnvironment.bGenerateDependenciesFile)
+					{
+						FileItem DependencyListFile = FileItem.GetItemByFileReference(FileReference.Combine(OutputDir, InlineArchName(Path.GetFileName(ISPCFile.AbsolutePath) + ".d", Arch, true)));
+						Arguments.Add(String.Format("-M -MF \"{0}\"", DependencyListFile.AbsolutePath.Replace('\\', '/')));
+						CompileAction.DependencyListFile = DependencyListFile;
+						CompileAction.ProducedItems.Add(DependencyListFile);
+					}
 
-						FileItem ISPCFinalHeaderFile = FileItem.GetItemByFileReference(
-							FileReference.Combine(
-								InlineArchIncludeFolder(OutputDir, Arch, GPUArchitecture),
-								Path.GetFileName(ISPCFile.AbsolutePath) + ".generated.h"
-								)
-							);
+					CompileAction.ProducedItems.Add(ISPCIncludeHeaderFile);
 
-						// Fix interrupted build issue by copying header after generation completes
-						FileReference SourceFile = ISPCIncludeHeaderFile.Location;
-						FileReference TargetFile = ISPCFinalHeaderFile.Location;
+					CompileAction.CommandArguments = String.Join(" ", Arguments);
 
-						FileItem SourceFileItem = FileItem.GetItemByFileReference(SourceFile);
-						FileItem TargetFileItem = FileItem.GetItemByFileReference(TargetFile);
+					// Add the source file and its included files to the prerequisite item list.
+					CompileAction.PrerequisiteItems.Add(ISPCFile);
+					CompileAction.StatusDescription = string.Format("{0} [{1}]", Path.GetFileName(ISPCFile.AbsolutePath), Arch.Replace("-", ""));
 
-						Action CopyAction = Graph.CreateAction(ActionType.BuildProject);
-						CopyAction.CommandDescription = "Copy";
-						CopyAction.CommandPath = BuildHostPlatform.Current.Shell;
-						if (BuildHostPlatform.Current.ShellType == ShellType.Cmd)
-						{
-							CopyAction.CommandArguments = String.Format("/C \"copy /Y \"{0}\" \"{1}\" 1>nul\"", SourceFile, TargetFile);
-						}
-						else
-						{
-							CopyAction.CommandArguments = String.Format("-c 'cp -f \"{0}\" \"{1}\"'", SourceFile.FullName, TargetFile.FullName);
-						}
-						CopyAction.WorkingDirectory = UnrealBuildTool.EngineSourceDirectory;
-						CopyAction.PrerequisiteItems.Add(SourceFileItem);
-						CopyAction.ProducedItems.Add(TargetFileItem);
-						CopyAction.StatusDescription = TargetFileItem.Location.GetFileName();
-						CopyAction.bCanExecuteRemotely = false;
-						CopyAction.bShouldOutputStatusDescription = false;
+					FileItem ISPCFinalHeaderFile = FileItem.GetItemByFileReference(
+						FileReference.Combine(
+							InlineArchIncludeFolder(OutputDir, Arch),
+							Path.GetFileName(ISPCFile.AbsolutePath) + ".generated.h"
+							)
+						);
+
+					// Fix interrupted build issue by copying header after generation completes
+					FileReference SourceFile = ISPCIncludeHeaderFile.Location;
+					FileReference TargetFile = ISPCFinalHeaderFile.Location;
 
-						Result.GeneratedHeaderFiles.Add(TargetFileItem);
+					FileItem SourceFileItem = FileItem.GetItemByFileReference(SourceFile);
+					FileItem TargetFileItem = FileItem.GetItemByFileReference(TargetFile);
 
-						Log.TraceVerbose("   ISPC Generating Header " + CompileAction.StatusDescription + ": \"" + CompileAction.CommandPath + "\"" + CompileAction.CommandArguments);
+					Action CopyAction = Graph.CreateAction(ActionType.BuildProject);
+					CopyAction.CommandDescription = "Copy";
+					CopyAction.CommandPath = BuildHostPlatform.Current.Shell;
+					if (BuildHostPlatform.Current.ShellType == ShellType.Cmd)
+					{
+						CopyAction.CommandArguments = String.Format("/C \"copy /Y \"{0}\" \"{1}\" 1>nul\"", SourceFile, TargetFile);
 					}
+					else
+					{
+						CopyAction.CommandArguments = String.Format("-c 'cp -f \"{0}\" \"{1}\"'", SourceFile.FullName, TargetFile.FullName);
+					}
+					CopyAction.WorkingDirectory = UnrealBuildTool.EngineSourceDirectory;
+					CopyAction.PrerequisiteItems.Add(SourceFileItem);
+					CopyAction.ProducedItems.Add(TargetFileItem);
+					CopyAction.StatusDescription = TargetFileItem.Location.GetFileName();
+					CopyAction.bCanExecuteRemotely = false;
+					CopyAction.bShouldOutputStatusDescription = false;
+
+					Result.GeneratedHeaderFiles.Add(TargetFileItem);
+
+					Log.TraceVerbose("   ISPC Generating Header " + CompileAction.StatusDescription + ": \"" + CompileAction.CommandPath + "\"" + CompileAction.CommandArguments);
 				}
 			}
 
@@ -1787,199 +1740,189 @@ namespace UnrealBuildTool
 
 			foreach (string Arch in Arches)
 			{
-				foreach (string GPUArchitecture in GPUArchitectures)
+				List<string> CompileTargets = GetISPCCompileTargets(CompileEnvironment.Platform, Arch);
+
+				foreach (FileItem ISPCFile in InputFiles)
 				{
-					List<string> CompileTargets = GetISPCCompileTargets(CompileEnvironment.Platform, Arch);
+					Action CompileAction = Graph.CreateAction(ActionType.Compile);
+					CompileAction.CommandDescription = "Compile";
+					CompileAction.WorkingDirectory = UnrealBuildTool.EngineSourceDirectory;
+					CompileAction.CommandPath = new FileReference(GetISPCHostCompilerPath(BuildHostPlatform.Current.Platform));
+					CompileAction.StatusDescription = Path.GetFileName(ISPCFile.AbsolutePath);
 
-					foreach (FileItem ISPCFile in InputFiles)
-					{
-						Action CompileAction = Graph.CreateAction(ActionType.Compile);
-						CompileAction.CommandDescription = "Compile";
-						CompileAction.WorkingDirectory = UnrealBuildTool.EngineSourceDirectory;
-						CompileAction.CommandPath = new FileReference(GetISPCHostCompilerPath(BuildHostPlatform.Current.Platform));
-						CompileAction.StatusDescription = Path.GetFileName(ISPCFile.AbsolutePath);
+					// Disable remote execution to workaround mismatched case on XGE
+					CompileAction.bCanExecuteRemotely = false;
 
-						// Disable remote execution to workaround mismatched case on XGE
-						CompileAction.bCanExecuteRemotely = false;
+					List<string> Arguments = new List<string>();
 
-						List<string> Arguments = new List<string>();
+					// Add the ISPC file to be compiled.
+					Arguments.Add(String.Format(" \"{0}\"", ISPCFile.AbsolutePath));
 
-						// Add the ISPC file to be compiled.
-						Arguments.Add(String.Format(" \"{0}\"", ISPCFile.AbsolutePath));
+					List<FileItem> CompiledISPCObjFiles = new List<FileItem>();
+					List<FileItem> FinalISPCObjFiles = new List<FileItem>();
+					string TargetString = "";
 
-						List<FileItem> CompiledISPCObjFiles = new List<FileItem>();
-						List<FileItem> FinalISPCObjFiles = new List<FileItem>();
-						string TargetString = "";
+					foreach (string Target in CompileTargets)
+					{
+						string ObjTarget = Target;
 
-						foreach (string Target in CompileTargets)
+						if (Target.Contains("-"))
 						{
-							string ObjTarget = Target;
-
-							if (Target.Contains("-"))
-							{
-								// Remove lane width and gang size from obj file name
-								ObjTarget = Target.Split('-')[0];
-							}
+							// Remove lane width and gang size from obj file name
+							ObjTarget = Target.Split('-')[0];
+						}
 
-							FileItem CompiledISPCObjFile;
-							FileItem FinalISPCObjFile;
+						FileItem CompiledISPCObjFile;
+						FileItem FinalISPCObjFile;
 
-							if (CompileTargets.Count > 1)
-							{
-								CompiledISPCObjFile = FileItem.GetItemByFileReference(
-								FileReference.Combine(
-									OutputDir,
-									Path.GetFileNameWithoutExtension(InlineArchName(Path.GetFileName(ISPCFile.AbsolutePath) + ".o", Arch, GPUArchitecture, true)) + "_" + ObjTarget + ".o"
-									)
-								);
+						if (CompileTargets.Count > 1)
+						{
+							CompiledISPCObjFile = FileItem.GetItemByFileReference(
+							FileReference.Combine(
+								OutputDir,
+								Path.GetFileNameWithoutExtension(InlineArchName(Path.GetFileName(ISPCFile.AbsolutePath) + ".o", Arch, true)) + "_" + ObjTarget + ".o"
+								)
+							);
 
-								FinalISPCObjFile = FileItem.GetItemByFileReference(
+							FinalISPCObjFile = FileItem.GetItemByFileReference(
+							FileReference.Combine(
+								OutputDir,
+								Path.GetFileName(ISPCFile.AbsolutePath) + "_" + ObjTarget + InlineArchName(".o", Arch, true)
+								)
+							);
+						}
+						else
+						{
+							CompiledISPCObjFile = FileItem.GetItemByFileReference(
 								FileReference.Combine(
 									OutputDir,
-									Path.GetFileName(ISPCFile.AbsolutePath) + "_" + ObjTarget + InlineArchName(".o", Arch, GPUArchitecture, true)
+									InlineArchName(Path.GetFileName(ISPCFile.AbsolutePath) + ".o", Arch, true)
 									)
 								);
-							}
-							else
-							{
-								CompiledISPCObjFile = FileItem.GetItemByFileReference(
-									FileReference.Combine(
-										OutputDir,
-										InlineArchName(Path.GetFileName(ISPCFile.AbsolutePath) + ".o", Arch, GPUArchitecture, true)
-										)
-									);
-
-								FinalISPCObjFile = CompiledISPCObjFile;
-							}
-
-							// Add the ISA specific ISPC obj files to the produced item list.
-							CompiledISPCObjFiles.Add(CompiledISPCObjFile);
-							FinalISPCObjFiles.Add(FinalISPCObjFile);
 
-							// Build target string. No comma on last
-							if (Target == CompileTargets.Last())
-							{
-								TargetString += Target;
-							}
-							else
-							{
-								TargetString += Target + ",";
-							}
+							FinalISPCObjFile = CompiledISPCObjFile;
 						}
 
-						// Add the common ISPC obj file to the produced item list.
-						FileItem CompiledISPCObjFileNoISA = FileItem.GetItemByFileReference(
-							FileReference.Combine(
-								OutputDir,
-								InlineArchName(Path.GetFileName(ISPCFile.AbsolutePath) + ".o", Arch, GPUArchitecture, true)
-								)
-							);
-
-						CompiledISPCObjFiles.Add(CompiledISPCObjFileNoISA);
-						FinalISPCObjFiles.Add(CompiledISPCObjFileNoISA);
+						// Add the ISA specific ISPC obj files to the produced item list.
+						CompiledISPCObjFiles.Add(CompiledISPCObjFile);
+						FinalISPCObjFiles.Add(FinalISPCObjFile);
 
-						// Add the output ISPC obj file
-						Arguments.Add(String.Format("-o \"{0}\"", CompiledISPCObjFileNoISA));
-
-						// Build target triplet
-						Arguments.Add(String.Format("--target-os=\"{0}\"", GetISPCOSTarget(CompileEnvironment.Platform)));
-						Arguments.Add(String.Format("--arch=\"{0}\"", GetISPCArchTarget(CompileEnvironment.Platform, Arch)));
-						Arguments.Add(String.Format("--target=\"{0}\"", TargetString));
-
-						if (CompileEnvironment.Configuration == CppConfiguration.Debug)
+						// Build target string. No comma on last
+						if (Target == CompileTargets.Last())
 						{
-							Arguments.Add("-g -O0");
+							TargetString += Target;
 						}
 						else
 						{
-							Arguments.Add("-O2");
+							TargetString += Target + ",";
 						}
+					}
+
+					// Add the common ISPC obj file to the produced item list.
+					FileItem CompiledISPCObjFileNoISA = FileItem.GetItemByFileReference(
+						FileReference.Combine(
+							OutputDir,
+							InlineArchName(Path.GetFileName(ISPCFile.AbsolutePath) + ".o", Arch, true)
+							)
+						);
+
+					CompiledISPCObjFiles.Add(CompiledISPCObjFileNoISA);
+					FinalISPCObjFiles.Add(CompiledISPCObjFileNoISA);
+
+					// Add the output ISPC obj file
+					Arguments.Add(String.Format("-o \"{0}\"", CompiledISPCObjFileNoISA));
+
+					// Build target triplet
+					Arguments.Add(String.Format("--target-os=\"{0}\"", GetISPCOSTarget(CompileEnvironment.Platform)));
+					Arguments.Add(String.Format("--arch=\"{0}\"", GetISPCArchTarget(CompileEnvironment.Platform, Arch)));
+					Arguments.Add(String.Format("--target=\"{0}\"", TargetString));
+
+					if (CompileEnvironment.Configuration == CppConfiguration.Debug)
+					{
+						Arguments.Add("-g -O0");
+					}
+					else
+					{
+						Arguments.Add("-O2");
+					}
 
-						Arguments.Add("--pic");
+					Arguments.Add("--pic");
 
-						// Add include paths to the argument list (filtered by architecture)
-						foreach (DirectoryReference IncludePath in CompileEnvironment.SystemIncludePaths)
+					// Add include paths to the argument list (filtered by architecture)
+					foreach (DirectoryReference IncludePath in CompileEnvironment.SystemIncludePaths)
+					{
+						if (IsDirectoryForArch(IncludePath.FullName, Arch))
 						{
-							if (IsDirectoryForArch(IncludePath.FullName, Arch))
-							{
-								Arguments.Add(string.Format(" -I\"{0}\"", IncludePath));
-							}
+							Arguments.Add(string.Format(" -I\"{0}\"", IncludePath));
 						}
-						foreach (DirectoryReference IncludePath in CompileEnvironment.UserIncludePaths)
+					}
+					foreach (DirectoryReference IncludePath in CompileEnvironment.UserIncludePaths)
+					{
+						if (IsDirectoryForArch(IncludePath.FullName, Arch))
 						{
-							if (IsDirectoryForArch(IncludePath.FullName, Arch))
-							{
-								Arguments.Add(string.Format(" -I\"{0}\"", IncludePath));
-							}
+							Arguments.Add(string.Format(" -I\"{0}\"", IncludePath));
 						}
+					}
 
-						// Preprocessor definitions.
-						foreach (string Definition in CompileEnvironment.Definitions)
-						{
-							Arguments.Add(String.Format("-D\"{0}\"", Definition));
-						}
+					// Preprocessor definitions.
+					foreach (string Definition in CompileEnvironment.Definitions)
+					{
+						Arguments.Add(String.Format("-D\"{0}\"", Definition));
+					}
 
-						// Consume the included header dependency list
-						if (CompileEnvironment.bGenerateDependenciesFile)
-						{
-							FileItem DependencyListFile = FileItem.GetItemByFileReference(FileReference.Combine(OutputDir, InlineArchName(Path.GetFileName(ISPCFile.AbsolutePath) + ".d", Arch, GPUArchitecture, true)));
-							CompileAction.DependencyListFile = DependencyListFile;
-							CompileAction.PrerequisiteItems.Add(DependencyListFile);
-						}
+					// Consume the included header dependency list
+					if (CompileEnvironment.bGenerateDependenciesFile)
+					{
+						FileItem DependencyListFile = FileItem.GetItemByFileReference(FileReference.Combine(OutputDir, InlineArchName(Path.GetFileName(ISPCFile.AbsolutePath) + ".d", Arch, true)));
+						CompileAction.DependencyListFile = DependencyListFile;
+						CompileAction.PrerequisiteItems.Add(DependencyListFile);
+					}
 
-						CompileAction.ProducedItems.AddRange(CompiledISPCObjFiles);
+					CompileAction.ProducedItems.AddRange(CompiledISPCObjFiles);
 
-						CompileAction.CommandArguments = String.Join(" ", Arguments);
+					CompileAction.CommandArguments = String.Join(" ", Arguments);
 
-						// Add the source file and its included files to the prerequisite item list.
-						CompileAction.PrerequisiteItems.Add(ISPCFile);
+					// Add the source file and its included files to the prerequisite item list.
+					CompileAction.PrerequisiteItems.Add(ISPCFile);
+
+					CompileAction.StatusDescription = string.Format("{0} [{1}]", Path.GetFileName(ISPCFile.AbsolutePath), Arch.Replace("-", ""));
+
+					for(int i = 0; i < CompiledISPCObjFiles.Count; i++)
+					{
+						// ISPC compiler can't add suffix on the end of the arch, so copy to put into what linker expects
+						FileReference SourceFile = CompiledISPCObjFiles[i].Location;
+						FileReference TargetFile = FinalISPCObjFiles[i].Location;
 
-						if (GPUArchitecture.Length > 0)
+						if (SourceFile.Equals(TargetFile))
 						{
-							CompileAction.StatusDescription = string.Format("{0} [{1}-{2}]", Path.GetFileName(ISPCFile.AbsolutePath), Arch.Replace("-", ""), GPUArchitecture.Replace("-", ""));
+							continue;
 						}
-						else
+
+						FileItem SourceFileItem = FileItem.GetItemByFileReference(SourceFile);
+						FileItem TargetFileItem = FileItem.GetItemByFileReference(TargetFile);
+
+						Action CopyAction = Graph.CreateAction(ActionType.BuildProject);
+						CopyAction.CommandDescription = "Copy";
+						CopyAction.CommandPath = BuildHostPlatform.Current.Shell;
+						if (BuildHostPlatform.Current.ShellType == ShellType.Cmd)
 						{
-							CompileAction.StatusDescription = string.Format("{0} [{1}]", Path.GetFileName(ISPCFile.AbsolutePath), Arch.Replace("-", ""));
+							CopyAction.CommandArguments = String.Format("/C \"copy /Y \"{0}\" \"{1}\" 1>nul\"", SourceFile, TargetFile);
 						}
-
-						for(int i = 0; i < CompiledISPCObjFiles.Count; i++)
+						else
 						{
-							// ISPC compiler can't add suffix on the end of the arch, so copy to put into what linker expects
-							FileReference SourceFile = CompiledISPCObjFiles[i].Location;
-							FileReference TargetFile = FinalISPCObjFiles[i].Location;
-
-							if (SourceFile.Equals(TargetFile))
-							{
-								continue;
-							}
-
-							FileItem SourceFileItem = FileItem.GetItemByFileReference(SourceFile);
-							FileItem TargetFileItem = FileItem.GetItemByFileReference(TargetFile);
-
-							Action CopyAction = Graph.CreateAction(ActionType.BuildProject);
-							CopyAction.CommandDescription = "Copy";
-							CopyAction.CommandPath = BuildHostPlatform.Current.Shell;
-							if (BuildHostPlatform.Current.ShellType == ShellType.Cmd)
-							{
-								CopyAction.CommandArguments = String.Format("/C \"copy /Y \"{0}\" \"{1}\" 1>nul\"", SourceFile, TargetFile);
-							}
-							else
-							{
-								CopyAction.CommandArguments = String.Format("-c 'cp -f \"{0}\" \"{1}\"'", SourceFile.FullName, TargetFile.FullName);
-							}
-							CopyAction.WorkingDirectory = UnrealBuildTool.EngineSourceDirectory;
-							CopyAction.PrerequisiteItems.Add(SourceFileItem);
-							CopyAction.ProducedItems.Add(TargetFileItem);
-							CopyAction.StatusDescription = TargetFileItem.Location.GetFileName();
-							CopyAction.bCanExecuteRemotely = false;
-							CopyAction.bShouldOutputStatusDescription = false;
+							CopyAction.CommandArguments = String.Format("-c 'cp -f \"{0}\" \"{1}\"'", SourceFile.FullName, TargetFile.FullName);
 						}
+						CopyAction.WorkingDirectory = UnrealBuildTool.EngineSourceDirectory;
+						CopyAction.PrerequisiteItems.Add(SourceFileItem);
+						CopyAction.ProducedItems.Add(TargetFileItem);
+						CopyAction.StatusDescription = TargetFileItem.Location.GetFileName();
+						CopyAction.bCanExecuteRemotely = false;
+						CopyAction.bShouldOutputStatusDescription = false;
+					}
 
-						Result.ObjectFiles.AddRange(FinalISPCObjFiles);
+					Result.ObjectFiles.AddRange(FinalISPCObjFiles);
 
-						Log.TraceVerbose("   ISPC Compiling " + CompileAction.StatusDescription + ": \"" + CompileAction.CommandPath + "\"" + CompileAction.CommandArguments);
-					}
+					Log.TraceVerbose("   ISPC Compiling " + CompileAction.StatusDescription + ": \"" + CompileAction.CommandPath + "\"" + CompileAction.CommandArguments);
 				}
 			}
 
@@ -2000,263 +1943,259 @@ namespace UnrealBuildTool
 			{
 				string Arch = Arches[ArchIndex];
 
-				for (int GPUArchIndex = 0; GPUArchIndex < GPUArchitectures.Count; GPUArchIndex++)
-				{
-					string GPUArchitecture = GPUArchitectures[GPUArchIndex];
-					int OutputPathIndex = ArchIndex * GPUArchitectures.Count + GPUArchIndex;
+				int OutputPathIndex = ArchIndex;
 
-					// Android will have an array of outputs
-					if (!LinkEnvironment.bIsBuildingDLL && // DLL compiles don't have the Arch embedded in the name
-						(LinkEnvironment.OutputFilePaths.Count < OutputPathIndex ||
-						!LinkEnvironment.OutputFilePaths[OutputPathIndex].GetFileNameWithoutExtension().EndsWith(Arch + GPUArchitecture)))
-					{
-						throw new BuildException("The OutputFilePaths array didn't match the Arches array in AndroidToolChain.LinkAllFiles");
-					}
+				// Android will have an array of outputs
+				if (!LinkEnvironment.bIsBuildingDLL && // DLL compiles don't have the Arch embedded in the name
+					(LinkEnvironment.OutputFilePaths.Count < OutputPathIndex ||
+					!LinkEnvironment.OutputFilePaths[OutputPathIndex].GetFileNameWithoutExtension().EndsWith(Arch)))
+				{
+					throw new BuildException("The OutputFilePaths array didn't match the Arches array in AndroidToolChain.LinkAllFiles");
+				}
 
-					// Create an action that invokes the linker.
-					Action LinkAction = Graph.CreateAction(ActionType.Link);
-					LinkAction.WorkingDirectory = UnrealBuildTool.EngineSourceDirectory;
+				// Create an action that invokes the linker.
+				Action LinkAction = Graph.CreateAction(ActionType.Link);
+				LinkAction.WorkingDirectory = UnrealBuildTool.EngineSourceDirectory;
 
-					if (LinkEnvironment.bIsBuildingLibrary)
-					{
-						switch (Arch)
-						{
-							case "-arm64": LinkAction.CommandPath = new FileReference(ArPathArm64); break;
-							case "-x64": LinkAction.CommandPath = new FileReference(ArPathx64); break;
-						}
-					}
-					else
+				if (LinkEnvironment.bIsBuildingLibrary)
+				{
+					switch (Arch)
 					{
-						LinkAction.CommandPath = new FileReference(ClangPath);
+						case "-arm64": LinkAction.CommandPath = new FileReference(ArPathArm64); break;
+						case "-x64": LinkAction.CommandPath = new FileReference(ArPathx64); break;
 					}
+				}
+				else
+				{
+					LinkAction.CommandPath = new FileReference(ClangPath);
+				}
 
-					DirectoryReference LinkerPath = LinkAction.WorkingDirectory;
+				DirectoryReference LinkerPath = LinkAction.WorkingDirectory;
 
-					LinkAction.WorkingDirectory = LinkEnvironment.IntermediateDirectory;
+				LinkAction.WorkingDirectory = LinkEnvironment.IntermediateDirectory;
 
-					// Get link arguments.
-					LinkAction.CommandArguments = LinkEnvironment.bIsBuildingLibrary ? GetArArguments(LinkEnvironment) : GetLinkArguments(LinkEnvironment, Arch);
+				// Get link arguments.
+				LinkAction.CommandArguments = LinkEnvironment.bIsBuildingLibrary ? GetArArguments(LinkEnvironment) : GetLinkArguments(LinkEnvironment, Arch);
 
-					// Add the output file as a production of the link action.
-					FileItem OutputFile;
-					OutputFile = FileItem.GetItemByFileReference(LinkEnvironment.OutputFilePaths[OutputPathIndex]);
-					Outputs.Add(OutputFile);
-					LinkAction.ProducedItems.Add(OutputFile);
-					LinkAction.StatusDescription = string.Format("{0}", Path.GetFileName(OutputFile.AbsolutePath));
-					LinkAction.CommandVersion = AndroidClangBuild;
+				// Add the output file as a production of the link action.
+				FileItem OutputFile;
+				OutputFile = FileItem.GetItemByFileReference(LinkEnvironment.OutputFilePaths[OutputPathIndex]);
+				Outputs.Add(OutputFile);
+				LinkAction.ProducedItems.Add(OutputFile);
+				LinkAction.StatusDescription = string.Format("{0}", Path.GetFileName(OutputFile.AbsolutePath));
+				LinkAction.CommandVersion = AndroidClangBuild;
 
-					// LinkAction.bPrintDebugInfo = true;
+				// LinkAction.bPrintDebugInfo = true;
 
-					// Add the output file to the command-line.
-					if (LinkEnvironment.bIsBuildingLibrary)
-					{
-						LinkAction.CommandArguments += string.Format(" \"{0}\"", OutputFile.AbsolutePath);
-					}
-					else
+				// Add the output file to the command-line.
+				if (LinkEnvironment.bIsBuildingLibrary)
+				{
+					LinkAction.CommandArguments += string.Format(" \"{0}\"", OutputFile.AbsolutePath);
+				}
+				else
+				{
+					LinkAction.CommandArguments += string.Format(" -o \"{0}\"", OutputFile.AbsolutePath);
+				}
+
+				// Add the input files to a response file, and pass the response file on the command-line.
+				List<string> InputFileNames = new List<string>();
+				foreach (FileItem InputFile in LinkEnvironment.InputFiles)
+				{
+					// make sure it's for current Arch
+					if (Path.GetFileNameWithoutExtension(InputFile.AbsolutePath).EndsWith(ShortArchNames[Arch]))
 					{
-						LinkAction.CommandArguments += string.Format(" -o \"{0}\"", OutputFile.AbsolutePath);
+						string InputPath;
+						if (InputFile.Location.IsUnderDirectory(LinkEnvironment.IntermediateDirectory))
+						{
+							InputPath = InputFile.Location.MakeRelativeTo(LinkEnvironment.IntermediateDirectory);
+						}
+						else
+						{
+							InputPath = InputFile.Location.FullName;
+						}
+						InputFileNames.Add(string.Format("\"{0}\"", InputPath.Replace('\\', '/')));
+
+						LinkAction.PrerequisiteItems.Add(InputFile);
 					}
+				}
+
+				string LinkResponseArguments = "";
 
-					// Add the input files to a response file, and pass the response file on the command-line.
-					List<string> InputFileNames = new List<string>();
-					foreach (FileItem InputFile in LinkEnvironment.InputFiles)
+				// libs don't link in other libs
+				if (!LinkEnvironment.bIsBuildingLibrary)
+				{
+					// Make a list of library paths to search
+					List<string> AdditionalLibraryPaths = new List<string>();
+					List<string> AdditionalLibraries = new List<string>();
+
+					// Add the library paths to the additional path list
+					foreach (DirectoryReference LibraryPath in LinkEnvironment.SystemLibraryPaths)
 					{
-						// make sure it's for current Arch
-						if (Path.GetFileNameWithoutExtension(InputFile.AbsolutePath).EndsWith(ShortArchNames[Arch] + ShortArchNames[GPUArchitecture]))
+						// LinkerPaths could be relative or absolute
+						string AbsoluteLibraryPath = Utils.ExpandVariables(LibraryPath.FullName);
+						if (IsDirectoryForArch(AbsoluteLibraryPath, Arch))
 						{
-							string InputPath;
-							if (InputFile.Location.IsUnderDirectory(LinkEnvironment.IntermediateDirectory))
+							// environment variables aren't expanded when using the $( style
+							if (Path.IsPathRooted(AbsoluteLibraryPath) == false)
 							{
-								InputPath = InputFile.Location.MakeRelativeTo(LinkEnvironment.IntermediateDirectory);
+								AbsoluteLibraryPath = Path.Combine(LinkerPath.FullName, AbsoluteLibraryPath);
 							}
-							else
+							AbsoluteLibraryPath = Utils.CollapseRelativeDirectories(AbsoluteLibraryPath);
+							if (!AdditionalLibraryPaths.Contains(AbsoluteLibraryPath))
 							{
-								InputPath = InputFile.Location.FullName;
+								AdditionalLibraryPaths.Add(AbsoluteLibraryPath);
 							}
-							InputFileNames.Add(string.Format("\"{0}\"", InputPath.Replace('\\', '/')));
-
-							LinkAction.PrerequisiteItems.Add(InputFile);
 						}
 					}
 
-					string LinkResponseArguments = "";
-
-					// libs don't link in other libs
-					if (!LinkEnvironment.bIsBuildingLibrary)
+					// discover additional libraries and their paths
+					foreach (string SystemLibrary in LinkEnvironment.SystemLibraries)
 					{
-						// Make a list of library paths to search
-						List<string> AdditionalLibraryPaths = new List<string>();
-						List<string> AdditionalLibraries = new List<string>();
-
-						// Add the library paths to the additional path list
-						foreach (DirectoryReference LibraryPath in LinkEnvironment.SystemLibraryPaths)
+						if (!ShouldSkipLib(SystemLibrary, Arch))
 						{
-							// LinkerPaths could be relative or absolute
-							string AbsoluteLibraryPath = Utils.ExpandVariables(LibraryPath.FullName);
-							if (IsDirectoryForArch(AbsoluteLibraryPath, Arch))
+							if (String.IsNullOrEmpty(Path.GetDirectoryName(SystemLibrary)))
 							{
-								// environment variables aren't expanded when using the $( style
-								if (Path.IsPathRooted(AbsoluteLibraryPath) == false)
+								if (SystemLibrary.StartsWith("lib"))
 								{
-									AbsoluteLibraryPath = Path.Combine(LinkerPath.FullName, AbsoluteLibraryPath);
+									AdditionalLibraries.Add(SystemLibrary);
 								}
-								AbsoluteLibraryPath = Utils.CollapseRelativeDirectories(AbsoluteLibraryPath);
-								if (!AdditionalLibraryPaths.Contains(AbsoluteLibraryPath))
+								else
 								{
-									AdditionalLibraryPaths.Add(AbsoluteLibraryPath);
+									AdditionalLibraries.Add("lib" + SystemLibrary);
 								}
 							}
 						}
-
-						// discover additional libraries and their paths
-						foreach (string SystemLibrary in LinkEnvironment.SystemLibraries)
+					}
+					foreach (FileReference Library in LinkEnvironment.Libraries)
+					{
+						if (!ShouldSkipLib(Library.FullName, Arch))
 						{
-							if (!ShouldSkipLib(SystemLibrary, Arch, GPUArchitecture))
+							string AbsoluteLibraryPath = Path.GetDirectoryName(Library.FullName);
+							LinkAction.PrerequisiteItems.Add(FileItem.GetItemByFileReference(Library));
+
+							string Lib = Path.GetFileNameWithoutExtension(Library.FullName);
+							if (Lib.StartsWith("lib"))
 							{
-								if (String.IsNullOrEmpty(Path.GetDirectoryName(SystemLibrary)))
+								AdditionalLibraries.Add(Lib);
+								if (!AdditionalLibraryPaths.Contains(AbsoluteLibraryPath))
 								{
-									if (SystemLibrary.StartsWith("lib"))
-									{
-										AdditionalLibraries.Add(SystemLibrary);
-									}
-									else
-									{
-										AdditionalLibraries.Add("lib" + SystemLibrary);
-									}
+									AdditionalLibraryPaths.Add(AbsoluteLibraryPath);
 								}
 							}
-						}
-						foreach (FileReference Library in LinkEnvironment.Libraries)
-						{
-							if (!ShouldSkipLib(Library.FullName, Arch, GPUArchitecture))
+							else
 							{
-								string AbsoluteLibraryPath = Path.GetDirectoryName(Library.FullName);
-								LinkAction.PrerequisiteItems.Add(FileItem.GetItemByFileReference(Library));
-
-								string Lib = Path.GetFileNameWithoutExtension(Library.FullName);
-								if (Lib.StartsWith("lib"))
-								{
-									AdditionalLibraries.Add(Lib);
-									if (!AdditionalLibraryPaths.Contains(AbsoluteLibraryPath))
-									{
-										AdditionalLibraryPaths.Add(AbsoluteLibraryPath);
-									}
-								}
-								else
-								{
-									AdditionalLibraries.Add(AbsoluteLibraryPath);
-								}
+								AdditionalLibraries.Add(AbsoluteLibraryPath);
 							}
 						}
+					}
+
+					// add the library paths to response
+					foreach (string LibaryPath in AdditionalLibraryPaths)
+					{
+						LinkResponseArguments += string.Format(" -L\"{0}\"", LibaryPath);
+					}
 
-						// add the library paths to response
-						foreach (string LibaryPath in AdditionalLibraryPaths)
+					// add libraries in a library group
+					LinkResponseArguments += string.Format(" -Wl,--start-group");
+					foreach (string AdditionalLibrary in AdditionalLibraries)
+					{
+						if (AdditionalLibrary.StartsWith("lib"))
 						{
-							LinkResponseArguments += string.Format(" -L\"{0}\"", LibaryPath);
+							LinkResponseArguments += string.Format(" \"-l{0}\"", AdditionalLibrary.Substring(3));
 						}
-
-						// add libraries in a library group
-						LinkResponseArguments += string.Format(" -Wl,--start-group");
-						foreach (string AdditionalLibrary in AdditionalLibraries)
+						else
 						{
-							if (AdditionalLibrary.StartsWith("lib"))
-							{
-								LinkResponseArguments += string.Format(" \"-l{0}\"", AdditionalLibrary.Substring(3));
-							}
-							else
-							{
-								LinkResponseArguments += string.Format(" \"{0}\"", AdditionalLibrary);
-							}
+							LinkResponseArguments += string.Format(" \"{0}\"", AdditionalLibrary);
 						}
-						LinkResponseArguments += string.Format(" -Wl,--end-group");
+					}
+					LinkResponseArguments += string.Format(" -Wl,--end-group");
 
-						// Write the MAP file to the output directory.
-						if (LinkEnvironment.bCreateMapFile)
-						{
-							FileReference MAPFilePath = FileReference.Combine(LinkEnvironment.OutputDirectory, Path.GetFileNameWithoutExtension(OutputFile.AbsolutePath) + ".map");
-							FileItem MAPFile = FileItem.GetItemByFileReference(MAPFilePath);
-							LinkResponseArguments += String.Format(" -Wl,--cref -Wl,-Map,\"{0}\"", MAPFilePath);
-							LinkAction.ProducedItems.Add(MAPFile);
+					// Write the MAP file to the output directory.
+					if (LinkEnvironment.bCreateMapFile)
+					{
+						FileReference MAPFilePath = FileReference.Combine(LinkEnvironment.OutputDirectory, Path.GetFileNameWithoutExtension(OutputFile.AbsolutePath) + ".map");
+						FileItem MAPFile = FileItem.GetItemByFileReference(MAPFilePath);
+						LinkResponseArguments += String.Format(" -Wl,--cref -Wl,-Map,\"{0}\"", MAPFilePath);
+						LinkAction.ProducedItems.Add(MAPFile);
 
-							// Export a list of object file paths, so we can locate the object files referenced by the map file
-							ExportObjectFilePaths(LinkEnvironment, Path.ChangeExtension(MAPFilePath.FullName, ".objpaths"));
-						}
+						// Export a list of object file paths, so we can locate the object files referenced by the map file
+						ExportObjectFilePaths(LinkEnvironment, Path.ChangeExtension(MAPFilePath.FullName, ".objpaths"));
 					}
+				}
 
-					// Add the additional arguments specified by the environment.
-					LinkResponseArguments += LinkEnvironment.AdditionalArguments;
+				// Add the additional arguments specified by the environment.
+				LinkResponseArguments += LinkEnvironment.AdditionalArguments;
 
-					// Write out a response file
-					FileReference ResponseFileName = GetResponseFileName(LinkEnvironment, OutputFile);
-					InputFileNames.Add(LinkResponseArguments.Replace("\\", "/"));
+				// Write out a response file
+				FileReference ResponseFileName = GetResponseFileName(LinkEnvironment, OutputFile);
+				InputFileNames.Add(LinkResponseArguments.Replace("\\", "/"));
 
-					FileItem ResponseFileItem = Graph.CreateIntermediateTextFile(ResponseFileName, InputFileNames, StringComparison.InvariantCultureIgnoreCase);
+				FileItem ResponseFileItem = Graph.CreateIntermediateTextFile(ResponseFileName, InputFileNames, StringComparison.InvariantCultureIgnoreCase);
 
-					LinkAction.CommandArguments += string.Format(" @\"{0}\"", ResponseFileName);
-					LinkAction.PrerequisiteItems.Add(ResponseFileItem);
+				LinkAction.CommandArguments += string.Format(" @\"{0}\"", ResponseFileName);
+				LinkAction.PrerequisiteItems.Add(ResponseFileItem);
 
-					// Fix up the paths in commandline
-					LinkAction.CommandArguments = LinkAction.CommandArguments.Replace("\\", "/");
+				// Fix up the paths in commandline
+				LinkAction.CommandArguments = LinkAction.CommandArguments.Replace("\\", "/");
 
-					// Only execute linking on the local PC.
-					LinkAction.bCanExecuteRemotely = false;
+				// Only execute linking on the local PC.
+				LinkAction.bCanExecuteRemotely = false;
 
-					if(bExecuteCompilerThroughShell)
-					{
-						SetupActionToExecuteCompilerThroughShell(ref LinkAction, LinkAction.CommandPath.FullName, LinkAction.CommandArguments, "Link");
-					}
+				if(bExecuteCompilerThroughShell)
+				{
+					SetupActionToExecuteCompilerThroughShell(ref LinkAction, LinkAction.CommandPath.FullName, LinkAction.CommandArguments, "Link");
+				}
 
-					//Log.TraceInformation("Link: {0} {1}", LinkAction.CommandPath.FullName, LinkAction.CommandArguments);
+				//Log.TraceInformation("Link: {0} {1}", LinkAction.CommandPath.FullName, LinkAction.CommandArguments);
 
-					// Windows can run into an issue with too long of a commandline when clang tries to call ld to link.
-					// To work around this we call clang to just get the command it would execute and generate a
-					// second response file to directly call ld with the right arguments instead of calling through clang.
+				// Windows can run into an issue with too long of a commandline when clang tries to call ld to link.
+				// To work around this we call clang to just get the command it would execute and generate a
+				// second response file to directly call ld with the right arguments instead of calling through clang.
 /* disable while tracking down some linker errors this introduces
-					if (RuntimePlatform.IsWindows)
+				if (RuntimePlatform.IsWindows)
+				{
+					// capture the actual link command without running it
+					ProcessStartInfo StartInfo = new ProcessStartInfo();
+					StartInfo.WorkingDirectory = LinkEnvironment.IntermediateDirectory.FullName;
+					StartInfo.FileName = LinkAction.CommandPath;
+					StartInfo.Arguments = "-### " + LinkAction.CommandArguments;
+					StartInfo.UseShellExecute = false;
+					StartInfo.CreateNoWindow = true;
+					StartInfo.RedirectStandardError = true;
+
+					LinkerCommandline = "";
+
+					Process Proc = new Process();
+					Proc.StartInfo = StartInfo;
+					Proc.ErrorDataReceived += new DataReceivedEventHandler(OutputReceivedForLinker);
+					Proc.Start();
+					Proc.BeginErrorReadLine();
+					Proc.WaitForExit(5000);
+
+					LinkerCommandline = LinkerCommandline.Trim();
+
+					// the command should be in quotes; if not we'll just use clang to link as usual
+					int FirstQuoteIndex = LinkerCommandline.IndexOf('"');
+					if (FirstQuoteIndex >= 0)
 					{
-						// capture the actual link command without running it
-						ProcessStartInfo StartInfo = new ProcessStartInfo();
-						StartInfo.WorkingDirectory = LinkEnvironment.IntermediateDirectory.FullName;
-						StartInfo.FileName = LinkAction.CommandPath;
-						StartInfo.Arguments = "-### " + LinkAction.CommandArguments;
-						StartInfo.UseShellExecute = false;
-						StartInfo.CreateNoWindow = true;
-						StartInfo.RedirectStandardError = true;
-
-						LinkerCommandline = "";
-
-						Process Proc = new Process();
-						Proc.StartInfo = StartInfo;
-						Proc.ErrorDataReceived += new DataReceivedEventHandler(OutputReceivedForLinker);
-						Proc.Start();
-						Proc.BeginErrorReadLine();
-						Proc.WaitForExit(5000);
-
-						LinkerCommandline = LinkerCommandline.Trim();
-
-						// the command should be in quotes; if not we'll just use clang to link as usual
-						int FirstQuoteIndex = LinkerCommandline.IndexOf('"');
-						if (FirstQuoteIndex >= 0)
+						int SecondQuoteIndex = LinkerCommandline.Substring(FirstQuoteIndex + 1).IndexOf('"');
+						if (SecondQuoteIndex >= 0)
 						{
-							int SecondQuoteIndex = LinkerCommandline.Substring(FirstQuoteIndex + 1).IndexOf('"');
-							if (SecondQuoteIndex >= 0)
-							{
-								LinkAction.CommandPath = LinkerCommandline.Substring(FirstQuoteIndex + 1, SecondQuoteIndex - FirstQuoteIndex);
-								LinkAction.CommandArguments = LinkerCommandline.Substring(FirstQuoteIndex + SecondQuoteIndex + 3);
+							LinkAction.CommandPath = LinkerCommandline.Substring(FirstQuoteIndex + 1, SecondQuoteIndex - FirstQuoteIndex);
+							LinkAction.CommandArguments = LinkerCommandline.Substring(FirstQuoteIndex + SecondQuoteIndex + 3);
 
-								// replace double backslashes
-								LinkAction.CommandPath = LinkAction.CommandPath.Replace("\\\\", "/");
+							// replace double backslashes
+							LinkAction.CommandPath = LinkAction.CommandPath.Replace("\\\\", "/");
 
-								// now create a response file for the full command using ld directly
-								FileReference FinalResponseFileName = FileReference.Combine(LinkEnvironment.IntermediateDirectory, OutputFile.Location.GetFileName() + ".responseFinal");
-								FileItem FinalResponseFileItem = Graph.CreateIntermediateTextFile(FinalResponseFileName, LinkAction.CommandArguments);
-								LinkAction.CommandArguments = string.Format("@\"{0}\"", FinalResponseFileName);
-								LinkAction.PrerequisiteItems.Add(FinalResponseFileItem);
-							}
+							// now create a response file for the full command using ld directly
+							FileReference FinalResponseFileName = FileReference.Combine(LinkEnvironment.IntermediateDirectory, OutputFile.Location.GetFileName() + ".responseFinal");
+							FileItem FinalResponseFileItem = Graph.CreateIntermediateTextFile(FinalResponseFileName, LinkAction.CommandArguments);
+							LinkAction.CommandArguments = string.Format("@\"{0}\"", FinalResponseFileName);
+							LinkAction.PrerequisiteItems.Add(FinalResponseFileItem);
 						}
 					}
-*/
 				}
+*/
 			}
 
 			return Outputs.ToArray();
diff --git a/Engine/Source/Programs/UnrealBuildTool/Platform/Android/UEBuildAndroid.cs b/Engine/Source/Programs/UnrealBuildTool/Platform/Android/UEBuildAndroid.cs
--- a/Engine/Source/Programs/UnrealBuildTool/Platform/Android/UEBuildAndroid.cs
+++ b/Engine/Source/Programs/UnrealBuildTool/Platform/Android/UEBuildAndroid.cs
@@ -173,23 +173,16 @@ namespace UnrealBuildTool
 
 		static bool IsBuildProductWithArch(string Name, string[] NamePrefixes, string[] NameSuffixes, string Extension)
 		{
-			// Strip off the extension, then a GPU suffix, then a CPU suffix, before testing whether it matches a build product name.
+			// Strip off the extension, then a CPU suffix, before testing whether it matches a build product name.
 			if (Name.EndsWith(Extension, StringComparison.InvariantCultureIgnoreCase))
 			{
 				int ExtensionEndIdx = Name.Length - Extension.Length;
-				foreach (string GpuSuffix in AndroidToolChain.AllGpuSuffixes)
+				foreach (string CpuSuffix in AndroidToolChain.AllCpuSuffixes)
 				{
-					int GpuIdx = ExtensionEndIdx - GpuSuffix.Length;
-					if (GpuIdx > 0 && String.Compare(Name, GpuIdx, GpuSuffix, 0, GpuSuffix.Length, StringComparison.InvariantCultureIgnoreCase) == 0)
+					int CpuIdx = ExtensionEndIdx - CpuSuffix.Length;
+					if (CpuIdx > 0 && String.Compare(Name, CpuIdx, CpuSuffix, 0, CpuSuffix.Length, StringComparison.InvariantCultureIgnoreCase) == 0)
 					{
-						foreach (string CpuSuffix in AndroidToolChain.AllCpuSuffixes)
-						{
-							int CpuIdx = GpuIdx - CpuSuffix.Length;
-							if (CpuIdx > 0 && String.Compare(Name, CpuIdx, CpuSuffix, 0, CpuSuffix.Length, StringComparison.InvariantCultureIgnoreCase) == 0)
-							{
-								return IsBuildProductName(Name, 0, CpuIdx, NamePrefixes, NameSuffixes);
-							}
-						}
+						return IsBuildProductName(Name, 0, CpuIdx, NamePrefixes, NameSuffixes);
 					}
 				}
 			}
@@ -330,26 +323,22 @@ namespace UnrealBuildTool
 			AndroidToolChain ToolChain = CreateToolChain(Target) as AndroidToolChain;
 
 			List<string> Architectures = ToolChain.GetAllArchitectures();
-			List<string> GPUArchitectures = ToolChain.GetAllGPUArchitectures();
 
 			// make multiple output binaries
 			List<FileReference> AllBinaries = new List<FileReference>();
 			foreach (string Architecture in Architectures)
 			{
-				foreach (string GPUArchitecture in GPUArchitectures)
+				string BinaryPath;
+				if (Target.bShouldCompileAsDLL)
 				{
-					string BinaryPath;
-					if (Target.bShouldCompileAsDLL)
-					{
-						BinaryPath = Path.Combine(BinaryName.Directory.FullName, Target.Configuration.ToString(), "libUnreal.so");
-					}
-					else
-					{
-						BinaryPath = AndroidToolChain.InlineArchName(BinaryName.FullName, Architecture, GPUArchitecture);
-					}
-
-					AllBinaries.Add(new FileReference(BinaryPath));
+					BinaryPath = Path.Combine(BinaryName.Directory.FullName, Target.Configuration.ToString(), "libUnreal.so");
 				}
+				else
+				{
+					BinaryPath = AndroidToolChain.InlineArchName(BinaryName.FullName, Architecture);
+				}
+
+				AllBinaries.Add(new FileReference(BinaryPath));
 			}
 
 			return AllBinaries;
diff --git a/Engine/Source/Programs/UnrealBuildTool/Platform/Android/UEDeployAndroid.cs b/Engine/Source/Programs/UnrealBuildTool/Platform/Android/UEDeployAndroid.cs
--- a/Engine/Source/Programs/UnrealBuildTool/Platform/Android/UEDeployAndroid.cs
+++ b/Engine/Source/Programs/UnrealBuildTool/Platform/Android/UEDeployAndroid.cs
@@ -1928,12 +1928,6 @@ namespace UnrealBuildTool
 				CurrentSettings.AppendFormat("Arch={0}{1}", Arch, Environment.NewLine);
 			}
 
-			List<string> GPUArchitectures = ToolChain.GetAllGPUArchitectures();
-			foreach (string GPUArch in GPUArchitectures)
-			{
-				CurrentSettings.AppendFormat("GPUArch={0}{1}", GPUArch, Environment.NewLine);
-			}
-
 			// Modifying some settings in the GameMapsSettings could trigger the OBB regeneration
 			// and make the cached OBBData.java mismatch to the actually data. 
 			// So we insert the relevant keys into CurrentSettings to capture the change, to
@@ -1965,69 +1959,65 @@ namespace UnrealBuildTool
 			string CookFlavor, string OutputPath, bool bMakeSeparateApks, bool bPackageDataInsideApk)
 		{
 			List<string> Arches = ToolChain.GetAllArchitectures();
-			List<string> GPUArchitectures = ToolChain.GetAllGPUArchitectures();
 
 			// check all input files (.so, java files, .ini files, etc)
 			bool bAllInputsCurrent = true;
 			foreach (string Arch in Arches)
 			{
-				foreach (string GPUArch in GPUArchitectures)
-				{
-					string SourceSOName = AndroidToolChain.InlineArchName(OutputPath, Arch, GPUArch);
-					// if the source binary was UnrealGame, replace it with the new project name, when re-packaging a binary only build
-					string ApkFilename = Path.GetFileNameWithoutExtension(OutputPath).Replace("UnrealGame", ProjectName);
-					string DestApkName = Path.Combine(ProjectDirectory, "Binaries/Android/") + ApkFilename + ".apk";
+				string SourceSOName = AndroidToolChain.InlineArchName(OutputPath, Arch);
+				// if the source binary was UnrealGame, replace it with the new project name, when re-packaging a binary only build
+				string ApkFilename = Path.GetFileNameWithoutExtension(OutputPath).Replace("UnrealGame", ProjectName);
+				string DestApkName = Path.Combine(ProjectDirectory, "Binaries/Android/") + ApkFilename + ".apk";
 
-					// if we making multiple Apks, we need to put the architecture into the name
-					if (bMakeSeparateApks)
-					{
-						DestApkName = AndroidToolChain.InlineArchName(DestApkName, Arch, GPUArch);
-					}
+				// if we making multiple Apks, we need to put the architecture into the name
+				if (bMakeSeparateApks)
+				{
+					DestApkName = AndroidToolChain.InlineArchName(DestApkName, Arch);
+				}
 
-					// check to see if it's out of date before trying the slow make apk process (look at .so and all Engine and Project build files to be safe)
-					List<String> InputFiles = new List<string>();
-					InputFiles.Add(SourceSOName);
-					InputFiles.AddRange(Directory.EnumerateFiles(UnrealBuildFilesPath, "*.*", SearchOption.AllDirectories));
-					if (Directory.Exists(GameBuildFilesPath))
-					{
-						InputFiles.AddRange(Directory.EnumerateFiles(GameBuildFilesPath, "*.*", SearchOption.AllDirectories));
-					}
+				// check to see if it's out of date before trying the slow make apk process (look at .so and all Engine and Project build files to be safe)
+				List<String> InputFiles = new List<string>();
+				InputFiles.Add(SourceSOName);
+				InputFiles.AddRange(Directory.EnumerateFiles(UnrealBuildFilesPath, "*.*", SearchOption.AllDirectories));
+				if (Directory.Exists(GameBuildFilesPath))
+				{
+					InputFiles.AddRange(Directory.EnumerateFiles(GameBuildFilesPath, "*.*", SearchOption.AllDirectories));
+				}
 
-					// make sure changed java files will rebuild apk
-					InputFiles.AddRange(SettingsFiles);
+				// make sure changed java files will rebuild apk
+				InputFiles.AddRange(SettingsFiles);
 
-					// rebuild if .pak files exist for OBB in APK case
-					if (bPackageDataInsideApk)
+				// rebuild if .pak files exist for OBB in APK case
+				if (bPackageDataInsideApk)
+				{
+					string PAKFileLocation = ProjectDirectory + "/Saved/StagedBuilds/Android" + CookFlavor + "/" + ProjectName + "/Content/Paks";
+					if (Directory.Exists(PAKFileLocation))
 					{
-						string PAKFileLocation = ProjectDirectory + "/Saved/StagedBuilds/Android" + CookFlavor + "/" + ProjectName + "/Content/Paks";
-						if (Directory.Exists(PAKFileLocation))
+						IEnumerable<string> PakFiles = Directory.EnumerateFiles(PAKFileLocation, "*.pak", SearchOption.TopDirectoryOnly);
+						foreach (string Name in PakFiles)
 						{
-							IEnumerable<string> PakFiles = Directory.EnumerateFiles(PAKFileLocation, "*.pak", SearchOption.TopDirectoryOnly);
-							foreach (string Name in PakFiles)
-							{
-								InputFiles.Add(Name);
-							}
+							InputFiles.Add(Name);
 						}
 					}
+				}
 
-					// look for any newer input file
-					DateTime ApkTime = File.GetLastWriteTimeUtc(DestApkName);
-					foreach (string InputFileName in InputFiles)
+				// look for any newer input file
+				DateTime ApkTime = File.GetLastWriteTimeUtc(DestApkName);
+				foreach (string InputFileName in InputFiles)
+				{
+					if (File.Exists(InputFileName))
 					{
-						if (File.Exists(InputFileName))
+						// skip .log files
+						if (Path.GetExtension(InputFileName) == ".log")
 						{
-							// skip .log files
-							if (Path.GetExtension(InputFileName) == ".log")
-							{
-								continue;
-							}
-							DateTime InputFileTime = File.GetLastWriteTimeUtc(InputFileName);
-							if (InputFileTime.CompareTo(ApkTime) > 0)
-							{
-								bAllInputsCurrent = false;
-								Log.TraceInformation("{0} is out of date due to newer input file {1}", DestApkName, InputFileName);
-								break;
-							}
+							continue;
+						}
+						DateTime InputFileTime = File.GetLastWriteTimeUtc(InputFileName);
+						if (InputFileTime.CompareTo(ApkTime) > 0)
+						{
+							bAllInputsCurrent = false;
+							Log.TraceInformation("{0} is out of date due to newer input file {1}", DestApkName, InputFileName);
+							break;
 						}
 					}
 				}
@@ -2406,7 +2396,7 @@ namespace UnrealBuildTool
 		}
 
 
-		private string GenerateManifest(AndroidToolChain ToolChain, string ProjectName, TargetType InTargetType, string EngineDirectory, bool bIsForDistribution, bool bPackageDataInsideApk, string GameBuildFilesPath, bool bHasOBBFiles, bool bDisableVerifyOBBOnStartUp, string UnrealArch, string GPUArch, string CookFlavor, bool bUseExternalFilesDir, string Configuration, int SDKLevelInt, bool bIsEmbedded, bool bEnableBundle)
+		private string GenerateManifest(AndroidToolChain ToolChain, string ProjectName, TargetType InTargetType, string EngineDirectory, bool bIsForDistribution, bool bPackageDataInsideApk, string GameBuildFilesPath, bool bHasOBBFiles, bool bDisableVerifyOBBOnStartUp, string UnrealArch, string CookFlavor, bool bUseExternalFilesDir, string Configuration, int SDKLevelInt, bool bIsEmbedded, bool bEnableBundle)
 		{
 			// Read the engine version
 			string EngineVersion = ReadEngineVersion();
@@ -3585,7 +3575,6 @@ namespace UnrealBuildTool
 
 			// Get list of all architecture and GPU targets for build
 			List<string> Arches = ToolChain.GetAllArchitectures();
-			List<string> GPUArchitectures = ToolChain.GetAllGPUArchitectures();
 
 			// we do not need to really build an engine UnrealGame.apk so short-circuit it
 			if (!ForceAPKGeneration && ProjectName == "UnrealGame" && OutputPath.Replace("\\", "/").Contains("/Engine/Binaries/Android/") && Path.GetFileNameWithoutExtension(OutputPath).StartsWith("UnrealGame"))
@@ -3935,25 +3924,23 @@ namespace UnrealBuildTool
 			// Initialize UPL contexts for each architecture enabled
 			UPL.Init(NDKArches, bForDistribution, EngineDirectory, IntermediateAndroidPath, ProjectDirectory, Configuration.ToString(), bSkipGradleBuild, bPerArchBuildDir:true, ArchRemapping:ArchRemapping);
 
-			IEnumerable<Tuple<string, string, string>> BuildList = null;
+			IEnumerable<Tuple<string, string>> BuildList = null;
 
 			bool bRequiresOBB = RequiresOBB(bDisallowPackagingDataInApk, ObbFileLocation);
 			if (!bBuildSettingsMatch)
 			{
 				BuildList = from Arch in Arches
-							from GPUArch in GPUArchitectures
-							let manifest = GenerateManifest(ToolChain, ProjectName, InTargetType, EngineDirectory, bForDistribution, bPackageDataInsideApk, GameBuildFilesPath, bRequiresOBB, bDisableVerifyOBBOnStartUp, Arch, GPUArch, CookFlavor, bUseExternalFilesDir, Configuration.ToString(), SDKLevelInt, bSkipGradleBuild, bEnableBundle)
-							select Tuple.Create(Arch, GPUArch, manifest);
+							let manifest = GenerateManifest(ToolChain, ProjectName, InTargetType, EngineDirectory, bForDistribution, bPackageDataInsideApk, GameBuildFilesPath, bRequiresOBB, bDisableVerifyOBBOnStartUp, Arch, CookFlavor, bUseExternalFilesDir, Configuration.ToString(), SDKLevelInt, bSkipGradleBuild, bEnableBundle)
+							select Tuple.Create(Arch, manifest);
 			}
 			else
 			{
 				BuildList = from Arch in Arches
-							from GPUArch in GPUArchitectures
-							let manifestFile = Path.Combine(IntermediateAndroidPath, Arch + (GPUArch.Length > 0 ? ("_" + GPUArch.Substring(1)) : "") + "_AndroidManifest.xml")
-							let manifest = GenerateManifest(ToolChain, ProjectName, InTargetType, EngineDirectory, bForDistribution, bPackageDataInsideApk, GameBuildFilesPath, bRequiresOBB, bDisableVerifyOBBOnStartUp, Arch, GPUArch, CookFlavor, bUseExternalFilesDir, Configuration.ToString(), SDKLevelInt, bSkipGradleBuild, bEnableBundle)
+							let manifestFile = Path.Combine(IntermediateAndroidPath, Arch + "_AndroidManifest.xml")
+							let manifest = GenerateManifest(ToolChain, ProjectName, InTargetType, EngineDirectory, bForDistribution, bPackageDataInsideApk, GameBuildFilesPath, bRequiresOBB, bDisableVerifyOBBOnStartUp, Arch, CookFlavor, bUseExternalFilesDir, Configuration.ToString(), SDKLevelInt, bSkipGradleBuild, bEnableBundle)
 							let OldManifest = File.Exists(manifestFile) ? File.ReadAllText(manifestFile) : ""
 							where manifest != OldManifest
-							select Tuple.Create(Arch, GPUArch, manifest);
+							select Tuple.Create(Arch, manifest);
 			}
 
 			// Now we have to spin over all the arch/gpu combinations to make sure they all match
@@ -3973,11 +3960,10 @@ namespace UnrealBuildTool
 			Replacements.Add("${PY_VISUALIZER_PATH}", Path.GetFullPath(Path.Combine(EngineDirectory, "Extras", "LLDBDataFormatters", "UEDataFormatters_2ByteChars.py")));
 
 			// steps run for each build combination (note: there should only be one GPU in future)
-			foreach (Tuple<string, string, string> build in BuildList)
+			foreach (Tuple<string, string> build in BuildList)
 			{
 				string Arch = build.Item1;
-				string GPUArchitecture = build.Item2;
-				string Manifest = build.Item3;
+				string Manifest = build.Item2;
 				string NDKArch = GetNDKArch(Arch);
 
 				Log.TraceInformation("\n===={0}====PREPARING NATIVE CODE====={1}============================================================", DateTime.Now.ToString(), Arch);
@@ -4065,7 +4051,7 @@ namespace UnrealBuildTool
 				string CompileSDKVersion = SDKAPILevel.Replace("android-", "");
 
 				// Write the manifest to the correct locations (cache and real)
-				String ManifestFile = Path.Combine(IntermediateAndroidPath, Arch + (GPUArchitecture.Length > 0 ? ("_" + GPUArchitecture.Substring(1)) : "") + "_AndroidManifest.xml");
+				String ManifestFile = Path.Combine(IntermediateAndroidPath, Arch + "_AndroidManifest.xml");
 				File.WriteAllText(ManifestFile, Manifest);
 				ManifestFile = Path.Combine(UnrealBuildPath, "AndroidManifest.xml");
 				File.WriteAllText(ManifestFile, Manifest);
@@ -4098,13 +4084,13 @@ namespace UnrealBuildTool
 				}
 				else
 				{
-					string SourceSOName = AndroidToolChain.InlineArchName(OutputPath, Arch, GPUArchitecture);
+					string SourceSOName = AndroidToolChain.InlineArchName(OutputPath, Arch);
 					// if the source binary was UnrealGame, replace it with the new project name, when re-packaging a binary only build
 					string ApkFilename = Path.GetFileNameWithoutExtension(OutputPath).Replace("UnrealGame", ProjectName);
 					DestApkName = Path.Combine(DestApkDirectory, ApkFilename + ".apk");
 
 					// As we are always making seperate APKs we need to put the architecture into the name
-					DestApkName = AndroidToolChain.InlineArchName(DestApkName, Arch, GPUArchitecture);
+					DestApkName = AndroidToolChain.InlineArchName(DestApkName, Arch);
 
 					if (!File.Exists(SourceSOName))
 					{
@@ -4319,7 +4305,7 @@ namespace UnrealBuildTool
 				{
 					// Copy .so with symbols to 
 					int StoreVersion = GetStoreVersion(Arch);
-					string SymbolSODirectory = Path.Combine(DestApkDirectory, ProjectName + "_Symbols_v" + StoreVersion + "/" + ProjectName + Arch + GPUArchitecture);
+					string SymbolSODirectory = Path.Combine(DestApkDirectory, ProjectName + "_Symbols_v" + StoreVersion + "/" + ProjectName + Arch);
 					string SymbolifiedSOPath = Path.Combine(SymbolSODirectory, Path.GetFileName(FinalSOName));
 					MakeDirectoryIfRequired(SymbolifiedSOPath);
 					Log.TraceInformation("Writing symbols to {0}", SymbolifiedSOPath);
@@ -4345,11 +4331,10 @@ namespace UnrealBuildTool
 				String ABIFilter = "";
 
 				// loop through and merge the different architecture gradle directories
-				foreach (Tuple<string, string, string> build in BuildList)
+				foreach (Tuple<string, string> build in BuildList)
 				{
 					string Arch = build.Item1;
-					string GPUArchitecture = build.Item2;
-					string Manifest = build.Item3;
+					string Manifest = build.Item2;
 					string NDKArch = GetNDKArch(Arch);
 
 					string UnrealBuildPath = Path.Combine(IntermediateAndroidPath, Arch.Substring(1).Replace("-", "_"));
@@ -4676,11 +4661,10 @@ namespace UnrealBuildTool
 				else
 				{
 					// generate an AAB for each architecture separately, was unable to merge
-					foreach (Tuple<string, string, string> build in BuildList)
+					foreach (Tuple<string, string> build in BuildList)
 					{
 						string Arch = build.Item1;
-						string GPUArchitecture = build.Item2;
-						string Manifest = build.Item3;
+						string Manifest = build.Item2;
 						string NDKArch = GetNDKArch(Arch);
 
 						Log.TraceInformation("\n===={0}====GENERATING BUNDLE====={1}================================================================", DateTime.Now.ToString(), Arch);
diff --git a/Engine/Source/Runtime/Core/Public/Android/AndroidPlatformProperties.h b/Engine/Source/Runtime/Core/Public/Android/AndroidPlatformProperties.h
--- a/Engine/Source/Runtime/Core/Public/Android/AndroidPlatformProperties.h
+++ b/Engine/Source/Runtime/Core/Public/Android/AndroidPlatformProperties.h
@@ -71,11 +71,7 @@ struct FAndroidPlatformProperties
 
 	static FORCEINLINE bool SupportsLowQualityLightmaps()
 	{
-#if PLATFORM_ANDROIDGL4
-		return false;
-#else
 		return true;
-#endif
 	}
 
 	static FORCEINLINE bool SupportsDistanceFieldShadows()
diff --git a/Engine/Source/Runtime/Core/Public/HAL/Platform.h b/Engine/Source/Runtime/Core/Public/HAL/Platform.h
--- a/Engine/Source/Runtime/Core/Public/HAL/Platform.h
+++ b/Engine/Source/Runtime/Core/Public/HAL/Platform.h
@@ -44,12 +44,6 @@
 #if !defined(PLATFORM_ANDROID_X64)
 	#define PLATFORM_ANDROID_X64 0
 #endif
-#if !defined(PLATFORM_ANDROID_VULKAN)
-	#define PLATFORM_ANDROID_VULKAN 0
-#endif
-#if !defined(PLATFORM_ANDROIDGL4)
-	#define PLATFORM_ANDROIDGL4 0
-#endif
 #if !defined(PLATFORM_APPLE)
 	#define PLATFORM_APPLE 0
 #endif
