Editegrate fix from //UE5/Release-Engine-Staging changelist 18086920 to //UE5/Main:

Enable a project file generator to set ShouldGenerateIntelliSenseData()=true to get module dependency information, but ShouldGenerateIntelliSenseCompileEnvironments()=false to skip generating compile environments for those modules (including writing Definitions.h files).

diff --git a/Engine/Source/Programs/UnrealBuildTool/ProjectFiles/ProjectFileGenerator.cs b/Engine/Source/Programs/UnrealBuildTool/ProjectFiles/ProjectFileGenerator.cs
--- a/Engine/Source/Programs/UnrealBuildTool/ProjectFiles/ProjectFileGenerator.cs
+++ b/Engine/Source/Programs/UnrealBuildTool/ProjectFiles/ProjectFileGenerator.cs
@@ -396,6 +396,14 @@ namespace UnrealBuildTool
 			return bGenerateIntelliSenseData;
 		}
 
+		/// <summary>
+		/// True if this project generator requires a compile environment for Intellisense data.
+		/// </summary>
+		public virtual bool ShouldGenerateIntelliSenseCompileEnvironments()
+		{
+			return ShouldGenerateIntelliSenseData();
+		}
+
 		/// <summary>
 		/// Allows each project generator to indicate whether target rules should be used to explicitly enable or disable plugins.
 		/// Default is false - since usually not needed for project generation unless project files indicate whether referenced plugins should be built or not.
@@ -1900,29 +1908,33 @@ namespace UnrealBuildTool
 							UEBuildTarget Target = UEBuildTarget.Create(TargetDesc, false, false, bUsePrecompiled);
 							AddTargetForIntellisense(Target);
 
-							// Generate a compile environment for each module in the binary
-							CppCompileEnvironment GlobalCompileEnvironment = Target.CreateCompileEnvironmentForProjectFiles();
-							foreach (UEBuildBinary Binary in Target.Binaries)
+							// If the project generator just cares about the result of UEBuildTarget.Create, skip generating the compile environments.
+							if (ShouldGenerateIntelliSenseCompileEnvironments())
 							{
-								CppCompileEnvironment BinaryCompileEnvironment = Binary.CreateBinaryCompileEnvironment(GlobalCompileEnvironment);
-								foreach (UEBuildModuleCPP Module in Binary.Modules.OfType<UEBuildModuleCPP>())
+								// Generate a compile environment for each module in the binary
+								CppCompileEnvironment GlobalCompileEnvironment = Target.CreateCompileEnvironmentForProjectFiles();
+								foreach (UEBuildBinary Binary in Target.Binaries)
 								{
-									ProjectFile ProjectFileForIDE;
-									if (ModuleToProjectFileMap.TryGetValue(Module.Name, out ProjectFileForIDE) && ProjectFileForIDE == TargetProjectFile)
+									CppCompileEnvironment BinaryCompileEnvironment = Binary.CreateBinaryCompileEnvironment(GlobalCompileEnvironment);
+									foreach (UEBuildModuleCPP Module in Binary.Modules.OfType<UEBuildModuleCPP>())
 									{
-										CppCompileEnvironment ModuleCompileEnvironment = Module.CreateCompileEnvironmentForIntellisense(Target.Rules, BinaryCompileEnvironment);
-										ProjectFileForIDE.AddModule(Module, ModuleCompileEnvironment);
+										ProjectFile ProjectFileForIDE;
+										if (ModuleToProjectFileMap.TryGetValue(Module.Name, out ProjectFileForIDE) && ProjectFileForIDE == TargetProjectFile)
+										{
+											CppCompileEnvironment ModuleCompileEnvironment = Module.CreateCompileEnvironmentForIntellisense(Target.Rules, BinaryCompileEnvironment);
+											ProjectFileForIDE.AddModule(Module, ModuleCompileEnvironment);
+										}
 									}
 								}
-							}
 
-							// If we're generating project files, then go ahead and wipe out the existing UBTMakefile for every target, to make sure that
-							// it gets a full dependency scan next time.
-							// NOTE: This is just a safeguard and doesn't have to be perfect.  We also check for newer project file timestamps in LoadUBTMakefile()
-							FileReference MakefileLocation = TargetMakefile.GetLocation(TargetDesc.ProjectFile, TargetDesc.Name, TargetDesc.Platform, TargetDesc.Architecture, TargetDesc.Configuration);
-							if (FileReference.Exists(MakefileLocation))
-							{
-								FileReference.Delete(MakefileLocation);
+								// If we're generating project files, then go ahead and wipe out the existing UBTMakefile for every target, to make sure that
+								// it gets a full dependency scan next time.
+								// NOTE: This is just a safeguard and doesn't have to be perfect.  We also check for newer project file timestamps in LoadUBTMakefile()
+								FileReference MakefileLocation = TargetMakefile.GetLocation(TargetDesc.ProjectFile, TargetDesc.Name, TargetDesc.Platform, TargetDesc.Architecture, TargetDesc.Configuration);
+								if (FileReference.Exists(MakefileLocation))
+								{
+									FileReference.Delete(MakefileLocation);
+								}
 							}
 						}
 						catch (Exception Ex)
