Added a way of suppressing implicit plugin dependencies for specific build targets

  - A new property on TargetRules allows a target to define any number of plugins that are allowed to be depended upon, even though a plugin may not explicitly state their dependency in the plugin descriptor
  - This is useful for some projects where the inclusion of a plugin can only be made as part of a Target.cs file or Build.cs file

diff --git a/Engine/Source/Programs/UnrealBuildTool/Configuration/TargetRules.cs b/Engine/Source/Programs/UnrealBuildTool/Configuration/TargetRules.cs
--- a/Engine/Source/Programs/UnrealBuildTool/Configuration/TargetRules.cs
+++ b/Engine/Source/Programs/UnrealBuildTool/Configuration/TargetRules.cs
@@ -355,6 +355,11 @@ namespace UnrealBuildTool
 		[CommandLine("-DisablePlugin=", ListSeparator = '+')]
 		public List<string> DisablePlugins = new List<string>();
 
+		/// <summary>
+		/// A list of Plugin names that are allowed to exist as dependencies without being defined in the uplugin descriptor
+		/// </summary>
+		public List<string> InternalPluginDependencies = new List<string>();
+
 		/// <summary>
 		/// Path to the set of pak signing keys to embed in the executable.
 		/// </summary>
@@ -2056,6 +2061,10 @@ namespace UnrealBuildTool
 			get { return Inner.BuildPlugins; }
 		}
 
+		public IEnumerable<string> InternalPluginDependencies
+		{
+			get { return Inner.InternalPluginDependencies; }
+		}
 		public bool bRuntimeDependenciesComeFromBuildPlugins
 		{
 			get { return Inner.bRuntimeDependenciesComeFromBuildPlugins; }
diff --git a/Engine/Source/Programs/UnrealBuildTool/Configuration/UEBuildTarget.cs b/Engine/Source/Programs/UnrealBuildTool/Configuration/UEBuildTarget.cs
--- a/Engine/Source/Programs/UnrealBuildTool/Configuration/UEBuildTarget.cs
+++ b/Engine/Source/Programs/UnrealBuildTool/Configuration/UEBuildTarget.cs
@@ -1708,7 +1708,9 @@ namespace UnrealBuildTool
 					foreach(UEBuildModule DependencyModule in DependencyModules)
 					{
 						UEBuildPlugin? DependencyPlugin;
-						if(ModuleToPlugin.TryGetValue(DependencyModule, out DependencyPlugin) && DependencyPlugin != Plugin && !Plugin.Dependencies!.Contains(DependencyPlugin))
+						if(ModuleToPlugin.TryGetValue(DependencyModule, out DependencyPlugin) && DependencyPlugin != Plugin
+							&& !Rules.InternalPluginDependencies.Contains(DependencyPlugin.Name)
+							&& !Plugin.Dependencies!.Contains(DependencyPlugin))
 						{
 							Log.TraceWarning("Warning: Plugin '{0}' does not list plugin '{1}' as a dependency, but module '{2}' depends on '{3}'.", Plugin.Name, DependencyPlugin.Name, Module.Name, DependencyModule.Name);
 						}
