UBT: Prevent timestamp being updated on *.version files if contents have not changed.

diff --git a/Engine/Source/Programs/UnrealBuildTool/Modes/WriteMetadataMode.cs b/Engine/Source/Programs/UnrealBuildTool/Modes/WriteMetadataMode.cs
--- a/Engine/Source/Programs/UnrealBuildTool/Modes/WriteMetadataMode.cs
+++ b/Engine/Source/Programs/UnrealBuildTool/Modes/WriteMetadataMode.cs
@@ -189,7 +189,7 @@ namespace UnrealBuildTool
 			{
 				DirectoryReference.CreateDirectory(TargetInfo.VersionFile.Directory);
 				TargetInfo.Version.BuildId = BuildId;
-				TargetInfo.Version.Write(TargetInfo.VersionFile);
+				TargetInfo.Version.WriteIfModified(TargetInfo.VersionFile);
 			}
 
 			// Write out the receipt
diff --git a/Engine/Source/Programs/UnrealBuildTool/System/BuildVersion.cs b/Engine/Source/Programs/UnrealBuildTool/System/BuildVersion.cs
--- a/Engine/Source/Programs/UnrealBuildTool/System/BuildVersion.cs
+++ b/Engine/Source/Programs/UnrealBuildTool/System/BuildVersion.cs
@@ -178,6 +178,17 @@ namespace UnrealBuildTool
 			}
 		}
 
+		/// <summary>
+		/// Exports this object as Json
+		/// </summary>
+		/// <param name="FileName">The filename to write to</param>
+		public void WriteIfModified(FileReference FileName)
+		{
+			using StringWriter Writer = new StringWriter();
+			Write(Writer);
+			Utils.WriteFileIfChanged(FileName, Writer.ToString());
+		}
+
 		/// <summary>
 		/// Exports this object as Json
 		/// </summary>
