Fix UBT issue where a game project outside of the engine directory tree could end up with a BuildId that isn't in sync with the engine.

#rnx

diff --git a/Engine/Source/Programs/UnrealBuildTool/Modes/WriteMetadataMode.cs b/Engine/Source/Programs/UnrealBuildTool/Modes/WriteMetadataMode.cs
--- a/Engine/Source/Programs/UnrealBuildTool/Modes/WriteMetadataMode.cs
+++ b/Engine/Source/Programs/UnrealBuildTool/Modes/WriteMetadataMode.cs
@@ -214,7 +214,7 @@ namespace UnrealBuildTool
 			foreach (FileReference ManifestFileName in FileToManifest.Keys)
 			{
 				ModuleManifest? Manifest;
-				if (TryReadManifest(ManifestFileName, out Manifest) && Manifest.BuildId == BuildId)
+				if (ManifestFileName.IsUnderDirectory(Unreal.EngineDirectory) && TryReadManifest(ManifestFileName, out Manifest) && Manifest.BuildId == BuildId)
 				{
 					DateTime ManifestTime = FileReference.GetLastWriteTimeUtc(ManifestFileName);
 					foreach (string FileName in Manifest.ModuleNameToFileName.Values)
