Remove (typed) references to NDA platforms

diff --git a/Engine/Source/Programs/UnrealBuildTool/Configuration/ModuleRules.cs b/Engine/Source/Programs/UnrealBuildTool/Configuration/ModuleRules.cs
--- a/Engine/Source/Programs/UnrealBuildTool/Configuration/ModuleRules.cs
+++ b/Engine/Source/Programs/UnrealBuildTool/Configuration/ModuleRules.cs
@@ -1464,5 +1464,38 @@ namespace UnrealBuildTool
 				PublicSystemIncludePaths.Add(Path.Combine(VcPackageRoot, "include"));
 			}
 		}
+
+		/// <summary>
+		/// Replace an expected value in a list of definitions with a new value
+		/// </summary>
+		/// <param name="Definitions">List of definitions e.g. PublicDefinitions</param>
+		/// <param name="Name">Name of the define to change</param>
+		/// <param name="PreviousValue">Expected value</param>
+		/// <param name="NewValue">New value</param>
+		/// <exception cref="Exception"></exception>
+		protected static void ChangeDefinition(List<string> Definitions, string Name, string PreviousValue, string NewValue)
+		{
+			if (!Definitions.Remove($"{Name}={PreviousValue}"))
+			{
+				throw new Exception("Failed to removed expected definition");
+			}
+			Definitions.Add($"{Name}={NewValue}");
+		}
+
+		/// <summary>
+		/// Replace an expected value in a list of module names with a new value
+		/// </summary>
+		/// <param name="Definitions">List of module names e.g. PublicDependencyModuleNames</param>
+		/// <param name="PreviousModule">Expected value</param>
+		/// <param name="NewModule">New value</param>
+		/// <exception cref="Exception"></exception>
+		protected static void ReplaceModule(List<string> Definitions, string PreviousModule, string NewModule)
+		{
+			if (!Definitions.Remove(PreviousModule))
+			{
+				throw new Exception("Failed to removed expected module name");
+			}
+			Definitions.Add(NewModule);
+		}
 	}
 }
diff --git a/Engine/Source/Programs/UnrealBuildTool/Configuration/UEBuildTarget.cs b/Engine/Source/Programs/UnrealBuildTool/Configuration/UEBuildTarget.cs
--- a/Engine/Source/Programs/UnrealBuildTool/Configuration/UEBuildTarget.cs
+++ b/Engine/Source/Programs/UnrealBuildTool/Configuration/UEBuildTarget.cs
@@ -235,6 +235,16 @@ namespace UnrealBuildTool
 			return UEBuildPlatform.IsPlatformInGroup(this, Group);
 		}
 
+		/// <summary>
+		/// Helper that just calls UEBuildPlatform.IsPlatformInGroup
+		/// </summary>
+		/// <param name="Group"></param>
+		/// <returns></returns>
+		public bool IsInGroup(string Group)
+		{
+			return UnrealPlatformGroup.TryParse(Group, out UnrealPlatformGroup ActualGroup) && UEBuildPlatform.IsPlatformInGroup(this, ActualGroup);
+		}
+
 		/// <summary>
 		/// 64-bit Windows
 		/// </summary>
