UBT: Fix post build steps running earlier than they should, due to missing dependency on metadata file being written.

diff --git a/Engine/Source/Programs/UnrealBuildTool/Configuration/UEBuildTarget.cs b/Engine/Source/Programs/UnrealBuildTool/Configuration/UEBuildTarget.cs
--- a/Engine/Source/Programs/UnrealBuildTool/Configuration/UEBuildTarget.cs
+++ b/Engine/Source/Programs/UnrealBuildTool/Configuration/UEBuildTarget.cs
@@ -1976,6 +1976,7 @@ namespace UnrealBuildTool
 					BinaryFormatterUtils.SaveIfDifferent(TargetInfoFile, TargetInfo);
 
 					Action WriteTargetMetadata = CreateWriteMetadataAction(Makefile, ReceiptFileName.GetFileName(), TargetInfoFile);
+					WriteTargetMetadata.PrerequisiteItems.AddRange(Binaries.SelectMany(x => x.OutputFilePaths).Select(x => FileItem.GetItemByFileReference(x)));
 					WriteTargetMetadata.ProducedItems.Add(FileItem.GetItemByFileReference(ReceiptFileName));
 					Makefile.OutputItems.AddRange(WriteTargetMetadata.ProducedItems);
 				}
