UnrealBuildTool: Make platform handling more consistent

diff --git a/Engine/Source/Programs/UnrealBuildTool/Configuration/PlatformTargetRules.cs b/Engine/Source/Programs/UnrealBuildTool/Configuration/PlatformTargetRules.cs
deleted file mode 100644
diff --git a/Engine/Source/Programs/UnrealBuildTool/Configuration/TargetRules.cs b/Engine/Source/Programs/UnrealBuildTool/Configuration/TargetRules.cs
--- a/Engine/Source/Programs/UnrealBuildTool/Configuration/TargetRules.cs
+++ b/Engine/Source/Programs/UnrealBuildTool/Configuration/TargetRules.cs
@@ -1554,30 +1554,12 @@ namespace UnrealBuildTool
 		[ConfigSubObject]
 		public MacTargetRules MacPlatform = new MacTargetRules();
 
-		/// <summary>
-		/// PS4-specific target settings.
-		/// </summary>
-		[ConfigSubObject]
-		public PS4TargetRules PS4Platform = new PS4TargetRules();
-
-		/// <summary>
-		/// Switch-specific target settings.
-		/// </summary>
-		[ConfigSubObject]
-		public SwitchTargetRules SwitchPlatform = new SwitchTargetRules();
-
 		/// <summary>
 		/// Windows-specific target settings.
 		/// </summary>
 		[ConfigSubObject]
 		public WindowsTargetRules WindowsPlatform; // Requires 'this' parameter; initialized in constructor
 
-		/// <summary>
-		/// Xbox One-specific target settings.
-		/// </summary>
-		[ConfigSubObject]
-		public XboxOneTargetRules XboxOnePlatform = new XboxOneTargetRules();
-
 		/// <summary>
 		/// HoloLens-specific target settings.
 		/// </summary>
@@ -1933,10 +1915,7 @@ namespace UnrealBuildTool
 			IOSPlatform = new ReadOnlyIOSTargetRules(Inner.IOSPlatform);
 			LinuxPlatform = new ReadOnlyLinuxTargetRules(Inner.LinuxPlatform);
 			MacPlatform = new ReadOnlyMacTargetRules(Inner.MacPlatform);
-			PS4Platform = new ReadOnlyPS4TargetRules(Inner.PS4Platform);
-			SwitchPlatform = new ReadOnlySwitchTargetRules(Inner.SwitchPlatform);
 			WindowsPlatform = new ReadOnlyWindowsTargetRules(Inner.WindowsPlatform);
-			XboxOnePlatform = new ReadOnlyXboxOneTargetRules(Inner.XboxOnePlatform);
 			HoloLensPlatform = new ReadOnlyHoloLensTargetRules(Inner.HoloLensPlatform);
 		}
 
@@ -2836,18 +2815,6 @@ namespace UnrealBuildTool
 			private set;
 		}
 
-		public ReadOnlyPS4TargetRules PS4Platform
-		{
-			get;
-			private set;
-		}
-
-		public ReadOnlySwitchTargetRules SwitchPlatform
-		{
-			get;
-			private set;
-		}
-
 		public ReadOnlyWindowsTargetRules WindowsPlatform
 		{
 			get;
@@ -2860,12 +2827,6 @@ namespace UnrealBuildTool
 			private set;
 		}
 
-		public ReadOnlyXboxOneTargetRules XboxOnePlatform
-		{
-			get;
-			private set;
-		}
-
 		public bool bShouldCompileAsDLL
 		{
 			get { return Inner.bShouldCompileAsDLL; }
diff --git a/Engine/Source/Programs/UnrealBuildTool/Configuration/UEBuildTarget.cs b/Engine/Source/Programs/UnrealBuildTool/Configuration/UEBuildTarget.cs
--- a/Engine/Source/Programs/UnrealBuildTool/Configuration/UEBuildTarget.cs
+++ b/Engine/Source/Programs/UnrealBuildTool/Configuration/UEBuildTarget.cs
@@ -260,11 +260,6 @@ namespace UnrealBuildTool
 		/// </summary>
 		public static UnrealTargetPlatform Mac = FindOrAddByName("Mac");
 
-		/// <summary>
-		/// Playstation 4
-		/// </summary>
-		public static UnrealTargetPlatform PS4 = FindOrAddByName("PS4");
-
 		/// <summary>
 		/// iOS
 		/// </summary>
@@ -289,11 +284,6 @@ namespace UnrealBuildTool
 		/// TVOS
 		/// </summary>
 		public static UnrealTargetPlatform TVOS = FindOrAddByName("TVOS");
-
-		/// <summary>
-		/// Nintendo Switch
-		/// </summary>
-		public static UnrealTargetPlatform Switch = FindOrAddByName("Switch");
 	}
 
 	internal class UnrealTargetPlatformTypeConverter : TypeConverter
@@ -565,16 +555,6 @@ namespace UnrealBuildTool
 		/// </summary>
 		public static UnrealPlatformGroup Android = FindOrAddByName("Android");
 
-		/// <summary>
-		/// Sony platforms
-		/// </summary>
-		public static UnrealPlatformGroup Sony = FindOrAddByName("Sony");
-
-		/// <summary>
-		/// XboxCommon group
-		/// </summary>
-		public static UnrealPlatformGroup XboxCommon = FindOrAddByName("XboxCommon");
-
 		/// <summary>
 		/// Desktop group - used by UnrealPlatformClass.Desktop
 		/// </summary>
diff --git a/Engine/Source/Programs/UnrealBuildTool/Platform/Linux/UEBuildLinux.cs b/Engine/Source/Programs/UnrealBuildTool/Platform/Linux/UEBuildLinux.cs
--- a/Engine/Source/Programs/UnrealBuildTool/Platform/Linux/UEBuildLinux.cs
+++ b/Engine/Source/Programs/UnrealBuildTool/Platform/Linux/UEBuildLinux.cs
@@ -30,7 +30,7 @@ namespace UnrealBuildTool
 	/// <summary>
 	/// Linux-specific target settings
 	/// </summary>
-	public partial class LinuxTargetRules
+	public class LinuxTargetRules
 	{
 		/// <summary>
 		/// Constructor
@@ -84,8 +84,19 @@ namespace UnrealBuildTool
 	/// <summary>
 	/// Read-only wrapper for Linux-specific target settings
 	/// </summary>
-	public partial class ReadOnlyLinuxTargetRules
+	public class ReadOnlyLinuxTargetRules
 	{
+		private LinuxTargetRules Inner;
+
+		/// <summary>
+		/// Constructor
+		/// </summary>
+		/// <param name="Inner">The settings object to wrap</param>
+		public ReadOnlyLinuxTargetRules(LinuxTargetRules Inner)
+		{
+			this.Inner = Inner;
+		}
+
 		/// <summary>
 		/// Accessors for fields on the inner TargetRules instance
 		/// </summary>
@@ -597,7 +608,6 @@ namespace UnrealBuildTool
 		}
 	}
 
-
 	class LinuxPlatformFactory : UEBuildPlatformFactory
 	{
 		public override UnrealTargetPlatform TargetPlatform
diff --git a/Engine/Source/Programs/UnrealBuildTool/ProjectFiles/VisualStudio/VCProject.cs b/Engine/Source/Programs/UnrealBuildTool/ProjectFiles/VisualStudio/VCProject.cs
--- a/Engine/Source/Programs/UnrealBuildTool/ProjectFiles/VisualStudio/VCProject.cs
+++ b/Engine/Source/Programs/UnrealBuildTool/ProjectFiles/VisualStudio/VCProject.cs
@@ -1741,7 +1741,7 @@ namespace UnrealBuildTool
 
 					if (TargetRulesObject.Type == TargetType.Game || TargetRulesObject.Type == TargetType.Client || TargetRulesObject.Type == TargetType.Server)
 					{
-						// Allow platforms to add any special properties they require... like aumid override for Xbox One
+						// Allow platforms to add any special properties they require
 						PlatformProjectGenerators.GenerateGamePlatformSpecificProperties(Platform, Configuration, TargetRulesObject.Type, VCProjectFileContent, RootOutputDirectory, TargetFilePath);
 					}
 
