Fix include path when we changed dir from Linux -> Unix (thanks arcolight! PR #8640)

diff --git a/Engine/Source/Programs/UnrealBuildTool/ProjectFiles/Rider/RiderProjectFile.cs b/Engine/Source/Programs/UnrealBuildTool/ProjectFiles/Rider/RiderProjectFile.cs
--- a/Engine/Source/Programs/UnrealBuildTool/ProjectFiles/Rider/RiderProjectFile.cs
+++ b/Engine/Source/Programs/UnrealBuildTool/ProjectFiles/Rider/RiderProjectFile.cs
@@ -507,8 +507,8 @@ namespace UnrealBuildTool
 					    Target.Architecture.StartsWith("i686"))
 					{
 						// libc++ include directories
-						Writer.WriteValue(Path.Combine(EngineDirectory, "Source/ThirdParty/Linux/LibCxx/include/"));
-						Writer.WriteValue(Path.Combine(EngineDirectory, "Source/ThirdParty/Linux/LibCxx/include/c++/v1"));
+						Writer.WriteValue(Path.Combine(EngineDirectory, "Source/ThirdParty/Unix/LibCxx/include/"));
+						Writer.WriteValue(Path.Combine(EngineDirectory, "Source/ThirdParty/Unix/LibCxx/include/c++/v1"));
 					}
 				}
 
