UnrealBuildTool WriteDocumentaion: nested class lookup in Documentation.xml

When looking up member names in documentation xml, replace + with .

Type.FullName uses '+' to indicate a nested class (e.g. UnrealBuildTool.UnrealBuildTool+GlobalOptions), but generated documentation xml always uses '.'

diff --git a/Engine/Source/Programs/UnrealBuildTool/System/RulesDocumentation.cs b/Engine/Source/Programs/UnrealBuildTool/System/RulesDocumentation.cs
--- a/Engine/Source/Programs/UnrealBuildTool/System/RulesDocumentation.cs
+++ b/Engine/Source/Programs/UnrealBuildTool/System/RulesDocumentation.cs
@@ -191,20 +191,23 @@ namespace UnrealBuildTool
 			XmlNode? Node = null;
 			for (; ; )
 			{
-				Node = Documentation.SelectSingleNode($"//member[@name='{MemberName}']/exclude");
+				// Nested types are delineated with '+' in Type.FullName, but generated Documentation.xml uses '.' in all cases
+				string XmlMemberName = MemberName.Replace('+', '.');
+
+				Node = Documentation.SelectSingleNode($"//member[@name='{XmlMemberName}']/exclude");
 				if (Node != null)
 				{
 					bExclude = true;
 					break;
 				}
 
-				Node = Documentation.SelectSingleNode($"//member[@name='{MemberName}']/summary");
+				Node = Documentation.SelectSingleNode($"//member[@name='{XmlMemberName}']/summary");
 				if (Node != null)
 				{
 					break;
 				}
 
-				XmlNode InheritNode = Documentation.SelectSingleNode($"//member[@name='{MemberName}']/inheritdoc/@cref");
+				XmlNode InheritNode = Documentation.SelectSingleNode($"//member[@name='{XmlMemberName}']/inheritdoc/@cref");
 				if (InheritNode == null || !VisitedProperties.Add(InheritNode.InnerText))
 				{
 					break;
