Force use of cl-filter if printing timing information for CTA.

#rnx

diff --git a/Engine/Source/Programs/UnrealBuildTool/Platform/Windows/VCCompileAction.cs b/Engine/Source/Programs/UnrealBuildTool/Platform/Windows/VCCompileAction.cs
--- a/Engine/Source/Programs/UnrealBuildTool/Platform/Windows/VCCompileAction.cs
+++ b/Engine/Source/Programs/UnrealBuildTool/Platform/Windows/VCCompileAction.cs
@@ -133,6 +133,11 @@ namespace UnrealBuildTool
 		/// </summary>
 		public bool bShowIncludes { get; set; }
 
+		/// <summary>
+		/// Whether to override the normal logic for UsingClFilter and force it on.
+		/// </summary>
+		public bool ForceClFilter = false;
+
 		#region Public IAction implementation
 
 		/// <summary>
@@ -220,7 +225,7 @@ namespace UnrealBuildTool
 		/// <summary>
 		/// Whether to use cl-filter
 		/// </summary>
-		bool UsingClFilter => DependencyListFile != null && !DependencyListFile.HasExtension(".json") && !DependencyListFile.HasExtension(".d");
+		bool UsingClFilter => ForceClFilter || (DependencyListFile != null && !DependencyListFile.HasExtension(".json") && !DependencyListFile.HasExtension(".d"));
 
 		/// <inheritdoc/>
 		FileReference IExternalAction.CommandPath
diff --git a/Engine/Source/Programs/UnrealBuildTool/Platform/Windows/VCToolChain.cs b/Engine/Source/Programs/UnrealBuildTool/Platform/Windows/VCToolChain.cs
--- a/Engine/Source/Programs/UnrealBuildTool/Platform/Windows/VCToolChain.cs
+++ b/Engine/Source/Programs/UnrealBuildTool/Platform/Windows/VCToolChain.cs
@@ -1340,10 +1340,17 @@ namespace UnrealBuildTool
 					CompileAction.AdditionalProducedItems.Add(IfcFile);
 					CompileAction.CompiledModuleInterfaceFile = IfcFile;
 				}
-				
+
+				if (Target.bPrintToolChainTimingInfo || Target.WindowsPlatform.bCompilerTrace)
+				{
+					CompileAction.ForceClFilter = true;
+					CompileAction.TimingFile = FileItem.GetItemByFileReference(FileReference.Combine(OutputDir, String.Format("{0}.timing", SourceFile.Location.GetFileName())));
+					GenerateParseTimingInfoAction(SourceFile, CompileAction.TimingFile, Graph);
+				}
+
 				if (CompileEnvironment.bGenerateDependenciesFile)
 				{
-					if (EnvVars.ToolChainVersion >= VersionNumber.Parse("14.27") && Target.WindowsPlatform.Compiler >= WindowsCompiler.VisualStudio2019)
+					if (EnvVars.ToolChainVersion >= VersionNumber.Parse("14.27") && Target.WindowsPlatform.Compiler >= WindowsCompiler.VisualStudio2019 && !CompileAction.ForceClFilter)
 					{
 						CompileAction.DependencyListFile = FileItem.GetItemByFileReference(FileReference.Combine(OutputDir, String.Format("{0}.json", SourceFile.Location.GetFileName())));
 					}
@@ -1358,12 +1365,6 @@ namespace UnrealBuildTool
 					}
 				}
 
-				if (Target.bPrintToolChainTimingInfo || Target.WindowsPlatform.bCompilerTrace)
-				{
-					CompileAction.TimingFile = FileItem.GetItemByFileReference(FileReference.Combine(OutputDir, String.Format("{0}.timing", SourceFile.Location.GetFileName())));
-					GenerateParseTimingInfoAction(SourceFile, CompileAction.TimingFile, Graph);
-				}
-
 				if (!ProjectFileGenerator.bGenerateProjectFiles)
 				{
 					CompileAction.WriteResponseFile(Graph);
