UnrealBuildTool: Cache BuildConfiguration.xml file locations

diff --git a/Engine/Source/Programs/UnrealBuildTool/ProjectFiles/ProjectFileGenerator.cs b/Engine/Source/Programs/UnrealBuildTool/ProjectFiles/ProjectFileGenerator.cs
--- a/Engine/Source/Programs/UnrealBuildTool/ProjectFiles/ProjectFileGenerator.cs
+++ b/Engine/Source/Programs/UnrealBuildTool/ProjectFiles/ProjectFileGenerator.cs
@@ -1151,8 +1151,7 @@ namespace UnrealBuildTool
 					Path.Combine("Programs", "UnrealBuildTool")
 				));
 
-			List<XmlConfig.InputFile> InputFiles = XmlConfig.FindInputFiles();
-			foreach (XmlConfig.InputFile InputFile in InputFiles)
+			foreach (XmlConfig.InputFile InputFile in XmlConfig.InputFiles)
 			{
 				EngineProject.AddAliasedFileToProject(
 						new AliasedFile(
diff --git a/Engine/Source/Programs/UnrealBuildTool/System/TargetMakefile.cs b/Engine/Source/Programs/UnrealBuildTool/System/TargetMakefile.cs
--- a/Engine/Source/Programs/UnrealBuildTool/System/TargetMakefile.cs
+++ b/Engine/Source/Programs/UnrealBuildTool/System/TargetMakefile.cs
@@ -384,8 +384,7 @@ namespace UnrealBuildTool
 				}
 
 				// Check to see if any BuildConfiguration files have changed since the last build
-				List<XmlConfig.InputFile> InputFiles = XmlConfig.FindInputFiles();
-				foreach (XmlConfig.InputFile InputFile in InputFiles)
+				foreach (XmlConfig.InputFile InputFile in XmlConfig.InputFiles)
 				{
 					FileInfo InputFileInfo = new FileInfo(InputFile.Location.FullName);
 					if (InputFileInfo.LastWriteTime > MakefileInfo.LastWriteTime)
diff --git a/Engine/Source/Programs/UnrealBuildTool/System/XmlConfig.cs b/Engine/Source/Programs/UnrealBuildTool/System/XmlConfig.cs
--- a/Engine/Source/Programs/UnrealBuildTool/System/XmlConfig.cs
+++ b/Engine/Source/Programs/UnrealBuildTool/System/XmlConfig.cs
@@ -98,16 +98,16 @@ namespace UnrealBuildTool
 				}
 
 				// Find all the input files
-				FileReference[] InputFiles = FindInputFiles().Select(x => x.Location).ToArray();
+				FileReference[] InputFileLocations = InputFiles.Select(x => x.Location).ToArray();
 
 				// Get the path to the schema
 				FileReference SchemaFile = GetSchemaLocation();
 
 				// Try to read the existing cache from disk
 				XmlConfigData? CachedValues;
-				if(IsCacheUpToDate(CacheFile, InputFiles) && FileReference.Exists(SchemaFile))
+				if(IsCacheUpToDate(CacheFile, InputFileLocations) && FileReference.Exists(SchemaFile))
 				{
-					if(XmlConfigData.TryRead(CacheFile, ConfigTypes, out CachedValues) && Enumerable.SequenceEqual(InputFiles, CachedValues.InputFiles))
+					if(XmlConfigData.TryRead(CacheFile, ConfigTypes, out CachedValues) && Enumerable.SequenceEqual(InputFileLocations, CachedValues.InputFiles))
 					{
 						Values = CachedValues;
 					}
@@ -130,7 +130,7 @@ namespace UnrealBuildTool
 					// Read all the XML files and validate them against the schema
 					Dictionary<Type, Dictionary<FieldInfo, XmlConfigData.ValueInfo>> TypeToValues =
 						new Dictionary<Type, Dictionary<FieldInfo, XmlConfigData.ValueInfo>>();
-					foreach(FileReference InputFile in InputFiles)
+					foreach(FileReference InputFile in InputFileLocations)
 					{
 						if(!TryReadFile(InputFile, CategoryToFields, TypeToValues, Schema))
 						{
@@ -142,7 +142,7 @@ namespace UnrealBuildTool
 					DirectoryReference.CreateDirectory(CacheFile.Directory);
 
 					// Create the new cache
-					Values = new XmlConfigData(InputFiles, TypeToValues.ToDictionary(
+					Values = new XmlConfigData(InputFileLocations, TypeToValues.ToDictionary(
 							x => x.Key, 
 							x => x.Value.Select(x => x.Value).ToArray()));
 					Values.Write(CacheFile);
@@ -217,73 +217,85 @@ namespace UnrealBuildTool
 			return FileReference.Combine(Unreal.EngineDirectory, "Saved", "UnrealBuildTool", "BuildConfiguration.Schema.xsd");
 		}
 
+		static InputFile[]? CachedInputFiles;
+
 		/// <summary>
 		/// Initialize the list of input files
 		/// </summary>
-		public static List<InputFile> FindInputFiles()
+		public static InputFile[] InputFiles
 		{
-			// Find all the input file locations
-			List<InputFile> InputFiles = new List<InputFile>();
-
-			// Skip all the config files under the Engine folder if it's an installed build
-			if(!Unreal.IsEngineInstalled())
-			{
- 				// Check for the config file under /Engine/Programs/NotForLicensees/UnrealBuildTool
- 				FileReference NotForLicenseesConfigLocation = FileReference.Combine(Unreal.EngineDirectory, "Restricted", "NotForLicensees", "Programs", "UnrealBuildTool", "BuildConfiguration.xml");
- 				if(FileReference.Exists(NotForLicenseesConfigLocation))
- 				{
- 					InputFiles.Add(new InputFile(NotForLicenseesConfigLocation, "NotForLicensees"));
- 				}
-                else
-                {
-	                Log.TraceLog($"No config file at {NotForLicenseesConfigLocation}");
-                }
-
-				// Check for the user config file under /Engine/Saved/UnrealBuildTool
-				FileReference UserConfigLocation = FileReference.Combine(Unreal.EngineDirectory, "Saved", "UnrealBuildTool", "BuildConfiguration.xml");
-				if(!FileReference.Exists(UserConfigLocation))
+			get
+			{
+				if (CachedInputFiles != null)
 				{
-					Log.TraceLog($"Creating default config file at {UserConfigLocation}");
-					CreateDefaultConfigFile(UserConfigLocation);
+					return CachedInputFiles;
 				}
-				InputFiles.Add(new InputFile(UserConfigLocation, "User"));
-			}
 
-			// Check for the global config file under AppData/Unreal Engine/UnrealBuildTool
-			string AppDataFolder = Environment.GetFolderPath(Environment.SpecialFolder.ApplicationData);
-			if(!String.IsNullOrEmpty(AppDataFolder))
-			{
-				FileReference AppDataConfigLocation = FileReference.Combine(new DirectoryReference(AppDataFolder), "Unreal Engine", "UnrealBuildTool", "BuildConfiguration.xml");
-				if(!FileReference.Exists(AppDataConfigLocation))
+				// Find all the input file locations
+				List<InputFile> InputFilesFound = new List<InputFile>(4);
+
+				// Skip all the config files under the Engine folder if it's an installed build
+				if (!Unreal.IsEngineInstalled())
 				{
-					Log.TraceLog($"Creating default config file at {AppDataConfigLocation}");
-					CreateDefaultConfigFile(AppDataConfigLocation);
+					// Check for the config file under /Engine/Programs/NotForLicensees/UnrealBuildTool
+					FileReference NotForLicenseesConfigLocation = FileReference.Combine(Unreal.EngineDirectory, "Restricted", "NotForLicensees", "Programs", "UnrealBuildTool", "BuildConfiguration.xml");
+					if (FileReference.Exists(NotForLicenseesConfigLocation))
+					{
+						InputFilesFound.Add(new InputFile(NotForLicenseesConfigLocation, "NotForLicensees"));
+					}
+					else
+					{
+						Log.TraceLog($"No config file at {NotForLicenseesConfigLocation}");
+					}
+
+					// Check for the user config file under /Engine/Saved/UnrealBuildTool
+					FileReference UserConfigLocation = FileReference.Combine(Unreal.EngineDirectory, "Saved", "UnrealBuildTool", "BuildConfiguration.xml");
+					if (!FileReference.Exists(UserConfigLocation))
+					{
+						Log.TraceLog($"Creating default config file at {UserConfigLocation}");
+						CreateDefaultConfigFile(UserConfigLocation);
+					}
+					InputFilesFound.Add(new InputFile(UserConfigLocation, "User"));
 				}
-				InputFiles.Add(new InputFile(AppDataConfigLocation, "Global (AppData)"));
-			}
 
-			// Check for the global config file under My Documents/Unreal Engine/UnrealBuildTool
-			string PersonalFolder = Environment.GetFolderPath(Environment.SpecialFolder.Personal);
-			if(!String.IsNullOrEmpty(PersonalFolder))
-			{
-				FileReference PersonalConfigLocation = FileReference.Combine(new DirectoryReference(PersonalFolder), "Unreal Engine", "UnrealBuildTool", "BuildConfiguration.xml");
-				if(FileReference.Exists(PersonalConfigLocation))
+				// Check for the global config file under AppData/Unreal Engine/UnrealBuildTool
+				string AppDataFolder = Environment.GetFolderPath(Environment.SpecialFolder.ApplicationData);
+				if (!String.IsNullOrEmpty(AppDataFolder))
 				{
-					InputFiles.Add(new InputFile(PersonalConfigLocation, "Global (Documents)"));
+					FileReference AppDataConfigLocation = FileReference.Combine(new DirectoryReference(AppDataFolder), "Unreal Engine", "UnrealBuildTool", "BuildConfiguration.xml");
+					if (!FileReference.Exists(AppDataConfigLocation))
+					{
+						Log.TraceLog($"Creating default config file at {AppDataConfigLocation}");
+						CreateDefaultConfigFile(AppDataConfigLocation);
+					}
+					InputFilesFound.Add(new InputFile(AppDataConfigLocation, "Global (AppData)"));
 				}
-				else
+
+				// Check for the global config file under My Documents/Unreal Engine/UnrealBuildTool
+				string PersonalFolder = Environment.GetFolderPath(Environment.SpecialFolder.Personal);
+				if (!String.IsNullOrEmpty(PersonalFolder))
 				{
-					Log.TraceLog($"No config file at {PersonalConfigLocation}");
+					FileReference PersonalConfigLocation = FileReference.Combine(new DirectoryReference(PersonalFolder), "Unreal Engine", "UnrealBuildTool", "BuildConfiguration.xml");
+					if (FileReference.Exists(PersonalConfigLocation))
+					{
+						InputFilesFound.Add(new InputFile(PersonalConfigLocation, "Global (Documents)"));
+					}
+					else
+					{
+						Log.TraceLog($"No config file at {PersonalConfigLocation}");
+					}
 				}
-			}
 
-			Log.TraceLog("Configuration will be read from:");
-			foreach (InputFile InputFile in InputFiles)
-			{
-				Log.TraceLog($"  {InputFile.Location.FullName}");
+				CachedInputFiles = InputFilesFound.ToArray();
+
+				Log.TraceLog("Configuration will be read from:");
+				foreach (InputFile InputFile in InputFiles)
+				{
+					Log.TraceLog($"  {InputFile.Location.FullName}");
+				}
+
+				return CachedInputFiles;
 			}
-			
-			return InputFiles;
 		}
 
 		/// <summary>
