Link with static libc++ on Android

diff --git a/Engine/Source/Programs/UnrealBuildTool/Platform/Android/AndroidToolChain.cs b/Engine/Source/Programs/UnrealBuildTool/Platform/Android/AndroidToolChain.cs
--- a/Engine/Source/Programs/UnrealBuildTool/Platform/Android/AndroidToolChain.cs
+++ b/Engine/Source/Programs/UnrealBuildTool/Platform/Android/AndroidToolChain.cs
@@ -863,6 +863,7 @@ namespace UnrealBuildTool
 			string Result = "";
 
 			//Result += " -nostdlib";
+			Result += " -static-libstdc++";
 			Result += " -no-canonical-prefixes";
 			Result += " -shared";
 			Result += " -Wl,-Bsymbolic";
