UnrealBuildTool: track history for AdditionalPluginDirectories

External plugins will be built in project_dir/Intermediate/External/...

diff --git a/Engine/Source/Programs/UnrealBuildTool/Configuration/UEBuildTarget.cs b/Engine/Source/Programs/UnrealBuildTool/Configuration/UEBuildTarget.cs
--- a/Engine/Source/Programs/UnrealBuildTool/Configuration/UEBuildTarget.cs
+++ b/Engine/Source/Programs/UnrealBuildTool/Configuration/UEBuildTarget.cs
@@ -1089,7 +1089,7 @@ namespace UnrealBuildTool
 			bDeployAfterCompile = InRules.bDeployAfterCompile && !InRules.bDisableLinking && InDescriptor.SpecificFilesToCompile.Count == 0 && InDescriptor.OnlyModuleNames.Count == 0;
 
 			// now that we have the platform, we can set the intermediate path to include the platform/architecture name
-			PlatformIntermediateFolder = GetPlatformIntermediateFolder(Platform, Architecture);
+			PlatformIntermediateFolder = GetPlatformIntermediateFolder(Platform, Architecture, false);
 
 			TargetRulesFile = InRules.File;
 
@@ -1151,11 +1151,12 @@ namespace UnrealBuildTool
 		/// </summary>
 		/// <param name="Platform">Platform to get the folder for</param>
 		/// <param name="Architecture">Architecture to get the folder for</param>
+		/// <param name="External">Insert External after Intermediate - for out-of-tree plugins</param>
 		/// <returns>The output directory for intermediates</returns>
-		public static string GetPlatformIntermediateFolder(UnrealTargetPlatform Platform, string Architecture)
+		public static string GetPlatformIntermediateFolder(UnrealTargetPlatform Platform, string Architecture, bool External)
 		{
 			// now that we have the platform, we can set the intermediate path to include the platform/architecture name
-			return Path.Combine("Intermediate", "Build", Platform.ToString(), UEBuildPlatform.GetBuildPlatform(Platform).GetFolderNameForArchitecture(Architecture));
+			return Path.Combine("Intermediate", External ? "External" : String.Empty, "Build", Platform.ToString(), UEBuildPlatform.GetBuildPlatform(Platform).GetFolderNameForArchitecture(Architecture));
 		}
 
 		/// <summary>
@@ -2991,7 +2992,7 @@ namespace UnrealBuildTool
 		{
 			// Get the root output directory and base name (target name/app name) for this binary
 			DirectoryReference BaseOutputDirectory;
-			if(bUseSharedBuildEnvironment)
+			if(bUseSharedBuildEnvironment && ModuleRules.Plugin?.Type != PluginType.External)
 			{
 				BaseOutputDirectory = ModuleRules.Context.DefaultOutputBaseDir;
 			}
@@ -3020,7 +3021,7 @@ namespace UnrealBuildTool
 			}
 
 			// Get the output and intermediate directories for this module
-			DirectoryReference IntermediateDirectory = DirectoryReference.Combine(BaseOutputDirectory, PlatformIntermediateFolder, AppName, ModuleConfiguration.ToString());
+			DirectoryReference IntermediateDirectory = DirectoryReference.Combine(BaseOutputDirectory, GetPlatformIntermediateFolder(Platform, Architecture, ModuleRules.Plugin?.Type == PluginType.External), AppName, ModuleConfiguration.ToString());
 
 			// Append a subdirectory if the module rules specifies one
 			if (!String.IsNullOrEmpty(ModuleRules.BinariesSubFolder))
diff --git a/Engine/Source/Programs/UnrealBuildTool/Modes/CleanMode.cs b/Engine/Source/Programs/UnrealBuildTool/Modes/CleanMode.cs
--- a/Engine/Source/Programs/UnrealBuildTool/Modes/CleanMode.cs
+++ b/Engine/Source/Programs/UnrealBuildTool/Modes/CleanMode.cs
@@ -171,13 +171,13 @@ namespace UnrealBuildTool
 				{
 					foreach (string NamePrefix in NamePrefixes)
 					{
-						DirectoryReference GeneratedCodeDir = DirectoryReference.Combine(BaseDir, UEBuildTarget.GetPlatformIntermediateFolder(Target.Platform, Target.Architecture), NamePrefix, "Inc");
+						DirectoryReference GeneratedCodeDir = DirectoryReference.Combine(BaseDir, UEBuildTarget.GetPlatformIntermediateFolder(Target.Platform, Target.Architecture, false), NamePrefix, "Inc");
 						if (DirectoryReference.Exists(GeneratedCodeDir))
 						{
 							DirectoriesToDelete.Add(GeneratedCodeDir);
 						}
 
-						DirectoryReference IntermediateDir = DirectoryReference.Combine(BaseDir, UEBuildTarget.GetPlatformIntermediateFolder(Target.Platform, Target.Architecture), NamePrefix, Target.Configuration.ToString());
+						DirectoryReference IntermediateDir = DirectoryReference.Combine(BaseDir, UEBuildTarget.GetPlatformIntermediateFolder(Target.Platform, Target.Architecture, false), NamePrefix, Target.Configuration.ToString());
 						if (DirectoryReference.Exists(IntermediateDir))
 						{
 							DirectoriesToDelete.Add(IntermediateDir);
@@ -185,6 +185,8 @@ namespace UnrealBuildTool
 					}
 				}
 
+				// todo: handle external plugin intermediates, written to the Project's Intermediate/External directory
+
 				// List of additional files and directories to clean, specified by the target platform
 				List<FileReference> AdditionalFilesToDelete = new List<FileReference>();
 				List<DirectoryReference> AdditionalDirectoriesToDelete = new List<DirectoryReference>();
diff --git a/Engine/Source/Programs/UnrealBuildTool/ProjectFiles/CMake/CMakefileGenerator.cs b/Engine/Source/Programs/UnrealBuildTool/ProjectFiles/CMake/CMakefileGenerator.cs
--- a/Engine/Source/Programs/UnrealBuildTool/ProjectFiles/CMake/CMakefileGenerator.cs
+++ b/Engine/Source/Programs/UnrealBuildTool/ProjectFiles/CMake/CMakefileGenerator.cs
@@ -395,7 +395,7 @@ namespace UnrealBuildTool
 			}
 
 			// Create SourceFiles, HeaderFiles, and ConfigFiles sections.
-			List<FileReference> AllModuleFiles = DiscoverModules(FindGameProjects());
+			List<FileReference> AllModuleFiles = DiscoverModules(FindGameProjects(), null);
 			foreach (FileReference CurModuleFile in AllModuleFiles)
 			{
 				List<FileReference> FoundFiles = SourceFileSearch.FindModuleSourceFiles(CurModuleFile);
diff --git a/Engine/Source/Programs/UnrealBuildTool/ProjectFiles/ProjectFileGenerator.cs b/Engine/Source/Programs/UnrealBuildTool/ProjectFiles/ProjectFileGenerator.cs
--- a/Engine/Source/Programs/UnrealBuildTool/ProjectFiles/ProjectFileGenerator.cs
+++ b/Engine/Source/Programs/UnrealBuildTool/ProjectFiles/ProjectFileGenerator.cs
@@ -856,9 +856,20 @@ namespace UnrealBuildTool
 			// Remove any game projects that don't have a target
 			AllGameProjects.RemoveAll(x => !AllTargetFiles.Any(y => y.IsUnderDirectory(x.Directory)));
 
+			Dictionary<FileReference, List<DirectoryReference>> AdditionalSearchPaths = new Dictionary<FileReference, List<DirectoryReference>>();
+
+			foreach (FileReference GameProject in AllGameProjects)
+			{
+				ProjectDescriptor Project = ProjectDescriptor.FromFile(GameProject);
+				if (Project.AdditionalPluginDirectories.Count > 0)
+				{
+					AdditionalSearchPaths[GameProject] = Project.AdditionalPluginDirectories;
+				}
+			}
+
 			// Find all of the module files.  This will filter out any modules or targets that don't belong to platforms
 			// we're generating project files for.
-			List<FileReference> AllModuleFiles = DiscoverModules(AllGameProjects);
+			List<FileReference> AllModuleFiles = DiscoverModules(AllGameProjects, AdditionalSearchPaths.Values.SelectMany(x => x).ToList());
 
 			ProjectFile? EngineProject;
 			List<ProjectFile> GameProjects;
@@ -1031,7 +1042,7 @@ namespace UnrealBuildTool
 			}
 
 			// Setup "stub" projects for all modules
-			AddProjectsForAllModules(AllGameProjects, ProgramProjects, ModProjects, AllModuleFiles, bGatherThirdPartySource);
+			AddProjectsForAllModules(AllGameProjects, ProgramProjects, ModProjects, AllModuleFiles, AdditionalSearchPaths, bGatherThirdPartySource);
 
 			{
 				if (bIncludeDotNetPrograms)
@@ -1473,10 +1484,10 @@ namespace UnrealBuildTool
 		/// Finds all module files (filtering by platform)
 		/// </summary>
 		/// <returns>Filtered list of module files</returns>
-		protected List<FileReference> DiscoverModules(List<FileReference> AllGameProjects)
+		protected List<FileReference> DiscoverModules(List<FileReference> AllGameProjects, List<DirectoryReference>? AdditionalSearchPaths)
 		{
 			// Locate all modules (*.Build.cs files)
-			return Rules.FindAllRulesSourceFiles(Rules.RulesFileType.Module, GameFolders: AllGameProjects.Select(x => x.Directory).ToList(), ForeignPlugins: null, AdditionalSearchPaths: null);
+			return Rules.FindAllRulesSourceFiles(Rules.RulesFileType.Module, GameFolders: AllGameProjects.Select(x => x.Directory).ToList(), ForeignPlugins: null, AdditionalSearchPaths: AdditionalSearchPaths);
 		}
 
 		/// <summary>
@@ -2076,8 +2087,9 @@ namespace UnrealBuildTool
 		/// <param name="ProgramProjects">All program projects</param>
 		/// <param name="ModProjects">All mod projects</param>
 		/// <param name="AllModuleFiles">List of *.Build.cs files for all engine programs and games</param>
+		/// <param name="AdditionalSearchPaths"></param>
 		/// <param name="bGatherThirdPartySource">True to gather source code from third party projects too</param>
-		protected void AddProjectsForAllModules(List<FileReference> AllGames, Dictionary<FileReference, ProjectFile> ProgramProjects, List<ProjectFile> ModProjects, List<FileReference> AllModuleFiles, bool bGatherThirdPartySource)
+		protected void AddProjectsForAllModules(List<FileReference> AllGames, Dictionary<FileReference, ProjectFile> ProgramProjects, List<ProjectFile> ModProjects, List<FileReference> AllModuleFiles, Dictionary<FileReference, List<DirectoryReference>> AdditionalSearchPaths, bool bGatherThirdPartySource)
 		{
 			HashSet<ProjectFile> ProjectsWithPlugins = new HashSet<ProjectFile>();
 			foreach (FileReference CurModuleFile in AllModuleFiles)
@@ -2107,7 +2119,7 @@ namespace UnrealBuildTool
 				if (WantProjectFileForModule)
 				{
 					DirectoryReference BaseFolder;
-					ProjectFile ProjectFile = FindProjectForModule(CurModuleFile, AllGames, ProgramProjects, ModProjects, out BaseFolder);
+					ProjectFile ProjectFile = FindProjectForModule(CurModuleFile, AllGames, ProgramProjects, ModProjects, AdditionalSearchPaths, out BaseFolder);
 
 					// Update our module map
 					ModuleToProjectFileMap[ModuleName] = ProjectFile;
@@ -2188,7 +2200,7 @@ namespace UnrealBuildTool
 			return FindOrAddProject(ProjectFileName, InBaseFolder, IncludeInGeneratedProjects: true, bAlreadyExisted: out bProjectAlreadyExisted);
 		}
 
-		private ProjectFile FindProjectForModule(FileReference CurModuleFile, List<FileReference> AllGames, Dictionary<FileReference, ProjectFile> ProgramProjects, List<ProjectFile> ModProjects, out DirectoryReference BaseFolder)
+		private ProjectFile FindProjectForModule(FileReference CurModuleFile, List<FileReference> AllGames, Dictionary<FileReference, ProjectFile> ProgramProjects, List<ProjectFile> ModProjects, Dictionary<FileReference, List<DirectoryReference>> AdditionalSearchPaths, out DirectoryReference BaseFolder)
 		{
 			// Starting at the base directory of the module find a project which has the same directory as base, walking up the directory hierarchy until a match is found
 
@@ -2243,6 +2255,18 @@ namespace UnrealBuildTool
 				Path = Path.ParentDirectory!;
 			}
 
+			foreach ((FileReference ProjectPath, List<DirectoryReference> AdditionalPluginDirectories) in AdditionalSearchPaths)
+			{
+				foreach (DirectoryReference AdditionalPluginDirectory in AdditionalPluginDirectories)
+				{
+					if (CurModuleFile.IsUnderDirectory(AdditionalPluginDirectory))
+					{
+						BaseFolder = ProjectPath.Directory;
+						return FindOrAddProjectHelper(ProjectPath.GetFileNameWithoutExtension(), BaseFolder);
+					}
+				}
+			}
+
 			throw new BuildException("Found a module file (" + CurModuleFile + ") that did not exist within any of the known game folders or other source locations");
 		}
 
diff --git a/Engine/Source/Programs/UnrealBuildTool/ProjectFiles/QMake/QMakefileGenerator.cs b/Engine/Source/Programs/UnrealBuildTool/ProjectFiles/QMake/QMakefileGenerator.cs
--- a/Engine/Source/Programs/UnrealBuildTool/ProjectFiles/QMake/QMakefileGenerator.cs
+++ b/Engine/Source/Programs/UnrealBuildTool/ProjectFiles/QMake/QMakefileGenerator.cs
@@ -265,7 +265,7 @@ namespace UnrealBuildTool
 			);
 
 			// Create SourceFiles, HeaderFiles, and ConfigFiles sections.
-			List<FileReference> AllModuleFiles = DiscoverModules(FindGameProjects());
+			List<FileReference> AllModuleFiles = DiscoverModules(FindGameProjects(), null);
 			foreach (FileReference CurModuleFile in AllModuleFiles)
 			{
 				List<FileReference> FoundFiles = SourceFileSearch.FindModuleSourceFiles(CurModuleFile);
diff --git a/Engine/Source/Programs/UnrealBuildTool/System/ActionHistory.cs b/Engine/Source/Programs/UnrealBuildTool/System/ActionHistory.cs
--- a/Engine/Source/Programs/UnrealBuildTool/System/ActionHistory.cs
+++ b/Engine/Source/Programs/UnrealBuildTool/System/ActionHistory.cs
@@ -197,7 +197,7 @@ namespace UnrealBuildTool
 				AppName = UEBuildTarget.GetAppNameForTargetType(TargetType);
 			}
 
-			return FileReference.Combine(Unreal.EngineDirectory, UEBuildTarget.GetPlatformIntermediateFolder(Platform, Architecture), AppName, "ActionHistory.bin");
+			return FileReference.Combine(Unreal.EngineDirectory, UEBuildTarget.GetPlatformIntermediateFolder(Platform, Architecture, false), AppName, "ActionHistory.bin");
 		}
 
 		/// <summary>
@@ -210,7 +210,7 @@ namespace UnrealBuildTool
 		/// <returns>Path to the project action history</returns>
 		public static FileReference GetProjectLocation(FileReference ProjectFile, string TargetName, UnrealTargetPlatform Platform, string Architecture)
 		{
-			return FileReference.Combine(ProjectFile.Directory, UEBuildTarget.GetPlatformIntermediateFolder(Platform, Architecture), TargetName, "ActionHistory.dat");
+			return FileReference.Combine(ProjectFile.Directory, UEBuildTarget.GetPlatformIntermediateFolder(Platform, Architecture, false), TargetName, "ActionHistory.dat");
 		}
 
 		/// <summary>
diff --git a/Engine/Source/Programs/UnrealBuildTool/System/CppDependencyCache.cs b/Engine/Source/Programs/UnrealBuildTool/System/CppDependencyCache.cs
--- a/Engine/Source/Programs/UnrealBuildTool/System/CppDependencyCache.cs
+++ b/Engine/Source/Programs/UnrealBuildTool/System/CppDependencyCache.cs
@@ -323,13 +323,13 @@ namespace UnrealBuildTool
 					AppName = UEBuildTarget.GetAppNameForTargetType(TargetType);
 				}
 
-				FileReference EngineCacheLocation = FileReference.Combine(Unreal.EngineDirectory, UEBuildTarget.GetPlatformIntermediateFolder(Platform, Architecture), AppName, Configuration.ToString(), "DependencyCache.bin");
+				FileReference EngineCacheLocation = FileReference.Combine(Unreal.EngineDirectory, UEBuildTarget.GetPlatformIntermediateFolder(Platform, Architecture, false), AppName, Configuration.ToString(), "DependencyCache.bin");
 				FindOrAddPartition(EngineCacheLocation, Unreal.EngineDirectory);
 			}
 
 			if (ProjectFile != null)
 			{
-				FileReference ProjectCacheLocation = FileReference.Combine(ProjectFile.Directory, UEBuildTarget.GetPlatformIntermediateFolder(Platform, Architecture), TargetName, Configuration.ToString(), "DependencyCache.bin");
+				FileReference ProjectCacheLocation = FileReference.Combine(ProjectFile.Directory, UEBuildTarget.GetPlatformIntermediateFolder(Platform, Architecture, false), TargetName, Configuration.ToString(), "DependencyCache.bin");
 				FindOrAddPartition(ProjectCacheLocation, ProjectFile.Directory);
 			}
 		}
diff --git a/Engine/Source/Programs/UnrealBuildTool/System/HotReload.cs b/Engine/Source/Programs/UnrealBuildTool/System/HotReload.cs
--- a/Engine/Source/Programs/UnrealBuildTool/System/HotReload.cs
+++ b/Engine/Source/Programs/UnrealBuildTool/System/HotReload.cs
@@ -107,7 +107,7 @@ namespace UnrealBuildTool
 		public static FileReference GetLocation(FileReference? ProjectFile, string TargetName, UnrealTargetPlatform Platform, UnrealTargetConfiguration Configuration, string Architecture)
 		{
 			DirectoryReference BaseDir = DirectoryReference.FromFile(ProjectFile) ?? Unreal.EngineDirectory;
-			return FileReference.Combine(BaseDir, UEBuildTarget.GetPlatformIntermediateFolder(Platform, Architecture), TargetName, Configuration.ToString(), "HotReload.state");
+			return FileReference.Combine(BaseDir, UEBuildTarget.GetPlatformIntermediateFolder(Platform, Architecture, false), TargetName, Configuration.ToString(), "HotReload.state");
 		}
 
 		/// <summary>
diff --git a/Engine/Source/Programs/UnrealBuildTool/System/TargetMakefile.cs b/Engine/Source/Programs/UnrealBuildTool/System/TargetMakefile.cs
--- a/Engine/Source/Programs/UnrealBuildTool/System/TargetMakefile.cs
+++ b/Engine/Source/Programs/UnrealBuildTool/System/TargetMakefile.cs
@@ -684,7 +684,7 @@ namespace UnrealBuildTool
 		public static FileReference GetLocation(FileReference? ProjectFile, string TargetName, UnrealTargetPlatform Platform, string Architecture, UnrealTargetConfiguration Configuration)
 		{
 			DirectoryReference BaseDirectory = DirectoryReference.FromFile(ProjectFile) ?? Unreal.EngineDirectory;
-			return FileReference.Combine(BaseDirectory, UEBuildTarget.GetPlatformIntermediateFolder(Platform, Architecture), TargetName, Configuration.ToString(), "Makefile.bin");
+			return FileReference.Combine(BaseDirectory, UEBuildTarget.GetPlatformIntermediateFolder(Platform, Architecture, false), TargetName, Configuration.ToString(), "Makefile.bin");
 		}
 
 		/// <inheritdoc/>
