UBT: Allow overriding the target name via command line options, to avoid clobbering intermediates from builds with different options.

diff --git a/Engine/Source/Programs/UnrealBuildTool/Configuration/TargetRules.cs b/Engine/Source/Programs/UnrealBuildTool/Configuration/TargetRules.cs
--- a/Engine/Source/Programs/UnrealBuildTool/Configuration/TargetRules.cs
+++ b/Engine/Source/Programs/UnrealBuildTool/Configuration/TargetRules.cs
@@ -207,6 +207,10 @@ namespace UnrealBuildTool
 		/// </summary>
 		public bool IsNameOverriden() { return !String.IsNullOrEmpty(NameOverride); }
 
+		/// <summary>
+		/// Override the name used for this target
+		/// </summary>
+		[CommandLine("-TargetNameOverride=")]
 		private string? NameOverride;
 
 		private readonly string DefaultName;
