UnrealBuildTool: Makefile records external dependency on platform/group-specific target rule source file

diff --git a/Engine/Source/Programs/UnrealBuildTool/Configuration/TargetRules.cs b/Engine/Source/Programs/UnrealBuildTool/Configuration/TargetRules.cs
--- a/Engine/Source/Programs/UnrealBuildTool/Configuration/TargetRules.cs
+++ b/Engine/Source/Programs/UnrealBuildTool/Configuration/TargetRules.cs
@@ -212,10 +212,15 @@ namespace UnrealBuildTool
 		private readonly string DefaultName;
 
 		/// <summary>
-		/// File containing this target
+		/// File containing the general type for this target (not including platform/group)
 		/// </summary>
 		internal FileReference? File { get; set; }
 
+		/// <summary>
+		/// File containing the platform/group-specific type for this target
+		/// </summary>
+		internal FileReference? TargetSourceFile { get; set; }
+
 		/// <summary>
 		/// Platform that this target is being built for.
 		/// </summary>
@@ -1935,6 +1940,11 @@ namespace UnrealBuildTool
 			get { return Inner.File!; }
 		}
 
+		internal FileReference TargetSourceFile 
+		{
+			get { return Inner.TargetSourceFile!; }
+		}
+
 		public UnrealTargetPlatform Platform
 		{
 			get { return Inner.Platform; }
diff --git a/Engine/Source/Programs/UnrealBuildTool/Configuration/UEBuildTarget.cs b/Engine/Source/Programs/UnrealBuildTool/Configuration/UEBuildTarget.cs
--- a/Engine/Source/Programs/UnrealBuildTool/Configuration/UEBuildTarget.cs
+++ b/Engine/Source/Programs/UnrealBuildTool/Configuration/UEBuildTarget.cs
@@ -2085,6 +2085,7 @@ namespace UnrealBuildTool
 
 			// Add all the input files to the predicate store
 			Makefile.ExternalDependencies.Add(FileItem.GetItemByFileReference(TargetRulesFile));
+			Makefile.ExternalDependencies.Add(FileItem.GetItemByFileReference(Rules.TargetSourceFile));
 			foreach(UEBuildModule Module in Modules.Values)
 			{
 				Makefile.ExternalDependencies.Add(FileItem.GetItemByFileReference(Module.RulesFile));
diff --git a/Engine/Source/Programs/UnrealBuildTool/System/RulesAssembly.cs b/Engine/Source/Programs/UnrealBuildTool/System/RulesAssembly.cs
--- a/Engine/Source/Programs/UnrealBuildTool/System/RulesAssembly.cs
+++ b/Engine/Source/Programs/UnrealBuildTool/System/RulesAssembly.cs
@@ -566,9 +566,12 @@ namespace UnrealBuildTool
 				Rules.DefaultBuildSettings = DefaultBuildSettings.Value;
 			}
 
-			// Return the base target file name to the caller. This affects where the resulting build product is created so the platform/group is not desired in this case.
+			// The base target file name: this affects where the resulting build product is created so the platform/group is not desired in this case.
 			Rules.File = TargetNameToTargetFile[TargetInfo.Name];
 
+			// The platform/group-specific target file name
+			Rules.TargetSourceFile = TargetNameToTargetFile.TryGetValue(PlatformRulesName, out FileReference? PlatformTargetFile) ? PlatformTargetFile : Rules.File;
+
 			// Find the constructor
 			ConstructorInfo? Constructor = RulesType.GetConstructor(new Type[] { typeof(TargetInfo) });
 			if(Constructor == null)
