Add bPreferThinLTO target option so that targets can opt for ThinLTO on supported platform via the .Target.cs file, without needing to pass additional UBT command line switches via a build script.

diff --git a/Engine/Source/Programs/UnrealBuildTool/Configuration/TargetRules.cs b/Engine/Source/Programs/UnrealBuildTool/Configuration/TargetRules.cs
--- a/Engine/Source/Programs/UnrealBuildTool/Configuration/TargetRules.cs
+++ b/Engine/Source/Programs/UnrealBuildTool/Configuration/TargetRules.cs
@@ -1129,6 +1129,14 @@ namespace UnrealBuildTool
 		[XmlConfigFile(Category = "BuildConfiguration")]
 		public bool bAllowLTCG = false;
 
+		/// <summary>
+		/// When Link Time Code Generation (LTCG) is enabled, whether to 
+		/// prefer using the lighter weight version on supported platforms.
+		/// </summary>
+		[CommandLine("-ThinLTO")]
+		[XmlConfigFile(Category = "BuildConfiguration")]
+		public bool bPreferThinLTO = false;
+
 		/// <summary>
 		/// Whether to enable Profile Guided Optimization (PGO) instrumentation in this build.
 		/// </summary>
@@ -2561,6 +2569,12 @@ namespace UnrealBuildTool
 		{
 			get { return Inner.bAllowLTCG; }
 		}
+
+		public bool bPreferThinLTO
+		{
+			get { return Inner.bPreferThinLTO; }
+		}
+
 		public bool bPGOProfile
 		{
 			get { return Inner.bPGOProfile; }
diff --git a/Engine/Source/Programs/UnrealBuildTool/Platform/Linux/UEBuildLinux.cs b/Engine/Source/Programs/UnrealBuildTool/Platform/Linux/UEBuildLinux.cs
--- a/Engine/Source/Programs/UnrealBuildTool/Platform/Linux/UEBuildLinux.cs
+++ b/Engine/Source/Programs/UnrealBuildTool/Platform/Linux/UEBuildLinux.cs
@@ -68,12 +68,6 @@ namespace UnrealBuildTool
 		[XmlConfigFile(Category = "BuildConfiguration", Name = "bEnableMemorySanitizer")]
 		public bool bEnableMemorySanitizer = false;
 
-		/// <summary>
-		/// Enables "thin" LTO
-		/// </summary>
-		[CommandLine("-ThinLTO")]
-		public bool bEnableThinLTO = false;
-
 		/// <summary>
 		/// Whether or not to preserve the portable symbol file produced by dump_syms
 		/// </summary>
@@ -128,11 +122,6 @@ namespace UnrealBuildTool
 			get { return Inner.bEnableMemorySanitizer; }
 		}
 
-		public bool bEnableThinLTO
-		{
-			get { return Inner.bEnableThinLTO; }
-		}
-
 		#pragma warning restore CS1591
 		#endregion
 	}
@@ -202,11 +191,6 @@ namespace UnrealBuildTool
 
 		public override void ValidateTarget(TargetRules Target)
 		{
-			if(Target.LinuxPlatform.bEnableThinLTO)
-			{
-				Target.bAllowLTCG = true;
-			}
-
 			if (!Target.IsNameOverriden())
 			{
 				string? SanitizerSuffix = null;
@@ -575,7 +559,7 @@ namespace UnrealBuildTool
 					throw new BuildException("Memory Sanitizer (MSan) unsupported for non-monolithic builds");
 				}
 			}
-			if (Target.LinuxPlatform.bEnableThinLTO)
+			if (Target.bAllowLTCG && Target.bPreferThinLTO)
 			{
 				Options |= LinuxToolChainOptions.EnableThinLTO;
 			}
