Fix duplicate 'GetISPCObjectFileFormat' definition caused by integration

diff --git a/Engine/Source/Programs/UnrealBuildTool/ToolChain/UEToolChain.cs b/Engine/Source/Programs/UnrealBuildTool/ToolChain/UEToolChain.cs
--- a/Engine/Source/Programs/UnrealBuildTool/ToolChain/UEToolChain.cs
+++ b/Engine/Source/Programs/UnrealBuildTool/ToolChain/UEToolChain.cs
@@ -380,33 +380,6 @@ namespace UnrealBuildTool
 			return Format;
 		}
 
-		/// <summary>
-		/// Get object file format for ISPC.
-		/// </summary>
-		/// <param name="Platform">Which OS build platform is running on.</param>
-		/// <returns>Object file suffix</returns>
-		public virtual string GetISPCObjectFileFormat(UnrealTargetPlatform Platform)
-		{
-			string Format = "";
-
-			if (UEBuildPlatform.IsPlatformInGroup(Platform, UnrealPlatformGroup.Windows))
-			{
-				Format += "obj";
-			}
-			else if (UEBuildPlatform.IsPlatformInGroup(Platform, UnrealPlatformGroup.Unix) ||
-					Platform == UnrealTargetPlatform.Mac ||
-					Platform == UnrealTargetPlatform.Android)
-			{
-				Format += "obj";
-			}
-			else
-			{
-				Log.TraceWarning("Unsupported ISPC platform target!");
-			}
-
-			return Format;
-		}
-
 		/// <summary>
 		/// Get object file suffix for ISPC.
 		/// </summary>
