Fixed UHT warnings about returning null from a function that expects a non-nullable string.

diff --git a/Engine/Source/Programs/UnrealBuildTool/ToolChain/UEToolChain.cs b/Engine/Source/Programs/UnrealBuildTool/ToolChain/UEToolChain.cs
--- a/Engine/Source/Programs/UnrealBuildTool/ToolChain/UEToolChain.cs
+++ b/Engine/Source/Programs/UnrealBuildTool/ToolChain/UEToolChain.cs
@@ -288,7 +288,7 @@ namespace UnrealBuildTool
 		/// </summary>
 		/// <param name="Platform">Which OS platform to target.</param>
 		/// <returns>CPU string passed to ISPC compiler</returns>
-		public virtual string GetISPCCpuTarget(UnrealTargetPlatform Platform)
+		public virtual string? GetISPCCpuTarget(UnrealTargetPlatform Platform)
 		{
 			return null;  // no specific CPU selected
 		}
@@ -465,7 +465,7 @@ namespace UnrealBuildTool
 				Arguments.Add(String.Format("--target={0}", TargetString));
 				Arguments.Add(String.Format("--emit-{0}", GetISPCObjectFileFormat(CompileEnvironment.Platform)));
 				
-				string CpuTarget = GetISPCCpuTarget(CompileEnvironment.Platform);
+				string? CpuTarget = GetISPCCpuTarget(CompileEnvironment.Platform);
 				if (!String.IsNullOrEmpty(CpuTarget))
                 {
 					Arguments.Add(String.Format("--cpu={0}", CpuTarget));
