UBT: Fix -NoEngineChanges blocking builds when switching between targets. Metadata file for engine binaries was invalidating the action causing *.modules files to be generated.

diff --git a/Engine/Source/Programs/UnrealBuildTool/Configuration/UEBuildTarget.cs b/Engine/Source/Programs/UnrealBuildTool/Configuration/UEBuildTarget.cs
--- a/Engine/Source/Programs/UnrealBuildTool/Configuration/UEBuildTarget.cs
+++ b/Engine/Source/Programs/UnrealBuildTool/Configuration/UEBuildTarget.cs
@@ -1991,10 +1991,10 @@ namespace UnrealBuildTool
 						}
 					}
 
-					BinaryFormatterUtils.SaveIfDifferent(EngineInfoFile, EngineInfo);
+					WriteMetadataFile(Makefile, EngineInfoFile, EngineInfo);
 					Makefile.OutputItems.AddRange(WriteEngineMetadata.ProducedItems);
 
-					BinaryFormatterUtils.SaveIfDifferent(TargetInfoFile, TargetInfo);
+					WriteMetadataFile(Makefile, EngineInfoFile, EngineInfo);
 					Makefile.OutputItems.AddRange(WriteTargetMetadata.ProducedItems);
 				}
 
@@ -2139,11 +2139,22 @@ namespace UnrealBuildTool
 			WriteMetadataAction.WorkingDirectory = UnrealBuildTool.EngineSourceDirectory;
 			WriteMetadataAction.StatusDescription = StatusDescription;
 			WriteMetadataAction.bCanExecuteRemotely = false;
-			WriteMetadataAction.PrerequisiteItems.Add(FileItem.GetItemByFileReference(InfoFile));
 			WriteMetadataAction.bUseActionHistory = false; // Different files for each target; do not want to invalidate based on this.
 			return WriteMetadataAction;
 		}
 
+		static void WriteMetadataFile(TargetMakefile Makefile, FileReference InfoFile, WriteMetadataTargetInfo TargetInfo)
+		{
+			BinaryFormatterUtils.Save(InfoFile, TargetInfo);
+
+			// Note: Deliberately do not add a prerequisite from the action onto the info file. It is always beneath the project intermediate directory
+			// (even for the engine) and will cause -NoEngineChanges to block the build when updated. The behavior we want is just to depend on the engine
+			// DLL timestamps. Instead, we add a makefile dependency on it, causing it to be regenerated if missing.
+			FileItem InfoFileItem = FileItem.GetItemByFileReference(InfoFile);
+			InfoFileItem.ResetCachedInfo();
+			Makefile.InternalDependencies.Add(InfoFileItem);
+		}
+
 		/// <summary>
 		/// Gets the output directory for the main executable
 		/// </summary>
