Adding the option to change the Visual Studio platform name on a per-project basis.

#rnx
#[robomerge][starship] UE5-Main

diff --git a/Engine/Source/Programs/UnrealBuildTool/ProjectFiles/PlatformProjectGenerator.cs b/Engine/Source/Programs/UnrealBuildTool/ProjectFiles/PlatformProjectGenerator.cs
--- a/Engine/Source/Programs/UnrealBuildTool/ProjectFiles/PlatformProjectGenerator.cs
+++ b/Engine/Source/Programs/UnrealBuildTool/ProjectFiles/PlatformProjectGenerator.cs
@@ -72,6 +72,19 @@ namespace UnrealBuildTool
 			return InPlatform.ToString();
 		}
 
+		/// <summary>
+		/// Return the VisualStudio platform name for this build platform
+		/// </summary>
+		/// <param name="InPlatform">  The UnrealTargetPlatform being built</param>
+		/// <param name="InConfiguration"> The UnrealTargetConfiguration being built</param>
+		/// <param name="InProjectDir">  The root directory for the current project</param>
+		/// <returns>string    The name of the platform that VisualStudio recognizes</returns>
+		public virtual string GetVisualStudioPlatformName(UnrealTargetPlatform InPlatform, UnrealTargetConfiguration InConfiguration, DirectoryReference InProjectDir)
+		{
+			// By default, return the project-independent platform string
+			return GetVisualStudioPlatformName(InPlatform, InConfiguration);
+		}
+
 		/// <summary>
 		/// Return project configuration settings that must be included before the default props file
 		/// </summary>
diff --git a/Engine/Source/Programs/UnrealBuildTool/ProjectFiles/VisualStudio/VCProject.cs b/Engine/Source/Programs/UnrealBuildTool/ProjectFiles/VisualStudio/VCProject.cs
--- a/Engine/Source/Programs/UnrealBuildTool/ProjectFiles/VisualStudio/VCProject.cs
+++ b/Engine/Source/Programs/UnrealBuildTool/ProjectFiles/VisualStudio/VCProject.cs
@@ -342,7 +342,7 @@ namespace UnrealBuildTool
 				// Note that the actual name of the platform on the Visual Studio side may be different than what
 				// UnrealBuildTool calls it (e.g. "Win64" -> "x64".) GetVisualStudioPlatformName() will figure this out.
 				ProjectConfigurationName = Configuration.ToString();
-				ProjectPlatformName = PlatformProjectGenerator!.GetVisualStudioPlatformName(Platform, Configuration);
+				ProjectPlatformName = PlatformProjectGenerator!.GetVisualStudioPlatformName(Platform, Configuration, BaseDir);
 			}
 			else
 			{
