VCProjectGen: Add std version to AdditionalOptions to fix intellisense

#rnx

diff --git a/Engine/Source/Programs/UnrealBuildTool/ProjectFiles/VisualStudio/VCProject.cs b/Engine/Source/Programs/UnrealBuildTool/ProjectFiles/VisualStudio/VCProject.cs
--- a/Engine/Source/Programs/UnrealBuildTool/ProjectFiles/VisualStudio/VCProject.cs
+++ b/Engine/Source/Programs/UnrealBuildTool/ProjectFiles/VisualStudio/VCProject.cs
@@ -1229,7 +1229,9 @@ namespace UnrealBuildTool
 							VCProjectFileContent.AppendLine("      <ForcedIncludeFiles>{0}</ForcedIncludeFiles>", ForceIncludePaths);
 							if (PchHeaderFile != null && ProjectFileFormat >= VCProjectFileFormat.VisualStudio2022)
 							{
-								VCProjectFileContent.AppendLine("      <AdditionalOptions>/Yu\"{0}\"</AdditionalOptions>", PchHeaderFile);
+								VCProjectFileContent.AppendLine("      <AdditionalOptions>{0} /Yu\"{1}\"</AdditionalOptions>",
+									GetCppStandardCompileArgument(GetIntelliSenseCppVersion()),
+									PchHeaderFile);
 							}
 						}
 						VCProjectFileContent.AppendLine("    </{0}>", VCFileType);
