UnrealBuildTool: Added TargetRules.bCompileAgainstEditor to enable our low level tests programs to compile against and test editor code.
- Drives the WITH_EDITOR preprocessor definition.
- Updated TargetRules.bBuildEditor to check the target type or bCompileAgainstEditor.
- Removed usages of TargetRules.bBuildEditor from UBT and replaced it with checks against the target type or the new bCompileAgainstEditor flag where necessary.

diff --git a/Engine/Source/Programs/UnrealBuildTool/Configuration/ModuleRules.cs b/Engine/Source/Programs/UnrealBuildTool/Configuration/ModuleRules.cs
--- a/Engine/Source/Programs/UnrealBuildTool/Configuration/ModuleRules.cs
+++ b/Engine/Source/Programs/UnrealBuildTool/Configuration/ModuleRules.cs
@@ -1198,7 +1198,7 @@ namespace UnrealBuildTool
 					PublicDefinitions.Add("WITH_APEX=0");
 					PublicDefinitions.Add("WITH_APEX_CLOTHING=0");
 					PublicDefinitions.Add("WITH_CLOTH_COLLISION_DETECTION=0");
-					PublicDefinitions.Add(string.Format("WITH_PHYSX_COOKING={0}", Target.bBuildEditor && Target.bCompilePhysX ? 1 : 0));  // without APEX, we only need cooking in editor builds
+					PublicDefinitions.Add(string.Format("WITH_PHYSX_COOKING={0}", Target.Type == TargetType.Editor && Target.bCompilePhysX ? 1 : 0));  // without APEX, we only need cooking in editor builds
 				}
 
 				if (Target.bCompileNvCloth == true)
@@ -1223,7 +1223,7 @@ namespace UnrealBuildTool
 				PublicDefinitions.Add("PHYSICS_INTERFACE_PHYSX=0");
 				PublicDefinitions.Add("WITH_APEX=0");
 				PublicDefinitions.Add("WITH_APEX_CLOTHING=0");
-				PublicDefinitions.Add(string.Format("WITH_PHYSX_COOKING={0}", Target.bBuildEditor && Target.bCompilePhysX ? 1 : 0));  // without APEX, we only need cooking in editor builds
+				PublicDefinitions.Add(string.Format("WITH_PHYSX_COOKING={0}", Target.Type == TargetType.Editor && Target.bCompilePhysX ? 1 : 0));  // without APEX, we only need cooking in editor builds
 				PublicDefinitions.Add("WITH_NVCLOTH=0");
 
 				if(Target.bUseChaos)
diff --git a/Engine/Source/Programs/UnrealBuildTool/Configuration/TargetRules.cs b/Engine/Source/Programs/UnrealBuildTool/Configuration/TargetRules.cs
--- a/Engine/Source/Programs/UnrealBuildTool/Configuration/TargetRules.cs
+++ b/Engine/Source/Programs/UnrealBuildTool/Configuration/TargetRules.cs
@@ -503,11 +503,11 @@ namespace UnrealBuildTool
 		public bool bCompilePython = true;
 
 		/// <summary>
-		/// Whether to compile the editor or not. Only desktop platforms (Windows or Mac) will use this, other platforms force this to false.
+		/// Whether we are compiling editor code or not. Prefer the more explicit bCompileAgainstEditor instead.
 		/// </summary>
 		public bool bBuildEditor
 		{
-			get { return (Type == TargetType.Editor); }
+			get { return (Type == TargetType.Editor || bCompileAgainstEditor); }
 			set { Log.TraceWarning("Setting {0}.bBuildEditor is deprecated. Set {0}.Type instead.", GetType().Name); }
 		}
 
@@ -605,6 +605,22 @@ namespace UnrealBuildTool
 		[RequiresUniqueBuildEnvironment]
 		public bool bCompileAgainstApplicationCore = true;
 
+		/// <summary>
+		/// Manually specified value for bCompileAgainstEditor.
+		/// </summary>
+		bool? bCompileAgainstEditorOverride;
+
+		/// <summary>
+		/// Enabled for editor builds (TargetType.Editor). Can be overridden for programs (TargetType.Program) that would need to compile against editor code. Not available for other target types.
+		/// Mainly drives the value of WITH_EDITOR.
+		/// </summary>
+		[RequiresUniqueBuildEnvironment]
+		public bool bCompileAgainstEditor
+		{
+			set { bCompileAgainstEditorOverride = value; }
+			get { return bCompileAgainstEditorOverride ?? (Type == TargetType.Editor); }
+		}
+
 		/// <summary>
 		/// Whether to compile Recast navmesh generation.
 		/// </summary>
@@ -631,7 +647,7 @@ namespace UnrealBuildTool
 		bool? bOverrideCompileSpeedTree;
 
 		/// <summary>
-		/// Whether we should compile in support for Simplygon or not.
+		/// Whether we should compile in support for SpeedTree or not.
 		/// </summary>
 		[RequiresUniqueBuildEnvironment]
 		public bool bCompileSpeedTree
@@ -687,7 +703,7 @@ namespace UnrealBuildTool
 		{
 			get
 			{
-				return bWithPushModelOverride ?? (bBuildEditor);
+				return bWithPushModelOverride ?? (Type == TargetType.Editor);
 			}
 			set
 			{
@@ -2235,6 +2251,11 @@ namespace UnrealBuildTool
 			get { return Inner.bCompileAgainstApplicationCore; }
 		}
 
+		public bool bCompileAgainstEditor
+		{
+			get { return Inner.bCompileAgainstEditor; }
+		}
+
 		public bool bCompileRecast
 		{
 			get { return Inner.bCompileRecast; }
diff --git a/Engine/Source/Programs/UnrealBuildTool/Configuration/UEBuildTarget.cs b/Engine/Source/Programs/UnrealBuildTool/Configuration/UEBuildTarget.cs
--- a/Engine/Source/Programs/UnrealBuildTool/Configuration/UEBuildTarget.cs
+++ b/Engine/Source/Programs/UnrealBuildTool/Configuration/UEBuildTarget.cs
@@ -3554,8 +3554,8 @@ namespace UnrealBuildTool
 			GlobalCompileEnvironment.Definitions.Add(String.Format("IS_PROGRAM={0}", TargetType == TargetType.Program ? "1" : "0"));
 			GlobalCompileEnvironment.Definitions.AddRange(Rules.GlobalDefinitions);
 			GlobalCompileEnvironment.bUseSharedBuildEnvironment = (Rules.BuildEnvironment == TargetBuildEnvironment.Shared);
-			GlobalCompileEnvironment.bEnableExceptions = Rules.bForceEnableExceptions || Rules.bBuildEditor;
-			GlobalCompileEnvironment.bEnableObjCExceptions = Rules.bForceEnableObjCExceptions || Rules.bBuildEditor;
+			GlobalCompileEnvironment.bEnableExceptions = Rules.bForceEnableExceptions || Rules.bCompileAgainstEditor;
+			GlobalCompileEnvironment.bEnableObjCExceptions = Rules.bForceEnableObjCExceptions || Rules.bCompileAgainstEditor;
 			GlobalCompileEnvironment.DefaultWarningLevel = Rules.DefaultWarningLevel;
 			GlobalCompileEnvironment.DeprecationWarningLevel = Rules.DeprecationWarningLevel;
 			GlobalCompileEnvironment.ShadowVariableWarningLevel = Rules.ShadowVariableWarningLevel;
@@ -3676,7 +3676,7 @@ namespace UnrealBuildTool
 			if (Rules.bForceDisableAutomationTests)
 			{
 				bCompileDevTests = bCompilePerfTests = false;
-				}
+			}
 
 			GlobalCompileEnvironment.Definitions.Add("WITH_DEV_AUTOMATION_TESTS=" + (bCompileDevTests ? "1" : "0"));
 			GlobalCompileEnvironment.Definitions.Add("WITH_PERF_AUTOMATION_TESTS=" + (bCompilePerfTests ? "1" : "0"));
@@ -3800,8 +3800,7 @@ namespace UnrealBuildTool
 				GlobalCompileEnvironment.Definitions.Add("USE_ESTIMATED_UTCNOW=0");
 			}
 
-			// bBuildEditor has now been set appropriately for all platforms, so this is here to make sure the #define
-			if (Rules.bBuildEditor)
+			if (Rules.bCompileAgainstEditor && (Rules.Type == TargetType.Editor || Rules.Type == TargetType.Program))
 			{
 				GlobalCompileEnvironment.Definitions.Add("WITH_EDITOR=1");
 				GlobalCompileEnvironment.Definitions.Add("WITH_IOSTORE_IN_EDITOR=1");
diff --git a/Engine/Source/Programs/UnrealBuildTool/System/RulesAssembly.cs b/Engine/Source/Programs/UnrealBuildTool/System/RulesAssembly.cs
--- a/Engine/Source/Programs/UnrealBuildTool/System/RulesAssembly.cs
+++ b/Engine/Source/Programs/UnrealBuildTool/System/RulesAssembly.cs
@@ -610,10 +610,10 @@ namespace UnrealBuildTool
 				Rules.bCompileAgainstCoreUObject = true;
 			}
 
-			// Must have editor only data if building the editor.
-			if (Rules.bBuildEditor)
+			if (Rules.Type == TargetType.Editor)
 			{
-				Rules.bBuildWithEditorOnlyData = true;
+				Rules.bBuildWithEditorOnlyData = true; // Must have editor only data if building the editor.
+				Rules.bCompileAgainstEditor = true;
 			}
 
 			// Apply the override to force debug info to be enabled
