Editgrate 18353193 - UnrealBuildTool: track history for AdditionalPluginDirectories
External plugins will be built in project_dir/Intermediate/External/...

diff --git a/Engine/Source/Programs/UnrealBuildTool/ProjectFiles/ProjectFileGenerator.cs b/Engine/Source/Programs/UnrealBuildTool/ProjectFiles/ProjectFileGenerator.cs
--- a/Engine/Source/Programs/UnrealBuildTool/ProjectFiles/ProjectFileGenerator.cs
+++ b/Engine/Source/Programs/UnrealBuildTool/ProjectFiles/ProjectFileGenerator.cs
@@ -1484,7 +1484,9 @@ namespace UnrealBuildTool
 		/// Finds all module files (filtering by platform)
 		/// </summary>
 		/// <returns>Filtered list of module files</returns>
+#pragma warning disable CS8632 // The annotation for nullable reference types should only be used in code within a '#nullable' annotations context. (TODO: Remove warning disable when file no longer has #nullable disable)
 		protected List<FileReference> DiscoverModules(List<FileReference> AllGameProjects, List<DirectoryReference>? AdditionalSearchPaths)
+#pragma warning restore CS8632 // The annotation for nullable reference types should only be used in code within a '#nullable' annotations context.
 		{
 			// Locate all modules (*.Build.cs files)
 			return Rules.FindAllRulesSourceFiles(Rules.RulesFileType.Module, GameFolders: AllGameProjects.Select(x => x.Directory).ToList(), ForeignPlugins: null, AdditionalSearchPaths: AdditionalSearchPaths);
