Add ExludedFilePaths to vcproject generation for special situations where a developer does not want to see certain files in visual studio project

diff --git a/Engine/Source/Programs/UnrealBuildTool/ProjectFiles/VisualStudio/VCProject.cs b/Engine/Source/Programs/UnrealBuildTool/ProjectFiles/VisualStudio/VCProject.cs
--- a/Engine/Source/Programs/UnrealBuildTool/ProjectFiles/VisualStudio/VCProject.cs
+++ b/Engine/Source/Programs/UnrealBuildTool/ProjectFiles/VisualStudio/VCProject.cs
@@ -740,14 +740,15 @@ namespace UnrealBuildTool
 		}
 
 		private string[]? FilteredList = null;
+		private string[]? FilteredPathsLis = null;
 
-		bool IncludePathIsFilteredOut(DirectoryReference IncludePath)
+		bool PathIsFilteredOut(DirectoryReference InPath, ref string[] FilteredList)
 		{
 			// Turn the filter string into an array, remove whitespace, and normalize any path statements the first time
 			// we are asked to check a path.
 			if (FilteredList == null)
 			{
-				IEnumerable<string> CleanPaths = Settings.ExcludedIncludePaths.Split(new[] { ';' }, StringSplitOptions.RemoveEmptyEntries)
+				IEnumerable<string> CleanPaths = Settings.ExcludedFilePaths.Split(new[] { ';' }, StringSplitOptions.RemoveEmptyEntries)
 					.Select(P => P.Trim())
 					.Select(P => P.Replace('/', Path.DirectorySeparatorChar));
 
@@ -759,14 +760,14 @@ namespace UnrealBuildTool
 			// we'll compare to our list so we can Contains and EndsWith to catch:
 			// <Path>\Foo\Dir
 			// <Path>\Foo
-			string IncludePathWithSeparator = IncludePath.FullName + Path.DirectorySeparatorChar;
+			string PathWithSeparator = InPath.FullName + Path.DirectorySeparatorChar;
 
 			if (FilteredList.Length > 0)
 			{
 				foreach (string Entry in FilteredList)
 				{
-					if (IncludePathWithSeparator.Contains(Entry, StringComparison.OrdinalIgnoreCase)
-						|| IncludePathWithSeparator.EndsWith(Entry, StringComparison.OrdinalIgnoreCase))
+					if (PathWithSeparator.Contains(Entry, StringComparison.OrdinalIgnoreCase)
+						|| PathWithSeparator.EndsWith(Entry, StringComparison.OrdinalIgnoreCase))
 					{
 						return true;
 					}
@@ -776,6 +777,16 @@ namespace UnrealBuildTool
 			return false;
 		}
 
+		bool IncludePathIsFilteredOut(DirectoryReference IncludePath)
+		{
+			return PathIsFilteredOut(IncludePath, ref FilteredIncludeList);
+		}
+
+		bool FilePathIsFilteredOut(DirectoryReference InPath)
+		{
+			return PathIsFilteredOut(InPath, ref FilteredPathsList);
+		}
+
 		bool TryGetBuildEnvironment(DirectoryReference BaseDir, [NotNullWhen(true)] out BuildEnvironment? OutBuildEnvironment)
 		{
 			for (DirectoryReference? CurrentDir = BaseDir; CurrentDir != null; CurrentDir = CurrentDir.ParentDirectory)
@@ -1164,6 +1175,12 @@ namespace UnrealBuildTool
 						continue;
 					}
 
+					// Allow filtering of any type of file
+					if (FilePathIsFilteredOut(new DirectoryReference(AliasedFile.FileSystemPath)))
+					{
+						continue;
+					}
+
 					if (VCFileType != "ClCompile")
 					{
 						VCProjectFileContent.AppendLine("    <{0} Include=\"{1}\"/>", VCFileType, EscapeFileName(AliasedFile.FileSystemPath));
diff --git a/Engine/Source/Programs/UnrealBuildTool/ProjectFiles/VisualStudio/VCProjectFileGenerator.cs b/Engine/Source/Programs/UnrealBuildTool/ProjectFiles/VisualStudio/VCProjectFileGenerator.cs
--- a/Engine/Source/Programs/UnrealBuildTool/ProjectFiles/VisualStudio/VCProjectFileGenerator.cs
+++ b/Engine/Source/Programs/UnrealBuildTool/ProjectFiles/VisualStudio/VCProjectFileGenerator.cs
@@ -42,6 +42,13 @@ namespace UnrealBuildTool
 		[XmlConfigFile(Category = "VCProjectFileGenerator")]
 		public string ExcludedIncludePaths = "";
 
+		/// <summary>
+		/// Semi-colon separated list of paths that should not be added to the projects. Useful for omitting third-party files
+		/// (e.g ThirdParty/WebRTC) from intellisense suggestions and reducing memory footprints.
+		/// </summary>
+		[XmlConfigFile(Category = "VCProjectFileGenerator")]
+		public string ExcludedFilePaths = "";
+
 		/// <summary>
 		/// Whether to write a solution option (suo) file for the sln.
 		/// </summary>
