Update code to match new nullable requirements

diff --git a/Engine/Source/Programs/UnrealBuildTool/ProjectFiles/VisualStudio/VCProject.cs b/Engine/Source/Programs/UnrealBuildTool/ProjectFiles/VisualStudio/VCProject.cs
--- a/Engine/Source/Programs/UnrealBuildTool/ProjectFiles/VisualStudio/VCProject.cs
+++ b/Engine/Source/Programs/UnrealBuildTool/ProjectFiles/VisualStudio/VCProject.cs
@@ -739,10 +739,10 @@ namespace UnrealBuildTool
 			return ProjectFiles;
 		}
 
-		private string[]? FilteredList = null;
-		private string[]? FilteredPathsLis = null;
+		private string[]? FilteredIncludeList = null;
+		private string[]? FilteredPathsList = null;
 
-		bool PathIsFilteredOut(DirectoryReference InPath, ref string[] FilteredList)
+		bool PathIsFilteredOut(DirectoryReference InPath, ref string[]? FilteredList)
 		{
 			// Turn the filter string into an array, remove whitespace, and normalize any path statements the first time
 			// we are asked to check a path.
