bUseFastPDBLinking is inverted since Visual Studio 2017

diff --git a/Engine/Source/Programs/UnrealBuildTool/Platform/Windows/VCToolChain.cs b/Engine/Source/Programs/UnrealBuildTool/Platform/Windows/VCToolChain.cs
--- a/Engine/Source/Programs/UnrealBuildTool/Platform/Windows/VCToolChain.cs
+++ b/Engine/Source/Programs/UnrealBuildTool/Platform/Windows/VCToolChain.cs
@@ -965,6 +965,18 @@ namespace UnrealBuildTool
 							break;
 					}
 				}
+				else
+				{
+					// Beginning in Visual Studio 2017 /DEBUG defaults to /DEBUG:FASTLINK for debug builds
+					switch (Target.WindowsPlatform.Compiler)
+					{
+						case WindowsCompiler.VisualStudio2017:
+						case WindowsCompiler.VisualStudio2019:
+						case WindowsCompiler.VisualStudio2022:
+							Arguments[Arguments.Count - 1] += ":FULL";
+							break;
+					}
+				}
 			}
 
 			// Prompt the user before reporting internal errors to Microsoft.
