Build fix

diff --git a/Engine/Source/Programs/UnrealBuildTool/Platform/Windows/VCToolChain.cs b/Engine/Source/Programs/UnrealBuildTool/Platform/Windows/VCToolChain.cs
--- a/Engine/Source/Programs/UnrealBuildTool/Platform/Windows/VCToolChain.cs
+++ b/Engine/Source/Programs/UnrealBuildTool/Platform/Windows/VCToolChain.cs
@@ -970,7 +970,6 @@ namespace UnrealBuildTool
 					// Beginning in Visual Studio 2017 /DEBUG defaults to /DEBUG:FASTLINK for debug builds
 					switch (Target.WindowsPlatform.Compiler)
 					{
-						case WindowsCompiler.VisualStudio2017:
 						case WindowsCompiler.VisualStudio2019:
 						case WindowsCompiler.VisualStudio2022:
 							Arguments[Arguments.Count - 1] += ":FULL";
