bUseVerse is now also harvested from BuildConfiguration.xml

diff --git a/Engine/Source/Programs/UnrealBuildTool/Configuration/TargetRules.cs b/Engine/Source/Programs/UnrealBuildTool/Configuration/TargetRules.cs
--- a/Engine/Source/Programs/UnrealBuildTool/Configuration/TargetRules.cs
+++ b/Engine/Source/Programs/UnrealBuildTool/Configuration/TargetRules.cs
@@ -417,6 +417,7 @@ namespace UnrealBuildTool
 		[RequiresUniqueBuildEnvironment]
 		[CommandLine("-NoUseVerse", Value = "false")]
 		[CommandLine("-UseVerse", Value = "true")]
+		[XmlConfigFile(Category = "BuildConfiguration")]
 		public bool bUseVerse = false;
 
 		/// <summary>
