UBT: Prevent engine metadata files being overwritten in installed builds.

diff --git a/Engine/Source/Programs/UnrealBuildTool/Configuration/UEBuildTarget.cs b/Engine/Source/Programs/UnrealBuildTool/Configuration/UEBuildTarget.cs
--- a/Engine/Source/Programs/UnrealBuildTool/Configuration/UEBuildTarget.cs
+++ b/Engine/Source/Programs/UnrealBuildTool/Configuration/UEBuildTarget.cs
@@ -1952,50 +1952,34 @@ namespace UnrealBuildTool
 				// Create an action which to generate the module receipts
 				if (VersionFile == null)
 				{
-					WriteMetadataTargetInfo TargetInfo = new WriteMetadataTargetInfo(ProjectFile, null, null, ReceiptFileName, Receipt);
+					WriteMetadataTargetInfo TargetInfo = new WriteMetadataTargetInfo(ProjectFile, null, null, ReceiptFileName, Receipt, FileNameToModuleManifest);
 					FileReference TargetInfoFile = FileReference.Combine(ProjectIntermediateDirectory, "TargetMetadata.dat");
-					BinaryFormatterUtils.SaveIfDifferent(TargetInfoFile, TargetInfo);
-
-					Action WriteTargetMetadata = CreateWriteMetadataAction(Makefile, ReceiptFileName.GetFileName(), TargetInfoFile);
-					WriteTargetMetadata.PrerequisiteItems.AddRange(Binaries.SelectMany(x => x.OutputFilePaths).Select(x => FileItem.GetItemByFileReference(x)));
-					WriteTargetMetadata.ProducedItems.Add(FileItem.GetItemByFileReference(ReceiptFileName));
-					Makefile.OutputItems.AddRange(WriteTargetMetadata.ProducedItems);
+					CreateWriteMetadataAction(Makefile, ReceiptFileName.GetFileName(), TargetInfoFile, TargetInfo);
 				}
 				else
 				{
-					WriteMetadataTargetInfo EngineInfo = new WriteMetadataTargetInfo(null, VersionFile, Receipt.Version, null, null);
-					FileReference EngineInfoFile = FileReference.Combine(ProjectIntermediateDirectory, "EngineMetadata.dat");
-					Action WriteEngineMetadata = CreateWriteMetadataAction(Makefile, VersionFile.GetFileName(), EngineInfoFile);
-					WriteEngineMetadata.ProducedItems.Add(FileItem.GetItemByFileReference(VersionFile));
-
-					WriteMetadataTargetInfo TargetInfo = new WriteMetadataTargetInfo(ProjectFile, VersionFile, null, ReceiptFileName, Receipt);
-					FileReference TargetInfoFile = FileReference.Combine(ProjectIntermediateDirectory, "TargetMetadata.dat");
-					Action WriteTargetMetadata = CreateWriteMetadataAction(Makefile, ReceiptFileName.GetFileName(), TargetInfoFile);
-					WriteTargetMetadata.PrerequisiteItems.Add(FileItem.GetItemByFileReference(VersionFile));
-					WriteTargetMetadata.ProducedItems.Add(FileItem.GetItemByFileReference(ReceiptFileName));
+					Dictionary<FileReference, ModuleManifest> EngineFileToManifest = new Dictionary<FileReference, ModuleManifest>();
+					Dictionary<FileReference, ModuleManifest> TargetFileToManifest = new Dictionary<FileReference, ModuleManifest>();
 
 					foreach ((FileReference File, ModuleManifest Manifest) in FileNameToModuleManifest)
 					{
-						FileItem FileItem = FileItem.GetItemByFileReference(File);
 						if (File.IsUnderDirectory(Unreal.EngineDirectory))
 						{
-							EngineInfo.FileToManifest.Add(File, Manifest);
-							WriteEngineMetadata.ProducedItems.Add(FileItem);
-							WriteEngineMetadata.PrerequisiteItems.AddRange(Manifest.ModuleNameToFileName.Values.Select(x => FileItem.GetItemByFileReference(FileReference.Combine(File.Directory, x))));
+							EngineFileToManifest.Add(File, Manifest);
 						}
 						else
 						{
-							TargetInfo.FileToManifest.Add(File, Manifest);
-							WriteTargetMetadata.ProducedItems.Add(FileItem);
-							WriteTargetMetadata.PrerequisiteItems.AddRange(Manifest.ModuleNameToFileName.Values.Select(x => FileItem.GetItemByFileReference(FileReference.Combine(File.Directory, x))));
+							TargetFileToManifest.Add(File, Manifest);
 						}
 					}
 
-					WriteMetadataFile(Makefile, EngineInfoFile, EngineInfo);
-					Makefile.OutputItems.AddRange(WriteEngineMetadata.ProducedItems);
+					WriteMetadataTargetInfo EngineInfo = new WriteMetadataTargetInfo(null, VersionFile, Receipt.Version, null, null, EngineFileToManifest);
+					FileReference EngineInfoFile = FileReference.Combine(ProjectIntermediateDirectory, "EngineMetadata.dat");
+					CreateWriteMetadataAction(Makefile, VersionFile.GetFileName(), EngineInfoFile, EngineInfo);
 
-					WriteMetadataFile(Makefile, TargetInfoFile, TargetInfo);
-					Makefile.OutputItems.AddRange(WriteTargetMetadata.ProducedItems);
+					WriteMetadataTargetInfo TargetInfo = new WriteMetadataTargetInfo(ProjectFile, VersionFile, null, ReceiptFileName, Receipt, TargetFileToManifest);
+					FileReference TargetInfoFile = FileReference.Combine(ProjectIntermediateDirectory, "TargetMetadata.dat");
+					CreateWriteMetadataAction(Makefile, ReceiptFileName.GetFileName(), TargetInfoFile, TargetInfo);
 				}
 
 				// Create actions to run the post build steps
@@ -2125,34 +2109,55 @@ namespace UnrealBuildTool
 			return Makefile;
 		}
 
-		Action CreateWriteMetadataAction(TargetMakefile Makefile, string StatusDescription, FileReference InfoFile)
+		void CreateWriteMetadataAction(TargetMakefile Makefile, string StatusDescription, FileReference InfoFile, WriteMetadataTargetInfo Info)
 		{
-			StringBuilder WriteMetadataArguments = new StringBuilder();
-			WriteMetadataArguments.AppendFormat("-Input={0}", Utils.MakePathSafeToUseWithCommandLine(InfoFile));
-			WriteMetadataArguments.AppendFormat(" -Version={0}", WriteMetadataMode.CurrentVersionNumber);
-			if (Rules.bNoManifestChanges)
+			List<FileReference> ProducedItems = new List<FileReference>(Info.FileToManifest.Keys);
+			if (Info.Version != null && Info.VersionFile != null)
+			{
+				ProducedItems.Add(Info.VersionFile);
+			}
+			if (Info.ReceiptFile != null)
 			{
-				WriteMetadataArguments.Append(" -NoManifestChanges");
+				ProducedItems.Add(Info.ReceiptFile);
 			}
 
-			Action WriteMetadataAction = Makefile.CreateRecursiveAction<WriteMetadataMode>(ActionType.WriteMetadata, WriteMetadataArguments.ToString());
-			WriteMetadataAction.WorkingDirectory = UnrealBuildTool.EngineSourceDirectory;
-			WriteMetadataAction.StatusDescription = StatusDescription;
-			WriteMetadataAction.bCanExecuteRemotely = false;
-			WriteMetadataAction.bUseActionHistory = false; // Different files for each target; do not want to invalidate based on this.
-			return WriteMetadataAction;
-		}
+			if (!ProducedItems.Any(x => UnrealBuildTool.IsFileInstalled(x)))
+			{
+				// Save the info file, but deliberately do not add a prerequisite from the action onto it. It is always beneath the
+				// project intermediate directory (even for the engine) and will cause -NoEngineChanges to block the build when
+				// updated. The behavior we want is just to depend on the engine DLL timestamps. Instead, we add a makefile dependency
+				// on it, causing it to be regenerated if missing.
+				FileItem InfoFileItem = FileItem.GetItemByFileReference(InfoFile);
+				BinaryFormatterUtils.Save(InfoFile, Info);
+				InfoFileItem.ResetCachedInfo();
+				Makefile.InternalDependencies.Add(InfoFileItem);
 
-		static void WriteMetadataFile(TargetMakefile Makefile, FileReference InfoFile, WriteMetadataTargetInfo TargetInfo)
-		{
-			BinaryFormatterUtils.Save(InfoFile, TargetInfo);
-
-			// Note: Deliberately do not add a prerequisite from the action onto the info file. It is always beneath the project intermediate directory
-			// (even for the engine) and will cause -NoEngineChanges to block the build when updated. The behavior we want is just to depend on the engine
-			// DLL timestamps. Instead, we add a makefile dependency on it, causing it to be regenerated if missing.
-			FileItem InfoFileItem = FileItem.GetItemByFileReference(InfoFile);
-			InfoFileItem.ResetCachedInfo();
-			Makefile.InternalDependencies.Add(InfoFileItem);
+				// Get the argument list
+				StringBuilder WriteMetadataArguments = new StringBuilder();
+				WriteMetadataArguments.AppendFormat("-Input={0}", Utils.MakePathSafeToUseWithCommandLine(InfoFile));
+				WriteMetadataArguments.AppendFormat(" -Version={0}", WriteMetadataMode.CurrentVersionNumber);
+				if (Rules.bNoManifestChanges)
+				{
+					WriteMetadataArguments.Append(" -NoManifestChanges");
+				}
+
+				// Create the action
+				Action WriteMetadataAction = Makefile.CreateRecursiveAction<WriteMetadataMode>(ActionType.WriteMetadata, WriteMetadataArguments.ToString());
+				WriteMetadataAction.WorkingDirectory = UnrealBuildTool.EngineSourceDirectory;
+				WriteMetadataAction.StatusDescription = StatusDescription;
+				WriteMetadataAction.bCanExecuteRemotely = false;
+				WriteMetadataAction.bUseActionHistory = false; // Different files for each target; do not want to invalidate based on this.
+				WriteMetadataAction.ProducedItems.AddRange(ProducedItems.Select(x => FileItem.GetItemByFileReference(x)));
+				WriteMetadataAction.PrerequisiteItems.AddRange(Makefile.OutputItems);
+
+				if (Info.Version == null && Info.VersionFile != null)
+				{
+					WriteMetadataAction.PrerequisiteItems.Add(FileItem.GetItemByFileReference(Info.VersionFile));
+				}
+
+				// Add the produced items as leaf nodes to be built
+				Makefile.OutputItems.AddRange(WriteMetadataAction.ProducedItems);
+			}
 		}
 
 		/// <summary>
diff --git a/Engine/Source/Programs/UnrealBuildTool/Modes/WriteMetadataMode.cs b/Engine/Source/Programs/UnrealBuildTool/Modes/WriteMetadataMode.cs
--- a/Engine/Source/Programs/UnrealBuildTool/Modes/WriteMetadataMode.cs
+++ b/Engine/Source/Programs/UnrealBuildTool/Modes/WriteMetadataMode.cs
@@ -47,7 +47,7 @@ namespace UnrealBuildTool
 		/// <summary>
 		/// Map of module manifest filenames to their location on disk.
 		/// </summary>
-		public Dictionary<FileReference, ModuleManifest> FileToManifest = new Dictionary<FileReference, ModuleManifest>();
+		public Dictionary<FileReference, ModuleManifest> FileToManifest;
 
 		/// <summary>
 		/// Constructor
@@ -57,13 +57,15 @@ namespace UnrealBuildTool
 		/// <param name="Version"></param>
 		/// <param name="ReceiptFile"></param>
 		/// <param name="Receipt"></param>
-		public WriteMetadataTargetInfo(FileReference? ProjectFile, FileReference? VersionFile, BuildVersion? Version, FileReference? ReceiptFile, TargetReceipt? Receipt)
+		/// <param name="FileToManifest"></param>
+		public WriteMetadataTargetInfo(FileReference? ProjectFile, FileReference? VersionFile, BuildVersion? Version, FileReference? ReceiptFile, TargetReceipt? Receipt, Dictionary<FileReference, ModuleManifest> FileToManifest)
 		{
 			this.ProjectFile = ProjectFile;
 			this.VersionFile = VersionFile;
 			this.Version = Version;
 			this.ReceiptFile = ReceiptFile;
 			this.Receipt = Receipt;
+			this.FileToManifest = FileToManifest;
 		}
 	}
 
