PR #8663: Add Linux and Mac support to GenerateClangDatabase in UBT (Contributed by [REDACTED_PERSON])

diff --git a/Engine/Source/Programs/UnrealBuildTool/Modes/GenerateClangDatabase.cs b/Engine/Source/Programs/UnrealBuildTool/Modes/GenerateClangDatabase.cs
--- a/Engine/Source/Programs/UnrealBuildTool/Modes/GenerateClangDatabase.cs
+++ b/Engine/Source/Programs/UnrealBuildTool/Modes/GenerateClangDatabase.cs
@@ -38,10 +38,10 @@ namespace UnrealBuildTool
 
 			// Parse the filter argument
 			FileFilter? FileFilter = null;
-			if(FilterRules.Count > 0)
+			if (FilterRules.Count > 0)
 			{
 				FileFilter = new FileFilter(FileFilterType.Exclude);
-				foreach(string FilterRule in FilterRules)
+				foreach (string FilterRule in FilterRules)
 				{
 					FileFilter.AddRules(FilterRule.Split(';'));
 				}
@@ -64,8 +64,7 @@ namespace UnrealBuildTool
 					UEBuildTarget Target = UEBuildTarget.Create(TargetDescriptor, BuildConfiguration.bSkipRulesCompile, BuildConfiguration.bForceRulesCompile, BuildConfiguration.bUsePrecompiled);
 
 					// Find the location of the compiler
-					VCEnvironment Environment = VCEnvironment.Create(WindowsCompiler.Clang, Target.Platform, Target.Rules.WindowsPlatform.Architecture, null, Target.Rules.WindowsPlatform.WindowsSdkVersion, null);
-					FileReference ClangPath = FileReference.Combine(Environment.CompilerDir, "bin", "clang++.exe");
+					FileReference ClangPath = FindClangCompiler(Target);
 
 					// Convince each module to output its generated code include path
 					foreach (UEBuildBinary Binary in Target.Binaries)
@@ -83,7 +82,7 @@ namespace UnrealBuildTool
 						CppCompileEnvironment BinaryCompileEnvironment = Binary.CreateBinaryCompileEnvironment(GlobalCompileEnvironment);
 						foreach (UEBuildModuleCPP Module in Binary.Modules.OfType<UEBuildModuleCPP>())
 						{
-							if(!Module.Rules.bUsePrecompiled)
+							if (!Module.Rules.bUsePrecompiled)
 							{
 								UEBuildModuleCPP.InputFileCollection InputFileCollection = Module.FindInputFiles(Target.Platform, new Dictionary<DirectoryItem, FileItem[]>());
 
@@ -131,7 +130,7 @@ namespace UnrealBuildTool
 
 								foreach (FileItem InputFile in InputFiles)
 								{
-									if(FileFilter == null || FileFilter.Matches(InputFile.Location.MakeRelativeTo(Unreal.RootDirectory)))
+									if (FileFilter == null || FileFilter.Matches(InputFile.Location.MakeRelativeTo(Unreal.RootDirectory)))
 									{
 										FileToCommand[InputFile.Location] = String.Format("{0} \"{1}\"", CommandBuilder, InputFile.FullName);
 									}
@@ -146,7 +145,7 @@ namespace UnrealBuildTool
 				using (JsonWriter Writer = new JsonWriter(DatabaseFile))
 				{
 					Writer.WriteArrayStart();
-					foreach(KeyValuePair<FileReference, string> FileCommandPair in FileToCommand.OrderBy(x => x.Key.FullName))
+					foreach (KeyValuePair<FileReference, string> FileCommandPair in FileToCommand.OrderBy(x => x.Key.FullName))
 					{
 						Writer.WriteObjectStart();
 						Writer.WriteValue("file", FileCommandPair.Key.FullName);
@@ -160,5 +159,44 @@ namespace UnrealBuildTool
 
 			return 0;
 		}
+
+		/// <summary>
+		/// Searches for the Clang compiler for the given platform.
+		/// </summary>
+		/// <param name="Target">The build platform to use to search for the Clang compiler.</param>
+		/// <returns>The path to the Clang compiler.</returns>
+		private static FileReference FindClangCompiler(UEBuildTarget Target)
+		{
+			UnrealTargetPlatform HostPlatform = BuildHostPlatform.Current.Platform;
+
+			if (HostPlatform == UnrealTargetPlatform.Win64)
+			{
+				VCEnvironment Environment = VCEnvironment.Create(WindowsCompiler.Clang, Target.Platform,
+					Target.Rules.WindowsPlatform.Architecture, null,
+					Target.Rules.WindowsPlatform.WindowsSdkVersion, null);
+
+				return Environment.CompilerPath;
+			}
+			else if (HostPlatform == UnrealTargetPlatform.Linux)
+			{
+				string? Clang = LinuxCommon.WhichClang();
+
+				if (Clang != null)
+				{
+					return FileReference.FromString(Clang);
+				}
+			}
+			else if (HostPlatform == UnrealTargetPlatform.Mac)
+			{
+				MacToolChainSettings Settings = new MacToolChainSettings(false);
+				DirectoryReference? ToolchainDir = DirectoryReference.FromString(Settings.ToolchainDir);
+
+				if (ToolchainDir != null)
+				{
+					return FileReference.Combine(ToolchainDir, "clang++");
+				}
+			}
+			return FileReference.FromString("clang++");
+		}
 	}
 }
