XGE: Always set caption so the task in the build monitor displays correctly

#rnx

diff --git a/Engine/Source/Programs/UnrealBuildTool/Executors/XGE.cs b/Engine/Source/Programs/UnrealBuildTool/Executors/XGE.cs
--- a/Engine/Source/Programs/UnrealBuildTool/Executors/XGE.cs
+++ b/Engine/Source/Programs/UnrealBuildTool/Executors/XGE.cs
@@ -526,15 +526,7 @@ namespace UnrealBuildTool
 				XmlElement TaskElement = XGETaskDocument.CreateElement("Task");
 				ProjectElement.AppendChild(TaskElement);
 				TaskElement.SetAttribute("SourceFile", "");
-				if (!Action.bShouldOutputStatusDescription)
-				{
-					// If we were configured to not output a status description, then we'll instead
-					// set 'caption' text for this task, so that the XGE coordinator has something
-					// to display within the progress bars.  For tasks that are outputting a
-					// description, XGE automatically displays that text in the progress bar, so we
-					// only need to do this for tasks that output their own progress.
-					TaskElement.SetAttribute("Caption", Action.StatusDescription);
-				}
+				TaskElement.SetAttribute("Caption", Action.StatusDescription);
 				TaskElement.SetAttribute("Name", string.Format("Action{0}", ActionIndex));
 				TaskElement.SetAttribute("Tool", string.Format("Tool{0}", ActionIndex));
 				TaskElement.SetAttribute("WorkingDir", Action.WorkingDirectory.FullName);
