UBT: Fixed bAdaptiveUnityCompilesHeaderFiles dependency tracking

The header is now added to the working set rather than the generated source file.

#rnx

diff --git a/Engine/Source/Programs/UnrealBuildTool/System/Unity.cs b/Engine/Source/Programs/UnrealBuildTool/System/Unity.cs
--- a/Engine/Source/Programs/UnrealBuildTool/System/Unity.cs
+++ b/Engine/Source/Programs/UnrealBuildTool/System/Unity.cs
@@ -340,6 +340,14 @@ namespace UnrealBuildTool
 				Files.Add(CPPFile);
 			}
 
+			// Add adaptive files to the working set that will invalidate the makefile if it changes.
+			foreach (FileItem File in AdaptiveFiles)
+			{
+				Graph.AddFileToWorkingSet(File);
+			}
+
+			List<FileItem> AdaptiveHeaderFiles = new List<FileItem>();
+
 			// Add header files in the working set to the adaptive files if they are not the first include of a source file.
 			if (Target.bAdaptiveUnityCompilesHeaderFiles)
 			{
@@ -349,37 +357,27 @@ namespace UnrealBuildTool
 					OutputHeaderCPPWriter.WriteLine("// This file is automatically generated at compile-time to include a modified header file.");
 					OutputHeaderCPPWriter.WriteLine($"#include \"{HeaderFile.AbsolutePath.Replace('\\', '/')}\"");
 
-					string HeaderCPPFileName = string.Format($"{ModulePrefix}{BaseName}.Header.{Path.GetFileNameWithoutExtension(HeaderFile.AbsolutePath)}.cpp");
+					string HeaderCPPFileName = $"{ModulePrefix}{BaseName}.Header.{Path.GetFileNameWithoutExtension(HeaderFile.AbsolutePath)}.cpp";
 					FileReference HeaderCPPFilePath = FileReference.Combine(IntermediateDirectory, HeaderCPPFileName);
 
-					FileItem HeaderCPPFile = Graph.CreateIntermediateTextFile(HeaderCPPFilePath, OutputHeaderCPPWriter.ToString());
-					AdaptiveFiles.Add(HeaderCPPFile);
-				}
-			}
+					AdaptiveFiles.Add(Graph.CreateIntermediateTextFile(HeaderCPPFilePath, OutputHeaderCPPWriter.ToString()));
 
-			// Add adaptive files to the working set that will invalidate the makefile if it changes.
-			foreach (FileItem File in AdaptiveFiles)
-			{
-				Graph.AddFileToWorkingSet(File);
+					// Add the header file to the working set that will invalidate the makefile if it changes.
+					Graph.AddFileToWorkingSet(HeaderFile);
+					AdaptiveHeaderFiles.Add(HeaderFile);
+				}
 			}
 
 			HashSet<FileItem> CandidateAdaptiveFiles = new HashSet<FileItem>();
 			CandidateAdaptiveFiles.UnionWith(CPPFiles);
 			CandidateAdaptiveFiles.UnionWith(HeaderFiles);
 			CandidateAdaptiveFiles.ExceptWith(AdaptiveFiles);
+			CandidateAdaptiveFiles.ExceptWith(AdaptiveHeaderFiles);
 
 			foreach (FileItem File in CandidateAdaptiveFiles)
 			{
 				Graph.AddCandidateForWorkingSet(File);
 			}
 		}
-
-		static void AddUniqueDiagnostic(TargetMakefile Makefile, string Message)
-		{
-			if(!Makefile.Diagnostics.Contains(Message, StringComparer.Ordinal))
-			{
-				Makefile.Diagnostics.Add(Message);
-			}
-		}
 	}
 }
