UBT: Fix name for tests project.

diff --git a/Engine/Source/Programs/UnrealBuildTool/Properties/AssemblyInfo.cs b/Engine/Source/Programs/UnrealBuildTool/Properties/AssemblyInfo.cs
--- a/Engine/Source/Programs/UnrealBuildTool/Properties/AssemblyInfo.cs
+++ b/Engine/Source/Programs/UnrealBuildTool/Properties/AssemblyInfo.cs
@@ -10,4 +10,4 @@ using System.Runtime.CompilerServices;
 [assembly: AssemblyDescription("A tool to build Unreal Engine")]
 [assembly: AssemblyConfiguration("")]
 
-[assembly: InternalsVisibleTo("UnrealBuildToolTests")]
+[assembly: InternalsVisibleTo("UnrealBuildTool.Tests")]
