UBT: Add a -nocpp option for generating solutions without any C++ source when working on C# projects.

diff --git a/Engine/Source/Programs/UnrealBuildTool/ProjectFiles/ProjectFileGenerator.cs b/Engine/Source/Programs/UnrealBuildTool/ProjectFiles/ProjectFileGenerator.cs
--- a/Engine/Source/Programs/UnrealBuildTool/ProjectFiles/ProjectFileGenerator.cs
+++ b/Engine/Source/Programs/UnrealBuildTool/ProjectFiles/ProjectFileGenerator.cs
@@ -296,6 +296,12 @@ namespace UnrealBuildTool
 		[XmlConfigFile]
 		protected bool IncludeEnginePrograms = true;
 
+		/// <summary>
+		/// Whether to include C++ targets
+		/// </summary>
+		[XmlConfigFile]
+		protected bool IncludeCppSource = true;
+
 		/// <summary>
 		/// True if we should include csharp program projects in the generated solution. Pass "-DotNet" to enable this.
 		/// </summary>
@@ -871,17 +877,18 @@ namespace UnrealBuildTool
 			// we're generating project files for.
 			List<FileReference> AllModuleFiles = DiscoverModules(AllGameProjects, AdditionalSearchPaths.Values.SelectMany(x => x).ToList());
 
-			ProjectFile? EngineProject;
-			List<ProjectFile> GameProjects;
-			List<ProjectFile> ModProjects;
-			Dictionary<FileReference, ProjectFile> ProgramProjects;
-			Dictionary<RulesAssembly, DirectoryReference> RulesAssemblies;
+			ProjectFile? EngineProject = null;
+			List<ProjectFile> GameProjects = new List<ProjectFile>();
+			List<ProjectFile> ModProjects = new List<ProjectFile>();
+			Dictionary<FileReference, ProjectFile> ProgramProjects = new Dictionary<FileReference, ProjectFile>();
+			Dictionary<RulesAssembly, DirectoryReference> RulesAssemblies = new Dictionary<RulesAssembly, DirectoryReference>();
+			if(IncludeCppSource)
 			{
 				// Setup buildable projects for all targets
-				AddProjectsForAllTargets(PlatformProjectGenerators, AllGameProjects, AllTargetFiles, Arguments, out EngineProject, out GameProjects, out ProgramProjects, out RulesAssemblies);
+				AddProjectsForAllTargets(PlatformProjectGenerators, AllGameProjects, AllTargetFiles, Arguments, ref EngineProject, GameProjects, ProgramProjects, RulesAssemblies);
 
 				// Add projects for mods
-				AddProjectsForMods(GameProjects, out ModProjects);
+				AddProjectsForMods(GameProjects, ModProjects);
 
 				// Add all game projects and game config files
 				AddAllGameProjects(GameProjects);
@@ -1270,6 +1277,10 @@ namespace UnrealBuildTool
 								bAlwaysIncludeEngineModules = true;
 								break;
 
+							case "-NOCPP":
+								IncludeCppSource = false;
+								break;
+
 							case "-NOINTELLISENSE":
 								bGenerateIntelliSenseData = false;
 								break;
@@ -2115,6 +2126,11 @@ namespace UnrealBuildTool
 					}
 				}
 
+				if (!IncludeCppSource)
+				{
+					continue;
+				}
+
 				if (WantProjectFileForModule)
 				{
 					DirectoryReference BaseFolder;
@@ -2285,17 +2301,11 @@ namespace UnrealBuildTool
 			List<FileReference> AllGames,
 			List<FileReference> AllTargetFiles,
 			String[] Arguments,
-			out ProjectFile? EngineProject,
-			out List<ProjectFile> GameProjects,
-			out Dictionary<FileReference, ProjectFile> ProgramProjects,
-			out Dictionary<RulesAssembly, DirectoryReference> RulesAssemblies)
+			ref ProjectFile? EngineProject,
+			List<ProjectFile> GameProjects,
+			Dictionary<FileReference, ProjectFile> ProgramProjects,
+			Dictionary<RulesAssembly, DirectoryReference> RulesAssemblies)
 		{
-			// As we're creating project files, we'll also keep track of whether we created an "engine" project and return that if we have one
-			EngineProject = null;
-			GameProjects = new List<ProjectFile>();
-			ProgramProjects = new Dictionary<FileReference, ProjectFile>();
-			RulesAssemblies = new Dictionary<RulesAssembly, DirectoryReference>();
-
 			// Separate the .target.cs files that are platform extension specializations, per target name. These will be added alongside their base target.cs
 			Dictionary<string, List<FileReference>> AllSubTargetFilesPerTarget = new Dictionary<string, List<FileReference>>();
 			HashSet<FileReference> AllSubTargetFiles = new HashSet<FileReference>();
@@ -2569,10 +2579,9 @@ namespace UnrealBuildTool
 		/// </summary>
 		/// <param name="GameProjects">List of game project files</param>
 		/// <param name="ModProjects">Receives the list of mod projects on success</param>
-		protected void AddProjectsForMods(List<ProjectFile> GameProjects, out List<ProjectFile> ModProjects)
+		protected void AddProjectsForMods(List<ProjectFile> GameProjects, List<ProjectFile> ModProjects)
 		{
 			// Find all the mods for game projects
-			ModProjects = new List<ProjectFile>();
 			if (GameProjects.Count == 1)
 			{
 				ProjectFile GameProject = GameProjects.First();
diff --git a/Engine/Source/Programs/UnrealBuildTool/ProjectFiles/Rider/RiderProjectFileGenerator.cs b/Engine/Source/Programs/UnrealBuildTool/ProjectFiles/Rider/RiderProjectFileGenerator.cs
--- a/Engine/Source/Programs/UnrealBuildTool/ProjectFiles/Rider/RiderProjectFileGenerator.cs
+++ b/Engine/Source/Programs/UnrealBuildTool/ProjectFiles/Rider/RiderProjectFileGenerator.cs
@@ -452,14 +452,14 @@ namespace UnrealBuildTool
 		{ 
 			ProjectFile? EngineProject = null;
 			List<ProjectFile>? GameProjects = null;
-			List<ProjectFile>? ModProjects = null;
+			List<ProjectFile>? ModProjects = new List<ProjectFile>();
 			Dictionary<FileReference, ProjectFile>? ProgramProjects = null;
 
 			// Setup buildable projects for all targets
 			AddProjectsForAllTargets(PlatformProjectGenerators, AllGameProjects, out EngineProject,
 				out GameProjects, out ProgramProjects);
 
-			AddProjectsForMods(GameProjects, out ModProjects);
+			AddProjectsForMods(GameProjects, ModProjects);
 			AddAllGameProjects(GameProjects);
 
 			// If we're still missing an engine project because we don't have any targets for it, make one up.
