UnrealBuildTool: Log filepath when failing to parse a dependency file

#rnx

diff --git a/Engine/Source/Programs/UnrealBuildTool/System/CppDependencyCache.cs b/Engine/Source/Programs/UnrealBuildTool/System/CppDependencyCache.cs
--- a/Engine/Source/Programs/UnrealBuildTool/System/CppDependencyCache.cs
+++ b/Engine/Source/Programs/UnrealBuildTool/System/CppDependencyCache.cs
@@ -393,7 +393,7 @@ namespace UnrealBuildTool
 
 				if (TokenIdx + 1 >= Tokens.Count || Tokens[TokenIdx + 1] != ":")
 				{
-					throw new BuildException("Unable to parse dependency file");
+					throw new BuildException($"Unable to parse dependency file {InputFile.Location}");
 				}
 
 				TokenIdx += 2;
@@ -411,7 +411,7 @@ namespace UnrealBuildTool
 
 				if (TokenIdx != Tokens.Count)
 				{
-					throw new BuildException("Unable to parse dependency file");
+					throw new BuildException($"Unable to parse dependency file {InputFile.Location}");
 				}
 
 				return new DependencyInfo(InputFile.LastWriteTimeUtc.Ticks, null, null, NewDependencyFiles);
