UnrealBuildTool: Replace Compiler >= WindowsCompiler.VisualStudio2019 with Compiler.IsMSVC() extension function

#rnx

diff --git a/Engine/Source/Programs/UnrealBuildTool/Platform/HoloLens/HoloLensToolChain.cs b/Engine/Source/Programs/UnrealBuildTool/Platform/HoloLens/HoloLensToolChain.cs
--- a/Engine/Source/Programs/UnrealBuildTool/Platform/HoloLens/HoloLensToolChain.cs
+++ b/Engine/Source/Programs/UnrealBuildTool/Platform/HoloLens/HoloLensToolChain.cs
@@ -130,7 +130,7 @@ namespace UnrealBuildTool
 			// Experimental deterministic compile support
 			if (Target.WindowsPlatform.bDeterministic)
 			{
-				if (Target.WindowsPlatform.Compiler >= WindowsCompiler.VisualStudio2019)
+				if (Target.WindowsPlatform.Compiler.IsMSVC())
 				{
 					Arguments.Add("/experimental:deterministic");
 				}
@@ -710,7 +710,7 @@ namespace UnrealBuildTool
 				
 				if (CompileEnvironment.bGenerateDependenciesFile)
 				{
-					if (EnvVars.ToolChainVersion >= VersionNumber.Parse("14.27") && Target.WindowsPlatform.Compiler >= WindowsCompiler.VisualStudio2019)
+					if (EnvVars.ToolChainVersion >= VersionNumber.Parse("14.27") && Target.WindowsPlatform.Compiler.IsMSVC())
 					{
 						CompileAction.DependencyListFile = FileItem.GetItemByFileReference(FileReference.Combine(OutputDir, String.Format("{0}.json", SourceFile.Location.GetFileName())));
 						CompileAction.ProducedItems.Add(CompileAction.DependencyListFile);
diff --git a/Engine/Source/Programs/UnrealBuildTool/Platform/HoloLens/UEBuildHoloLens.cs b/Engine/Source/Programs/UnrealBuildTool/Platform/HoloLens/UEBuildHoloLens.cs
--- a/Engine/Source/Programs/UnrealBuildTool/Platform/HoloLens/UEBuildHoloLens.cs
+++ b/Engine/Source/Programs/UnrealBuildTool/Platform/HoloLens/UEBuildHoloLens.cs
@@ -347,7 +347,7 @@ namespace UnrealBuildTool
 
 		public static DirectoryReference? GetCppCXMetadataLocation(WindowsCompiler Compiler, DirectoryReference SelectedToolChainDir)
 		{
-			if (Compiler >= WindowsCompiler.VisualStudio2019)
+			if (Compiler.IsMSVC())
 			{
 				return DirectoryReference.Combine(SelectedToolChainDir, "lib", "x86", "Store", "references");
 			}
diff --git a/Engine/Source/Programs/UnrealBuildTool/Platform/Windows/UEBuildWindows.cs b/Engine/Source/Programs/UnrealBuildTool/Platform/Windows/UEBuildWindows.cs
--- a/Engine/Source/Programs/UnrealBuildTool/Platform/Windows/UEBuildWindows.cs
+++ b/Engine/Source/Programs/UnrealBuildTool/Platform/Windows/UEBuildWindows.cs
@@ -47,6 +47,23 @@ namespace UnrealBuildTool
 		VisualStudio2022,
 	}
 
+
+	/// <summary>
+	/// Extension methods for WindowsCompilier enum
+	/// </summary>
+	public static class WindowsCompilerExtensions
+	{
+		/// <summary>
+		/// Returns if this compiler toolchain based on MSVC
+		/// </summary>
+		/// <param name="Compiler">The compiler to check</param>
+		/// <returns>true if MSVC based</returns>
+		public static bool IsMSVC(this WindowsCompiler Compiler)
+		{
+			return Compiler >= WindowsCompiler.VisualStudio2019;
+		}
+	}
+
 	/// <summary>
 	/// Which static analyzer to use
 	/// </summary>
@@ -254,7 +271,7 @@ namespace UnrealBuildTool
 		/// Microsoft provides legacy_stdio_definitions library to enable building with VS2015 until they fix everything up.
 		public bool bNeedsLegacyStdioDefinitionsLib
 		{
-			get { return Compiler == WindowsCompiler.VisualStudio2019 || Compiler == WindowsCompiler.VisualStudio2022 || Compiler == WindowsCompiler.Clang; }
+			get { return Compiler.IsMSVC() || Compiler == WindowsCompiler.Clang; }
 		}
 
 		/// <summary>
@@ -937,6 +954,10 @@ namespace UnrealBuildTool
 				{
 					Target.WindowsPlatform.Compiler = WindowsCompiler.VisualStudio2019;
 				}
+				if (Target.WindowsPlatform.StaticAnalyzer == WindowsStaticAnalyzer.PVSStudio && HasCompiler(WindowsCompiler.VisualStudio2022))
+				{
+					Target.WindowsPlatform.Compiler = WindowsCompiler.VisualStudio2022;
+				}
 				else
 				{
 					Target.WindowsPlatform.Compiler = GetDefaultCompiler(Target.ProjectFile);
@@ -1009,7 +1030,7 @@ namespace UnrealBuildTool
 
 //			@Todo: Still getting reports of frequent OOM issues with this enabled as of 15.7.
 //			// Enable fast PDB linking if we're on VS2017 15.7 or later. Previous versions have OOM issues with large projects.
-//			if(!Target.bFormalBuild && !Target.bUseFastPDBLinking.HasValue && Target.WindowsPlatform.Compiler >= WindowsCompiler.VisualStudio2019)
+//			if(!Target.bFormalBuild && !Target.bUseFastPDBLinking.HasValue && Target.WindowsPlatform.Compiler.IsMSVC())
 //			{
 //				VersionNumber Version;
 //				DirectoryReference ToolChainDir;
@@ -1286,7 +1307,7 @@ namespace UnrealBuildTool
 							}
 						}
 					}
-				    else if(Compiler >= WindowsCompiler.VisualStudio2019)
+				    else if(Compiler.IsMSVC())
 				    {
 						// Enumerate all the manually installed toolchains
 						List<VisualStudioInstallation> Installations = FindVisualStudioInstallations(Compiler);
diff --git a/Engine/Source/Programs/UnrealBuildTool/Platform/Windows/VCEnvironment.cs b/Engine/Source/Programs/UnrealBuildTool/Platform/Windows/VCEnvironment.cs
--- a/Engine/Source/Programs/UnrealBuildTool/Platform/Windows/VCEnvironment.cs
+++ b/Engine/Source/Programs/UnrealBuildTool/Platform/Windows/VCEnvironment.cs
@@ -180,7 +180,7 @@ namespace UnrealBuildTool
 		/// <returns>Directory containing the 32-bit toolchain binaries</returns>
 		public static DirectoryReference GetVCToolPath(WindowsCompiler Compiler, DirectoryReference VCToolChainDir, WindowsArchitecture Architecture)
 		{
-			if (Compiler >= WindowsCompiler.VisualStudio2019)
+			if (Compiler.IsMSVC())
 			{
 				FileReference NativeCompilerPath = FileReference.Combine(VCToolChainDir, "bin", "HostX64", WindowsExports.GetArchitectureSubpath(Architecture), "cl.exe");
 				if (FileReference.Exists(NativeCompilerPath))
@@ -325,7 +325,7 @@ namespace UnrealBuildTool
 			string ArchFolder = WindowsExports.GetArchitectureSubpath(Architecture);
 
 			// Add the standard Visual C++ library paths
-			if (ToolChain >= WindowsCompiler.VisualStudio2019)
+			if (ToolChain.IsMSVC())
 			{
 				if (Platform == UnrealTargetPlatform.HoloLens)
 				{
@@ -367,7 +367,7 @@ namespace UnrealBuildTool
 			LibraryPaths.Add( GetToolChainLibsDir(Platform));
 
 			// If we're on >= Visual Studio 2015 and using pre-Windows 10 SDK, we need to find a Windows 10 SDK and add the UCRT include paths
-			if(ToolChain >= WindowsCompiler.VisualStudio2019 && WindowsSdkVersion < new VersionNumber(10))
+			if(ToolChain.IsMSVC() && WindowsSdkVersion < new VersionNumber(10))
 			{
 				KeyValuePair<VersionNumber, DirectoryReference> Pair = MicrosoftPlatformSDK.FindUniversalCrtDirs().OrderByDescending(x => x.Key).FirstOrDefault();
 				if(Pair.Key == null || Pair.Key < new VersionNumber(10))
@@ -526,7 +526,7 @@ namespace UnrealBuildTool
 				}
 				else
 				{
-					throw new BuildException("{0}, or {1} must be installed in order to build this target.", WindowsPlatform.GetCompilerName(WindowsCompiler.VisualStudio2019), WindowsPlatform.GetCompilerName(WindowsCompiler.VisualStudio2022));
+					throw new BuildException("{0} or {1} must be installed in order to build this target.", WindowsPlatform.GetCompilerName(WindowsCompiler.VisualStudio2019), WindowsPlatform.GetCompilerName(WindowsCompiler.VisualStudio2022));
 				}
 			}
 			else
diff --git a/Engine/Source/Programs/UnrealBuildTool/Platform/Windows/VCToolChain.cs b/Engine/Source/Programs/UnrealBuildTool/Platform/Windows/VCToolChain.cs
--- a/Engine/Source/Programs/UnrealBuildTool/Platform/Windows/VCToolChain.cs
+++ b/Engine/Source/Programs/UnrealBuildTool/Platform/Windows/VCToolChain.cs
@@ -313,7 +313,7 @@ namespace UnrealBuildTool
 				{
 					Arguments.Add("/clang:--trace-includes");
 				}
-				else if (Target.WindowsPlatform.Compiler >= WindowsCompiler.VisualStudio2019 )
+				else if (Target.WindowsPlatform.Compiler.IsMSVC() )
 				{
 					Arguments.Add("/showIncludes");
 				}
@@ -324,7 +324,7 @@ namespace UnrealBuildTool
 			{
 				Arguments.Add("-fdiagnostics-absolute-paths");
 			}
-			else if (Target.WindowsPlatform.Compiler >= WindowsCompiler.VisualStudio2019)
+			else if (Target.WindowsPlatform.Compiler.IsMSVC())
 			{
 				Arguments.Add("/FC");
 			}
@@ -374,7 +374,7 @@ namespace UnrealBuildTool
 			Arguments.Add("/wd4819");
 
 			// Disable Microsoft extensions on VS2017+ for improved standards compliance.
-			if (Target.WindowsPlatform.Compiler >= WindowsCompiler.VisualStudio2019 && Target.WindowsPlatform.bStrictConformanceMode)
+			if (Target.WindowsPlatform.Compiler.IsMSVC() && Target.WindowsPlatform.bStrictConformanceMode)
 			{
 				// This define is needed to ensure that MSVC static analysis mode doesn't declare attributes that are incompatible with strict conformance mode
 				AddDefinition(Arguments, "SAL_NO_ATTRIBUTE_DECLARATIONS=1");
@@ -387,13 +387,13 @@ namespace UnrealBuildTool
 			// Previously %s meant "the current character set" and %S meant "the other one".
 			// Now %s means multibyte and %S means wide. %Ts means "natural width".
 			// Reverting this behaviour until the UE4 source catches up.
-			if (Target.WindowsPlatform.Compiler >= WindowsCompiler.VisualStudio2019 || Target.WindowsPlatform.Compiler == WindowsCompiler.Clang)
+			if (Target.WindowsPlatform.Compiler.IsMSVC() || Target.WindowsPlatform.Compiler == WindowsCompiler.Clang)
 			{
 				AddDefinition(Arguments, "_CRT_STDIO_LEGACY_WIDE_SPECIFIERS=1");
 			}
 
 			// @todo HoloLens: Silence the hash_map deprecation errors for now. This should be replaced with unordered_map for the real fix.
-			if (Target.WindowsPlatform.Compiler >= WindowsCompiler.VisualStudio2019 || Target.WindowsPlatform.Compiler == WindowsCompiler.Clang)
+			if (Target.WindowsPlatform.Compiler.IsMSVC() || Target.WindowsPlatform.Compiler == WindowsCompiler.Clang)
 			{
 				AddDefinition(Arguments, "_SILENCE_STDEXT_HASH_DEPRECATION_WARNINGS=1");
 			}
@@ -414,7 +414,7 @@ namespace UnrealBuildTool
 			AddDefinition(Arguments, "_DISABLE_EXTENDED_ALIGNED_STORAGE");
 
 			// Fix Incredibuild errors with helpers using heterogeneous character sets
-			if (Target.WindowsPlatform.Compiler >= WindowsCompiler.VisualStudio2019)
+			if (Target.WindowsPlatform.Compiler.IsMSVC())
 			{
 				Arguments.Add("/source-charset:utf-8");
 				Arguments.Add("/execution-charset:utf-8");
@@ -434,7 +434,7 @@ namespace UnrealBuildTool
 			// Experimental deterministic compile support
 			if (Target.WindowsPlatform.bDeterministic)
 			{
-				if (Target.WindowsPlatform.Compiler >= WindowsCompiler.VisualStudio2019)
+				if (Target.WindowsPlatform.Compiler.IsMSVC())
 				{
 					Arguments.Add("/experimental:deterministic");
 				}
@@ -655,7 +655,7 @@ namespace UnrealBuildTool
 			}
 
 			//@todo: Disable warnings for VS2017. These should be reenabled as we clear the reasons for them out of the engine source and the VS2015 toolchain evolves.
-			if (Target.WindowsPlatform.Compiler >= WindowsCompiler.VisualStudio2019)
+			if (Target.WindowsPlatform.Compiler.IsMSVC())
 			{
 				// Disable shadow variable warnings
 				if (CompileEnvironment.ShadowVariableWarningLevel == WarningLevel.Off)
@@ -1189,7 +1189,7 @@ namespace UnrealBuildTool
 			// Generate the timing info
 			if (CompileEnvironment.bPrintTimingInfo || Target.WindowsPlatform.bCompilerTrace)
 			{
-				if (Target.WindowsPlatform.Compiler >= WindowsCompiler.VisualStudio2019)
+				if (Target.WindowsPlatform.Compiler.IsMSVC())
 				{
 					if (CompileEnvironment.bPrintTimingInfo)
 					{
@@ -1370,7 +1370,7 @@ namespace UnrealBuildTool
 
 				if (CompileEnvironment.bGenerateDependenciesFile)
 				{
-					if (EnvVars.ToolChainVersion >= VersionNumber.Parse("14.27") && Target.WindowsPlatform.Compiler >= WindowsCompiler.VisualStudio2019 && !CompileAction.ForceClFilter)
+					if (EnvVars.ToolChainVersion >= VersionNumber.Parse("14.27") && Target.WindowsPlatform.Compiler.IsMSVC() && !CompileAction.ForceClFilter)
 					{
 						CompileAction.DependencyListFile = FileItem.GetItemByFileReference(FileReference.Combine(OutputDir, String.Format("{0}.json", SourceFile.Location.GetFileName())));
 					}
diff --git a/Engine/Source/Programs/UnrealBuildTool/ProjectFiles/Rider/RiderProjectFile.cs b/Engine/Source/Programs/UnrealBuildTool/ProjectFiles/Rider/RiderProjectFile.cs
--- a/Engine/Source/Programs/UnrealBuildTool/ProjectFiles/Rider/RiderProjectFile.cs
+++ b/Engine/Source/Programs/UnrealBuildTool/ProjectFiles/Rider/RiderProjectFile.cs
@@ -637,7 +637,7 @@ namespace UnrealBuildTool
 				ToolchainInfo.Architecture = WindowsExports.GetArchitectureSubpath(CurrentTarget.Rules.WindowsPlatform.Architecture);
 				
 				WindowsCompiler WindowsPlatformCompiler = CurrentTarget.Rules.WindowsPlatform.Compiler;
-				ToolchainInfo.bStrictConformanceMode = WindowsPlatformCompiler >= WindowsCompiler.VisualStudio2019 && CurrentTarget.Rules.WindowsPlatform.bStrictConformanceMode;
+				ToolchainInfo.bStrictConformanceMode = WindowsPlatformCompiler.IsMSVC() && CurrentTarget.Rules.WindowsPlatform.bStrictConformanceMode;
 				ToolchainInfo.Compiler = WindowsPlatformCompiler.ToString();
 			}
 			else
