Remove references to i686

#trivial

diff --git a/Engine/Source/Programs/UnrealBuildTool/Platform/Linux/UEBuildLinux.cs b/Engine/Source/Programs/UnrealBuildTool/Platform/Linux/UEBuildLinux.cs
--- a/Engine/Source/Programs/UnrealBuildTool/Platform/Linux/UEBuildLinux.cs
+++ b/Engine/Source/Programs/UnrealBuildTool/Platform/Linux/UEBuildLinux.cs
@@ -22,9 +22,6 @@ namespace UnrealBuildTool
 
 		/** Arm64, ARM 64-bit */
 		AArch64UnknownLinuxGnueabi,
-
-		/** i686, Intel 32-bit */
-		I686UnknownLinuxGnu
 	}
 
 	/// <summary>
diff --git a/Engine/Source/Programs/UnrealBuildTool/ProjectFiles/Rider/RiderProjectFile.cs b/Engine/Source/Programs/UnrealBuildTool/ProjectFiles/Rider/RiderProjectFile.cs
--- a/Engine/Source/Programs/UnrealBuildTool/ProjectFiles/Rider/RiderProjectFile.cs
+++ b/Engine/Source/Programs/UnrealBuildTool/ProjectFiles/Rider/RiderProjectFile.cs
@@ -503,8 +503,7 @@ namespace UnrealBuildTool
 				if (string.IsNullOrEmpty(UseLibcxxEnvVarOverride) || UseLibcxxEnvVarOverride == "1")
 				{
 					if (Target.Architecture.StartsWith("x86_64") ||
-					    Target.Architecture.StartsWith("aarch64") ||
-					    Target.Architecture.StartsWith("i686"))
+					    Target.Architecture.StartsWith("aarch64"))
 					{
 						// libc++ include directories
 						Writer.WriteValue(Path.Combine(EngineDirectory, "Source/ThirdParty/Unix/LibCxx/include/"));
@@ -514,7 +513,7 @@ namespace UnrealBuildTool
 
 				UEBuildPlatform BuildPlatform;
 
-				if (Target.Architecture.StartsWith("x86_64") || Target.Architecture.StartsWith("i686"))
+				if (Target.Architecture.StartsWith("x86_64"))
 				{
 					BuildPlatform = UEBuildPlatform.GetBuildPlatform(UnrealTargetPlatform.Linux);
 				}
