UnrealBuildTool: lld-link.exe should only be used to link binaries, not libraries, when enabled on clang

#rnx

diff --git a/Engine/Source/Programs/UnrealBuildTool/Platform/Windows/VCEnvironment.cs b/Engine/Source/Programs/UnrealBuildTool/Platform/Windows/VCEnvironment.cs
--- a/Engine/Source/Programs/UnrealBuildTool/Platform/Windows/VCEnvironment.cs
+++ b/Engine/Source/Programs/UnrealBuildTool/Platform/Windows/VCEnvironment.cs
@@ -237,15 +237,14 @@ namespace UnrealBuildTool
 		/// </summary>
 		static FileReference GetLinkerToolPath(UnrealTargetPlatform Platform, WindowsCompiler Compiler, DirectoryReference CompilerDir, DirectoryReference DefaultLinkerDir)
 		{
-			// If we were asked to use Clang, then we'll redirect the path to the compiler to the LLVM installation directory
+			// Regardless of the target, if we're linking on a 64 bit machine, we want to use the 64 bit linker (it's faster than the 32 bit linker)
 			if (Compiler == WindowsCompiler.Clang && WindowsPlatform.bAllowClangLinker)
 			{
 				return FileReference.Combine(CompilerDir, "bin", "lld-link.exe");
 			}
-			else
-			{
-				return FileReference.Combine(DefaultLinkerDir, "link.exe");
-			}
+
+			return FileReference.Combine(DefaultLinkerDir, "link.exe");
+
 		}
 
 		/// <summary>
@@ -254,15 +253,13 @@ namespace UnrealBuildTool
 		static FileReference GetLibraryLinkerToolPath(UnrealTargetPlatform Platform, WindowsCompiler Compiler, DirectoryReference CompilerDir, DirectoryReference DefaultLinkerDir)
 		{
 			// Regardless of the target, if we're linking on a 64 bit machine, we want to use the 64 bit linker (it's faster than the 32 bit linker)
-			// If we were asked to use Clang, then we'll redirect the path to the compiler to the LLVM installation directory
 			if (Compiler == WindowsCompiler.Clang && WindowsPlatform.bAllowClangLinker)
 			{
-				return FileReference.Combine(CompilerDir, "bin", "lld-link.exe");
-			}
-			else
-			{
-				return FileReference.Combine(DefaultLinkerDir, "lib.exe");
+				// @todo: lld-link is not currently working for building .lib
+				//return FileReference.Combine(CompilerDir, "bin", "lld-link.exe");
 			}
+
+			return FileReference.Combine(DefaultLinkerDir, "lib.exe");
 		}
 
 		/// <summary>
