Remove (remnants of) GCC support from Linux toolchain, remove references to i686, assume minimal valid clang version of 10.0.0

diff --git a/Engine/Source/Programs/UnrealBuildTool/Platform/Linux/LinuxCommon.cs b/Engine/Source/Programs/UnrealBuildTool/Platform/Linux/LinuxCommon.cs
--- a/Engine/Source/Programs/UnrealBuildTool/Platform/Linux/LinuxCommon.cs
+++ b/Engine/Source/Programs/UnrealBuildTool/Platform/Linux/LinuxCommon.cs
@@ -57,11 +57,6 @@ namespace UnrealBuildTool
 
 			return null;
 		}
-
-		public static string? WhichGcc()
-		{
-			return Which("g++");
-		}
 	}
 }
 
diff --git a/Engine/Source/Programs/UnrealBuildTool/Platform/Linux/LinuxPlatformSDK.cs b/Engine/Source/Programs/UnrealBuildTool/Platform/Linux/LinuxPlatformSDK.cs
--- a/Engine/Source/Programs/UnrealBuildTool/Platform/Linux/LinuxPlatformSDK.cs
+++ b/Engine/Source/Programs/UnrealBuildTool/Platform/Linux/LinuxPlatformSDK.cs
@@ -262,7 +262,7 @@ namespace UnrealBuildTool
 
 			if (ForceUseSystemCompiler())
 			{
-				if (!String.IsNullOrEmpty(LinuxCommon.WhichClang()) || !String.IsNullOrEmpty(LinuxCommon.WhichGcc()))
+				if (!String.IsNullOrEmpty(LinuxCommon.WhichClang()))
 				{
 					return SDKStatus.Valid;
 				}
diff --git a/Engine/Source/Programs/UnrealBuildTool/Platform/Linux/LinuxToolChain.cs b/Engine/Source/Programs/UnrealBuildTool/Platform/Linux/LinuxToolChain.cs
--- a/Engine/Source/Programs/UnrealBuildTool/Platform/Linux/LinuxToolChain.cs
+++ b/Engine/Source/Programs/UnrealBuildTool/Platform/Linux/LinuxToolChain.cs
@@ -65,9 +65,6 @@ namespace UnrealBuildTool
 		/** Flavor of the current build (target triplet)*/
 		string Architecture;
 
-		/** Cache to avoid making multiple checks for lld availability/usability */
-		bool bUseLld = false;
-
 		/** Whether the compiler is set up to produce PIE executables by default */
 		bool bSuppressPIE = false;
 
@@ -115,19 +112,12 @@ namespace UnrealBuildTool
 
 				// use native linux toolchain
 				ClangPath = LinuxCommon.WhichClang();
-				GCCPath = LinuxCommon.WhichGcc();
 				ArPath = LinuxCommon.Which("ar");
 				LlvmArPath = LinuxCommon.Which("llvm-ar");
 				RanlibPath = LinuxCommon.Which("ranlib");
 				StripPath = LinuxCommon.Which("strip");
 				ObjcopyPath = LinuxCommon.Which("objcopy");
 
-				// if clang is available, zero out gcc (@todo: support runtime switching?)
-				if (!String.IsNullOrEmpty(ClangPath))
-				{
-					GCCPath = null;
-				}
-
 				// When compiling on Linux, use a faster way to relink circularly dependent libraries.
 				// Race condition between actions linking to the .so and action overwriting it is avoided thanks to inodes
 				bUseFixdeps = false;
@@ -152,7 +142,6 @@ namespace UnrealBuildTool
 				ToolchainInfo = String.Format("toolchain located at '{0}'", BaseLinuxPath);
 
 				// set up the path to our toolchain
-				GCCPath = "";
 				ClangPath   = Path.Combine(BaseLinuxPath, @"bin", "clang++" + GetHostPlatformBinarySuffix());
 				ArPath      = Path.Combine(Path.Combine(BaseLinuxPath, String.Format("bin/{0}-{1}", Architecture, "ar" + GetHostPlatformBinarySuffix())));
 				LlvmArPath  = Path.Combine(Path.Combine(BaseLinuxPath, String.Format("bin/{0}", "llvm-ar" + GetHostPlatformBinarySuffix())));
@@ -181,26 +170,15 @@ namespace UnrealBuildTool
 			CheckDefaultCompilerSettings();
 
 			// refuse to use compilers that we know won't work
-			// disable that only if you are a dev and you know what you are doing
-			if (!UsingClang())
-			{
-				throw new BuildException("Unable to build: no compatible clang version found. Please run Setup.sh");
-			}
+
 			// prevent unknown clangs since the build is likely to fail on too old or too new compilers
-			else if ((CompilerVersionMajor * 10 + CompilerVersionMinor) > 110 || (CompilerVersionMajor * 10 + CompilerVersionMinor) < 60)
+			if ((CompilerVersionMajor * 10 + CompilerVersionMinor) > 110 || (CompilerVersionMajor * 10 + CompilerVersionMinor) < 100)
 			{
 				throw new BuildException(
-					string.Format("This version of the Unreal Engine can only be compiled with clang 11.0, 10.0, 9.0, 8.0, 7.0 and 6.0. clang {0} may not build it - please use a different version.",
+					string.Format("This version of the Unreal Engine can only be compiled with clang 11.0, and 10.0. clang {0} may not build it - please use a different version.",
 						CompilerVersionString)
 					);
 			}
-
-			// trust lld only for clang 5.x and above (FIXME: also find if present on the system?)
-			// NOTE: with early version you can run into errors like "failed to compute relocation:" and others
-			bUseLld = (CompilerVersionMajor >= 5);
-
-			// Add --gdb-index for Clang 9.0 and higher
-			bGdbIndexSection = (CompilerVersionMajor >= 9);
 		}
 
 		public LinuxToolChain(UnrealTargetPlatform InPlatform, string InArchitecture, LinuxPlatformSDK InSDK, bool InPreservePSYM = false, LinuxToolChainOptions InOptions = LinuxToolChainOptions.None)
@@ -228,11 +206,6 @@ namespace UnrealBuildTool
 		}
 
 
-		protected virtual bool UsingClang()
-		{
-			return !String.IsNullOrEmpty(ClangPath);
-		}
-
 		/// <summary>
 		/// Splits compiler version string into numerical components, leaving unchanged if not known
 		/// </summary>
@@ -375,22 +348,7 @@ namespace UnrealBuildTool
 				Proc.StartInfo.RedirectStandardOutput = true;
 				Proc.StartInfo.RedirectStandardError = true;
 
-				if (!String.IsNullOrEmpty(GCCPath))
-				{
-					Proc.StartInfo.FileName = GCCPath;
-					Proc.StartInfo.Arguments = " -dumpversion";
-
-					Proc.Start();
-					Proc.WaitForExit();
-
-					if (Proc.ExitCode == 0)
-					{
-						// read just the first string
-						CompilerVersionString = Proc.StandardOutput.ReadLine();
-						DetermineCompilerMajMinPatchFromVersionString();
-					}
-				}
-				else if (!String.IsNullOrEmpty(ClangPath))
+				if (!String.IsNullOrEmpty(ClangPath))
 				{
 					Proc.StartInfo.FileName = ClangPath;
 					Proc.StartInfo.Arguments = " --version";
@@ -543,7 +501,7 @@ namespace UnrealBuildTool
 			if (string.IsNullOrEmpty(UseLibcxxEnvVarOverride) || UseLibcxxEnvVarOverride == "1")
 			{
 				// at the moment ARM32 libc++ remains missing
-				return Architecture.StartsWith("x86_64") || Architecture.StartsWith("aarch64") || Architecture.StartsWith("i686");
+				return Architecture.StartsWith("x86_64") || Architecture.StartsWith("aarch64");
 			}
 			return false;
 		}
@@ -594,7 +552,7 @@ namespace UnrealBuildTool
 
 			Result += " -Wall -Werror";
 
-			if (!CompileEnvironment.Architecture.StartsWith("x86_64") && !CompileEnvironment.Architecture.StartsWith("i686"))
+			if (!CompileEnvironment.Architecture.StartsWith("x86_64"))
 			{
 				Result += " -funwind-tables";               // generate unwind tables as they are needed for backtrace (on x86(64) they are generated implicitly)
 			}
@@ -621,82 +579,42 @@ namespace UnrealBuildTool
 				Result += " -fvisibility-inlines-hidden";
 			}
 
-			if (String.IsNullOrEmpty(ClangPath))
-			{
-				// GCC only option
-				Result += " -fno-strict-aliasing";
-				Result += " -Wno-sign-compare"; // needed to suppress: comparison between signed and unsigned integer expressions
-				Result += " -Wno-enum-compare"; // Stats2.h triggers this (alignof(int64) <= DATA_ALIGN)
-				Result += " -Wno-return-type"; // Variant.h triggers this
-				Result += " -Wno-unused-local-typedefs";
-				Result += " -Wno-multichar";
-				Result += " -Wno-unused-but-set-variable";
-				Result += " -Wno-strict-overflow"; // Array.h:518
-			}
-			else
+			// Clang only options
+			if (CrossCompiling())
 			{
-				// Clang only options
-				if (CrossCompiling())
-				{
-					if (BuildHostPlatform.Current.Platform == UnrealTargetPlatform.Win64)
-					{
-						Result += " -fdiagnostics-format=msvc";     // make diagnostics compatible with MSVC when cross-compiling
-					}
-					else if (Log.ColorConsoleOutput)
-					{
-						Result += " -fcolor-diagnostics";
-					}
-				}
-
-				// output full paths to the files when the build fails, required 4.0+ of clang
-				if (CompilerVersionGreaterOrEqual(4, 0, 0))
-				{
-					Result += " -fdiagnostics-absolute-paths";
-				}
-
-				Result += " -Wno-unused-private-field";     // MultichannelTcpSocket.h triggers this, possibly more
-				// this hides the "warning : comparison of unsigned expression < 0 is always false" type warnings due to constant comparisons, which are possible with template arguments
-				Result += " -Wno-tautological-compare";
-
-				// this switch is understood by clang 3.5.0, but not clang-3.5 as packaged by Ubuntu 14.04 atm
-				if (CompilerVersionGreaterOrEqual(3, 5, 0))
-				{
-					Result += " -Wno-undefined-bool-conversion";	// hides checking if 'this' pointer is null
-				}
-
-				if (CompilerVersionGreaterOrEqual(3, 6, 0))
+				if (BuildHostPlatform.Current.Platform == UnrealTargetPlatform.Win64)
 				{
-					Result += " -Wno-unused-local-typedef";	// clang is being overly strict here? PhysX headers trigger this.
-					Result += " -Wno-inconsistent-missing-override";	// these have to be suppressed for UE 4.8, should be fixed later.
+					Result += " -fdiagnostics-format=msvc";     // make diagnostics compatible with MSVC when cross-compiling
 				}
-
-				if (CompilerVersionGreaterOrEqual(3, 9, 0))
+				else if (Log.ColorConsoleOutput)
 				{
-					Result += " -Wno-undefined-var-template"; // not really a good warning to disable
-				}
-
-				if (CompilerVersionGreaterOrEqual(5, 0, 0))
-				{
-					Result += " -Wno-unused-lambda-capture";  // suppressed because capturing of compile-time constants is seemingly inconsistent. And MSVC doesn't do that.
+					Result += " -fcolor-diagnostics";
 				}
 			}
 
+			Result += " -fdiagnostics-absolute-paths"; // output full paths to the files when the build fails
+
+			Result += " -Wno-unused-private-field"; // MultichannelTcpSocket.h triggers this, possibly more
+			Result += " -Wno-tautological-compare"; // this hides the "warning : comparison of unsigned expression < 0 is always false" type warnings due to constant comparisons, which are possible with template arguments
+			Result += " -Wno-undefined-bool-conversion"; // hides checking if 'this' pointer is null
+			Result += " -Wno-unused-local-typedef";	// clang is being overly strict here? PhysX headers trigger this.
+			Result += " -Wno-inconsistent-missing-override";	// these have to be suppressed for UE 4.8, should be fixed later.
+			Result += " -Wno-undefined-var-template"; // not really a good warning to disable
+			Result += " -Wno-unused-lambda-capture";  // suppressed because capturing of compile-time constants is seemingly inconsistent. And MSVC doesn't do that.
 			Result += " -Wno-unused-variable";
+
 			if (CompilerVersionGreaterOrEqual(13, 0, 0))
 			{
 				Result += " -Wno-unused-but-set-variable";
 				Result += " -Wno-unused-but-set-parameter";
 				Result += " -Wno-ordered-compare-function-pointers";
 			}
-			// this will hide the warnings about static functions in headers that aren't used in every single .cpp file
-			Result += " -Wno-unused-function";
-			// this hides the "enumeration value 'XXXXX' not handled in switch [-Wswitch]" warnings - we should maybe remove this at some point and add UE_LOG(, Fatal, ) to default cases
-			Result += " -Wno-switch";
+
+			Result += " -Wno-unused-function"; // this will hide the warnings about static functions in headers that aren't used in every single .cpp file
+			Result += " -Wno-switch"; // this hides the "enumeration value 'XXXXX' not handled in switch [-Wswitch]" warnings - we should maybe remove this at some point and add UE_LOG(, Fatal, ) to default cases
 			Result += " -Wno-unknown-pragmas";			// Slate triggers this (with its optimize on/off pragmas)
-			// needed to suppress warnings about using offsetof on non-POD types.
-			Result += " -Wno-invalid-offsetof";
-			// we use this feature to allow static FNames.
-			Result += " -Wno-gnu-string-literal-operator-template";
+			Result += " -Wno-invalid-offsetof"; // needed to suppress warnings about using offsetof on non-POD types.
+			Result += " -Wno-gnu-string-literal-operator-template"; // we use this feature to allow static FNames.
 
 			// Profile Guided Optimization (PGO) and Link Time Optimization (LTO)
 			// Whether we actually can enable that is checked in CanUseAdvancedLinkerFeatures() earlier
@@ -838,7 +756,7 @@ namespace UnrealBuildTool
 			Result += ArchitectureSpecificDefines(CompileEnvironment.Architecture);
 			if (CrossCompiling())
 			{
-				if (UsingClang() && !string.IsNullOrEmpty(CompileEnvironment.Architecture))
+				if (!String.IsNullOrEmpty(CompileEnvironment.Architecture))
 				{
 					Result += String.Format(" -target {0}", CompileEnvironment.Architecture);        // Set target triple
 				}
@@ -1061,14 +979,7 @@ namespace UnrealBuildTool
 
 				if (bSuppressPIE)
 				{
-					if (CompilerVersionGreaterOrEqual(7, 0, 0))
-					{
-						Result += " -Wl,-no-pie";
-					}
-					else
-					{
-						Result += " -Wl,-nopie";
-					}
+					Result += " -Wl,-no-pie";
 				}
 			}
 
@@ -1108,10 +1019,7 @@ namespace UnrealBuildTool
 
 			if (CrossCompiling())
 			{
-				if (UsingClang())
-				{
-					Result += String.Format(" -target {0}", LinkEnvironment.Architecture);        // Set target triple
-				}
+				Result += String.Format(" -target {0}", LinkEnvironment.Architecture);        // Set target triple
 				string SysRootPath = BaseLinuxPath!.TrimEnd(new char[] { '\\', '/' });
 				Result += String.Format(" \"--sysroot={0}\"", SysRootPath);
 
@@ -1137,7 +1045,6 @@ namespace UnrealBuildTool
 		protected bool bIsCrossCompiling;
 		protected string? BaseLinuxPath;
 		protected string? ClangPath;
-		protected string? GCCPath;
 		protected string? ArPath;
 		protected string? LlvmArPath;
 		protected string? RanlibPath;
@@ -1215,18 +1122,13 @@ namespace UnrealBuildTool
 		{
 			Log.TraceInformation("------- Build details --------");
 			Log.TraceInformation("Using {0}.", ToolchainInfo);
-			Log.TraceInformation("Using {0} ({1}) version '{2}' (string), {3} (major), {4} (minor), {5} (patch)",
-				String.IsNullOrEmpty(ClangPath) ? "gcc" : "clang",
-				String.IsNullOrEmpty(ClangPath) ? GCCPath : ClangPath,
-				CompilerVersionString, CompilerVersionMajor, CompilerVersionMinor, CompilerVersionPatch);
+			Log.TraceInformation("Using clang ({0}) version '{1}' (string), {2} (major), {3} (minor), {4} (patch)",
+				ClangPath, CompilerVersionString, CompilerVersionMajor, CompilerVersionMinor, CompilerVersionPatch);
 
-			if (UsingClang())
-			{
-				// inform the user which C++ library the engine is going to be compiled against - important for compatibility with third party code that uses STL
-				Log.TraceInformation("Using {0} standard C++ library.", ShouldUseLibcxx(CompileEnvironment.Architecture) ? "bundled libc++" : "compiler default (most likely libstdc++)");
-				Log.TraceInformation("Using {0}", UsingLld(CompileEnvironment.Architecture) ? "lld linker" : "default linker (ld)");
-				Log.TraceInformation("Using {0}", !String.IsNullOrEmpty(LlvmArPath) ? String.Format("llvm-ar : {0}", LlvmArPath) : String.Format("ar and ranlib: {0}, {1}", GetArPath(CompileEnvironment.Architecture), GetRanlibPath(CompileEnvironment.Architecture)));
-			}
+			// inform the user which C++ library the engine is going to be compiled against - important for compatibility with third party code that uses STL
+			Log.TraceInformation("Using {0} standard C++ library.", ShouldUseLibcxx(CompileEnvironment.Architecture) ? "bundled libc++" : "compiler default (most likely libstdc++)");
+			Log.TraceInformation("Using {0}", UsingLld(CompileEnvironment.Architecture) ? "lld linker" : "default linker (ld)");
+			Log.TraceInformation("Using {0}", !String.IsNullOrEmpty(LlvmArPath) ? String.Format("llvm-ar : {0}", LlvmArPath) : String.Format("ar and ranlib: {0}, {1}", GetArPath(CompileEnvironment.Architecture), GetRanlibPath(CompileEnvironment.Architecture)));
 
 			if (Options.HasFlag(LinuxToolChainOptions.EnableAddressSanitizer) ||
 				Options.HasFlag(LinuxToolChainOptions.EnableThreadSanitizer) ||
@@ -1407,14 +1309,7 @@ namespace UnrealBuildTool
 				}
 
 				CompileAction.WorkingDirectory = UnrealBuildTool.EngineSourceDirectory;
-				if (!UsingClang())
-				{
-					CompileAction.CommandPath = new FileReference(GCCPath!);
-				}
-				else
-				{
-					CompileAction.CommandPath = new FileReference(ClangPath!);
-				}
+				CompileAction.CommandPath = new FileReference(ClangPath!);
 
 				string AllArguments = (Arguments + FileArguments + CompileEnvironment.AdditionalArguments);
 				// all response lines should have / instead of \, but we cannot just bulk-replace it here since some \ are used to escape quotes, e.g. Definitions.Add("FOO=TEXT(\"Bar\")");
@@ -1442,7 +1337,7 @@ namespace UnrealBuildTool
 
 		protected virtual bool UsingLld(string Architecture)
 		{
-			return bUseLld && (Architecture.StartsWith("x86_64") || (CompilerVersionMajor >= 9));
+			return Architecture.StartsWith("x86_64");
 		}
 
 		/// <summary>
@@ -1602,14 +1497,7 @@ namespace UnrealBuildTool
 			LinkAction.WorkingDirectory = UnrealBuildTool.EngineSourceDirectory;
 
 			string LinkCommandString;
-			if (String.IsNullOrEmpty(ClangPath))
-			{
-				LinkCommandString = "\"" + GCCPath + "\"";
-			}
-			else
-			{
-				LinkCommandString = "\"" + ClangPath + "\"";
-			}
+			LinkCommandString = "\"" + ClangPath + "\"";
 
 			// Get link arguments.
 			LinkCommandString += GetLinkArguments(LinkEnvironment);
diff --git a/Engine/Source/Programs/UnrealBuildTool/ProjectFiles/CMake/CMakefileGenerator.cs b/Engine/Source/Programs/UnrealBuildTool/ProjectFiles/CMake/CMakefileGenerator.cs
--- a/Engine/Source/Programs/UnrealBuildTool/ProjectFiles/CMake/CMakefileGenerator.cs
+++ b/Engine/Source/Programs/UnrealBuildTool/ProjectFiles/CMake/CMakefileGenerator.cs
@@ -290,10 +290,6 @@ namespace UnrealBuildTool
 				BuildCommand = "cd \"" + UE4RootPath + "\" && bash \"" + UE4RootPath + "/Engine/Build/BatchFiles/" + HostArchitecture + "/Build.sh\"";
 
 				string? CompilerPath = LinuxCommon.WhichClang();
-				if (CompilerPath == null)
-				{
-					CompilerPath = LinuxCommon.WhichGcc();
-				}
 				SetCompiler = "set(CMAKE_CXX_COMPILER " + CompilerPath + ")\n\n";
 			}
 			else
