UnrealBuildTool: Remove bBuildLargeAddressAwareBinary as it is no longer needed as Win32 support was removed

#rnx

diff --git a/Engine/Source/Programs/UnrealBuildTool/Platform/Windows/UEBuildWindows.cs b/Engine/Source/Programs/UnrealBuildTool/Platform/Windows/UEBuildWindows.cs
--- a/Engine/Source/Programs/UnrealBuildTool/Platform/Windows/UEBuildWindows.cs
+++ b/Engine/Source/Programs/UnrealBuildTool/Platform/Windows/UEBuildWindows.cs
@@ -306,11 +306,6 @@ namespace UnrealBuildTool
 		[XmlConfigFile(Category = "WindowsPlatform")]
 		public string AdditionalLinkerOptions = "";
 
-		/// <summary>
-		/// True if we allow using addresses larger than 2GB on 32 bit builds
-		/// </summary>
-		public bool bBuildLargeAddressAwareBinary = true;
-
 		/// <summary>
 		/// Create an image that can be hot patched (/FUNCTIONPADMIN)
 		/// </summary>
@@ -610,11 +605,6 @@ namespace UnrealBuildTool
 			get { return Inner.AdditionalLinkerOptions; }
 		}
 
-		public bool bBuildLargeAddressAwareBinary
-		{
-			get { return Inner.bBuildLargeAddressAwareBinary; }
-		}
-
 		public bool bCreateHotpatchableImage
 		{
 			get { return Inner.bCreateHotPatchableImage; }
@@ -867,11 +857,6 @@ namespace UnrealBuildTool
 		/// </summary>
 		public static readonly bool bAllowClangLinker = false;
 
-		/// <summary>
-		/// True if we allow using addresses larger than 2GB on 32 bit builds
-		/// </summary>
-		public static readonly bool bBuildLargeAddressAwareBinary = true;
-
 		MicrosoftPlatformSDK SDK;
 
 		/// <summary>
