UnrealBuildTool: Linux clang 12 -fbinutils-version=2.36

#rnx

diff --git a/Engine/Source/Programs/UnrealBuildTool/Platform/Linux/LinuxToolChain.cs b/Engine/Source/Programs/UnrealBuildTool/Platform/Linux/LinuxToolChain.cs
--- a/Engine/Source/Programs/UnrealBuildTool/Platform/Linux/LinuxToolChain.cs
+++ b/Engine/Source/Programs/UnrealBuildTool/Platform/Linux/LinuxToolChain.cs
@@ -552,6 +552,11 @@ namespace UnrealBuildTool
 
 			Result += " -Wall -Werror";
 
+			if (CompilerVersionGreaterOrEqual(12, 0, 0))
+			{
+				Result += " -fbinutils-version=2.36";
+			}
+
 			if (!CompileEnvironment.Architecture.StartsWith("x86_64"))
 			{
 				Result += " -funwind-tables";               // generate unwind tables as they are needed for backtrace (on x86(64) they are generated implicitly)
