- Allow for UBT to look in the ini for DefaultEditorTarget (etc) when passed -TargetType=Editor and there are multiple editor target found

diff --git a/Engine/Source/Programs/UnrealBuildTool/System/RulesAssembly.cs b/Engine/Source/Programs/UnrealBuildTool/System/RulesAssembly.cs
--- a/Engine/Source/Programs/UnrealBuildTool/System/RulesAssembly.cs
+++ b/Engine/Source/Programs/UnrealBuildTool/System/RulesAssembly.cs
@@ -759,14 +759,26 @@ namespace UnrealBuildTool
 			}
 			else
 			{
-				if(Matches.Count == 1)
+				if (Matches.Count == 1)
 				{
 					return Matches[0];
 				}
-				else
+
+				// attempt to get a default target (like DefaultEditorTarget) from the Engine.ini
+				string KeyName = $"Default{Type}Target";
+				string? DefaultTargetName;
+				// read in the engine config hierarchy and get the value
+				ConfigHierarchy EngineConfig = ConfigCache.ReadHierarchy(ConfigHierarchyType.Engine, ProjectFile?.Directory, Platform);
+				if (EngineConfig.GetString("/Script/BuildSettings.BuildSettings", KeyName, out DefaultTargetName))
 				{
-					throw new BuildException("Found multiple targets with TargetType={0}: {1}", Type, String.Join(", ", Matches));
+					// if a value was found, make sure that this is one of the found targets
+					if (Matches.Contains(DefaultTargetName))
+					{
+						return DefaultTargetName;
+					}
 				}
+				
+				throw new BuildException("Found multiple targets with TargetType={0}: {1}.\nSpecify a default with a {2} entry in [/Script/BuildSettings.BuildSettings] section of your DefaultEngine.ini", Type, String.Join(", ", Matches), KeyName);
 			}
 		}
 
