UnrealBuildTool: rename MasterProject settings as PrimaryProject, and some other uses of 'master'

Variable, function and class renames only. XML-facing variables included deprecated support for the old names.

Renaming of the MasterProjectName.txt file will be handled in a subsequent change.

diff --git a/Engine/Source/Programs/UnrealBuildTool/ProjectFiles/CMake/CMakefileGenerator.cs b/Engine/Source/Programs/UnrealBuildTool/ProjectFiles/CMake/CMakefileGenerator.cs
--- a/Engine/Source/Programs/UnrealBuildTool/ProjectFiles/CMake/CMakefileGenerator.cs
+++ b/Engine/Source/Programs/UnrealBuildTool/ProjectFiles/CMake/CMakefileGenerator.cs
@@ -214,12 +214,12 @@ namespace UnrealBuildTool
 		
 
 		/// <summary>
-		/// Writes the master project file (e.g. Visual Studio Solution file)
+		/// Writes the primary project file (e.g. Visual Studio Solution file)
 		/// </summary>
 		/// <param name="UBTProject">The UnrealBuildTool project</param>
 		/// <param name="PlatformProjectGenerators">The registered platform project generators</param>
 		/// <returns>True if successful</returns>
-		protected override bool WriteMasterProjectFile(ProjectFile? UBTProject, PlatformProjectGeneratorCollection PlatformProjectGenerators)
+		protected override bool WritePrimaryProjectFile(ProjectFile? UBTProject, PlatformProjectGeneratorCollection PlatformProjectGenerators)
 		{
 			return true;
 		}
@@ -592,7 +592,7 @@ namespace UnrealBuildTool
 				CMakefileContent.AppendLine("add_executable(FakeTarget ${ENGINE_HEADER_FILES} ${ENGINE_SOURCE_FILES} ${ENGINE_CSHARP_FILES} ${ENGINE_SHADER_FILES} ${ENGINE_CONFIG_FILES} ${PROJECT_HEADER_FILES} ${PROJECT_SOURCE_FILES} ${PROJECT_CSHARP_FILES} ${PROJECT_SHADER_FILES} ${PROJECT_CONFIG_FILES})");
 			}
 
-			string FullFileName = Path.Combine(MasterProjectPath.FullName, ProjectFileName);
+			string FullFileName = Path.Combine(PrimaryProjectPath.FullName, ProjectFileName);
 
 			// Write out CMake files
 			bool bWriteMakeList = WriteFileIfChanged(FullFileName, CMakefileContent.ToString());
@@ -702,7 +702,7 @@ namespace UnrealBuildTool
 		/// Adds the include directory to the list, after converting it to relative to UE4 root
 		private static string GetIncludeDirectory(string IncludeDir, string ProjectDir)
 		{
-			string FullProjectPath = Path.GetFullPath(MasterProjectPath.FullName);
+			string FullProjectPath = Path.GetFullPath(PrimaryProjectPath.FullName);
 			string FullPath = "";
 			// Check for paths outside of both the engine and the project
 			if (Path.IsPathRooted(IncludeDir) &&
@@ -778,13 +778,13 @@ namespace UnrealBuildTool
 			return new CMakefileProjectFile(InitFilePath, BaseDir);
 		}
 
-		public override void CleanProjectFiles(DirectoryReference InMasterProjectDirectory, string InMasterProjectName, DirectoryReference InIntermediateProjectFilesDirectory)
+		public override void CleanProjectFiles(DirectoryReference InPrimaryProjectDirectory, string InPrimaryProjectName, DirectoryReference InIntermediateProjectFilesDirectory)
 		{
 			// Remove Project File
-			FileReference MasterProjectFile = FileReference.Combine(InMasterProjectDirectory, ProjectFileName);
-			if (FileReference.Exists(MasterProjectFile))
+			FileReference PrimaryProjectFile = FileReference.Combine(InPrimaryProjectDirectory, ProjectFileName);
+			if (FileReference.Exists(PrimaryProjectFile))
 			{
-				FileReference.Delete(MasterProjectFile);
+				FileReference.Delete(PrimaryProjectFile);
 			}
 
 			// Remove Headers Files
diff --git a/Engine/Source/Programs/UnrealBuildTool/ProjectFiles/CodeLite/CodeLiteGenerator.cs b/Engine/Source/Programs/UnrealBuildTool/ProjectFiles/CodeLite/CodeLiteGenerator.cs
--- a/Engine/Source/Programs/UnrealBuildTool/ProjectFiles/CodeLite/CodeLiteGenerator.cs
+++ b/Engine/Source/Programs/UnrealBuildTool/ProjectFiles/CodeLite/CodeLiteGenerator.cs
@@ -47,15 +47,15 @@ namespace UnrealBuildTool
 				return ".project";
 			}
 		}
-		protected override bool WriteMasterProjectFile(ProjectFile? UBTProject, PlatformProjectGeneratorCollection PlatformProjectGenerators)
+		protected override bool WritePrimaryProjectFile(ProjectFile? UBTProject, PlatformProjectGeneratorCollection PlatformProjectGenerators)
 		{
-			string SolutionFileName = MasterProjectName + SolutionExtension;
-			string CodeCompletionFile = MasterProjectName + CodeCompletionFileName;
-			string CodeCompletionPreProcessorFile = MasterProjectName + CodeCompletionPreProcessorFileName;
+			string SolutionFileName = PrimaryProjectName + SolutionExtension;
+			string CodeCompletionFile = PrimaryProjectName + CodeCompletionFileName;
+			string CodeCompletionPreProcessorFile = PrimaryProjectName + CodeCompletionPreProcessorFileName;
 
-			string FullCodeLiteMasterFile = Path.Combine(MasterProjectPath.FullName, SolutionFileName);
-			string FullCodeLiteCodeCompletionFile = Path.Combine(MasterProjectPath.FullName, CodeCompletionFile);
-			string FullCodeLiteCodeCompletionPreProcessorFile = Path.Combine(MasterProjectPath.FullName, CodeCompletionPreProcessorFile);
+			string FullCodeLitePrimaryFile = Path.Combine(PrimaryProjectPath.FullName, SolutionFileName);
+			string FullCodeLiteCodeCompletionFile = Path.Combine(PrimaryProjectPath.FullName, CodeCompletionFile);
+			string FullCodeLiteCodeCompletionPreProcessorFile = Path.Combine(PrimaryProjectPath.FullName, CodeCompletionPreProcessorFile);
 
 			//
 			// HACK 
@@ -109,7 +109,7 @@ namespace UnrealBuildTool
 			settings.Indent = true;
 
 			XElement CodeLiteWorkspace = new XElement("CodeLite_Workspace");
-			XAttribute CodeLiteWorkspaceName = new XAttribute("Name", MasterProjectName);
+			XAttribute CodeLiteWorkspaceName = new XAttribute("Name", PrimaryProjectName);
 			XAttribute CodeLiteWorkspaceSWTLW = new XAttribute("SWTLW", "Yes"); // This flag will only work in CodeLite version > 8.0. See below
 			CodeLiteWorkspace.Add(CodeLiteWorkspaceName);
 			CodeLiteWorkspace.Add(CodeLiteWorkspaceSWTLW);
@@ -157,7 +157,7 @@ namespace UnrealBuildTool
 				foreach (ProjectTarget CurrentTarget in CurProject.ProjectTargets.OfType<ProjectTarget>())
 				{
 					string[] tmp = CurrentTarget.ToString().Split('.');
-					string ProjectTargetFileName = CurProject.ProjectFilePath.Directory.MakeRelativeTo(MasterProjectPath) + "/" + tmp[0] + ProjectExtension;
+					string ProjectTargetFileName = CurProject.ProjectFilePath.Directory.MakeRelativeTo(PrimaryProjectPath) + "/" + tmp[0] + ProjectExtension;
 					String ProjectName = tmp[0];
 
 
@@ -281,7 +281,7 @@ namespace UnrealBuildTool
 			}
 
 			CodeLiteWorkspace.Add(CodeLiteWorkspaceBuildMatrix);
-			CodeLiteWorkspace.Save(FullCodeLiteMasterFile);
+			CodeLiteWorkspace.Save(FullCodeLitePrimaryFile);
 
 			return true;
 		}
@@ -291,20 +291,20 @@ namespace UnrealBuildTool
 			return new CodeLiteProject(InitFilePath, BaseDir, OnlyGameProject);
 		}
 
-		public override void CleanProjectFiles(DirectoryReference InMasterProjectDirectory, string InMasterProjectName, DirectoryReference InIntermediateProjectFilesDirectory)
+		public override void CleanProjectFiles(DirectoryReference InPrimaryProjectDirectory, string InPrimaryProjectName, DirectoryReference InIntermediateProjectFilesDirectory)
 		{
 			// TODO Delete all files here. Not finished yet.
-			string SolutionFileName = InMasterProjectName + SolutionExtension;
-			string CodeCompletionFile = InMasterProjectName + CodeCompletionFileName;
-			string CodeCompletionPreProcessorFile = InMasterProjectName + CodeCompletionPreProcessorFileName;
+			string SolutionFileName = InPrimaryProjectName + SolutionExtension;
+			string CodeCompletionFile = InPrimaryProjectName + CodeCompletionFileName;
+			string CodeCompletionPreProcessorFile = InPrimaryProjectName + CodeCompletionPreProcessorFileName;
 
-			FileReference FullCodeLiteMasterFile = FileReference.Combine(InMasterProjectDirectory, SolutionFileName);
-			FileReference FullCodeLiteCodeCompletionFile = FileReference.Combine(InMasterProjectDirectory, CodeCompletionFile);
-			FileReference FullCodeLiteCodeCompletionPreProcessorFile = FileReference.Combine(InMasterProjectDirectory, CodeCompletionPreProcessorFile);
+			FileReference FullCodeLitePrimaryFile = FileReference.Combine(InPrimaryProjectDirectory, SolutionFileName);
+			FileReference FullCodeLiteCodeCompletionFile = FileReference.Combine(InPrimaryProjectDirectory, CodeCompletionFile);
+			FileReference FullCodeLiteCodeCompletionPreProcessorFile = FileReference.Combine(InPrimaryProjectDirectory, CodeCompletionPreProcessorFile);
 
-			if (FileReference.Exists(FullCodeLiteMasterFile))
+			if (FileReference.Exists(FullCodeLitePrimaryFile))
 			{
-				FileReference.Delete(FullCodeLiteMasterFile);
+				FileReference.Delete(FullCodeLitePrimaryFile);
 			}
 			if (FileReference.Exists(FullCodeLiteCodeCompletionFile))
 			{
diff --git a/Engine/Source/Programs/UnrealBuildTool/ProjectFiles/Eddie/EddieGenerator.cs b/Engine/Source/Programs/UnrealBuildTool/ProjectFiles/Eddie/EddieGenerator.cs
--- a/Engine/Source/Programs/UnrealBuildTool/ProjectFiles/Eddie/EddieGenerator.cs
+++ b/Engine/Source/Programs/UnrealBuildTool/ProjectFiles/Eddie/EddieGenerator.cs
@@ -23,12 +23,12 @@ namespace UnrealBuildTool
 			}
 		}
 		
-		public override void CleanProjectFiles(DirectoryReference InMasterProjectDirectory, string InMasterProjectName, DirectoryReference InIntermediateProjectFilesPath)
+		public override void CleanProjectFiles(DirectoryReference InPrimaryProjectDirectory, string InPrimaryProjectName, DirectoryReference InIntermediateProjectFilesPath)
 		{
-			FileReference MasterProjDeleteFilename = FileReference.Combine(InMasterProjectDirectory, InMasterProjectName + ".wkst");
-			if (FileReference.Exists(MasterProjDeleteFilename))
+			FileReference PrimaryProjDeleteFilename = FileReference.Combine(InPrimaryProjectDirectory, InPrimaryProjectName + ".wkst");
+			if (FileReference.Exists(PrimaryProjDeleteFilename))
 			{
-				File.Delete(MasterProjDeleteFilename.FullName);
+				File.Delete(PrimaryProjDeleteFilename.FullName);
 			}
 
 			// Delete the project files folder
@@ -57,7 +57,7 @@ namespace UnrealBuildTool
 			
 			StringBuilder WorksetDataContent = new StringBuilder();
 			WorksetDataContent.Append("# @Eddie Workset@" + ProjectFileGenerator.NewLine);
-			WorksetDataContent.Append("AddWorkset \"" + MasterProjectName + ".wkst\" \"" + MasterProjectPath + "\"" + ProjectFileGenerator.NewLine);
+			WorksetDataContent.Append("AddWorkset \"" + PrimaryProjectName + ".wkst\" \"" + PrimaryProjectPath + "\"" + ProjectFileGenerator.NewLine);
 			
 			System.Action< String /*Path*/, List<PrimaryProjectFolder> /* Folders */>? AddProjectsFunction = null;
 			AddProjectsFunction = (Path, FolderList) =>
@@ -81,17 +81,17 @@ namespace UnrealBuildTool
 						WorksetDataContent.Append("EndFileGroup \"" + NewPath + "\"" + ProjectFileGenerator.NewLine);
 					}
 				};
-			AddProjectsFunction(MasterProjectName, RootFolder.SubFolders);
+			AddProjectsFunction(PrimaryProjectName, RootFolder.SubFolders);
 			
-			string ProjectName = MasterProjectName;
-			string FilePath = MasterProjectPath + "/" + ProjectName + ".wkst";
+			string ProjectName = PrimaryProjectName;
+			string FilePath = PrimaryProjectPath + "/" + ProjectName + ".wkst";
 			
 			bSuccess = WriteFileIfChanged(FilePath, WorksetDataContent.ToString(), new UTF8Encoding());
 			
 			return bSuccess;
 		}
 		
-		protected override bool WriteMasterProjectFile(ProjectFile? UBTProject, PlatformProjectGeneratorCollection PlatformProjectGenerators)
+		protected override bool WritePrimaryProjectFile(ProjectFile? UBTProject, PlatformProjectGeneratorCollection PlatformProjectGenerators)
 		{
 			return WriteEddieWorkset();
 		}
diff --git a/Engine/Source/Programs/UnrealBuildTool/ProjectFiles/KDevelop/KDevelopGenerator.cs b/Engine/Source/Programs/UnrealBuildTool/ProjectFiles/KDevelop/KDevelopGenerator.cs
--- a/Engine/Source/Programs/UnrealBuildTool/ProjectFiles/KDevelop/KDevelopGenerator.cs
+++ b/Engine/Source/Programs/UnrealBuildTool/ProjectFiles/KDevelop/KDevelopGenerator.cs
@@ -37,7 +37,7 @@ namespace UnrealBuildTool
 			}
 		}
 
-		protected override bool WriteMasterProjectFile(ProjectFile? UBTProject, PlatformProjectGeneratorCollection PlatformProjectGenerators)
+		protected override bool WritePrimaryProjectFile(ProjectFile? UBTProject, PlatformProjectGeneratorCollection PlatformProjectGenerators)
 		{
 			bool bSuccess = true;
 			return bSuccess;
@@ -49,7 +49,7 @@ namespace UnrealBuildTool
 		/// </summary>
 		/// <param name="FileContent">File content.</param>
 		/// <param name="Name">Name.</param>
-		private void WriteKDevMasterProjectSection(ref StringBuilder FileContent, string Name)
+		private void WriteKDevPrimaryProjectSection(ref StringBuilder FileContent, string Name)
 		{
 			FileContent.Append("\n");
 			FileContent.Append("[Project] \n");
@@ -272,7 +272,7 @@ namespace UnrealBuildTool
 
 				foreach (string CurPath in KDevelopProject.IntelliSenseIncludeSearchPaths)
 				{
-					string FullProjectPath = ProjectFileGenerator.MasterProjectPath.FullName;
+					string FullProjectPath = ProjectFileGenerator.PrimaryProjectPath.FullName;
 					string FullPath = "";
 
 					// need to test to see if this in the project souce tree
@@ -300,7 +300,7 @@ namespace UnrealBuildTool
 
 				foreach (string CurPath in KDevelopProject.IntelliSenseSystemIncludeSearchPaths)
 				{
-					string FullProjectPath = ProjectFileGenerator.MasterProjectPath.FullName;
+					string FullProjectPath = ProjectFileGenerator.PrimaryProjectPath.FullName;
 					string FullPath = "";
 
 					if (CurPath.StartsWith("/") && !CurPath.StartsWith(FullProjectPath))
@@ -450,48 +450,48 @@ namespace UnrealBuildTool
 			// RAKE! Take one KDevelopProjectFileContent and pass
 			// it through each function that writes out the sections.
 			StringBuilder KDevelopFileContent = new StringBuilder();
-			StringBuilder KDevelopMasterFileContent = new StringBuilder();
+			StringBuilder KDevelopPrimaryFileContent = new StringBuilder();
 
 			// These are temp files until we can write them to the 
 			// *.kdev4 filename directly 
 			StringBuilder DefinesFileContent = new StringBuilder();
 			StringBuilder IncludesFileContent = new StringBuilder();
 
-			string FileName = MasterProjectName + ".kdev4";
+			string FileName = PrimaryProjectName + ".kdev4";
 
 			string DefinesFileName = "Defines.txt"; // RAKE! TEMP!
 			string IncludesFileName = "Includes.txt"; // RAKE! TEMP!
 
-			WriteKDevMasterProjectSection(ref KDevelopMasterFileContent, MasterProjectName);
+			WriteKDevPrimaryProjectSection(ref KDevelopPrimaryFileContent, PrimaryProjectName);
 
 			WriteCommandSection(ref KDevelopFileContent);
 			WriteIncludeSection(ref IncludesFileContent);
 			WriteDefineSection(ref DefinesFileContent);
 			WriteExcludeSection(ref KDevelopFileContent);
 
-			// Write the master kdev file.
-			string FullMasterProjectPath = Path.Combine(MasterProjectPath.FullName, ".kdev4/");
+			// Write the primary kdev file.
+			string FullPrimaryProjectPath = Path.Combine(PrimaryProjectPath.FullName, ".kdev4/");
 
-			if (!Directory.Exists(FullMasterProjectPath))
+			if (!Directory.Exists(FullPrimaryProjectPath))
 			{
-				Directory.CreateDirectory(FullMasterProjectPath);
+				Directory.CreateDirectory(FullPrimaryProjectPath);
 			}
 
-			string FullKDevelopMasterFileName = Path.Combine(MasterProjectPath.FullName, FileName);
-			string FullKDevelopFileName = Path.Combine(FullMasterProjectPath, FileName);
+			string FullKDevelopPrimaryFileName = Path.Combine(PrimaryProjectPath.FullName, FileName);
+			string FullKDevelopFileName = Path.Combine(FullPrimaryProjectPath, FileName);
 
-			string FullDefinesFileName = Path.Combine(FullMasterProjectPath, DefinesFileName);
-			string FullIncludesFileName = Path.Combine(FullMasterProjectPath, IncludesFileName);
+			string FullDefinesFileName = Path.Combine(FullPrimaryProjectPath, DefinesFileName);
+			string FullIncludesFileName = Path.Combine(FullPrimaryProjectPath, IncludesFileName);
 
 			WriteFileIfChanged(FullDefinesFileName, DefinesFileContent.ToString());
 			WriteFileIfChanged(FullIncludesFileName, IncludesFileContent.ToString());
 
-			return WriteFileIfChanged(FullKDevelopMasterFileName, KDevelopMasterFileContent.ToString()) &&
+			return WriteFileIfChanged(FullKDevelopPrimaryFileName, KDevelopPrimaryFileContent.ToString()) &&
 			WriteFileIfChanged(FullKDevelopFileName, KDevelopFileContent.ToString());
 		}
 
 		/// ProjectFileGenerator interface
-		//protected override bool WriteMasterProjectFile( ProjectFile UBTProject )
+		//protected override bool WritePrimaryProjectFile( ProjectFile UBTProject )
 		protected override bool WriteProjectFiles(PlatformProjectGeneratorCollection PlatformProjectGenerators)
 		{
 			return WriteKDevelopPro();
@@ -510,7 +510,7 @@ namespace UnrealBuildTool
 		}
 
 		/// ProjectFileGenerator interface
-		public override void CleanProjectFiles(DirectoryReference InMasterProjectDirectory, string InMasterProjectName, DirectoryReference InIntermediateProjectFilesDirectory)
+		public override void CleanProjectFiles(DirectoryReference InPrimaryProjectDirectory, string InPrimaryProjectName, DirectoryReference InIntermediateProjectFilesDirectory)
 		{
 		}
 	}
diff --git a/Engine/Source/Programs/UnrealBuildTool/ProjectFiles/Make/MakefileGenerator.cs b/Engine/Source/Programs/UnrealBuildTool/ProjectFiles/Make/MakefileGenerator.cs
--- a/Engine/Source/Programs/UnrealBuildTool/ProjectFiles/Make/MakefileGenerator.cs
+++ b/Engine/Source/Programs/UnrealBuildTool/ProjectFiles/Make/MakefileGenerator.cs
@@ -48,7 +48,7 @@ namespace UnrealBuildTool
 			}
 		}
 
-		protected override bool WriteMasterProjectFile(ProjectFile? UBTProject, PlatformProjectGeneratorCollection PlatformProjectGenerators)
+		protected override bool WritePrimaryProjectFile(ProjectFile? UBTProject, PlatformProjectGeneratorCollection PlatformProjectGenerators)
 		{
 			bool bSuccess = true;
 			return bSuccess;
@@ -73,7 +73,7 @@ namespace UnrealBuildTool
 
 			BuildCommand = "BUILD = bash \"$(UNREALROOTPATH)/Engine/Build/BatchFiles/Linux/Build.sh\"\n";
 
-			string FileName = "Makefile"; // MasterProjectName + ".mk";
+			string FileName = "Makefile"; // PrimaryProjectName + ".mk";
 			StringBuilder MakefileContent = new StringBuilder();
 			MakefileContent.Append(
 				"# Makefile generated by MakefileGenerator.cs\n" +
@@ -180,12 +180,12 @@ namespace UnrealBuildTool
 			}
 
 			MakefileContent.Append("\n.PHONY: $(TARGETS)\n");
-			FileReference FullFileName = FileReference.Combine(MasterProjectPath, FileName);
+			FileReference FullFileName = FileReference.Combine(PrimaryProjectPath, FileName);
 			return WriteFileIfChanged(FullFileName.FullName, MakefileContent.ToString());
 		}
 
 		/// ProjectFileGenerator interface
-		//protected override bool WriteMasterProjectFile( ProjectFile UBTProject )
+		//protected override bool WritePrimaryProjectFile( ProjectFile UBTProject )
 		protected override bool WriteProjectFiles(PlatformProjectGeneratorCollection PlatformProjectGenerators)
 		{
 			return WriteMakefile();
@@ -204,7 +204,7 @@ namespace UnrealBuildTool
 		}
 
 		/// ProjectFileGenerator interface
-		public override void CleanProjectFiles(DirectoryReference InMasterProjectDirectory, string InMasterProjectName, DirectoryReference InIntermediateProjectFilesDirectory)
+		public override void CleanProjectFiles(DirectoryReference InPrimaryProjectDirectory, string InPrimaryProjectName, DirectoryReference InIntermediateProjectFilesDirectory)
 		{
 		}
 	}
diff --git a/Engine/Source/Programs/UnrealBuildTool/ProjectFiles/Project.cs b/Engine/Source/Programs/UnrealBuildTool/ProjectFiles/Project.cs
--- a/Engine/Source/Programs/UnrealBuildTool/ProjectFiles/Project.cs
+++ b/Engine/Source/Programs/UnrealBuildTool/ProjectFiles/Project.cs
@@ -172,7 +172,7 @@ namespace UnrealBuildTool
 		/// <summary>
 		/// Constructs a new project file object
 		/// </summary>
-		/// <param name="ProjectFilePath">The path to the project file, relative to the master project file</param>
+		/// <param name="ProjectFilePath">The path to the project file, relative to the primary project file</param>
 		/// <param name="BaseDir">The base directory for files within this project</param>
 		protected ProjectFile(FileReference ProjectFilePath, DirectoryReference BaseDir)
 		{
@@ -208,8 +208,8 @@ namespace UnrealBuildTool
 		}
 
 
-		/// Returns true if this is a "stub" project.  Stub projects function as dumb containers for source files
-		/// and are never actually "built" by the master project.  Stub projects are always "generated" projects.
+		/// Returns true if this is a "stub" project.  Stub projects function as simple containers for source files
+		/// and are never actually "built" by the primary project.  Stub projects are always "generated" projects.
 		public bool IsStubProject
 		{
 			get;
@@ -534,7 +534,7 @@ namespace UnrealBuildTool
 		public static string NormalizeProjectPath(FileSystemReference InputPath)
 		{
 			// Try to make it relative to the solution directory.
-			if (InputPath.IsUnderDirectory(ProjectFileGenerator.MasterProjectPath))
+			if (InputPath.IsUnderDirectory(ProjectFileGenerator.PrimaryProjectPath))
 			{
 				return InputPath.MakeRelativeTo(ProjectFileGenerator.IntermediateProjectFilesPath);
 			}
diff --git a/Engine/Source/Programs/UnrealBuildTool/ProjectFiles/ProjectFileGenerator.cs b/Engine/Source/Programs/UnrealBuildTool/ProjectFiles/ProjectFileGenerator.cs
--- a/Engine/Source/Programs/UnrealBuildTool/ProjectFiles/ProjectFileGenerator.cs
+++ b/Engine/Source/Programs/UnrealBuildTool/ProjectFiles/ProjectFileGenerator.cs
@@ -309,8 +309,8 @@ namespace UnrealBuildTool
 		bool bIncludeTempTargets = false;
 
 		/// <summary>
-		/// True if we should reflect "Source" sub-directories on disk in the master project as master project directories.
-		/// This (arguably) adds some visual clutter to the master project but it is truer to the on-disk file organization.
+		/// True if we should reflect "Source" sub-directories on disk in the primary project as project directories.
+		/// This (arguably) adds some visual clutter to the primary project but it is truer to the on-disk file organization.
 		/// </summary>
 		[XmlConfigFile]
 		bool bKeepSourceSubDirectories = true;
@@ -329,9 +329,9 @@ namespace UnrealBuildTool
 		string[]? ConfigurationNames = null;
 
 		/// <summary>
-		/// Relative path to the directory where the master project file will be saved to
+		/// Relative path to the directory where the primary project file will be saved to
 		/// </summary>
-		public static DirectoryReference MasterProjectPath = Unreal.RootDirectory; // We'll save the master project to our "root" folder
+		public static DirectoryReference PrimaryProjectPath = Unreal.RootDirectory; // We'll save the primary project to our "root" folder
 
 		/// <summary>
 		/// Name of the UE4 engine project that contains all of the engine code, config files and other files
@@ -356,17 +356,27 @@ namespace UnrealBuildTool
 		/// </summary>
 		bool bGatherThirdPartySource = false;
 
+		/// <exclude/>
+		[Obsolete("Deprecated in 5.0, use PrimaryProjectName instead.")]
+		[XmlConfigFile(Deprecated = true, NewAttributeName = "PrimaryProjectName")]
+		protected string MasterProjectName = "UE5";
+
+		/// <exclude/>
+		[Obsolete("Deprecated in 5.0, use bPrimaryProjecNameFromFolder instead.")]
+		[XmlConfigFile(Deprecated = true, NewAttributeName = "bPrimaryProjectNameFromFolder")]
+		protected bool bMasterProjectNameFromFolder = false;
+
 		/// <summary>
-		/// Name of the master project file -- for example, the base file name for the Visual Studio solution file, or the Xcode project file on Mac.
+		/// Name of the primary project file -- for example, the base file name for the Visual Studio solution file, or the Xcode project file on Mac.
 		/// </summary>
 		[XmlConfigFile]
-		protected string MasterProjectName = "UE5";
+		protected string PrimaryProjectName = "UE5";
 
 		/// <summary>
-		/// If true, sets the master project name according to the name of the folder it is in.
+		/// If true, sets the primary project name according to the name of the folder it is in.
 		/// </summary>
 		[XmlConfigFile]
-		protected bool bMasterProjectNameFromFolder = false;
+		protected bool bPrimaryProjectNameFromFolder = false;
 
 		/// <summary>
 		/// Maps all module names that were included in generated project files, to actual project file objects.
@@ -763,48 +773,48 @@ namespace UnrealBuildTool
 			{
 				Log.TraceInformation("Discovering modules, targets and source code for project...");
 
-				MasterProjectPath = OnlyGameProject!.Directory;
+				PrimaryProjectPath = OnlyGameProject!.Directory;
 
 				// Set the project file name
-				MasterProjectName = OnlyGameProject.GetFileNameWithoutExtension();
+				PrimaryProjectName = OnlyGameProject.GetFileNameWithoutExtension();
 
-				if (!DirectoryReference.Exists(DirectoryReference.Combine(MasterProjectPath, "Source")))
+				if (!DirectoryReference.Exists(DirectoryReference.Combine(PrimaryProjectPath, "Source")))
 				{
-					if (!DirectoryReference.Exists(DirectoryReference.Combine(MasterProjectPath, "Intermediate", "Source")))
+					if (!DirectoryReference.Exists(DirectoryReference.Combine(PrimaryProjectPath, "Intermediate", "Source")))
 					{
 						if (BuildHostPlatform.Current.Platform == UnrealTargetPlatform.Mac)
 						{
-							MasterProjectPath = Unreal.EngineDirectory;
+							PrimaryProjectPath = Unreal.EngineDirectory;
 							GameProjectName = "UnrealGame";
 						}
-						if (!DirectoryReference.Exists(DirectoryReference.Combine(MasterProjectPath, "Source")))
+						if (!DirectoryReference.Exists(DirectoryReference.Combine(PrimaryProjectPath, "Source")))
 						{
-							throw new BuildException("Directory '{0}' is missing 'Source' folder.", MasterProjectPath);
+							throw new BuildException("Directory '{0}' is missing 'Source' folder.", PrimaryProjectPath);
 						}
 					}
 				}
-				IntermediateProjectFilesPath = DirectoryReference.Combine(MasterProjectPath, "Intermediate", "ProjectFiles");
+				IntermediateProjectFilesPath = DirectoryReference.Combine(PrimaryProjectPath, "Intermediate", "ProjectFiles");
 			}
 			else
 			{
-				// Set the master project name from the folder name
+				// Set the primary project name from the folder name
 				if (Environment.GetEnvironmentVariable("UE_NAME_PROJECT_AFTER_FOLDER") == "1")
 				{
-					MasterProjectName += "_" + Path.GetFileName(MasterProjectPath.ToString());
+					PrimaryProjectName += "_" + Path.GetFileName(PrimaryProjectPath.ToString());
 				}
-				else if (bMasterProjectNameFromFolder)
+				else if (bPrimaryProjectNameFromFolder)
 				{
-					string NewMasterProjectName = MasterProjectPath.GetDirectoryName();
-					if (!String.IsNullOrEmpty(NewMasterProjectName))
+					string NewPrimaryProjectName = PrimaryProjectPath.GetDirectoryName();
+					if (!String.IsNullOrEmpty(NewPrimaryProjectName))
 					{
-						MasterProjectName = NewMasterProjectName;
+						PrimaryProjectName = NewPrimaryProjectName;
 					}
 				}
 
-				// Write out the name of the master project file, so the runtime knows to use it
-				FileReference MasterProjectNameLocation = FileReference.Combine(Unreal.EngineDirectory, "Intermediate", "ProjectFiles", "MasterProjectName.txt");
-				DirectoryReference.CreateDirectory(MasterProjectNameLocation.Directory);
-				FileReference.WriteAllText(MasterProjectNameLocation, MasterProjectName);
+				// Write out the name of the primary project file, so the runtime knows to use it
+				FileReference PrimaryProjectNameLocation = FileReference.Combine(Unreal.EngineDirectory, "Intermediate", "ProjectFiles", "MasterProjectName.txt");
+				DirectoryReference.CreateDirectory(PrimaryProjectNameLocation.Directory);
+				FileReference.WriteAllText(PrimaryProjectNameLocation, PrimaryProjectName);
 			}
 
 			// Modify the name if specific platforms were given
@@ -818,23 +828,23 @@ namespace UnrealBuildTool
 				}
 				SortedPlatformNames.Sort();
 
-				MasterProjectName += "_";
+				PrimaryProjectName += "_";
 				foreach (string SortedPlatform in SortedPlatformNames)
 				{
-					MasterProjectName += SortedPlatform;
+					PrimaryProjectName += SortedPlatform;
 					IntermediateProjectFilesPath = new DirectoryReference(IntermediateProjectFilesPath.FullName + SortedPlatform);
 				}
 
-				// the master project name is always read from our intermediate directory and not the overriden one for this set of platforms
-				FileReference MasterProjectNameLocation = FileReference.Combine(Unreal.EngineDirectory, "Intermediate", "ProjectFiles", "MasterProjectName.txt");
-				DirectoryReference.CreateDirectory(MasterProjectNameLocation.Directory);
-				FileReference.WriteAllText(MasterProjectNameLocation, MasterProjectName);
+				// the primary project name is always read from our intermediate directory and not the overriden one for this set of platforms
+				FileReference PrimaryProjectNameLocation = FileReference.Combine(Unreal.EngineDirectory, "Intermediate", "ProjectFiles", "MasterProjectName.txt");
+				DirectoryReference.CreateDirectory(PrimaryProjectNameLocation.Directory);
+				FileReference.WriteAllText(PrimaryProjectNameLocation, PrimaryProjectName);
 			}
 
 			bool bCleanProjectFiles = Arguments.Any(x => x.Equals("-CleanProjects", StringComparison.InvariantCultureIgnoreCase));
 			if (bCleanProjectFiles)
 			{
-				CleanProjectFiles(MasterProjectPath, MasterProjectName, IntermediateProjectFilesPath);
+				CleanProjectFiles(PrimaryProjectPath, PrimaryProjectName, IntermediateProjectFilesPath);
 			}
 
 			// Figure out which platforms we should generate project files for.
@@ -1049,15 +1059,15 @@ namespace UnrealBuildTool
 				{
 					PrimaryProjectFolder ProgramsFolder = RootFolder.AddSubFolder("Programs");
 
-					// Add UnrealBuildTool to the master project
+					// Add UnrealBuildTool to the primary project
 					AddUnrealBuildToolProject(ProgramsFolder);
 
-					// Add AutomationTool to the master project
+					// Add AutomationTool to the primary project
 					VCSharpProjectFile AutomationToolProject = AddSimpleCSharpProject("AutomationTool", bShouldBuildForAllSolutionTargets: true, bForceDevelopmentConfiguration: true);
 					if (AutomationToolProject != null)
 						ProgramsFolder.ChildProjects.Add(AutomationToolProject);
 
-					// Add automation.csproj files to the master project
+					// Add automation.csproj files to the primary project
 					AddAutomationModules(AllGameProjects, RootFolder, ProgramsFolder);
 
 					// Add shared projects
@@ -1075,10 +1085,10 @@ namespace UnrealBuildTool
 				}
 
 
-				// Eliminate all redundant master project folders.  E.g., folders which contain only one project and that project
+				// Eliminate all redundant project folders.  E.g., folders which contain only one project and that project
 				// has the same name as the folder itself.  To the user, projects "feel like" folders already in the IDE, so we
 				// want to collapse them down where possible.
-				EliminateRedundantMasterProjectSubFolders(RootFolder, "");
+				EliminateRedundantPrimaryProjectSubFolders(RootFolder, "");
 
 				// Figure out which targets we need about IntelliSense for.  We only need to worry about targets for projects
 				// that we're actually generating in this session.
@@ -1177,10 +1187,10 @@ namespace UnrealBuildTool
 		/// <summary>
 		/// Clean project files
 		/// </summary>
-		/// <param name="InMasterProjectDirectory">The master project directory</param>
-		/// <param name="InMasterProjectName">The name of the master project</param>
+		/// <param name="InPrimaryProjectDirectory">The primary project directory</param>
+		/// <param name="InPrimaryProjectName">The name of the primary project</param>
 		/// <param name="InIntermediateProjectFilesDirectory">The intermediate path of project files</param>
-		public abstract void CleanProjectFiles(DirectoryReference InMasterProjectDirectory, string InMasterProjectName, DirectoryReference InIntermediateProjectFilesDirectory);
+		public abstract void CleanProjectFiles(DirectoryReference InPrimaryProjectDirectory, string InPrimaryProjectName, DirectoryReference InIntermediateProjectFilesDirectory);
 
 		/// <summary>
 		/// Configures project generator based on command-line options
@@ -1588,14 +1598,14 @@ namespace UnrealBuildTool
 
 		/// <summary>
 		/// Recursively collapses all sub-folders that are redundant.  Should only be called after we're done adding
-		/// files and projects to the master project.
+		/// files and projects to the primary project.
 		/// </summary>
 		/// <param name="Folder">The folder whose sub-folders we should potentially collapse into</param>
-		/// <param name="ParentMasterProjectFolderPath"></param>
-		void EliminateRedundantMasterProjectSubFolders(PrimaryProjectFolder Folder, string ParentMasterProjectFolderPath)
+		/// <param name="ParentPrimaryProjectFolderPath"></param>
+		void EliminateRedundantPrimaryProjectSubFolders(PrimaryProjectFolder Folder, string ParentPrimaryProjectFolderPath)
 		{
 			// NOTE: This is for diagnostics output only
-			string MasterProjectFolderPath = String.IsNullOrEmpty(ParentMasterProjectFolderPath) ? Folder.FolderName : (ParentMasterProjectFolderPath + "/" + Folder.FolderName);
+			string PrimaryProjectFolderPath = String.IsNullOrEmpty(ParentPrimaryProjectFolderPath) ? Folder.FolderName : (ParentPrimaryProjectFolderPath + "/" + Folder.FolderName);
 
 			// We can eliminate folders that meet all of these requirements:
 			//		1) Have only a single project file in them
@@ -1610,7 +1620,7 @@ namespace UnrealBuildTool
 			foreach (PrimaryProjectFolder SubFolder in Folder.SubFolders)
 			{
 				// Recurse
-				EliminateRedundantMasterProjectSubFolders(SubFolder, MasterProjectFolderPath);
+				EliminateRedundantPrimaryProjectSubFolders(SubFolder, PrimaryProjectFolderPath);
 			}
 
 			List<PrimaryProjectFolder> SubFoldersToAdd = new List<PrimaryProjectFolder>();
@@ -1677,7 +1687,7 @@ namespace UnrealBuildTool
 			Folder.SubFolders.AddRange(SubFoldersToAdd);
 
 			// After everything has been collapsed, do a bit of data validation
-			Validate(Folder, ParentMasterProjectFolderPath);
+			Validate(Folder, ParentPrimaryProjectFolderPath);
 		}
 
 		/// <summary>
@@ -1685,8 +1695,8 @@ namespace UnrealBuildTool
 		/// for project file names to be unique!
 		/// </summary>
 		/// <param name="Folder">Folder.</param>
-		/// <param name="MasterProjectFolderPath">Parent master project folder path.</param>
-		protected virtual void Validate(PrimaryProjectFolder Folder, string MasterProjectFolderPath)
+		/// <param name="PrimaryProjectFolderPath">Parent primary project folder path.</param>
+		protected virtual void Validate(PrimaryProjectFolder Folder, string PrimaryProjectFolderPath)
 		{
 			foreach (ProjectFile CurChildProject in Folder.ChildProjects)
 			{
@@ -1696,7 +1706,7 @@ namespace UnrealBuildTool
 					{
 						if (CurChildProject.ProjectFilePath.GetFileName().Equals(OtherChildProject.ProjectFilePath.GetFileNameWithoutExtension(), StringComparison.InvariantCultureIgnoreCase))
 						{
-							throw new BuildException("Detected collision between two project files with the same path in the same master project folder, " + OtherChildProject.ProjectFilePath.FullName + " and " + CurChildProject.ProjectFilePath.FullName + " (master project folder: " + MasterProjectFolderPath + ")");
+							throw new BuildException("Detected collision between two project files with the same path in the same primary project folder, " + OtherChildProject.ProjectFilePath.FullName + " and " + CurChildProject.ProjectFilePath.FullName + " (primary project folder: " + PrimaryProjectFolderPath + ")");
 						}
 					}
 				}
@@ -1711,14 +1721,14 @@ namespace UnrealBuildTool
 				{
 					if (CurChildProject.ProjectFilePath.GetFileNameWithoutExtension().Equals(SubFolder.FolderName, StringComparison.InvariantCultureIgnoreCase))
 					{
-						throw new BuildException("Detected collision between a master project sub-folder " + SubFolder.FolderName + " and a project within the outer folder " + CurChildProject.ProjectFilePath + " (master project folder: " + MasterProjectFolderPath + ")");
+						throw new BuildException("Detected collision between a primary project sub-folder " + SubFolder.FolderName + " and a project within the outer folder " + CurChildProject.ProjectFilePath + " (primary project folder: " + PrimaryProjectFolderPath + ")");
 					}
 				}
 				foreach (string CurFile in Folder.Files)
 				{
 					if (Path.GetFileName(CurFile).Equals(SubFolder.FolderName, StringComparison.InvariantCultureIgnoreCase))
 					{
-						throw new BuildException("Detected collision between a master project sub-folder " + SubFolder.FolderName + " and a file within the outer folder " + CurFile + " (master project folder: " + MasterProjectFolderPath + ")");
+						throw new BuildException("Detected collision between a primary project sub-folder " + SubFolder.FolderName + " and a file within the outer folder " + CurFile + " (primary project folder: " + PrimaryProjectFolderPath + ")");
 					}
 				}
 				foreach (PrimaryProjectFolder CurFolder in Folder.SubFolders)
@@ -1727,7 +1737,7 @@ namespace UnrealBuildTool
 					{
 						if (CurFolder.FolderName.Equals(SubFolder.FolderName, StringComparison.InvariantCultureIgnoreCase))
 						{
-							throw new BuildException("Detected collision between a master project sub-folder " + SubFolder.FolderName + " and a sibling folder " + CurFolder.FolderName + " (master project folder: " + MasterProjectFolderPath + ")");
+							throw new BuildException("Detected collision between a primary project sub-folder " + SubFolder.FolderName + " and a sibling folder " + CurFolder.FolderName + " (primary project folder: " + PrimaryProjectFolderPath + ")");
 						}
 					}
 				}
@@ -1735,7 +1745,7 @@ namespace UnrealBuildTool
 		}
 
 		/// <summary>
-		/// Adds UnrealBuildTool to the master project
+		/// Adds UnrealBuildTool to the primary project
 		/// </summary>
 		private void AddUnrealBuildToolProject(PrimaryProjectFolder ProgramsFolder)
 		{
@@ -1770,7 +1780,7 @@ namespace UnrealBuildTool
 		}
 
 		/// <summary>
-		/// Adds a C# project to the master project
+		/// Adds a C# project to the primary project
 		/// </summary>
 		/// <param name="ProjectName">Name of project file to add</param>
 		/// <param name="bShouldBuildForAllSolutionTargets"></param>
@@ -2776,7 +2786,7 @@ namespace UnrealBuildTool
 		{
 			using (ProgressWriter Progress = new ProgressWriter("Writing project files...", true))
 			{
-				int TotalProjectFileCount = GeneratedProjectFiles.Count + 1;    // +1 for the master project file, which we'll save next
+				int TotalProjectFileCount = GeneratedProjectFiles.Count + 1;    // +1 for the primary project file, which we'll save next
 
 				for (int ProjectFileIndex = 0; ProjectFileIndex < GeneratedProjectFiles.Count; ++ProjectFileIndex)
 				{
@@ -2789,7 +2799,7 @@ namespace UnrealBuildTool
 					Progress.Write(ProjectFileIndex + 1, TotalProjectFileCount);
 				}
 
-				WriteMasterProjectFile(UBTProject, PlatformProjectGenerators);
+				WritePrimaryProjectFile(UBTProject, PlatformProjectGenerators);
 				Progress.Write(TotalProjectFileCount, TotalProjectFileCount);
 			}
 
@@ -2797,12 +2807,12 @@ namespace UnrealBuildTool
 		}
 
 		/// <summary>
-		/// Writes the master project file (e.g. Visual Studio Solution file)
+		/// Writes the primary project file (e.g. Visual Studio Solution file)
 		/// </summary>
 		/// <param name="UBTProject">The UnrealBuildTool project</param>
 		/// <param name="PlatformProjectGenerators">The platform project file generators</param>
 		/// <returns>True if successful</returns>
-		protected abstract bool WriteMasterProjectFile(ProjectFile? UBTProject, PlatformProjectGeneratorCollection PlatformProjectGenerators);
+		protected abstract bool WritePrimaryProjectFile(ProjectFile? UBTProject, PlatformProjectGeneratorCollection PlatformProjectGenerators);
 
 
 		/// <summary>
@@ -3065,7 +3075,7 @@ namespace UnrealBuildTool
 
 		protected readonly List<ProjectFile> AutomationProjectFiles = new List<ProjectFile>();
 
-		/// List of top-level folders in the master project file
+		/// List of top-level folders in the primary project file
 		protected PrimaryProjectFolder RootFolder;
 	}
 
diff --git a/Engine/Source/Programs/UnrealBuildTool/ProjectFiles/QMake/QMakefileGenerator.cs b/Engine/Source/Programs/UnrealBuildTool/ProjectFiles/QMake/QMakefileGenerator.cs
--- a/Engine/Source/Programs/UnrealBuildTool/ProjectFiles/QMake/QMakefileGenerator.cs
+++ b/Engine/Source/Programs/UnrealBuildTool/ProjectFiles/QMake/QMakefileGenerator.cs
@@ -38,7 +38,7 @@ namespace UnrealBuildTool
 			}
 		}
 
-		protected override bool WriteMasterProjectFile(ProjectFile? UBTProject, PlatformProjectGeneratorCollection PlatformProjectGenerators)
+		protected override bool WritePrimaryProjectFile(ProjectFile? UBTProject, PlatformProjectGeneratorCollection PlatformProjectGenerators)
 		{
 			bool bSuccess = true;
 			return bSuccess;
@@ -69,7 +69,7 @@ namespace UnrealBuildTool
 		/// Adds the include directory to the list, after converting it to relative to UE4 root
 		private void AddIncludeDirectory(ref List<string> IncludeDirectories, string IncludeDir, string ProjectDir)
 		{
-			string FullProjectPath = ProjectFileGenerator.MasterProjectPath.FullName;
+			string FullProjectPath = ProjectFileGenerator.PrimaryProjectPath.FullName;
 			string FullPath = "";
 			if (IncludeDir.StartsWith("/") && !IncludeDir.StartsWith(FullProjectPath))
 			{
@@ -100,10 +100,10 @@ namespace UnrealBuildTool
 
 			// DefineList.Add ("");
 
-			string QMakeIncludesFileName = MasterProjectName + "Includes.pri";
+			string QMakeIncludesFileName = PrimaryProjectName + "Includes.pri";
 			StringBuilder QMakeIncludesPriFileContent = new StringBuilder();
 
-			string QMakeDefinesFileName = MasterProjectName + "Defines.pri";
+			string QMakeDefinesFileName = PrimaryProjectName + "Defines.pri";
 			StringBuilder QMakeDefinesPriFileContent = new StringBuilder();
 
 			string GameProjectPath = "";
@@ -219,11 +219,11 @@ namespace UnrealBuildTool
 
 			string UnrealRootPath = Unreal.RootDirectory.FullName;
 
-			string FileName = MasterProjectName + ".pro";
+			string FileName = PrimaryProjectName + ".pro";
 
-			string QMakeSourcePriFileName = MasterProjectName + "Source.pri";
-			string QMakeHeaderPriFileName = MasterProjectName + "Header.pri";
-			string QMakeConfigPriFileName = MasterProjectName + "Config.pri";
+			string QMakeSourcePriFileName = PrimaryProjectName + "Source.pri";
+			string QMakeHeaderPriFileName = PrimaryProjectName + "Header.pri";
+			string QMakeConfigPriFileName = PrimaryProjectName + "Config.pri";
 
 			StringBuilder QMakeFileContent = new StringBuilder();
 
@@ -392,13 +392,13 @@ namespace UnrealBuildTool
 
 			QMakeFileContent.Append(QMakeTargetList.TrimEnd('\\'));
 
-			string FullFileName = Path.Combine(MasterProjectPath.FullName, FileName);
+			string FullFileName = Path.Combine(PrimaryProjectPath.FullName, FileName);
 
-			string FullQMakeDefinesFileName = Path.Combine(MasterProjectPath.FullName, QMakeDefinesFileName);
-			string FullQMakeIncludesFileName = Path.Combine(MasterProjectPath.FullName, QMakeIncludesFileName);
-			string FullQMakeSourcePriFileName = Path.Combine(MasterProjectPath.FullName, QMakeSourcePriFileName);
-			string FullQMakeHeaderPriFileName = Path.Combine(MasterProjectPath.FullName, QMakeHeaderPriFileName);
-			string FullQMakeConfigPriFileName = Path.Combine(MasterProjectPath.FullName, QMakeConfigPriFileName);
+			string FullQMakeDefinesFileName = Path.Combine(PrimaryProjectPath.FullName, QMakeDefinesFileName);
+			string FullQMakeIncludesFileName = Path.Combine(PrimaryProjectPath.FullName, QMakeIncludesFileName);
+			string FullQMakeSourcePriFileName = Path.Combine(PrimaryProjectPath.FullName, QMakeSourcePriFileName);
+			string FullQMakeHeaderPriFileName = Path.Combine(PrimaryProjectPath.FullName, QMakeHeaderPriFileName);
+			string FullQMakeConfigPriFileName = Path.Combine(PrimaryProjectPath.FullName, QMakeConfigPriFileName);
 
 			WriteFileIfChanged(FullQMakeDefinesFileName, QMakeDefinesPriFileContent.ToString());
 			WriteFileIfChanged(FullQMakeIncludesFileName, QMakeIncludesPriFileContent.ToString());
@@ -411,7 +411,7 @@ namespace UnrealBuildTool
 		}
 
 		/// ProjectFileGenerator interface
-		//protected override bool WriteMasterProjectFile( ProjectFile UBTProject )
+		//protected override bool WritePrimaryProjectFile( ProjectFile UBTProject )
 		protected override bool WriteProjectFiles(PlatformProjectGeneratorCollection PlatformProjectGenerators)
 		{
 			return WriteQMakePro();
@@ -430,7 +430,7 @@ namespace UnrealBuildTool
 		}
 
 		/// ProjectFileGenerator interface
-		public override void CleanProjectFiles(DirectoryReference InMasterProjectDirectory, string InMasterProjectName, DirectoryReference InIntermediateProjectFilesDirectory)
+		public override void CleanProjectFiles(DirectoryReference InPrimaryProjectDirectory, string InPrimaryProjectName, DirectoryReference InIntermediateProjectFilesDirectory)
 		{
 		}
 	}
diff --git a/Engine/Source/Programs/UnrealBuildTool/ProjectFiles/Rider/RiderProjectFileGenerator.cs b/Engine/Source/Programs/UnrealBuildTool/ProjectFiles/Rider/RiderProjectFileGenerator.cs
--- a/Engine/Source/Programs/UnrealBuildTool/ProjectFiles/Rider/RiderProjectFileGenerator.cs
+++ b/Engine/Source/Programs/UnrealBuildTool/ProjectFiles/Rider/RiderProjectFileGenerator.cs
@@ -66,10 +66,10 @@ namespace UnrealBuildTool
 
 		public override bool ShouldGenerateIntelliSenseData() => true;
 
-		public override void CleanProjectFiles(DirectoryReference InMasterProjectDirectory, string InMasterProjectName,
+		public override void CleanProjectFiles(DirectoryReference InPrimaryProjectDirectory, string InPrimaryProjectName,
 			DirectoryReference InIntermediateProjectFilesDirectory)
 		{
-			DirectoryReference.Delete(InMasterProjectDirectory);
+			DirectoryReference.Delete(InPrimaryProjectDirectory);
 		}
 
 		private void ConfigureProjectFileGeneration()
@@ -407,29 +407,29 @@ namespace UnrealBuildTool
 
 			if (bGeneratingGameProjectFiles)
 			{
-				MasterProjectPath = OnlyGameProject!.Directory;
-				MasterProjectName = OnlyGameProject.GetFileNameWithoutExtension();
+				PrimaryProjectPath = OnlyGameProject!.Directory;
+				PrimaryProjectName = OnlyGameProject.GetFileNameWithoutExtension();
 
-				if (!DirectoryReference.Exists(DirectoryReference.Combine(MasterProjectPath, "Source")))
+				if (!DirectoryReference.Exists(DirectoryReference.Combine(PrimaryProjectPath, "Source")))
 				{
-					if (!DirectoryReference.Exists(DirectoryReference.Combine(MasterProjectPath, "Intermediate",
+					if (!DirectoryReference.Exists(DirectoryReference.Combine(PrimaryProjectPath, "Intermediate",
 						"Source")))
 					{
 						if (BuildHostPlatform.Current.Platform == UnrealTargetPlatform.Mac)
 						{
-							MasterProjectPath = Unreal.EngineDirectory;
+							PrimaryProjectPath = Unreal.EngineDirectory;
 							GameProjectName = "UnrealGame";
 						}
 
-						if (!DirectoryReference.Exists(DirectoryReference.Combine(MasterProjectPath, "Source")))
+						if (!DirectoryReference.Exists(DirectoryReference.Combine(PrimaryProjectPath, "Source")))
 						{
-							throw new BuildException("Directory '{0}' is missing 'Source' folder.", MasterProjectPath);
+							throw new BuildException("Directory '{0}' is missing 'Source' folder.", PrimaryProjectPath);
 						}
 					}
 				}
 
 				IntermediateProjectFilesPath =
-					DirectoryReference.Combine(MasterProjectPath, "Intermediate", "ProjectFiles");
+					DirectoryReference.Combine(PrimaryProjectPath, "Intermediate", "ProjectFiles");
 			}
 
 			SetupSupportedPlatformsAndConfigurations();
@@ -477,7 +477,7 @@ namespace UnrealBuildTool
 			}
 		}
 
-		protected override bool WriteMasterProjectFile(ProjectFile? ProjectFile,
+		protected override bool WritePrimaryProjectFile(ProjectFile? ProjectFile,
 			PlatformProjectGeneratorCollection PlatformProjectGenerators)
 		{
 			return true;
diff --git a/Engine/Source/Programs/UnrealBuildTool/ProjectFiles/VisualStudio/VCProject.cs b/Engine/Source/Programs/UnrealBuildTool/ProjectFiles/VisualStudio/VCProject.cs
--- a/Engine/Source/Programs/UnrealBuildTool/ProjectFiles/VisualStudio/VCProject.cs
+++ b/Engine/Source/Programs/UnrealBuildTool/ProjectFiles/VisualStudio/VCProject.cs
@@ -550,9 +550,9 @@ namespace UnrealBuildTool
 
 		static string InsertPathVariables(FileReference Location)
 		{
-			if (Location.IsUnderDirectory(ProjectFileGenerator.MasterProjectPath))
+			if (Location.IsUnderDirectory(ProjectFileGenerator.PrimaryProjectPath))
 			{
-				return String.Format("$(SolutionDir){0}", Location.MakeRelativeTo(ProjectFileGenerator.MasterProjectPath));
+				return String.Format("$(SolutionDir){0}", Location.MakeRelativeTo(ProjectFileGenerator.PrimaryProjectPath));
 			}
 			else
 			{
diff --git a/Engine/Source/Programs/UnrealBuildTool/ProjectFiles/VisualStudio/VCProjectFileGenerator.cs b/Engine/Source/Programs/UnrealBuildTool/ProjectFiles/VisualStudio/VCProjectFileGenerator.cs
--- a/Engine/Source/Programs/UnrealBuildTool/ProjectFiles/VisualStudio/VCProjectFileGenerator.cs
+++ b/Engine/Source/Programs/UnrealBuildTool/ProjectFiles/VisualStudio/VCProjectFileGenerator.cs
@@ -144,33 +144,33 @@ namespace UnrealBuildTool
 
 		/// <summary>
 		/// </summary>
-		public override void CleanProjectFiles(DirectoryReference InMasterProjectDirectory, string InMasterProjectName, DirectoryReference InIntermediateProjectFilesDirectory)
+		public override void CleanProjectFiles(DirectoryReference InPrimaryProjectDirectory, string InPrimaryProjectName, DirectoryReference InIntermediateProjectFilesDirectory)
 		{
-			FileReference MasterProjectFile = FileReference.Combine(InMasterProjectDirectory, InMasterProjectName);
-			FileReference MasterProjDeleteFilename = MasterProjectFile + ".sln";
-			if (FileReference.Exists(MasterProjDeleteFilename))
+			FileReference PrimaryProjectFile = FileReference.Combine(InPrimaryProjectDirectory, InPrimaryProjectName);
+			FileReference PrimaryProjDeleteFilename = PrimaryProjectFile + ".sln";
+			if (FileReference.Exists(PrimaryProjDeleteFilename))
 			{
-				FileReference.Delete(MasterProjDeleteFilename);
+				FileReference.Delete(PrimaryProjDeleteFilename);
 			}
-			MasterProjDeleteFilename = MasterProjectFile + ".sdf";
-			if (FileReference.Exists(MasterProjDeleteFilename))
+			PrimaryProjDeleteFilename = PrimaryProjectFile + ".sdf";
+			if (FileReference.Exists(PrimaryProjDeleteFilename))
 			{
-				FileReference.Delete(MasterProjDeleteFilename);
+				FileReference.Delete(PrimaryProjDeleteFilename);
 			}
-			MasterProjDeleteFilename = MasterProjectFile + ".suo";
-			if (FileReference.Exists(MasterProjDeleteFilename))
+			PrimaryProjDeleteFilename = PrimaryProjectFile + ".suo";
+			if (FileReference.Exists(PrimaryProjDeleteFilename))
 			{
-				FileReference.Delete(MasterProjDeleteFilename);
+				FileReference.Delete(PrimaryProjDeleteFilename);
 			}
-			MasterProjDeleteFilename = MasterProjectFile + ".v11.suo";
-			if (FileReference.Exists(MasterProjDeleteFilename))
+			PrimaryProjDeleteFilename = PrimaryProjectFile + ".v11.suo";
+			if (FileReference.Exists(PrimaryProjDeleteFilename))
 			{
-				FileReference.Delete(MasterProjDeleteFilename);
+				FileReference.Delete(PrimaryProjDeleteFilename);
 			}
-			MasterProjDeleteFilename = MasterProjectFile + ".v12.suo";
-			if (FileReference.Exists(MasterProjDeleteFilename))
+			PrimaryProjDeleteFilename = PrimaryProjectFile + ".v12.suo";
+			if (FileReference.Exists(PrimaryProjDeleteFilename))
 			{
-				FileReference.Delete(MasterProjDeleteFilename);
+				FileReference.Delete(PrimaryProjDeleteFilename);
 			}
 
 			// Delete the project files folder
@@ -469,11 +469,11 @@ namespace UnrealBuildTool
 		}
 
 
-		protected override bool WriteMasterProjectFile(ProjectFile? UBTProject, PlatformProjectGeneratorCollection PlatformProjectGenerators)
+		protected override bool WritePrimaryProjectFile(ProjectFile? UBTProject, PlatformProjectGeneratorCollection PlatformProjectGenerators)
 		{
 			bool bSuccess = true;
 
-			string SolutionFileName = MasterProjectName + ".sln";
+			string SolutionFileName = PrimaryProjectName + ".sln";
 
 			// Setup solution file content
 			StringBuilder VCSolutionFileContent = new StringBuilder();
@@ -546,7 +546,7 @@ namespace UnrealBuildTool
 					// Use the existing project's GUID that's already known to us
 					string ProjectGUID = CurProject.ProjectGUID.ToString("B").ToUpperInvariant();
 
-					VCSolutionFileContent.AppendLine("Project(\"" + ProjectTypeGUID + "\") = \"" + ProjectNameInSolution + "\", \"" + CurProject.ProjectFilePath.MakeRelativeTo(ProjectFileGenerator.MasterProjectPath) + "\", \"" + ProjectGUID + "\"");
+					VCSolutionFileContent.AppendLine("Project(\"" + ProjectTypeGUID + "\") = \"" + ProjectNameInSolution + "\", \"" + CurProject.ProjectFilePath.MakeRelativeTo(ProjectFileGenerator.PrimaryProjectPath) + "\", \"" + ProjectGUID + "\"");
 
 					// Setup dependency on UnrealBuildTool, if we need that.  This makes sure that UnrealBuildTool is
 					// freshly compiled before kicking off any build operations on this target project
@@ -584,7 +584,7 @@ namespace UnrealBuildTool
 				// Add the visualizers at the solution level. Doesn't seem to be picked up from a makefile project in VS2017 15.8.5.
 				VCSolutionFileContent.AppendLine(String.Format("Project(\"{0}\") = \"Visualizers\", \"Visualizers\", \"{{1CCEC849-CC72-4C59-8C36-2F7C38706D4C}}\"", SolutionFolderEntryGUID));
 				VCSolutionFileContent.AppendLine("\tProjectSection(SolutionItems) = preProject");
-				VCSolutionFileContent.AppendLine("\t\t{0} = {0}", VisualizersFile.MakeRelativeTo(MasterProjectPath));
+				VCSolutionFileContent.AppendLine("\t\t{0} = {0}", VisualizersFile.MakeRelativeTo(PrimaryProjectPath));
 				VCSolutionFileContent.AppendLine("\tEndProjectSection");
 				VCSolutionFileContent.AppendLine("EndProject");
 			}
@@ -755,9 +755,9 @@ namespace UnrealBuildTool
 						// filter) to.  This sets up the hierarchical solution explorer tree for all solution folders and projects.
 
 						System.Action<StringBuilder /* VCSolutionFileContent */, List<PrimaryProjectFolder> /* Folders */ >? FolderProcessorFunction = null;
-						FolderProcessorFunction = (LocalVCSolutionFileContent, LocalMasterProjectFolders) =>
+						FolderProcessorFunction = (LocalVCSolutionFileContent, LocalPrimaryProjectFolders) =>
 							{
-								foreach (PrimaryProjectFolder CurFolder in LocalMasterProjectFolders)
+								foreach (PrimaryProjectFolder CurFolder in LocalPrimaryProjectFolders)
 								{
 									string CurFolderGUIDString = ProjectFolderGuids[CurFolder].ToString("B").ToUpperInvariant();
 
@@ -790,7 +790,7 @@ namespace UnrealBuildTool
 			// Save the solution file
 			if (bSuccess)
 			{
-				string SolutionFilePath = FileReference.Combine(MasterProjectPath, SolutionFileName).FullName;
+				string SolutionFilePath = FileReference.Combine(PrimaryProjectPath, SolutionFileName).FullName;
 				bSuccess = WriteFileIfChanged(SolutionFilePath, VCSolutionFileContent.ToString());
 			}
 
@@ -803,10 +803,10 @@ namespace UnrealBuildTool
 				switch (Settings.ProjectFileFormat)
 				{
 					case VCProjectFileFormat.VisualStudio2019:
-						SolutionOptionsFileName = FileReference.Combine(MasterProjectPath, ".vs", Path.GetFileNameWithoutExtension(SolutionFileName), "v16", ".suo");
+						SolutionOptionsFileName = FileReference.Combine(PrimaryProjectPath, ".vs", Path.GetFileNameWithoutExtension(SolutionFileName), "v16", ".suo");
 						break;
 					case VCProjectFileFormat.VisualStudio2022:
-						SolutionOptionsFileName = FileReference.Combine(MasterProjectPath, ".vs", Path.GetFileNameWithoutExtension(SolutionFileName), "v17", ".suo");
+						SolutionOptionsFileName = FileReference.Combine(PrimaryProjectPath, ".vs", Path.GetFileNameWithoutExtension(SolutionFileName), "v17", ".suo");
 						break;
 					default:
 						throw new BuildException("Unsupported Visual Studio version");
diff --git a/Engine/Source/Programs/UnrealBuildTool/ProjectFiles/VisualStudioCode/VSCodeProjectFileGenerator.cs b/Engine/Source/Programs/UnrealBuildTool/ProjectFiles/VisualStudioCode/VSCodeProjectFileGenerator.cs
--- a/Engine/Source/Programs/UnrealBuildTool/ProjectFiles/VisualStudioCode/VSCodeProjectFileGenerator.cs
+++ b/Engine/Source/Programs/UnrealBuildTool/ProjectFiles/VisualStudioCode/VSCodeProjectFileGenerator.cs
@@ -231,7 +231,7 @@ namespace UnrealBuildTool
 			}
 		}
 
-		public override void CleanProjectFiles(DirectoryReference InMasterProjectDirectory, string InMasterProjectName, DirectoryReference InIntermediateProjectFilesPath)
+		public override void CleanProjectFiles(DirectoryReference InPrimaryProjectDirectory, string InPrimaryProjectName, DirectoryReference InIntermediateProjectFilesPath)
 		{
 		}
 
@@ -245,9 +245,9 @@ namespace UnrealBuildTool
 			return new VSCodeProject(InitFilePath, BaseDir);
 		}
 
-		protected override bool WriteMasterProjectFile(ProjectFile? UBTProject, PlatformProjectGeneratorCollection PlatformProjectGenerators)
+		protected override bool WritePrimaryProjectFile(ProjectFile? UBTProject, PlatformProjectGeneratorCollection PlatformProjectGenerators)
 		{
-			DirectoryReference VSCodeDir = DirectoryReference.Combine(MasterProjectPath, ".vscode");
+			DirectoryReference VSCodeDir = DirectoryReference.Combine(PrimaryProjectPath, ".vscode");
 			DirectoryReference.CreateDirectory(VSCodeDir);
 
 			bForeignProject = !VSCodeDir.IsUnderDirectory(ProjectRoot);
@@ -1499,7 +1499,7 @@ namespace UnrealBuildTool
 					OutFile.AppendLine(PathToExclude.Replace('\\', '/').Replace(WorkspaceRootPath, ""));
 				}
 			}
-			FileReference.WriteAllText(FileReference.Combine(MasterProjectPath, ".ignore"), OutFile.ToString());
+			FileReference.WriteAllText(FileReference.Combine(PrimaryProjectPath, ".ignore"), OutFile.ToString());
 		}
 
 		private void WriteWorkspaceIgnoreFileForPlugins(DirectoryReference PluginBaseDir, List<string> PathsToExclude)
@@ -1585,7 +1585,7 @@ namespace UnrealBuildTool
 			WorkspaceFile.EndRootObject();
 
 			string? WorkspaceName = bForeignProject ? GameProjectName : "UE5";
-			WorkspaceFile.Write(FileReference.Combine(MasterProjectPath, WorkspaceName + ".code-workspace"));
+			WorkspaceFile.Write(FileReference.Combine(PrimaryProjectPath, WorkspaceName + ".code-workspace"));
 		}
 
 		private void GetExcludePathsCPP(DirectoryReference BaseDir, List<string> PathsToExclude)
diff --git a/Engine/Source/Programs/UnrealBuildTool/ProjectFiles/Xcode/XcodeProjectFileGenerator.cs b/Engine/Source/Programs/UnrealBuildTool/ProjectFiles/Xcode/XcodeProjectFileGenerator.cs
--- a/Engine/Source/Programs/UnrealBuildTool/ProjectFiles/Xcode/XcodeProjectFileGenerator.cs
+++ b/Engine/Source/Programs/UnrealBuildTool/ProjectFiles/Xcode/XcodeProjectFileGenerator.cs
@@ -78,12 +78,12 @@ namespace UnrealBuildTool
 
 		/// <summary>
 		/// </summary>
-		public override void CleanProjectFiles(DirectoryReference InMasterProjectDirectory, string InMasterProjectName, DirectoryReference InIntermediateProjectFilesPath)
+		public override void CleanProjectFiles(DirectoryReference InPrimaryProjectDirectory, string InPrimaryProjectName, DirectoryReference InIntermediateProjectFilesPath)
 		{
-			DirectoryReference MasterProjDeleteFilename = DirectoryReference.Combine(InMasterProjectDirectory, InMasterProjectName + ".xcworkspace");
-			if (DirectoryReference.Exists(MasterProjDeleteFilename))
+			DirectoryReference PrimaryProjDeleteFilename = DirectoryReference.Combine(InPrimaryProjectDirectory, InPrimaryProjectName + ".xcworkspace");
+			if (DirectoryReference.Exists(PrimaryProjDeleteFilename))
 			{
-				DirectoryReference.Delete(MasterProjDeleteFilename, true);
+				DirectoryReference.Delete(PrimaryProjDeleteFilename, true);
 			}
 
 			// Delete the project files folder
@@ -187,7 +187,7 @@ namespace UnrealBuildTool
 						foreach (XcodeProjectFile XcodeProject in SupportedProjects)
 						{
 							WorkspaceDataContent.Append(Ident + "      <FileRef" + ProjectFileGenerator.NewLine);
-							WorkspaceDataContent.Append(Ident + "         location = \"group:" + XcodeProject.ProjectFilePath.MakeRelativeTo(ProjectFileGenerator.MasterProjectPath) + "\">" + ProjectFileGenerator.NewLine);
+							WorkspaceDataContent.Append(Ident + "         location = \"group:" + XcodeProject.ProjectFilePath.MakeRelativeTo(ProjectFileGenerator.PrimaryProjectPath) + "\">" + ProjectFileGenerator.NewLine);
 							WorkspaceDataContent.Append(Ident + "      </FileRef>" + ProjectFileGenerator.NewLine);							
 						}
 
@@ -231,18 +231,18 @@ namespace UnrealBuildTool
 				}
 			}
 
-			string ProjectName = MasterProjectName;
+			string ProjectName = PrimaryProjectName;
 			if (ProjectFilePlatform != XcodeProjectFilePlatform.All)
 			{
 				ProjectName += ProjectFilePlatform == XcodeProjectFilePlatform.Mac ? "_Mac" : (ProjectFilePlatform == XcodeProjectFilePlatform.iOS ? "_IOS" : "_TVOS");
 			}
-			string WorkspaceDataFilePath = MasterProjectPath + "/" + ProjectName + ".xcworkspace/contents.xcworkspacedata";
+			string WorkspaceDataFilePath = PrimaryProjectPath + "/" + ProjectName + ".xcworkspace/contents.xcworkspacedata";
 			bSuccess = WriteFileIfChanged(WorkspaceDataFilePath, WorkspaceDataContent.ToString(), new UTF8Encoding());
 			if (bSuccess)
 			{
-				string WorkspaceSettingsFilePath = MasterProjectPath + "/" + ProjectName + ".xcworkspace/xcuserdata/" + Environment.UserName + ".xcuserdatad/WorkspaceSettings.xcsettings";
+				string WorkspaceSettingsFilePath = PrimaryProjectPath + "/" + ProjectName + ".xcworkspace/xcuserdata/" + Environment.UserName + ".xcuserdatad/WorkspaceSettings.xcsettings";
 				bSuccess = WriteWorkspaceSettingsFile(WorkspaceSettingsFilePath);
-				string WorkspaceSharedSettingsFilePath = MasterProjectPath + "/" + ProjectName + ".xcworkspace/xcshareddata/WorkspaceSettings.xcsettings";
+				string WorkspaceSharedSettingsFilePath = PrimaryProjectPath + "/" + ProjectName + ".xcworkspace/xcshareddata/WorkspaceSettings.xcsettings";
 				bSuccess = WriteWorkspaceSharedSettingsFile(WorkspaceSharedSettingsFilePath);
 			}
 
@@ -250,7 +250,7 @@ namespace UnrealBuildTool
 			return bSuccess;
 		}
 
-		protected override bool WriteMasterProjectFile(ProjectFile? UBTProject, PlatformProjectGeneratorCollection PlatformProjectGenerators)
+		protected override bool WritePrimaryProjectFile(ProjectFile? UBTProject, PlatformProjectGeneratorCollection PlatformProjectGenerators)
 		{
 			return WriteXcodeWorkspace();
 		}
