PR #8702: Add an option for better LLDB support by UBT on Linux (Contributed by [REDACTED_PERSON])

diff --git a/Engine/Source/Programs/UnrealBuildTool/Platform/Linux/LinuxToolChain.cs b/Engine/Source/Programs/UnrealBuildTool/Platform/Linux/LinuxToolChain.cs
--- a/Engine/Source/Programs/UnrealBuildTool/Platform/Linux/LinuxToolChain.cs
+++ b/Engine/Source/Programs/UnrealBuildTool/Platform/Linux/LinuxToolChain.cs
@@ -58,6 +58,11 @@ namespace UnrealBuildTool
 		/// When we support larger the 4GB files with objcopy.exe this can be removed!
 		/// </summary>
 		DisableSplitDebugInfoWithObjCopy = 0x40,
+
+		/// <summary>
+		/// Enable tuning of debug info for LLDB
+		/// </summary>
+		TuneDebugInfoForLLDB = 0x80,
 	}
 
 	class LinuxToolChain : ISPCToolChain
@@ -695,6 +700,11 @@ namespace UnrealBuildTool
 					// GDB index. This option is only useful with a linker that can produce GDB index version 7.
 					Result += " -ggnu-pubnames";
 				}
+
+				if (Options.HasFlag(LinuxToolChainOptions.TuneDebugInfoForLLDB))
+				{
+					Result += " -glldb";
+				}
 			}
 
 			// optimization level
diff --git a/Engine/Source/Programs/UnrealBuildTool/Platform/Linux/UEBuildLinux.cs b/Engine/Source/Programs/UnrealBuildTool/Platform/Linux/UEBuildLinux.cs
--- a/Engine/Source/Programs/UnrealBuildTool/Platform/Linux/UEBuildLinux.cs
+++ b/Engine/Source/Programs/UnrealBuildTool/Platform/Linux/UEBuildLinux.cs
@@ -70,6 +70,13 @@ namespace UnrealBuildTool
 		/// </summary>
 		[ConfigFile(ConfigHierarchyType.Engine, "/Script/LinuxPlatform.LinuxTargetSettings")]
 		public bool bPreservePSYM = false;
+
+		/// <summary>
+		/// Turns on tuning of debug info for LLDB
+		/// </summary>
+		[CommandLine("-EnableLLDB")]
+		[XmlConfigFile(Category = "BuildConfiguration", Name = "bTuneDebugInfoForLLDB")]
+		public bool bTuneDebugInfoForLLDB = false;
 	}
 
 	/// <summary>
@@ -119,6 +126,11 @@ namespace UnrealBuildTool
 			get { return Inner.bEnableMemorySanitizer; }
 		}
 
+		public bool bTuneDebugInfoForLLDB
+		{
+			get { return Inner.bTuneDebugInfoForLLDB; }
+		}
+
 		#pragma warning restore CS1591
 		#endregion
 	}
@@ -571,6 +583,11 @@ namespace UnrealBuildTool
 				Options |= LinuxToolChainOptions.DisableSplitDebugInfoWithObjCopy;
 			}
 
+			if (Target.LinuxPlatform.bTuneDebugInfoForLLDB)
+			{
+				Options |= LinuxToolChainOptions.TuneDebugInfoForLLDB;
+			}
+
 			// Disable color logging if we are on a build machine
 			if (Environment.GetEnvironmentVariable("IsBuildMachine") == "1")
 			{
