UnrealBuildTool: Remove various references to UE4.

#trivial

diff --git a/Engine/Source/Programs/UnrealBuildTool/Configuration/ModuleRules.cs b/Engine/Source/Programs/UnrealBuildTool/Configuration/ModuleRules.cs
--- a/Engine/Source/Programs/UnrealBuildTool/Configuration/ModuleRules.cs
+++ b/Engine/Source/Programs/UnrealBuildTool/Configuration/ModuleRules.cs
@@ -963,7 +963,7 @@ namespace UnrealBuildTool
 		public List<string>? SubclassRules;
 
 		/// <summary>
-		/// Whether this module requires the IMPLEMENT_MODULE macro to be implemented. Most UE4 modules require this, since we use the IMPLEMENT_MODULE macro
+		/// Whether this module requires the IMPLEMENT_MODULE macro to be implemented. Most UE modules require this, since we use the IMPLEMENT_MODULE macro
 		/// to do other global overloads (eg. operator new/delete forwarding to GMalloc).
 		/// </summary>
 		public bool? bRequiresImplementModule;
diff --git a/Engine/Source/Programs/UnrealBuildTool/Configuration/TargetRules.cs b/Engine/Source/Programs/UnrealBuildTool/Configuration/TargetRules.cs
--- a/Engine/Source/Programs/UnrealBuildTool/Configuration/TargetRules.cs
+++ b/Engine/Source/Programs/UnrealBuildTool/Configuration/TargetRules.cs
@@ -449,7 +449,7 @@ namespace UnrealBuildTool
 		public bool bUseChaosMemoryTracking = false;
 
 		/// <summary>
-		/// Whether scene query acceleration is done by UE4. The physx scene query structure is still created, but we do not use it.
+		/// Whether scene query acceleration is done by UE. The physx scene query structure is still created, but we do not use it.
 		/// </summary>
 		[RequiresUniqueBuildEnvironment]
 		public bool bCustomSceneQueryStructure = false;
diff --git a/Engine/Source/Programs/UnrealBuildTool/Platform/Android/AndroidProjectGenerator.cs b/Engine/Source/Programs/UnrealBuildTool/Platform/Android/AndroidProjectGenerator.cs
--- a/Engine/Source/Programs/UnrealBuildTool/Platform/Android/AndroidProjectGenerator.cs
+++ b/Engine/Source/Programs/UnrealBuildTool/Platform/Android/AndroidProjectGenerator.cs
@@ -737,7 +737,7 @@ namespace UnrealBuildTool
 			Project.ShouldBuildForAllSolutionTargets = false;
 			Project.ShouldBuildByDefaultForSolutionTargets = false;
 
-			return Success ? new Tuple<ProjectFile, string>(Project, "UE4 Android Debug Projects") : null;
+			return Success ? new Tuple<ProjectFile, string>(Project, "Unreal Android Debug Projects") : null;
 
 		}
 
diff --git a/Engine/Source/Programs/UnrealBuildTool/Platform/HoloLens/ManifestGenerator.cs b/Engine/Source/Programs/UnrealBuildTool/Platform/HoloLens/ManifestGenerator.cs
--- a/Engine/Source/Programs/UnrealBuildTool/Platform/HoloLens/ManifestGenerator.cs
+++ b/Engine/Source/Programs/UnrealBuildTool/Platform/HoloLens/ManifestGenerator.cs
@@ -1285,7 +1285,7 @@ namespace UnrealBuildTool
 
 				XmlElement Identity = AppxManifestXmlDocument.CreateElement("Identity");
 				{
-					XmlAttribute PackageName = CreateStringAttribute("Name", "PackageName", "Package.Identity.Name", "/Script/EngineSettings.GeneralProjectSettings", "ProjectName", "DefaultUE4Project", ValidatePackageName);
+					XmlAttribute PackageName = CreateStringAttribute("Name", "PackageName", "Package.Identity.Name", "/Script/EngineSettings.GeneralProjectSettings", "ProjectName", "DefaultUnrealProject", ValidatePackageName);
 					Identity.Attributes.Append(PackageName);
 					XmlAttribute PublisherName = CreateStringAttribute("Publisher", "PublisherName", "Package.Identity.Publisher", "/Script/EngineSettings.GeneralProjectSettings", "CompanyDistinguishedName", "CN=NoPublisher");
 					Identity.Attributes.Append(PublisherName);
@@ -1828,7 +1828,7 @@ namespace UnrealBuildTool
 
 			if (!IsDlc)
 			{
-				XmlAttribute PackageName = CreateStringAttribute("Name", "PackageName", "Package.Identity.Name", "/Script/EngineSettings.GeneralProjectSettings", "ProjectName", "DefaultUE4Project", ValidatePackageName);
+				XmlAttribute PackageName = CreateStringAttribute("Name", "PackageName", "Package.Identity.Name", "/Script/EngineSettings.GeneralProjectSettings", "ProjectName", "DefaultUnrealProject", ValidatePackageName);
 				Identity.Attributes.Append(PackageName);
 			}
 			else
@@ -1872,7 +1872,7 @@ namespace UnrealBuildTool
 			XmlElement DisplayName = AppxManifestXmlDocument.CreateElement("DisplayName");
 			DisplayName.InnerText = "ms-resource:PackageDisplayName";
 			Properties.AppendChild(DisplayName);
-			AddResourceEntry("PackageDisplayName", "PackageDisplayName", "Package.Properties.DisplayName", "/Script/EngineSettings.GeneralProjectSettings", "ProjectDisplayedTitle", "DefaultUE4Project");
+			AddResourceEntry("PackageDisplayName", "PackageDisplayName", "Package.Properties.DisplayName", "/Script/EngineSettings.GeneralProjectSettings", "ProjectDisplayedTitle", "DefaultUnrealProject");
 
 			XmlElement PublisherDisplayName = AppxManifestXmlDocument.CreateElement("PublisherDisplayName");
 			PublisherDisplayName.InnerText = "ms-resource:PublisherDisplayName";
@@ -1945,7 +1945,7 @@ namespace UnrealBuildTool
 				XmlElement PackageDependency = AppxManifestXmlDocument.CreateElement("uap3:MainPackageDependency", "http://schemas.microsoft.com/appx/manifest/uap/windows10/3");
 				Dependencies.AppendChild(PackageDependency);
 
-				XmlAttribute NameAttribute = CreateStringAttribute("Name", "PackageName", "Package.Identity.Name", "/Script/EngineSettings.GeneralProjectSettings", "ProjectName", "DefaultUE4Project", ValidatePackageName);
+				XmlAttribute NameAttribute = CreateStringAttribute("Name", "PackageName", "Package.Identity.Name", "/Script/EngineSettings.GeneralProjectSettings", "ProjectName", "DefaultUnrealProject", ValidatePackageName);
 				PackageDependency.Attributes.Append(NameAttribute);
 			}
 
diff --git a/Engine/Source/Programs/UnrealBuildTool/Platform/IOS/IOSExports.cs b/Engine/Source/Programs/UnrealBuildTool/Platform/IOS/IOSExports.cs
--- a/Engine/Source/Programs/UnrealBuildTool/Platform/IOS/IOSExports.cs
+++ b/Engine/Source/Programs/UnrealBuildTool/Platform/IOS/IOSExports.cs
@@ -92,7 +92,7 @@ namespace UnrealBuildTool
 		/// <param name="ProjectFile"></param>
 		/// <param name="Config"></param>
 		/// <param name="ProjectDirectory"></param>
-		/// <param name="bIsUE4Game"></param>
+		/// <param name="bIsUnrealGame"></param>
 		/// <param name="GameName"></param>
 		/// <param name="bIsClient"></param>
 		/// <param name="ProjectName"></param>
@@ -102,10 +102,10 @@ namespace UnrealBuildTool
 		/// <param name="bSupportsPortrait"></param>
 		/// <param name="bSupportsLandscape"></param>
 		/// <returns></returns>
-		public static bool GeneratePList(FileReference ProjectFile, UnrealTargetConfiguration Config, DirectoryReference ProjectDirectory, bool bIsUE4Game, string GameName, bool bIsClient, string ProjectName, DirectoryReference InEngineDir, DirectoryReference AppDirectory, FileReference BuildReceiptFileName, out bool bSupportsPortrait, out bool bSupportsLandscape)
+		public static bool GeneratePList(FileReference ProjectFile, UnrealTargetConfiguration Config, DirectoryReference ProjectDirectory, bool bIsUnrealGame, string GameName, bool bIsClient, string ProjectName, DirectoryReference InEngineDir, DirectoryReference AppDirectory, FileReference BuildReceiptFileName, out bool bSupportsPortrait, out bool bSupportsLandscape)
 		{
 			TargetReceipt Receipt = TargetReceipt.Read(BuildReceiptFileName);
-			return new UEDeployIOS().GeneratePList(ProjectFile, Config, ProjectDirectory.FullName, bIsUE4Game, GameName, bIsClient, ProjectName, InEngineDir.FullName, AppDirectory.FullName, Receipt, out bSupportsPortrait, out bSupportsLandscape);
+			return new UEDeployIOS().GeneratePList(ProjectFile, Config, ProjectDirectory.FullName, bIsUnrealGame, GameName, bIsClient, ProjectName, InEngineDir.FullName, AppDirectory.FullName, Receipt, out bSupportsPortrait, out bSupportsLandscape);
 		}
 
 		/// <summary>
@@ -134,7 +134,7 @@ namespace UnrealBuildTool
 			bool bUserImagesExist = false;
 			DirectoryReference ResourcesDir = IOSToolChain.GenerateAssetCatalog(ProjectFile, Platform, ref bUserImagesExist);
 
-			// Don't attempt to do anything remotely if the user is using the default UE4 images.
+			// Don't attempt to do anything remotely if the user is using the default UE images.
 			if (!bUserImagesExist)
 			{
 				return;
diff --git a/Engine/Source/Programs/UnrealBuildTool/Platform/Mac/MacToolChain.cs b/Engine/Source/Programs/UnrealBuildTool/Platform/Mac/MacToolChain.cs
--- a/Engine/Source/Programs/UnrealBuildTool/Platform/Mac/MacToolChain.cs
+++ b/Engine/Source/Programs/UnrealBuildTool/Platform/Mac/MacToolChain.cs
@@ -678,7 +678,7 @@ namespace UnrealBuildTool
 			{
 				GameName = "EpicGamesLauncher";
 			}
-			else if (GameName == "UE4" && ProjectFile != null)
+			else if (GameName == "UE5" && ProjectFile != null)
 			{
 				GameName = ProjectFile.GetFileNameWithoutAnyExtensions();
 			}
diff --git a/Engine/Source/Programs/UnrealBuildTool/Platform/Windows/VCToolChain.cs b/Engine/Source/Programs/UnrealBuildTool/Platform/Windows/VCToolChain.cs
--- a/Engine/Source/Programs/UnrealBuildTool/Platform/Windows/VCToolChain.cs
+++ b/Engine/Source/Programs/UnrealBuildTool/Platform/Windows/VCToolChain.cs
@@ -383,10 +383,10 @@ namespace UnrealBuildTool
 				Arguments.Add("/Zc:strictStrings-"); // Have to disable strict const char* semantics due to Windows headers not being compliant.
 			}
 
-			// @todo HoloLens: UE4 is non-compliant when it comes to use of %s and %S
+			// @todo HoloLens: UE is non-compliant when it comes to use of %s and %S
 			// Previously %s meant "the current character set" and %S meant "the other one".
 			// Now %s means multibyte and %S means wide. %Ts means "natural width".
-			// Reverting this behaviour until the UE4 source catches up.
+			// Reverting this behaviour until the UE source catches up.
 			AddDefinition(Arguments, "_CRT_STDIO_LEGACY_WIDE_SPECIFIERS=1");
 
 			// @todo HoloLens: Silence the hash_map deprecation errors for now. This should be replaced with unordered_map for the real fix.
diff --git a/Engine/Source/Programs/UnrealBuildTool/ProjectFiles/CMake/CMakefileGenerator.cs b/Engine/Source/Programs/UnrealBuildTool/ProjectFiles/CMake/CMakefileGenerator.cs
--- a/Engine/Source/Programs/UnrealBuildTool/ProjectFiles/CMake/CMakefileGenerator.cs
+++ b/Engine/Source/Programs/UnrealBuildTool/ProjectFiles/CMake/CMakefileGenerator.cs
@@ -224,11 +224,11 @@ namespace UnrealBuildTool
 			return true;
 		}
 
-		private void AppendCleanedPathToList(StringBuilder EngineFiles, StringBuilder ProjectFiles, String SourceFileRelativeToRoot, String FullName, String GameProjectPath, String UE4RootPath, String GameRootPath)
+		private void AppendCleanedPathToList(StringBuilder EngineFiles, StringBuilder ProjectFiles, String SourceFileRelativeToRoot, String FullName, String GameProjectPath, String UnrealRootPath, String GameRootPath)
 		{
 			if (!SourceFileRelativeToRoot.StartsWith("..") && !Path.IsPathRooted(SourceFileRelativeToRoot))
 			{
-				EngineFiles.Append("\t\"" + UE4RootPath + "/Engine/" + Utils.CleanDirectorySeparators(SourceFileRelativeToRoot, '/') + "\"\n");
+				EngineFiles.Append("\t\"" + UnrealRootPath + "/Engine/" + Utils.CleanDirectorySeparators(SourceFileRelativeToRoot, '/') + "\"\n");
 			}
 			else
 			{
@@ -265,7 +265,7 @@ namespace UnrealBuildTool
 			StringBuilder IncludeDirectoriesList = new StringBuilder("include_directories( \n");
 			StringBuilder PreprocessorDefinitionsList = new StringBuilder("add_definitions( \n");
 
-            string UE4RootPath = Utils.CleanDirectorySeparators(Unreal.RootDirectory.FullName, '/');
+            string UnrealRootPath = Utils.CleanDirectorySeparators(Unreal.RootDirectory.FullName, '/');
             string CMakeGameRootPath = "";
 			string GameProjectPath = "";
 			string CMakeGameProjectFile = "";
@@ -275,19 +275,19 @@ namespace UnrealBuildTool
 			if (BuildHostPlatform.Current.Platform == UnrealTargetPlatform.Win64)
 			{
 				HostArchitecture = "Win64";
-				BuildCommand = "call \"" + UE4RootPath + "/Engine/Build/BatchFiles/Build.bat\"";
+				BuildCommand = "call \"" + UnrealRootPath + "/Engine/Build/BatchFiles/Build.bat\"";
 			}
 			else if (BuildHostPlatform.Current.Platform == UnrealTargetPlatform.Mac)
 			{
 				HostArchitecture = "Mac";
-				BuildCommand = "cd \"" + UE4RootPath + "\" && bash \"" + UE4RootPath + "/Engine/Build/BatchFiles/" + HostArchitecture + "/Build.sh\"";
+				BuildCommand = "cd \"" + UnrealRootPath + "\" && bash \"" + UnrealRootPath + "/Engine/Build/BatchFiles/" + HostArchitecture + "/Build.sh\"";
 				bIncludeIOSTargets = true;
 				bIncludeTVOSTargets = true;
 			}
 			else if (BuildHostPlatform.Current.Platform == UnrealTargetPlatform.Linux)
 			{
 				HostArchitecture = "Linux";
-				BuildCommand = "cd \"" + UE4RootPath + "\" && bash \"" + UE4RootPath + "/Engine/Build/BatchFiles/" + HostArchitecture + "/Build.sh\"";
+				BuildCommand = "cd \"" + UnrealRootPath + "\" && bash \"" + UnrealRootPath + "/Engine/Build/BatchFiles/" + HostArchitecture + "/Build.sh\"";
 
 				string? CompilerPath = LinuxCommon.WhichClang();
 				SetCompiler = "set(CMAKE_CXX_COMPILER " + CompilerPath + ")\n\n";
@@ -323,7 +323,7 @@ namespace UnrealBuildTool
 				"# Makefile generated by CMakefileGenerator.cs (v1.2)\n" +
 				"# *DO NOT EDIT*\n\n" +
 				"cmake_minimum_required (VERSION 2.6)\n" +
-				"project (UE4)\n\n" + 			
+				"project (Unreal)\n\n" + 			
 				"# CMake Flags\n" +
 				"set(CMAKE_CXX_STANDARD 14)\n" + // Need to keep this updated
 				"set(CMAKE_CXX_USE_RESPONSE_FILE_FOR_OBJECTS 1 CACHE BOOL \"\" FORCE)\n" +
@@ -352,7 +352,7 @@ namespace UnrealBuildTool
 					{
 						if (IncludeDirectory.Contains(Unreal.RootDirectory.FullName))
 						{
-							IncludeDirectories.Add(IncludeDirectory.Replace(Unreal.RootDirectory.FullName, UE4RootPath));
+							IncludeDirectories.Add(IncludeDirectory.Replace(Unreal.RootDirectory.FullName, UnrealRootPath));
 						}
 						else
 						{
@@ -404,23 +404,23 @@ namespace UnrealBuildTool
 					{
 						if (SourceFileRelativeToRoot.EndsWith(".cpp"))
 						{							
-							AppendCleanedPathToList(CMakeEngineSourceFilesList, CMakeProjectSourceFilesList, SourceFileRelativeToRoot, CurSourceFile.FullName, GameProjectPath, UE4RootPath, CMakeGameRootPath);
+							AppendCleanedPathToList(CMakeEngineSourceFilesList, CMakeProjectSourceFilesList, SourceFileRelativeToRoot, CurSourceFile.FullName, GameProjectPath, UnrealRootPath, CMakeGameRootPath);
 						}
 						else if (SourceFileRelativeToRoot.EndsWith(".h"))
 						{
-							AppendCleanedPathToList(CMakeEngineHeaderFilesList, CMakeProjectHeaderFilesList, SourceFileRelativeToRoot, CurSourceFile.FullName, GameProjectPath, UE4RootPath, CMakeGameRootPath);
+							AppendCleanedPathToList(CMakeEngineHeaderFilesList, CMakeProjectHeaderFilesList, SourceFileRelativeToRoot, CurSourceFile.FullName, GameProjectPath, UnrealRootPath, CMakeGameRootPath);
 						}
 						else if (SourceFileRelativeToRoot.EndsWith(".cs"))
 						{
-							AppendCleanedPathToList(CMakeEngineCSFilesList, CMakeProjectCSFilesList, SourceFileRelativeToRoot, CurSourceFile.FullName, GameProjectPath, UE4RootPath, CMakeGameRootPath);
+							AppendCleanedPathToList(CMakeEngineCSFilesList, CMakeProjectCSFilesList, SourceFileRelativeToRoot, CurSourceFile.FullName, GameProjectPath, UnrealRootPath, CMakeGameRootPath);
 						}
 						else if (SourceFileRelativeToRoot.EndsWith(".usf") || SourceFileRelativeToRoot.EndsWith(".ush"))
 						{
-							AppendCleanedPathToList(CMakeEngineShaderFilesList, CMakeProjectShaderFilesList, SourceFileRelativeToRoot, CurSourceFile.FullName, GameProjectPath, UE4RootPath, CMakeGameRootPath);
+							AppendCleanedPathToList(CMakeEngineShaderFilesList, CMakeProjectShaderFilesList, SourceFileRelativeToRoot, CurSourceFile.FullName, GameProjectPath, UnrealRootPath, CMakeGameRootPath);
 						}
 						else if (SourceFileRelativeToRoot.EndsWith(".ini"))
 						{
-							AppendCleanedPathToList(CMakeEngineConfigFilesList, CMakeProjectConfigFilesList, SourceFileRelativeToRoot, CurSourceFile.FullName, GameProjectPath, UE4RootPath, CMakeGameRootPath);
+							AppendCleanedPathToList(CMakeEngineConfigFilesList, CMakeProjectConfigFilesList, SourceFileRelativeToRoot, CurSourceFile.FullName, GameProjectPath, UnrealRootPath, CMakeGameRootPath);
 						}
 					}
 				}
@@ -453,7 +453,7 @@ namespace UnrealBuildTool
 				string SourceFileRelativeToRoot = CurSourceFile.MakeRelativeTo(Unreal.EngineDirectory);
 				if (SourceFileRelativeToRoot.EndsWith(".usf") || SourceFileRelativeToRoot.EndsWith(".ush"))
 				{
-					AppendCleanedPathToList(CMakeEngineShaderFilesList, CMakeProjectShaderFilesList, SourceFileRelativeToRoot, CurSourceFile.FullName, GameProjectPath, UE4RootPath, CMakeGameRootPath);
+					AppendCleanedPathToList(CMakeEngineShaderFilesList, CMakeProjectShaderFilesList, SourceFileRelativeToRoot, CurSourceFile.FullName, GameProjectPath, UnrealRootPath, CMakeGameRootPath);
 				}
 			}
 
@@ -464,7 +464,7 @@ namespace UnrealBuildTool
 				string SourceFileRelativeToRoot = CurSourceFile.MakeRelativeTo(Unreal.EngineDirectory);
 				if (SourceFileRelativeToRoot.EndsWith(".ini"))
 				{
-					AppendCleanedPathToList(CMakeEngineConfigFilesList, CMakeProjectConfigFilesList, SourceFileRelativeToRoot, CurSourceFile.FullName, GameProjectPath, UE4RootPath, CMakeGameRootPath);
+					AppendCleanedPathToList(CMakeEngineConfigFilesList, CMakeProjectConfigFilesList, SourceFileRelativeToRoot, CurSourceFile.FullName, GameProjectPath, UnrealRootPath, CMakeGameRootPath);
 				}
 			}
 
@@ -699,7 +699,7 @@ namespace UnrealBuildTool
 			return false;
 		}
 
-		/// Adds the include directory to the list, after converting it to relative to UE4 root
+		/// Adds the include directory to the list, after converting it to relative to Unreal root
 		private static string GetIncludeDirectory(string IncludeDir, string ProjectDir)
 		{
 			string FullProjectPath = Path.GetFullPath(PrimaryProjectPath.FullName);
diff --git a/Engine/Source/Programs/UnrealBuildTool/ProjectFiles/ProjectFileGenerator.cs b/Engine/Source/Programs/UnrealBuildTool/ProjectFiles/ProjectFileGenerator.cs
--- a/Engine/Source/Programs/UnrealBuildTool/ProjectFiles/ProjectFileGenerator.cs
+++ b/Engine/Source/Programs/UnrealBuildTool/ProjectFiles/ProjectFileGenerator.cs
@@ -334,7 +334,7 @@ namespace UnrealBuildTool
 		public static DirectoryReference PrimaryProjectPath = Unreal.RootDirectory; // We'll save the primary project to our "root" folder
 
 		/// <summary>
-		/// Name of the UE4 engine project that contains all of the engine code, config files and other files
+		/// Name of the UE engine project that contains all of the engine code, config files and other files
 		/// </summary>
 		public const string EngineProjectFileNameBase = "UE5";
 
@@ -533,7 +533,7 @@ namespace UnrealBuildTool
 		}
 
 		/// <summary>
-		/// Finds all csproj within Engine/Source/Programs, and add them if their UE4CSharp.prog file exists.
+		/// Finds all csproj within Engine/Source/Programs, and add them if their .uprogram file exists.
 		/// </summary>
 		void DiscoverCSharpProgramProjects(List<FileReference> AllGameProjects, PrimaryProjectFolder ProgramsFolder)
 		{
@@ -659,7 +659,7 @@ namespace UnrealBuildTool
 				{
 					break;
 				}
-				// For csproj files, add them to the sln if the UE4CSharp.prog file also exists.
+				// For csproj files, add them to the sln if the .uprogram file also exists.
 				if (bIsCsProj && FileReference.Exists(File.ChangeExtension(".uprogram")))
 				{
 					FoundProjects.Add(File);
diff --git a/Engine/Source/Programs/UnrealBuildTool/ProjectFiles/QMake/QMakefileGenerator.cs b/Engine/Source/Programs/UnrealBuildTool/ProjectFiles/QMake/QMakefileGenerator.cs
--- a/Engine/Source/Programs/UnrealBuildTool/ProjectFiles/QMake/QMakefileGenerator.cs
+++ b/Engine/Source/Programs/UnrealBuildTool/ProjectFiles/QMake/QMakefileGenerator.cs
@@ -66,7 +66,7 @@ namespace UnrealBuildTool
 			}
 		}
 
-		/// Adds the include directory to the list, after converting it to relative to UE4 root
+		/// Adds the include directory to the list, after converting it to relative to Unreal root
 		private void AddIncludeDirectory(ref List<string> IncludeDirectories, string IncludeDir, string ProjectDir)
 		{
 			string FullProjectPath = ProjectFileGenerator.PrimaryProjectPath.FullName;
@@ -251,7 +251,7 @@ namespace UnrealBuildTool
 				"CONFIG -= console\n" +
 				"CONFIG -= app_bundle\n" +
 				"CONFIG -= qt\n\n" +
-				"TARGET = UE4 \n\n" +
+				"TARGET = UE5 \n\n" +
 				"unrealRootPath=" + UnrealRootPath + "\n" +
 				GameProjectRootPath +
 				QMakeGameProjectFile +
diff --git a/Engine/Source/Programs/UnrealBuildTool/ProjectFiles/Rider/RiderProjectFile.cs b/Engine/Source/Programs/UnrealBuildTool/ProjectFiles/Rider/RiderProjectFile.cs
--- a/Engine/Source/Programs/UnrealBuildTool/ProjectFiles/Rider/RiderProjectFile.cs
+++ b/Engine/Source/Programs/UnrealBuildTool/ProjectFiles/Rider/RiderProjectFile.cs
@@ -36,7 +36,7 @@ namespace UnrealBuildTool
 		/// For every combination of <c>UnrealTargetPlatform</c>, <c>UnrealTargetConfiguration</c> and <c>TargetType</c>
 		/// will be generated separate JSON file.
 		/// Project file will be stored:
-		/// For UE4:  {UE4Root}/Engine/Intermediate/ProjectFiles/.Rider/{Platform}/{Configuration}/{TargetType}/{ProjectName}.json
+		/// For UE:  {UnrealRoot}/Engine/Intermediate/ProjectFiles/.Rider/{Platform}/{Configuration}/{TargetType}/{ProjectName}.json
 		/// For game: {GameRoot}/Intermediate/ProjectFiles/.Rider/{Platform}/{Configuration}/{TargetType}/{ProjectName}.json
 		/// </summary>
 		/// <remarks>
diff --git a/Engine/Source/Programs/UnrealBuildTool/ProjectFiles/VisualStudio/VCProject.cs b/Engine/Source/Programs/UnrealBuildTool/ProjectFiles/VisualStudio/VCProject.cs
--- a/Engine/Source/Programs/UnrealBuildTool/ProjectFiles/VisualStudio/VCProject.cs
+++ b/Engine/Source/Programs/UnrealBuildTool/ProjectFiles/VisualStudio/VCProject.cs
@@ -1686,7 +1686,7 @@ namespace UnrealBuildTool
 						VCProjectFileContent.Append(PathStrings);
 					}
 
-					// This is the standard UE4 based project NMake build line:
+					// This is the standard UE based project NMake build line:
 					//	..\..\Build\BatchFiles\Build.bat <TARGETNAME> <PLATFORM> <CONFIGURATION>
 					//	ie ..\..\Build\BatchFiles\Build.bat BlankProgram Win64 Debug
 
diff --git a/Engine/Source/Programs/UnrealBuildTool/System/BuildVersion.cs b/Engine/Source/Programs/UnrealBuildTool/System/BuildVersion.cs
--- a/Engine/Source/Programs/UnrealBuildTool/System/BuildVersion.cs
+++ b/Engine/Source/Programs/UnrealBuildTool/System/BuildVersion.cs
@@ -19,7 +19,7 @@ namespace UnrealBuildTool
 	public class BuildVersion
 	{
 		/// <summary>
-		/// The major engine version (4 for UE4)
+		/// The major engine version (5 for UE5)
 		/// </summary>
 		public int MajorVersion;
 
diff --git a/Engine/Source/Programs/UnrealBuildTool/System/ConfigFileAttribute.cs b/Engine/Source/Programs/UnrealBuildTool/System/ConfigFileAttribute.cs
--- a/Engine/Source/Programs/UnrealBuildTool/System/ConfigFileAttribute.cs
+++ b/Engine/Source/Programs/UnrealBuildTool/System/ConfigFileAttribute.cs
@@ -9,7 +9,7 @@ using System.Threading.Tasks;
 namespace UnrealBuildTool
 {
 	/// <summary>
-	/// Attribute indicating a value which should be populated from a UE4 .ini config file
+	/// Attribute indicating a value which should be populated from a UE .ini config file
 	/// </summary>
 	[AttributeUsage(AttributeTargets.Field)]
 	public class ConfigFileAttribute : Attribute
diff --git a/Engine/Source/Programs/UnrealBuildTool/System/ConfigHierarchy.cs b/Engine/Source/Programs/UnrealBuildTool/System/ConfigHierarchy.cs
--- a/Engine/Source/Programs/UnrealBuildTool/System/ConfigHierarchy.cs
+++ b/Engine/Source/Programs/UnrealBuildTool/System/ConfigHierarchy.cs
@@ -557,7 +557,7 @@ namespace UnrealBuildTool
 		}
 
 		/// <summary>
-		/// Attempts to parse the given line as a UE4 config object (eg. (Name="Foo",Number=1234)).
+		/// Attempts to parse the given line as a UE config object (eg. (Name="Foo",Number=1234)).
 		/// </summary>
 		/// <param name="Line">Line of text to parse</param>
 		/// <param name="Properties">Receives key/value pairs for the config object</param>
@@ -739,7 +739,7 @@ namespace UnrealBuildTool
 
 
 		/// <summary>
-		/// Attempts to parse the given line as a UE4 config array (eg. ("one", "two", "three") ).
+		/// Attempts to parse the given line as a UE config array (eg. ("one", "two", "three") ).
 		/// </summary>
 		/// <param name="Line">Line of text to parse</param>
 		/// <param name="Array">Receives array for the config array</param>
diff --git a/Engine/Source/Programs/UnrealBuildTool/System/ExternalExecution.cs b/Engine/Source/Programs/UnrealBuildTool/System/ExternalExecution.cs
--- a/Engine/Source/Programs/UnrealBuildTool/System/ExternalExecution.cs
+++ b/Engine/Source/Programs/UnrealBuildTool/System/ExternalExecution.cs
@@ -22,7 +22,7 @@ using OpenTracing.Util;
 namespace UnrealBuildTool
 {
 	// This enum has to be compatible with the one defined in the
-	// UE4\Engine\Source\Runtime\Core\Public\Misc\ComplilationResult.h 
+	// Engine\Source\Runtime\Core\Public\Misc\ComplilationResult.h 
 	// to keep communication between UHT, UBT and Editor compiling
 	// processes valid.
 	enum CompilationResult
@@ -1266,7 +1266,7 @@ namespace UnrealBuildTool
 						if (BuildHostPlatform.Current.Platform == UnrealTargetPlatform.Win64 && 
 							(int)(UHTResult) < 0)
 						{
-							Log.TraceError(String.Format("UnrealHeaderTool failed with exit code 0x{0:X} - check that UE4 prerequisites are installed.", (int)UHTResult));
+							Log.TraceError(String.Format("UnrealHeaderTool failed with exit code 0x{0:X} - check that Unreal Engine prerequisites are installed.", (int)UHTResult));
 						}
 
 						throw new CompilationResultException(UHTResult);
diff --git a/Engine/Source/Programs/UnrealBuildTool/System/JunkDeleter.cs b/Engine/Source/Programs/UnrealBuildTool/System/JunkDeleter.cs
--- a/Engine/Source/Programs/UnrealBuildTool/System/JunkDeleter.cs
+++ b/Engine/Source/Programs/UnrealBuildTool/System/JunkDeleter.cs
@@ -74,7 +74,7 @@ namespace UnrealBuildTool
 								}
 							}
 
-							// All paths within the manifest are UE4 root directory relative.
+							// All paths within the manifest are Unreal root directory relative.
 							// UBT's working directory is Engine\Source so add "..\..\" to each of the entires.
 							if (bIsValidJunkLine)
 							{
diff --git a/Engine/Source/Programs/UnrealBuildTool/System/LinkEnvironment.cs b/Engine/Source/Programs/UnrealBuildTool/System/LinkEnvironment.cs
--- a/Engine/Source/Programs/UnrealBuildTool/System/LinkEnvironment.cs
+++ b/Engine/Source/Programs/UnrealBuildTool/System/LinkEnvironment.cs
@@ -164,7 +164,7 @@ namespace UnrealBuildTool
 		public bool bIsBuildingConsoleApplication = false;
 
 		/// <summary>
-		/// If set, overrides the program entry function on Windows platform.  This is used by the base UE4
+		/// If set, overrides the program entry function on Windows platform.  This is used by the base Unreal
 		/// program so we can link in either command-line mode or windowed mode without having to recompile the Launch module.
 		/// </summary>
 		public string WindowsEntryPointOverride = String.Empty;
diff --git a/Engine/Source/Programs/UnrealBuildTool/System/ModuleDescriptor.cs b/Engine/Source/Programs/UnrealBuildTool/System/ModuleDescriptor.cs
--- a/Engine/Source/Programs/UnrealBuildTool/System/ModuleDescriptor.cs
+++ b/Engine/Source/Programs/UnrealBuildTool/System/ModuleDescriptor.cs
@@ -21,7 +21,7 @@ namespace UnrealBuildTool
 		Default,
 
 		/// <summary>
-		/// Any target using the UE4 runtime
+		/// Any target using the UE runtime
 		/// </summary>
 		Runtime,
 
diff --git a/Engine/Source/Programs/UnrealBuildTool/System/RestrictedFolders.cs b/Engine/Source/Programs/UnrealBuildTool/System/RestrictedFolders.cs
--- a/Engine/Source/Programs/UnrealBuildTool/System/RestrictedFolders.cs
+++ b/Engine/Source/Programs/UnrealBuildTool/System/RestrictedFolders.cs
@@ -185,7 +185,7 @@ namespace UnrealBuildTool
 		public static RestrictedFolder EpicInternal = FindOrAddByName("EpicInternal");
 
 		/// <summary>
-		/// Can be used by UE4 but not required
+		/// Can be used by UE but not required
 		/// </summary>
 		public static RestrictedFolder CarefullyRedist = FindOrAddByName("CarefullyRedist");
 
diff --git a/Engine/Source/Programs/UnrealBuildTool/System/RulesAssembly.cs b/Engine/Source/Programs/UnrealBuildTool/System/RulesAssembly.cs
--- a/Engine/Source/Programs/UnrealBuildTool/System/RulesAssembly.cs
+++ b/Engine/Source/Programs/UnrealBuildTool/System/RulesAssembly.cs
@@ -223,7 +223,7 @@ namespace UnrealBuildTool
 			PreprocessorDefines.Add("WITH_FORWARDED_MODULE_RULES_CTOR");
 			PreprocessorDefines.Add("WITH_FORWARDED_TARGET_RULES_CTOR");
 
-			// Define macros for the UE4 version, starting with 4.17
+			// Define macros for the Unreal engine version, starting with 4.17
 			BuildVersion? Version;
 			if (BuildVersion.TryRead(BuildVersion.GetDefaultFileName(), out Version))
 			{
