UnrealBuildTool: Remove references to UE4, changes affecting Linux

diff --git a/Engine/Source/Programs/UnrealBuildTool/Platform/Linux/LinuxToolChain.cs b/Engine/Source/Programs/UnrealBuildTool/Platform/Linux/LinuxToolChain.cs
--- a/Engine/Source/Programs/UnrealBuildTool/Platform/Linux/LinuxToolChain.cs
+++ b/Engine/Source/Programs/UnrealBuildTool/Platform/Linux/LinuxToolChain.cs
@@ -501,8 +501,8 @@ namespace UnrealBuildTool
 
 		private static bool ShouldUseLibcxx(string Architecture)
 		{
-			// set UE4_LINUX_USE_LIBCXX to either 0 or 1. If unset, defaults to 1.
-			string? UseLibcxxEnvVarOverride = Environment.GetEnvironmentVariable("UE4_LINUX_USE_LIBCXX");
+			// set UE_LINUX_USE_LIBCXX to either 0 or 1. If unset, defaults to 1.
+			string? UseLibcxxEnvVarOverride = Environment.GetEnvironmentVariable("UE_LINUX_USE_LIBCXX");
 			if (string.IsNullOrEmpty(UseLibcxxEnvVarOverride) || UseLibcxxEnvVarOverride == "1")
 			{
 				// at the moment ARM32 libc++ remains missing
diff --git a/Engine/Source/Programs/UnrealBuildTool/ProjectFiles/Rider/RiderProjectFile.cs b/Engine/Source/Programs/UnrealBuildTool/ProjectFiles/Rider/RiderProjectFile.cs
--- a/Engine/Source/Programs/UnrealBuildTool/ProjectFiles/Rider/RiderProjectFile.cs
+++ b/Engine/Source/Programs/UnrealBuildTool/ProjectFiles/Rider/RiderProjectFile.cs
@@ -499,7 +499,7 @@ namespace UnrealBuildTool
 			{
 				var EngineDirectory = Unreal.EngineDirectory.ToString();
 
-				string? UseLibcxxEnvVarOverride = Environment.GetEnvironmentVariable("UE4_LINUX_USE_LIBCXX");
+				string? UseLibcxxEnvVarOverride = Environment.GetEnvironmentVariable("UE_LINUX_USE_LIBCXX");
 				if (string.IsNullOrEmpty(UseLibcxxEnvVarOverride) || UseLibcxxEnvVarOverride == "1")
 				{
 					if (Target.Architecture.StartsWith("x86_64") ||
