UnrealBuildTool: Remove references to UE4, changes affecting Mac platforms.

diff --git a/Engine/Source/Programs/AutomationTool/Gauntlet/Platform/Mac/Gauntlet.MacBuildSource.cs b/Engine/Source/Programs/AutomationTool/Gauntlet/Platform/Mac/Gauntlet.MacBuildSource.cs
--- a/Engine/Source/Programs/AutomationTool/Gauntlet/Platform/Mac/Gauntlet.MacBuildSource.cs
+++ b/Engine/Source/Programs/AutomationTool/Gauntlet/Platform/Mac/Gauntlet.MacBuildSource.cs
@@ -112,7 +112,7 @@ namespace Gauntlet
 
 					// check both binary AND Engine content. Regular Unreal builds won't have the latter
 					if (Directory.Exists(MacOSDir)
-						&& Directory.Exists(Path.Combine(Di.FullName, "Contents", "UE4")))
+						&& Directory.Exists(Path.Combine(Di.FullName, "Contents", "UE")))
 					{
 						// Check there's an executable
 						FileInfo Executable = new DirectoryInfo(MacOSDir).GetFiles().Where(Fi => Fi.Name.StartsWith(ShortName, StringComparison.OrdinalIgnoreCase)).FirstOrDefault();
diff --git a/Engine/Source/Programs/AutomationTool/Mac/MacPlatform.Automation.cs b/Engine/Source/Programs/AutomationTool/Mac/MacPlatform.Automation.cs
--- a/Engine/Source/Programs/AutomationTool/Mac/MacPlatform.Automation.cs
+++ b/Engine/Source/Programs/AutomationTool/Mac/MacPlatform.Automation.cs
@@ -465,7 +465,7 @@ public class MacPlatform : Platform
 				}
 			}
 
-			string TargetPath = CombinePaths(BundlePath, "Contents", "UE4");
+			string TargetPath = CombinePaths(BundlePath, "Contents", "UE");
 			if (!SC.bIsCombiningMultiplePlatforms)
 			{
 				if (Directory.Exists(BundlePath))
diff --git a/Engine/Source/Programs/Mac/UnrealEditorServices/Private/Mac/UnrealEditorServicesAppDelegate.cpp b/Engine/Source/Programs/Mac/UnrealEditorServices/Private/Mac/UnrealEditorServicesAppDelegate.cpp
--- a/Engine/Source/Programs/Mac/UnrealEditorServices/Private/Mac/UnrealEditorServicesAppDelegate.cpp
+++ b/Engine/Source/Programs/Mac/UnrealEditorServices/Private/Mac/UnrealEditorServicesAppDelegate.cpp
@@ -427,7 +427,7 @@
 {
 	return AppBundle && ([[AppBundle bundleIdentifier] isEqualToString:@"com.epicgames.UnrealEditor"] || [[AppBundle bundleIdentifier] isEqualToString:@"com.epicgames.UnrealEditorServices"])
 					 && [[[AppBundle bundlePath] stringByDeletingLastPathComponent] hasSuffix:@"Engine/Binaries/Mac"]
-					 && [[AppBundle bundlePath] rangeOfString:@".app/Contents/UE4"].location == NSNotFound;
+					 && [[AppBundle bundlePath] rangeOfString:@".app/Contents/UE"].location == NSNotFound;
 }
 
 - (NSString*)enginePathForAppBundle:(NSBundle*)AppBundle
diff --git a/Engine/Source/Programs/UnrealBuildTool/Platform/IOS/UEBuildIOS.cs b/Engine/Source/Programs/UnrealBuildTool/Platform/IOS/UEBuildIOS.cs
--- a/Engine/Source/Programs/UnrealBuildTool/Platform/IOS/UEBuildIOS.cs
+++ b/Engine/Source/Programs/UnrealBuildTool/Platform/IOS/UEBuildIOS.cs
@@ -1171,7 +1171,7 @@ namespace UnrealBuildTool
 			int IntPart = (int)TargetFloat;
 			int FracPart = (int)((TargetFloat - IntPart) * 10);
 			int TargetNum = IntPart * 10000 + FracPart * 100;
-			CompileEnvironment.Definitions.Add("MINIMUM_UE4_COMPILED_IOS_VERSION=" + TargetNum);
+			CompileEnvironment.Definitions.Add("MINIMUM_UE_COMPILED_IOS_VERSION=" + TargetNum);
 
 			LinkEnvironment.AdditionalFrameworks.Add(new UEBuildFramework("GameKit"));
 			LinkEnvironment.AdditionalFrameworks.Add(new UEBuildFramework("StoreKit"));
diff --git a/Engine/Source/Programs/UnrealBuildTool/Platform/Mac/MacToolChain.cs b/Engine/Source/Programs/UnrealBuildTool/Platform/Mac/MacToolChain.cs
--- a/Engine/Source/Programs/UnrealBuildTool/Platform/Mac/MacToolChain.cs
+++ b/Engine/Source/Programs/UnrealBuildTool/Platform/Mac/MacToolChain.cs
@@ -697,7 +697,7 @@ namespace UnrealBuildTool
 				&& !LibraryFullPath.Contains("/Engine/Source/ThirdParty/") && LibraryDir != ExeDir && !RPaths.Contains(LibraryDir))
 			{
 				// macOS gatekeeper erroneously complains about not seeing the CEF3 framework in the codesigned Launcher because it's only present in one of the folders specified in RPATHs.
-				// To work around this we will only add a single RPATH entry for it, for the framework stored in .app/Contents/UE4/ subfolder of the packaged app bundle
+				// To work around this we will only add a single RPATH entry for it, for the framework stored in .app/Contents/UE/ subfolder of the packaged app bundle
 				bool bCanUseMultipleRPATHs = !ExeAbsolutePath.Contains("EpicGamesLauncher-Mac-Shipping") || !Library.Contains("CEF3");
 
 				// First, add a path relative to the executable.
@@ -724,17 +724,17 @@ namespace UnrealBuildTool
 					string EngineDir = Unreal.RootDirectory.ToString();
 					string? ProjectDir = ProjectFile?.Directory.FullName;
 
-					// In packaged games dylibs are stored in Contents/UE4 subfolders, for example in GameName.app/Contents/UE4/Engine/Binaries/ThirdParty/PhysX/Mac
-					string BundleUE4Dir = Path.GetFullPath(ExeDir + "/../../Contents/UE4");
+					// In packaged games dylibs are stored in Contents/UE subfolders, for example in GameName.app/Contents/UE/Engine/Binaries/ThirdParty/PhysX/Mac
+					string BundleUEDir = Path.GetFullPath(ExeDir + "/../../Contents/UE");
 					string? BundleLibraryDir = null;
 					if (LibraryDir.StartsWith(EngineDir + "/"))
 					{
-						BundleLibraryDir = LibraryDir.Replace(EngineDir, BundleUE4Dir);
+						BundleLibraryDir = LibraryDir.Replace(EngineDir, BundleUEDir);
 					}
 					else if (ProjectDir != null && LibraryDir.StartsWith(ProjectDir + "/"))
 					{
 						string GameName = GetGameNameFromExecutablePath(ExeAbsolutePath);
-						BundleLibraryDir = LibraryDir.Replace(ProjectDir, BundleUE4Dir + "/" + GameName);
+						BundleLibraryDir = LibraryDir.Replace(ProjectDir, BundleUEDir + "/" + GameName);
 					}
 
 					if (BundleLibraryDir != null)
@@ -750,8 +750,8 @@ namespace UnrealBuildTool
 					// For staged code-based games we need additional entry if the game is not stored directly in the engine's root directory
 					if (bCanUseMultipleRPATHs)
 					{
-						string StagedUE4Dir = Path.GetFullPath(ExeDir + "/../../../../../..");
-						string StagedLibraryDir = LibraryDir.Replace(EngineDir, StagedUE4Dir);
+						string StagedUEDir = Path.GetFullPath(ExeDir + "/../../../../../..");
+						string StagedLibraryDir = LibraryDir.Replace(EngineDir, StagedUEDir);
 						string StagedRelativeDir = Utils.MakePathRelativeTo(StagedLibraryDir, ExeDir).Replace("\\", "/");
 						if (StagedRelativeDir != RelativePath)
 						{
diff --git a/Engine/Source/Runtime/Core/Private/Mac/MacPlatformProcess.cpp b/Engine/Source/Runtime/Core/Private/Mac/MacPlatformProcess.cpp
--- a/Engine/Source/Runtime/Core/Private/Mac/MacPlatformProcess.cpp
+++ b/Engine/Source/Runtime/Core/Private/Mac/MacPlatformProcess.cpp
@@ -1026,11 +1026,11 @@ const TCHAR* FMacPlatformProcess::BaseDir()
 				NSString* BundledBinariesPath = NULL;
 				if (!FApp::IsProjectNameEmpty())
 				{
-					BundledBinariesPath = [BasePath stringByAppendingPathComponent : [NSString stringWithFormat : @"Contents/UE4/%s/Binaries/Mac", TCHAR_TO_UTF8(FApp::GetProjectName())]];
+					BundledBinariesPath = [BasePath stringByAppendingPathComponent : [NSString stringWithFormat : @"Contents/UE/%s/Binaries/Mac", TCHAR_TO_UTF8(FApp::GetProjectName())]];
 				}
 				if (!BundledBinariesPath || ![FileManager fileExistsAtPath:BundledBinariesPath])
 				{
-					BundledBinariesPath = [BasePath stringByAppendingPathComponent: @"Contents/UE4/Engine/Binaries/Mac"];
+					BundledBinariesPath = [BasePath stringByAppendingPathComponent: @"Contents/UE/Engine/Binaries/Mac"];
 				}
 				if ([FileManager fileExistsAtPath: BundledBinariesPath])
 				{
