UnrealBuildTool: Remove reference to UE4

#trivial

diff --git a/Engine/Source/Programs/UnrealBuildTool/ProjectFiles/CMake/CMakefileGenerator.cs b/Engine/Source/Programs/UnrealBuildTool/ProjectFiles/CMake/CMakefileGenerator.cs
--- a/Engine/Source/Programs/UnrealBuildTool/ProjectFiles/CMake/CMakefileGenerator.cs
+++ b/Engine/Source/Programs/UnrealBuildTool/ProjectFiles/CMake/CMakefileGenerator.cs
@@ -165,7 +165,7 @@ namespace UnrealBuildTool
 		{
 			get
 			{
-				return "cmake-headers-ue4" + CMakeExtension;
+				return "cmake-headers-ue" + CMakeExtension;
 			}
 		}
 		/// <summary>
