UnrealBuildTool: Update HoloLensToolChain to inherit from VcToolChain and remove duplicate code

#rnx

diff --git a/Engine/Source/Programs/UnrealBuildTool/Platform/HoloLens/HoloLensProjectGenerator.cs b/Engine/Source/Programs/UnrealBuildTool/Platform/HoloLens/HoloLensProjectGenerator.cs
--- a/Engine/Source/Programs/UnrealBuildTool/Platform/HoloLens/HoloLensProjectGenerator.cs
+++ b/Engine/Source/Programs/UnrealBuildTool/Platform/HoloLens/HoloLensProjectGenerator.cs
@@ -79,13 +79,13 @@ namespace UnrealBuildTool
 			}
 
 			WindowsCompiler Compiler = WindowsPlatform.GetDefaultCompiler(TargetFilePath);
-			DirectoryReference? PlatformWinMDLocation = HoloLens.GetCppCXMetadataLocation(Compiler, "Latest");
+			DirectoryReference? PlatformWinMDLocation = HoloLensPlatform.GetCppCXMetadataLocation(Compiler, "Latest");
 			if (PlatformWinMDLocation == null || !FileReference.Exists(FileReference.Combine(PlatformWinMDLocation, "platform.winmd")))
 			{
 				throw new BuildException("Unable to find platform.winmd for {0} toolchain; '{1}' is an invalid version", WindowsPlatform.GetCompilerName(Compiler), "Latest");
 			}
-			string FoundationWinMDPath = HoloLens.GetLatestMetadataPathForApiContract("Windows.Foundation.FoundationContract", Compiler);
-			string UniversalWinMDPath = HoloLens.GetLatestMetadataPathForApiContract("Windows.Foundation.UniversalApiContract", Compiler);
+			string FoundationWinMDPath = HoloLensPlatform.GetLatestMetadataPathForApiContract("Windows.Foundation.FoundationContract", Compiler);
+			string UniversalWinMDPath = HoloLensPlatform.GetLatestMetadataPathForApiContract("Windows.Foundation.UniversalApiContract", Compiler);
 			VCProjectFileContent.Append("		<AdditionalOptions>/ZW /ZW:nostdlib</AdditionalOptions>" + ProjectFileGenerator.NewLine);
 			VCProjectFileContent.Append("		<NMakePreprocessorDefinitions>$(NMakePreprocessorDefinitions);PLATFORM_HOLOLENS=1;HOLOLENS=1;</NMakePreprocessorDefinitions>" + ProjectFileGenerator.NewLine);
 			if (PlatformWinMDLocation != null)
diff --git a/Engine/Source/Programs/UnrealBuildTool/Platform/HoloLens/HoloLensToolChain.cs b/Engine/Source/Programs/UnrealBuildTool/Platform/HoloLens/HoloLensToolChain.cs
--- a/Engine/Source/Programs/UnrealBuildTool/Platform/HoloLens/HoloLensToolChain.cs
+++ b/Engine/Source/Programs/UnrealBuildTool/Platform/HoloLens/HoloLensToolChain.cs
@@ -21,241 +21,56 @@ namespace UnrealBuildTool
 		}
 	}
 
-	class HoloLensToolChain : UEToolChain
+	class HoloLensToolChain : VCToolChain
 	{
-
-		/// <summary>
-		/// The Visual C++ environment
-		/// </summary>
-		protected VCEnvironment EnvVars;
-
-		/// <summary>
-		/// The target being built
-		/// </summary>
-		protected ReadOnlyTargetRules Target;
-		
-		private readonly string CodeAnalysisExtension = ".nativecodeanalysis.xml";
-
 		public HoloLensToolChain(ReadOnlyTargetRules Target)
+			: base(Target)
 		{
-			this.Target = Target;
-			EnvVars = Target.WindowsPlatform.Environment!;
-
-			if (Target.WindowsPlatform.ObjSrcMapFile != null)
-			{
-				try
-				{
-					File.Delete(Target.WindowsPlatform.ObjSrcMapFile);
-				}
-				catch
-				{
-				}
-			}
-
 			// by default tools chains don't parse arguments, but we want to be able to check the -architectures flag defined above. This is
 			// only necessary when AndroidToolChain is used during UAT
 			CommandLine.ParseArguments(Environment.GetCommandLineArgs(), this);
 		}
-		public override void SetEnvironmentVariables()
-		{
-			EnvVars.SetEnvironmentVariables();
-		}
 
-		void AppendCLArguments_Global(CppCompileEnvironment CompileEnvironment, VCEnvironment EnvVars, List<string> Arguments)
+		protected override void AppendCLArguments_Global(CppCompileEnvironment CompileEnvironment, List<string> Arguments)
 		{
-			//Arguments.Add("/showIncludes");
-
-			// Suppress generation of object code for unreferenced inline functions. Enabling this option is more standards compliant, and causes a big reduction
-			// in object file sizes (and link times) due to the amount of stuff we inline.
-			Arguments.Add("/Zc:inline");
-
-			// Prevents the compiler from displaying its logo for each invocation.
-			Arguments.Add("/nologo");
-
-			// Enable intrinsic functions.
-			Arguments.Add("/Oi");
-
-			// Pack struct members on 8-byte boundaries.
-			// Note: we do this even for 32bit builds because it's the required packing mode
-			// for WinRT types, and currently enum class forward declarations (which are everywhere)
-			// are interpreted as WinRT.
-			Arguments.Add("/Zp8");
-
-			if (CompileEnvironment.Platform != UnrealTargetPlatform.HoloLens)
-			{
-				// Allow the compiler to generate SSE2 instructions. (On by default in 64bit)
-				Arguments.Add("/arch:SSE2");
-			}
-
-			// Separate functions for linker.
-			Arguments.Add("/Gy");
-
-			// Relaxes floating point precision semantics to allow more optimization.
-			Arguments.Add("/fp:fast");
-
-			// Compile into an .obj file, and skip linking.
-			Arguments.Add("/c");
-
-			// Allow 900% of the default memory allocation limit.
-			Arguments.Add("/Zm900");
-
-			// Allow large object files to avoid hitting the 2^16 section limit when running with -StressTestUnity.
-			Arguments.Add("/bigobj");
-
-			// Disable "The file contains a character that cannot be represented in the current code page" warning for non-US windows.
-			Arguments.Add("/wd4819");
-
-			{
-				VCToolChain.AddDefinition(Arguments, "_CRT_STDIO_LEGACY_WIDE_SPECIFIERS", "1");
-				//VCToolChain.AddDefinition(Arguments, "USE_SECURE_CRT", "1");
-			}
+			base.AppendCLArguments_Global(CompileEnvironment, Arguments);
 
 			// @todo HoloLens: Disable "unreachable code" warning since auto-included vccorlib.h triggers it
 			Arguments.Add("/wd4702");
 
-			// Disable "usage of ATL attributes is deprecated" since WRL headers generate this
+			// @todo HoloLens: Disable "usage of ATL attributes is deprecated" since WRL headers generate this
 			Arguments.Add("/wd4467");
 
-			// @todo HoloLens: Silence the hash_map deprecation errors for now. This should be replaced with unordered_map for the real fix.
-			{
-				VCToolChain.AddDefinition(Arguments, "_SILENCE_STDEXT_HASH_DEPRECATION_WARNINGS");
-			}
-
-			// If compiling as a DLL, set the relevant defines
-			if (CompileEnvironment.bIsBuildingDLL)
-			{
-				VCToolChain.AddDefinition(Arguments, "_WINDLL");
-			}
-
-			// Experimental deterministic compile support
-			if (Target.WindowsPlatform.bDeterministic)
-			{
-				if (Target.WindowsPlatform.Compiler.IsMSVC())
-				{
-					Arguments.Add("/experimental:deterministic");
-				}
-			}
-
-			//
-			//	Debug
-			//
-			if (CompileEnvironment.Configuration == CppConfiguration.Debug)
-			{
-				// Disable compiler optimization.
-				Arguments.Add("/Od");
-
-				// Favor code size (especially useful for embedded platforms).
-				Arguments.Add("/Os");
-
-				// Allow inline method expansion unless E&C support is requested
-				if( !CompileEnvironment.bSupportEditAndContinue )
-				{
-					Arguments.Add("/Ob2");
-				}
-
-				if (CompileEnvironment.bUseDebugCRT)
-				{
-					Arguments.Add("/RTC1");
-				}
-			}
-			//
-			//	Development and LTCG
-			//
-			else
-			{
-				// Maximum optimizations if desired.
-				if (CompileEnvironment.bOptimizeCode)
-				{
-					Arguments.Add("/Ox");
-
-					// Allow optimized code to be debugged more easily.  This makes PDBs a bit larger, but doesn't noticeably affect
-					// compile times.  The executable code is not affected at all by this switch, only the debugging information.
-					// VC2013 Update 3 has a new flag for doing this
-					Arguments.Add("/Zo");
-				}
-
-				// Favor code speed.
-				Arguments.Add("/Ot");
-
-				// Only omit frame pointers on the PC (which is implied by /Ox) if wanted.
-				if (CompileEnvironment.bOmitFramePointers == false)
-				{
-					Arguments.Add("/Oy-");
-				}
+			// @todo HoloLens: Disable "declaration of 'identifier' hides class member"
+			Arguments.Add("/wd4458");
 
-				// Allow inline method expansion
-				Arguments.Add("/Ob2");
+			// @todo HoloLens: Disable "not defined as a preprocessor macro, replacing with '0'"
+			Arguments.Add("/wd4668");
 
-				//
-				// LTCG
-				//
-				if (CompileEnvironment.Configuration == CppConfiguration.Shipping)
-				{
-					if (CompileEnvironment.bAllowLTCG)
-					{
-						// Enable link-time code generation.
-						Arguments.Add("/GL");
-					}
-				}
-			}
-
-			//
-			//	PC
-			//
-
-			// Prompt the user before reporting internal errors to Microsoft.
-			Arguments.Add("/errorReport:prompt");
-
-			// Enable C++ exception handling, but not C exceptions.
-			Arguments.Add("/EHsc");
-
-			// If enabled, create debug information.
-			if (CompileEnvironment.bCreateDebugInfo)
+			if (CompileEnvironment.Configuration == CppConfiguration.Debug && CompileEnvironment.bUseDebugCRT)
 			{
-				// Store debug info in .pdb files.
-				// @todo clang: PDB files are emited from Clang but do not fully work with Visual Studio yet (breakpoints won't hit due to "symbol read error")
-				if (CompileEnvironment.bUsePDBFiles)
-				{
-					// Create debug info suitable for E&C if wanted.
-					if (CompileEnvironment.bSupportEditAndContinue &&
-						// We only need to do this in debug as that's the only configuration that supports E&C.
-						CompileEnvironment.Configuration == CppConfiguration.Debug)
-					{
-						Arguments.Add("/ZI");
-					}
-					// Regular PDB debug information.
-					else
-					{
-						Arguments.Add("/Zi");
-					}
-				}
-				// Store C7-format debug info in the .obj files, which is faster.
-				else
-				{
-					Arguments.Add("/Z7");
-				}
+				// @todo HoloLens: Reconcile with "/RTCs" in VCToolChain
+				Arguments.Add("/RTC1");
 			}
 
-			// Static CRT not supported for HoloLens
-			if (CompileEnvironment.bUseDebugCRT)
+			// C2338: two-phase name lookup is not supported for C++/CLI or C++/CX; use /Zc:twoPhase-
+			if (Target.WindowsPlatform.Compiler.IsMSVC() && Target.WindowsPlatform.bStrictConformanceMode)
 			{
-				Arguments.Add("/MDd");
-			}
-			else
-			{
-				Arguments.Add("/MD");
+				Arguments.Add("/Zc:twoPhase-");
 			}
 
-			DirectoryReference? PlatformWinMDLocation = HoloLens.GetCppCXMetadataLocation(EnvVars.Compiler, EnvVars.ToolChainDir);
+			DirectoryReference? PlatformWinMDLocation = HoloLensPlatform.GetCppCXMetadataLocation(EnvVars.Compiler, EnvVars.ToolChainDir);
 			if (PlatformWinMDLocation != null)
 			{
-				Arguments.AddFormat(@" /AI""{0}""", PlatformWinMDLocation);
-				Arguments.AddFormat(@" /FU""{0}\platform.winmd""", PlatformWinMDLocation);
+				Arguments.AddFormat(@"/AI""{0}""", PlatformWinMDLocation);
+				Arguments.AddFormat(@"/FU""{0}\platform.winmd""", PlatformWinMDLocation);
 			}
 		}
 
-		static void AppendCLArguments_CPP(CppCompileEnvironment CompileEnvironment, List<string> Arguments)
+		protected override void AppendCLArguments_CPP(CppCompileEnvironment CompileEnvironment, List<string> Arguments)
 		{
+			base.AppendCLArguments_CPP(CompileEnvironment, Arguments);
+
 			// Enable Windows Runtime extensions.  Do this even for libs (plugins) so that these too can consume WinRT APIs
 			Arguments.Add("/ZW");
 
@@ -263,219 +78,31 @@ namespace UnrealBuildTool
 			// we've hit problems where types are somehow in windows.winmd on some installations but not others, leading to either
 			// missing or duplicated type references.
 			Arguments.Add("/ZW:nostdlib");
-			
-			// Note: Should be shared with all Windows-like platforms.
-			switch (CompileEnvironment.CppStandard)
-			{
-				case CppStandardVersion.Cpp14:
-					Arguments.Add("/std:c++14");
-					break;
-				case CppStandardVersion.Cpp17:
-					Arguments.Add("/std:c++17");
-					break;
-				case CppStandardVersion.Latest:
-					Arguments.Add("/std:c++latest");
-					break;
-				// Will be added when MSVC is feature-complete.
-				// https://docs.microsoft.com/en-us/cpp/build/reference/std-specify-language-standard-version?view=msvc-160
-				// case CppStandardVersion.Cpp20:
-				//	Arguments.Add("/std:c++20");
-				//  break;
-				default:
-					throw new BuildException($"Unsupported C++ standard type set: {CompileEnvironment.CppStandard}");
-			}
-
-			// Explicitly compile the file as C++.
-			Arguments.Add("/TP");
-
-			if (!CompileEnvironment.bEnableBufferSecurityChecks)
-			{
-				// This will disable buffer security checks (which are enabled by default) that the MS compiler adds around arrays on the stack,
-				// Which can add some performance overhead, especially in performance intensive code
-				// Only disable this if you know what you are doing, because it will be disabled for the entire module!
-				Arguments.Add("/GS-");
-			}
-
-			if (CompileEnvironment.bUseRTTI)
-			{
-				// Enable C++ RTTI.
-				Arguments.Add("/GR");
-			}
-			else
-			{
-				// Disable C++ RTTI.
-				Arguments.Add("/GR-");
-			}
-
-			// Level 4 warnings.
-			Arguments.Add("/W3");
-		}
-
-		static void AppendCLArguments_C(List<string> Arguments)
-		{
-			// Explicitly compile the file as C.
-			Arguments.Add("/TC");
-
-			// Level 0 warnings.  Needed for external C projects that produce warnings at higher warning levels.
-			Arguments.Add("/W0");
 		}
 
-		void AppendLinkArguments(LinkEnvironment LinkEnvironment, List<string> Arguments)
+		protected override void AppendLinkArguments(LinkEnvironment LinkEnvironment, List<string> Arguments)
 		{
-			// Don't create a side-by-side manifest file for the executable.
-			Arguments.Add("/MANIFEST:NO");
+			base.AppendLinkArguments(LinkEnvironment, Arguments);
 
-			// Prevents the linker from displaying its logo for each invocation.
-			Arguments.Add("/NOLOGO");
-
-			if (LinkEnvironment.bCreateDebugInfo)
+			// @todo HoloLens: This is logic from before inheriting from VCToolChain. Is not embedding the manifest necessary?
+			if (!Arguments.Contains("/MANIFEST:NO"))
 			{
-				// Output debug info for the linked executable.
+				Arguments.Remove("/MANIFEST:EMBED");
+				Arguments.RemoveAll((x) => x.StartsWith("/MANIFESTINPUT"));
 
-				// Allow partial PDBs for faster linking
-				if (LinkEnvironment.bUseFastPDBLinking)
-				{
-					Arguments.Add("/DEBUG:FASTLINK");
-				}
-				else
-				{ 
-					Arguments.Add("/DEBUG");
-				}
+				// Don't create a side-by-side manifest file for the executable.
+				Arguments.Add("/MANIFEST:NO");
 			}
 
-			// Prompt the user before reporting internal errors to Microsoft.
-			Arguments.Add("/errorReport:prompt");
-
-			//
-			//	PC
-			//
 			Arguments.Add("/APPCONTAINER");
+
 			// this helps with store API compliance validation tools, adding additional pdb info
 			Arguments.Add("/PROFILE");
-
-			Arguments.Add("/SUBSYSTEM:WINDOWS");
-
-			Arguments.Add("/MACHINE:" + Target.WindowsPlatform.GetArchitectureSubpath());
-
-			if (LinkEnvironment.bIsBuildingConsoleApplication && !LinkEnvironment.bIsBuildingDLL && !String.IsNullOrEmpty(LinkEnvironment.WindowsEntryPointOverride))
-			{
-				// Use overridden entry point
-				Arguments.Add("/ENTRY:" + LinkEnvironment.WindowsEntryPointOverride);
-			}
-
-			// Allow the OS to load the EXE at different base addresses than its preferred base address.
-			Arguments.Add("/FIXED:No");
-
-			// Explicitly declare that the executable is compatible with Data Execution Prevention.
-			Arguments.Add("/NXCOMPAT");
-
-			// Set the default stack size.
-			Arguments.Add("/STACK:5000000,131072");
-
-			// Allow delay-loaded DLLs to be explicitly unloaded.
-			Arguments.Add("/DELAY:UNLOAD");
-
-			if (LinkEnvironment.bIsBuildingDLL)
-			{
-				Arguments.Add("/DLL");
-			}
-
-			// Don't embed the full PDB path; we want to be able to move binaries elsewhere. They will always be side by side.
-			Arguments.Add("/PDBALTPATH:%_PDB%");
-
-			// Experimental deterministic compile support
-			if (Target.WindowsPlatform.bDeterministic)
-			{
-				if (Target.WindowsPlatform.Compiler.IsMSVC())
-				{
-					Arguments.Add("/experimental:deterministic");
-				}
-			}
-
-			//
-			//	Shipping & LTCG
-			//
-			if (LinkEnvironment.bAllowLTCG &&
-				LinkEnvironment.Configuration == CppConfiguration.Shipping)
-			{
-				// Use link-time code generation.
-				Arguments.Add("/LTCG");
-
-				// This is where we add in the PGO-Lite linkorder.txt if we are using PGO-Lite
-				//Arguments.Add("/ORDER:@linkorder.txt");
-				//Arguments.Add("/VERBOSE");
-			}
-
-			//
-			//	Shipping binary
-			//
-			if (LinkEnvironment.Configuration == CppConfiguration.Shipping)
-			{
-				// Generate an EXE checksum.
-				Arguments.Add("/RELEASE");
-
-				// Eliminate unreferenced symbols.
-				Arguments.Add("/OPT:REF");
-
-				// Remove redundant COMDATs.
-				Arguments.Add("/OPT:ICF");
-			}
-			//
-			//	Regular development binary. 
-			//
-			else
-			{
-				// Keep symbols that are unreferenced.
-				Arguments.Add("/OPT:NOREF");
-
-				// Disable identical COMDAT folding.
-				Arguments.Add("/OPT:NOICF");
-			}
-
-			// Enable incremental linking if wanted. ( avoid /INCREMENTAL getting ignored (LNK4075) due to /LTCG, /RELEASE, and /OPT:ICF )
-			if (LinkEnvironment.bUseIncrementalLinking && LinkEnvironment.Configuration != CppConfiguration.Shipping)
-			{
-				Arguments.Add("/INCREMENTAL");
-			}
-			else
-			{
-				Arguments.Add("/INCREMENTAL:NO");
-			}
-
-			// Disable 
-			//LINK : warning LNK4199: /DELAYLOAD:nvtt_64.dll ignored; no imports found from nvtt_64.dll
-			// type warning as we leverage the DelayLoad option to put third-party DLLs into a 
-			// non-standard location. This requires the module(s) that use said DLL to ensure that it 
-			// is loaded prior to using it.
-			Arguments.Add("/ignore:4199");
-
-			// Suppress warnings about missing PDB files for statically linked libraries.  We often don't want to distribute
-			// PDB files for these libraries.
-			Arguments.Add("/ignore:4099");    // warning LNK4099: PDB '<file>' was not found with '<file>'
 		}
 
-		static void AppendLibArguments(LinkEnvironment LinkEnvironment, List<string> Arguments)
+		protected override void AppendLibArguments(LinkEnvironment LinkEnvironment, List<string> Arguments)
 		{
-			// Prevents the linker from displaying its logo for each invocation.
-			Arguments.Add("/NOLOGO");
-
-			// Prompt the user before reporting internal errors to Microsoft.
-			Arguments.Add("/errorReport:prompt");
-
-			//
-			//	PC
-			//
-
-			Arguments.Add("/SUBSYSTEM:WINDOWS");
-
-			//
-			//	Shipping & LTCG
-			//
-			if (LinkEnvironment.Configuration == CppConfiguration.Shipping)
-			{
-				// Use link-time code generation.
-				Arguments.Add("/LTCG");
-			}
+			base.AppendLibArguments(LinkEnvironment, Arguments);
 
 			// Ignore warning about /ZW in static libraries. It's not relevant since UE modules
 			// have no reason to export new WinRT types, and ignoring it quiets noise when using
@@ -483,656 +110,6 @@ namespace UnrealBuildTool
 			Arguments.Add("/ignore:4264");
 		}
 
-		public override CPPOutput CompileCPPFiles(CppCompileEnvironment CompileEnvironment, List<FileItem> InputFiles, DirectoryReference OutputDir, string ModuleName, IActionGraphBuilder Graph)
-		{
-			List<string> SharedArguments = new List<string>();
-			AppendCLArguments_Global(CompileEnvironment, EnvVars, SharedArguments);
-
-			// Add include paths to the argument list.
-			foreach (DirectoryReference IncludePath in CompileEnvironment.UserIncludePaths)
-			{
-				VCToolChain.AddIncludePath(SharedArguments, IncludePath, Target.HoloLensPlatform.Compiler, CompileEnvironment.bPreprocessOnly);
-			}
-			foreach (DirectoryReference IncludePath in CompileEnvironment.SystemIncludePaths)
-			{
-				VCToolChain.AddIncludePath(SharedArguments, IncludePath, Target.HoloLensPlatform.Compiler, CompileEnvironment.bPreprocessOnly);
-			}
-
-			foreach (DirectoryReference IncludePath in EnvVars.IncludePaths)
-			{
-				VCToolChain.AddIncludePath(SharedArguments, IncludePath, Target.HoloLensPlatform.Compiler, CompileEnvironment.bPreprocessOnly);
-			}
-
-			// Add preprocessor definitions to the argument list.
-			foreach (string Definition in CompileEnvironment.Definitions)
-			{
-				// Escape all quotation marks so that they get properly passed with the command line.
-				var DefinitionArgument = Definition.Contains("\"") ? Definition.Replace("\"", "\\\"") : Definition;
-				VCToolChain.AddDefinition(SharedArguments,  DefinitionArgument);
-			}
-			
-			// Create a compile action for each source file.
-			CPPOutput Result = new CPPOutput();
-			foreach (FileItem SourceFile in InputFiles)
-			{
-				Action CompileAction = Graph.CreateAction(ActionType.Compile);
-				CompileAction.CommandDescription = "Compile";
-
-				List<string> FileArguments = new List<string>();
-				bool bIsPlainCFile = Path.GetExtension(SourceFile.AbsolutePath).ToUpperInvariant() == ".C";
-
-				// Add the C++ source file and its included files to the prerequisite item list.
-				CompileAction.PrerequisiteItems.Add(SourceFile);
-
-				bool bEmitsObjectFile = true;
-				if (CompileEnvironment.PrecompiledHeaderAction == PrecompiledHeaderAction.Create)
-				{
-					// Generate a CPP File that just includes the precompiled header.
-					FileReference PCHCPPPath = CompileEnvironment.PrecompiledHeaderIncludeFilename!.ChangeExtension(".cpp");
-					FileItem PCHCPPFile = Graph.CreateIntermediateTextFile(
-						PCHCPPPath,
-						string.Format("#include \"{0}\"\r\n", CompileEnvironment.PrecompiledHeaderIncludeFilename.FullName.Replace('\\', '/')));
-
-					// Make sure the original source directory the PCH header file existed in is added as an include
-					// path -- it might be a private PCH header and we need to make sure that its found!
-					string OriginalPCHHeaderDirectory = Path.GetDirectoryName(SourceFile.AbsolutePath)!;
-					FileArguments.AddFormat(" /I \"{0}\"", OriginalPCHHeaderDirectory);
-
-					var PrecompiledFileExtension = UEBuildPlatform.GetBuildPlatform(UnrealTargetPlatform.HoloLens).GetBinaryExtension(UEBuildBinaryType.PrecompiledHeader);
-					// Add the precompiled header file to the produced items list.
-					FileItem PrecompiledHeaderFile = FileItem.GetItemByFileReference(
-						FileReference.Combine(
-							OutputDir,
-							Path.GetFileName(SourceFile.AbsolutePath) + PrecompiledFileExtension
-							)
-						);
-					CompileAction.ProducedItems.Add(PrecompiledHeaderFile);
-					Result.PrecompiledHeaderFile = PrecompiledHeaderFile;
-
-					// Add the parameters needed to compile the precompiled header file to the command-line.
-					FileArguments.AddFormat(" /Yc\"{0}\"", CompileEnvironment.PrecompiledHeaderIncludeFilename);
-					FileArguments.AddFormat(" /Fp\"{0}\"", PrecompiledHeaderFile.AbsolutePath);
-
-					// If we're creating a PCH that will be used to compile source files for a library, we need
-					// the compiled modules to retain a reference to PCH's module, so that debugging information
-					// will be included in the library.  This is also required to avoid linker warning "LNK4206"
-					// when linking an application that uses this library.
-					if (CompileEnvironment.bIsBuildingLibrary)
-					{
-						// NOTE: The symbol name we use here is arbitrary, and all that matters is that it is
-						// unique per PCH module used in our library
-						string FakeUniquePCHSymbolName = CompileEnvironment.PrecompiledHeaderIncludeFilename.GetFileNameWithoutExtension();
-						FileArguments.AddFormat(" /Yl{0}", FakeUniquePCHSymbolName);
-					}
-
-					FileArguments.AddFormat(" \"{0}\"", PCHCPPFile.AbsolutePath);
-
-					CompileAction.StatusDescription = PCHCPPPath.GetFileName();
-				}
-				else
-				{
-					if (CompileEnvironment.PrecompiledHeaderAction == PrecompiledHeaderAction.Include)
-					{
-						CompileAction.PrerequisiteItems.Add(CompileEnvironment.PrecompiledHeaderFile!);
-
-						FileArguments.Add(String.Format("/FI\"{0}\"", CompileEnvironment.PrecompiledHeaderIncludeFilename!.FullName));
-						FileArguments.Add(String.Format("/Yu\"{0}\"", CompileEnvironment.PrecompiledHeaderIncludeFilename.FullName));
-						FileArguments.Add(String.Format("/Fp\"{0}\"", CompileEnvironment.PrecompiledHeaderFile!.AbsolutePath));
-					}
-
-					// Add the source file path to the command-line.
-					FileArguments.AddFormat(" \"{0}\"", SourceFile.AbsolutePath);
-
-					CompileAction.StatusDescription = Path.GetFileName(SourceFile.AbsolutePath);
-				}
-
-				foreach (FileItem ForceIncludeFile in CompileEnvironment.ForceIncludeFiles)
-				{
-					FileArguments.AddFormat(" /FI\"{0}\"", ForceIncludeFile.AbsolutePath);
-				}
-
-				if (bEmitsObjectFile)
-				{
-					var ObjectFileExtension = UEBuildPlatform.GetBuildPlatform(UnrealTargetPlatform.HoloLens).GetBinaryExtension(UEBuildBinaryType.Object);
-					// Add the object file to the produced item list.
-					FileItem ObjectFile = FileItem.GetItemByFileReference(
-						FileReference.Combine(
-							OutputDir,
-							Path.GetFileName(SourceFile.AbsolutePath) + ObjectFileExtension
-							)
-						);
-					CompileAction.ProducedItems.Add(ObjectFile);
-					Result.ObjectFiles.Add(ObjectFile);
-					FileArguments.AddFormat(" /Fo\"{0}\"", ObjectFile.AbsolutePath);
-				}
-
-				// Create PDB files if we were configured to do that.
-				if (CompileEnvironment.bUsePDBFiles)
-				{
-					string PDBFileName;
-					bool bActionProducesPDB = false;
-
-					// All files using the same PCH are required to share the same PDB that was used when compiling the PCH
-					if (CompileEnvironment.PrecompiledHeaderAction == PrecompiledHeaderAction.Include)
-					{
-						PDBFileName = CompileEnvironment.PrecompiledHeaderIncludeFilename!.GetFileName();
-					}
-					// Files creating a PCH use a PDB per file.
-					else if (CompileEnvironment.PrecompiledHeaderAction == PrecompiledHeaderAction.Create)
-					{
-						PDBFileName = CompileEnvironment.PrecompiledHeaderIncludeFilename!.GetFileName();
-						bActionProducesPDB = true;
-					}
-					// Ungrouped C++ files use a PDB per file.
-					else if (!bIsPlainCFile)
-					{
-						PDBFileName = Path.GetFileName(SourceFile.AbsolutePath);
-						bActionProducesPDB = true;
-					}
-					// Group all plain C files that doesn't use PCH into the same PDB
-					else
-					{
-						PDBFileName = "MiscPlainC";
-					}
-
-						// Specify the PDB file that the compiler should write to.
-					FileItem PDBFile = FileItem.GetItemByFileReference(
-							FileReference.Combine(
-									OutputDir,
-									PDBFileName + ".pdb"
-									)
-								);
-					FileArguments.AddFormat(" /Fd\"{0}\"", PDBFile.AbsolutePath);
-
-					// Only use the PDB as an output file if we want PDBs and this particular action is
-					// the one that produces the PDB (as opposed to no debug info, where the above code
-					// is needed, but not the output PDB, or when multiple files share a single PDB, so
-					// only the action that generates it should count it as output directly)
-					if (CompileEnvironment.bUsePDBFiles && bActionProducesPDB)
-					{
-						CompileAction.ProducedItems.Add(PDBFile);
-					}
-				}
-
-				if (Target.HoloLensPlatform.bRunNativeCodeAnalysis)
-				{
-					// Add the analysis log to the produced item list.
-					FileItem AnalysisLogFile = FileItem.GetItemByFileReference(
-						FileReference.Combine(
-							OutputDir,
-							Path.GetFileName(SourceFile.AbsolutePath) + CodeAnalysisExtension
-							)
-						);
-					CompileAction.ProducedItems.Add(AnalysisLogFile);
-					// Peform code analysis with results in a log file
-					FileArguments.AddFormat(" /analyze:log \"{0}\"", AnalysisLogFile.AbsolutePath);
-					// Suppress code analysis output
-					FileArguments.Add(" /analyze:quiet");
-					string? rulesetFile = Target.HoloLensPlatform.NativeCodeAnalysisRuleset;
-					if (!String.IsNullOrEmpty(rulesetFile))
-					{
-						if (!Path.IsPathRooted(rulesetFile))
-						{
-							rulesetFile = FileReference.Combine(Target.ProjectFile!.Directory, rulesetFile).FullName;
-						}
-						// A non default ruleset was specified
-						FileArguments.AddFormat(" /analyze:ruleset \"{0}\"", rulesetFile);
-					}
-				}
-				
-				// Add C or C++ specific compiler arguments.
-				if (bIsPlainCFile)
-				{
-					AppendCLArguments_C(FileArguments);
-				}
-				else
-				{
-					AppendCLArguments_CPP(CompileEnvironment, FileArguments);
-				}
-
-				CompileAction.WorkingDirectory = UnrealBuildTool.EngineSourceDirectory;
-				CompileAction.CommandPath = EnvVars.CompilerPath;
-				CompileAction.CommandVersion = EnvVars.ToolChainVersion.ToString();
-
-				string[] AdditionalArguments = String.IsNullOrEmpty(CompileEnvironment.AdditionalArguments) ? new string[0] : new string[] { CompileEnvironment.AdditionalArguments };
-
-				if (!ProjectFileGenerator.bGenerateProjectFiles
-					&& Target.WindowsPlatform.Compiler.IsMSVC()
-					&& CompileAction.ProducedItems.Count > 0)
-				{
-					FileItem TargetFile = CompileAction.ProducedItems[0];
-					FileReference ResponseFileName = new FileReference(TargetFile.AbsolutePath + ".response");
-					FileItem ResponseFileItem = Graph.CreateIntermediateTextFile(ResponseFileName, SharedArguments.Concat(FileArguments).Concat(AdditionalArguments).Select(x => Utils.ExpandVariables(x)));
-					CompileAction.CommandArguments = " @\"" + ResponseFileName + "\"";
-					CompileAction.PrerequisiteItems.Add(ResponseFileItem);
-				}
-				else
-				{
-					CompileAction.CommandArguments = String.Join(" ", SharedArguments.Concat(FileArguments).Concat(AdditionalArguments));
-				}
-				
-				if (CompileEnvironment.bGenerateDependenciesFile)
-				{
-					if (EnvVars.ToolChainVersion >= VersionNumber.Parse("14.27") && Target.WindowsPlatform.Compiler.IsMSVC())
-					{
-						CompileAction.DependencyListFile = FileItem.GetItemByFileReference(FileReference.Combine(OutputDir, String.Format("{0}.json", SourceFile.Location.GetFileName())));
-						CompileAction.ProducedItems.Add(CompileAction.DependencyListFile);
-						CompileAction.CommandArguments = string.Join(" ", CompileAction.CommandArguments, "/sourceDependencies", $"\"{CompileAction.DependencyListFile.AbsolutePath}\"");
-					}
-					else if (Target.WindowsPlatform.Compiler.IsClang())
-					{
-						CompileAction.DependencyListFile = FileItem.GetItemByFileReference(FileReference.Combine(OutputDir, String.Format("{0}.d", SourceFile.Location.GetFileName())));
-						CompileAction.ProducedItems.Add(CompileAction.DependencyListFile);
-						CompileAction.CommandArguments = string.Join(" ", CompileAction.CommandArguments, "/clang:-MD", $"/clang:-MF\"{CompileAction.DependencyListFile.AbsolutePath}\"");
-					}
-					else
-					{
-						List<string> CommandArguments = new List<string>();
-						CompileAction.DependencyListFile = FileItem.GetItemByFileReference(FileReference.Combine(OutputDir, String.Format("{0}.txt", SourceFile.Location.GetFileName())));
-						CompileAction.ProducedItems.Add(CompileAction.DependencyListFile);
-						CommandArguments.Add(String.Format("-dependencies={0}", Utils.MakePathSafeToUseWithCommandLine(CompileAction.DependencyListFile.Location)));
-
-						CommandArguments.Add(String.Format("-compiler={0}", Utils.MakePathSafeToUseWithCommandLine(CompileAction.CommandPath)));
-						CommandArguments.Add("--");
-						CommandArguments.Add(Utils.MakePathSafeToUseWithCommandLine(CompileAction.CommandPath));
-						CommandArguments.Add(CompileAction.CommandArguments);
-						CommandArguments.Add("/showIncludes");
-						CompileAction.CommandArguments = string.Join(" ", CommandArguments);
-						CompileAction.CommandPath = FileReference.Combine(Unreal.EngineDirectory, "Build", "Windows", "cl-filter", "cl-filter.exe");
-					}
-				}
-
-				if (CompileEnvironment.PrecompiledHeaderAction == PrecompiledHeaderAction.Create)
-				{
-					Log.TraceVerbose("Creating PCH: " + CompileEnvironment.PrecompiledHeaderIncludeFilename);
-					Log.TraceVerbose("	 Command: " + CompileAction.CommandArguments);
-				}
-				else
-				{
-					Log.TraceVerbose("   Compiling: " + CompileAction.StatusDescription);
-					Log.TraceVerbose("	 Command: " + CompileAction.CommandArguments);
-				}
-
-				// VC++ always outputs the source file name being compiled, so we don't need to emit this ourselves
-				CompileAction.bShouldOutputStatusDescription = false;
-
-				// Don't farm out creation of precompiled headers as it is the critical path task.
-				CompileAction.bCanExecuteRemotely =
-					CompileEnvironment.PrecompiledHeaderAction != PrecompiledHeaderAction.Create ||
-					CompileEnvironment.bAllowRemotelyCompiledPCHs
-					;
-
-			}
-			return Result;
-		}
-
-		public override CPPOutput CompileRCFiles(CppCompileEnvironment CompileEnvironment, List<FileItem> InputFiles, DirectoryReference OutputDir, IActionGraphBuilder Graph)
-		{
-			CPPOutput Result = new CPPOutput();
-
-			foreach (FileItem RCFile in InputFiles)
-			{
-				Action CompileAction = Graph.CreateAction(ActionType.Compile);
-				CompileAction.CommandDescription = "Resource";
-				CompileAction.WorkingDirectory = UnrealBuildTool.EngineSourceDirectory;
-				CompileAction.CommandPath = EnvVars.ResourceCompilerPath;
-				CompileAction.StatusDescription = Path.GetFileName(RCFile.AbsolutePath);
-				CompileAction.PrerequisiteItems.AddRange(CompileEnvironment.ForceIncludeFiles);
-
-				// Resource tool can run remotely if possible
-				CompileAction.bCanExecuteRemotely = true;
-
-				List<string> Arguments = new List<string>();
-
-				// Suppress header spew
-				Arguments.Add("/nologo");
-
-				// If we're compiling for 64-bit Windows, also add the _WIN64 definition to the resource
-				// compiler so that we can switch on that in the .rc file using #ifdef.
-				if (CompileEnvironment.Architecture == "x64" || CompileEnvironment.Architecture == "ARM64")
-				{
-					VCToolChain.AddDefinition(Arguments, "_WIN64");
-				}
-
-				// Language
-				Arguments.Add("/l 0x409");
-
-				// Include paths. Don't use AddIncludePath() here, since it uses the full path and exceeds the max command line length.
-				foreach (DirectoryReference IncludePath in CompileEnvironment.UserIncludePaths)
-				{
-					Arguments.Add(String.Format("/I \"{0}\"", IncludePath));
-				}
-
-				// System include paths.
-				foreach (DirectoryReference SystemIncludePath in CompileEnvironment.SystemIncludePaths)
-				{
-					Arguments.Add(String.Format("/I \"{0}\"", SystemIncludePath));
-				}
-				foreach (DirectoryReference SystemIncludePath in EnvVars.IncludePaths)
-				{
-					Arguments.Add(String.Format("/I \"{0}\"", SystemIncludePath));
-				}
-
-				// Preprocessor definitions.
-				foreach (string Definition in CompileEnvironment.Definitions)
-				{
-					if (!Definition.Contains("_API"))
-					{
-						VCToolChain.AddDefinition(Arguments, Definition);
-					}
-				}
-
-				// Add the RES file to the produced item list.
-				FileItem CompiledResourceFile = FileItem.GetItemByFileReference(
-					FileReference.Combine(
-						OutputDir,
-						Path.GetFileName(RCFile.AbsolutePath) + ".res"
-						)
-					);
-				CompileAction.ProducedItems.Add(CompiledResourceFile);
-				Arguments.Add(String.Format("/fo \"{0}\"", CompiledResourceFile.AbsolutePath));
-				Result.ObjectFiles.Add(CompiledResourceFile);
-
-				// Add the RC file as a prerequisite of the action.
-				Arguments.Add(String.Format(" \"{0}\"", RCFile.AbsolutePath));
-
-				CompileAction.CommandArguments = String.Join(" ", Arguments);
-
-				// Add the C++ source file and its included files to the prerequisite item list.
-				CompileAction.PrerequisiteItems.Add(RCFile);
-			}
-
-			return Result;
-		}
-
-		public override FileItem LinkFiles(LinkEnvironment LinkEnvironment, bool bBuildImportLibraryOnly, IActionGraphBuilder Graph) //(LinkEnvironment LinkEnvironment, bool bBuildImportLibraryOnly)
-		{
-			if (LinkEnvironment.bIsBuildingDotNetAssembly)
-			{
-				return FileItem.GetItemByFileReference(LinkEnvironment.OutputFilePath);
-			}
-
-			bool bIsBuildingLibrary = LinkEnvironment.bIsBuildingLibrary || bBuildImportLibraryOnly;
-
-			// Get link arguments.
-			List<string> Arguments = new List<string>();
-			if (bIsBuildingLibrary)
-			{
-				AppendLibArguments(LinkEnvironment, Arguments);
-			}
-			else
-			{
-				AppendLinkArguments(LinkEnvironment, Arguments);
-			}
-
-			if (Target.HoloLensPlatform.Compiler.IsMSVC() && LinkEnvironment.bPrintTimingInfo)
-			{
-				Arguments.Add("/time+");
-			}
-
-			// If we're only building an import library, add the '/DEF' option that tells the LIB utility
-			// to simply create a .LIB file and .EXP file, and don't bother validating imports
-			if (bBuildImportLibraryOnly)
-			{
-				Arguments.Add("/DEF");
-
-				// Ensure that the import library references the correct filename for the linked binary.
-				Arguments.Add(String.Format("/NAME:\"{0}\"", LinkEnvironment.OutputFilePath.GetFileName()));
-
-				// Ignore warnings about object files with no public symbols.
-				Arguments.Add("/IGNORE:4221");
-			}
-
-
-			// Add delay loaded DLLs.
-			if (!bIsBuildingLibrary)
-			{
-				// Delay-load these DLLs.
-				foreach (string DelayLoadDLL in LinkEnvironment.DelayLoadDLLs.Distinct())
-				{
-					Arguments.Add(String.Format("/DELAYLOAD:\"{0}\"", DelayLoadDLL));
-				}
-
-				// Pass the module definition file to the linker if we have one
-				if (LinkEnvironment.ModuleDefinitionFile != null && LinkEnvironment.ModuleDefinitionFile.Length > 0)
-				{
-					Arguments.Add(String.Format("/DEF:\"{0}\"", LinkEnvironment.ModuleDefinitionFile));
-				}
-			}
-
-			// Set up the library paths for linking this binary
-			if (bBuildImportLibraryOnly)
-			{
-				// When building an import library, ignore all the libraries included via embedded #pragma lib declarations. 
-				// We shouldn't need them to generate exports.
-				Arguments.Add("/NODEFAULTLIB");
-			}
-			else if (!LinkEnvironment.bIsBuildingLibrary)
-			{
-				// Add the library paths to the argument list.
-				foreach (DirectoryReference LibraryPath in LinkEnvironment.SystemLibraryPaths)
-				{
-					Arguments.Add(String.Format("/LIBPATH:\"{0}\"", LibraryPath));
-				}
-				foreach (DirectoryReference LibraryPath in EnvVars.LibraryPaths)
-				{
-					Arguments.Add(String.Format("/LIBPATH:\"{0}\"", LibraryPath));
-				}
-
-				// Add the excluded default libraries to the argument list.
-				foreach (string ExcludedLibrary in LinkEnvironment.ExcludedLibraries)
-				{
-					Arguments.Add(String.Format("/NODEFAULTLIB:\"{0}\"", ExcludedLibrary));
-				}
-			}
-
-			// For targets that are cross-referenced, we don't want to write a LIB file during the link step as that
-			// file will clobber the import library we went out of our way to generate during an earlier step.  This
-			// file is not needed for our builds, but there is no way to prevent MSVC from generating it when
-			// linking targets that have exports.  We don't want this to clobber our LIB file and invalidate the
-			// existing timstamp, so instead we simply emit it with a different name
-			FileReference ImportLibraryFilePath = FileReference.Combine(LinkEnvironment.IntermediateDirectory!,
-														 LinkEnvironment.OutputFilePath.GetFileNameWithoutExtension() + ".lib");
-
-			if (LinkEnvironment.bIsCrossReferenced && !bBuildImportLibraryOnly)
-			{
-				ImportLibraryFilePath = ImportLibraryFilePath.ChangeExtension(".suppressed" + ImportLibraryFilePath.GetExtension());
-			}
-
-			FileItem OutputFile;
-			if (bBuildImportLibraryOnly)
-			{
-				OutputFile = FileItem.GetItemByFileReference(ImportLibraryFilePath);
-			}
-			else
-			{
-				OutputFile = FileItem.GetItemByFileReference(LinkEnvironment.OutputFilePath);
-			}
-
-			List<FileItem> ProducedItems = new List<FileItem>();
-			ProducedItems.Add(OutputFile);
-
-			List<FileItem> PrerequisiteItems = new List<FileItem>();
-
-			// Add the input files to a response file, and pass the response file on the command-line.
-			List<string> InputFileNames = new List<string>();
-			foreach (FileItem InputFile in LinkEnvironment.InputFiles)
-			{
-				InputFileNames.Add(string.Format("\"{0}\"", InputFile.AbsolutePath));
-				PrerequisiteItems.Add(InputFile);
-			}
-
-			if (!bIsBuildingLibrary)
-			{
-				foreach (string SystemLibrary in LinkEnvironment.SystemLibraries)
-				{
-					InputFileNames.Add(string.Format("\"{0}\"", SystemLibrary));
-				}
-				foreach (FileReference Library in LinkEnvironment.Libraries)
-				{
-					InputFileNames.Add(string.Format("\"{0}\"", Library));
-					PrerequisiteItems.Add(FileItem.GetItemByFileReference(Library));
-				}
-			}
-
-			Arguments.AddRange(InputFileNames);
-
-			// Add the output file to the command-line.
-			Arguments.Add(String.Format("/OUT:\"{0}\"", OutputFile.AbsolutePath));
-
-			if (bBuildImportLibraryOnly || (LinkEnvironment.bHasExports && !bIsBuildingLibrary))
-			{
-				// An export file is written to the output directory implicitly; add it to the produced items list.
-				FileReference ExportFilePath = ImportLibraryFilePath.ChangeExtension(".exp");
-				FileItem ExportFile = FileItem.GetItemByFileReference(ExportFilePath);
-				ProducedItems.Add(ExportFile);
-			}
-
-			if (!bIsBuildingLibrary)
-			{
-				// There is anything to export
-				if (LinkEnvironment.bHasExports
-					// Shipping monolithic builds don't need exports
-					&& (!((LinkEnvironment.Configuration == CppConfiguration.Shipping) /*&& (LinkEnvironment.bShouldCompileMonolithic != false)*/)))
-				{
-					// Write the import library to the output directory for nFringe support.
-					FileItem ImportLibraryFile = FileItem.GetItemByFileReference(ImportLibraryFilePath);
-					Arguments.Add(String.Format("/IMPLIB:\"{0}\"", ImportLibraryFilePath));
-					ProducedItems.Add(ImportLibraryFile);
-				}
-
-				if (LinkEnvironment.bCreateDebugInfo)
-				{
-					// Write the PDB file to the output directory.			
-					{
-						FileReference PDBFilePath = FileReference.Combine(LinkEnvironment.OutputDirectory!, Path.GetFileNameWithoutExtension(OutputFile.AbsolutePath) + ".pdb");
-						FileItem PDBFile = FileItem.GetItemByFileReference(PDBFilePath);
-						Arguments.Add(String.Format("/PDB:\"{0}\"", PDBFilePath));
-						ProducedItems.Add(PDBFile);
-					}
-
-					// Write the MAP file to the output directory.			
-					if (LinkEnvironment.bCreateMapFile)
-					{
-						FileReference MAPFilePath = FileReference.Combine(LinkEnvironment.OutputDirectory!, Path.GetFileNameWithoutExtension(OutputFile.AbsolutePath) + ".map");
-						FileItem MAPFile = FileItem.GetItemByFileReference(MAPFilePath);
-						Arguments.Add(String.Format("/MAP:\"{0}\"", MAPFilePath));
-						ProducedItems.Add(MAPFile);
-
-						// Export a list of object file paths, so we can locate the object files referenced by the map file
-						ExportObjectFilePaths(LinkEnvironment, Path.ChangeExtension(MAPFilePath.FullName, ".objpaths"));
-					}
-				}
-
-				// Add the additional arguments specified by the environment.
-				if (!String.IsNullOrEmpty(LinkEnvironment.AdditionalArguments))
-				{
-					Arguments.Add(LinkEnvironment.AdditionalArguments.Trim());
-				}
-			}
-
-			// Create a response file for the linker, unless we're generating IntelliSense data
-			FileReference ResponseFileName = GetResponseFileName(LinkEnvironment, OutputFile);
-			if (!ProjectFileGenerator.bGenerateProjectFiles)
-			{
-				FileItem ResponseFile = Graph.CreateIntermediateTextFile(ResponseFileName, Arguments);
-				PrerequisiteItems.Add(ResponseFile);
-			}
-
-			// Create an action that invokes the linker.
-			Action LinkAction = Graph.CreateAction(ActionType.Link);
-			LinkAction.CommandDescription = "Link";
-			LinkAction.WorkingDirectory = UnrealBuildTool.EngineSourceDirectory;
-			LinkAction.CommandPath = bIsBuildingLibrary ? EnvVars.LibraryManagerPath : EnvVars.LinkerPath;
-			LinkAction.CommandArguments = String.Format("@\"{0}\"", ResponseFileName);
-			LinkAction.ProducedItems.AddRange(ProducedItems);
-			LinkAction.PrerequisiteItems.AddRange(PrerequisiteItems);
-			LinkAction.StatusDescription = Path.GetFileName(OutputFile.AbsolutePath);
-
-			// VS 15.3+ does not touch lib files if they do not contain any modifications, but we need to ensure the timestamps are updated to avoid repeatedly building them.
-			if (bBuildImportLibraryOnly || (LinkEnvironment.bHasExports && !bIsBuildingLibrary))
-			{
-				LinkAction.DeleteItems.AddRange(LinkAction.ProducedItems.Where(x => x.Location.HasExtension(".lib") || x.Location.HasExtension(".exp")));
-			}
-
-			// Delete PDB files for all produced items, since incremental updates are slower than full ones.
-			if (!LinkEnvironment.bUseIncrementalLinking)
-			{
-				LinkAction.DeleteItems.AddRange(LinkAction.ProducedItems.Where(x => x.Location.HasExtension(".pdb")));
-			}
-
-
-			// Tell the action that we're building an import library here and it should conditionally be
-			// ignored as a prerequisite for other actions
-			LinkAction.bProducesImportLibrary = bBuildImportLibraryOnly || LinkEnvironment.bIsBuildingDLL;
-
-			// Allow remote linking.  Especially in modular builds with many small DLL files, this is almost always very efficient
-			LinkAction.bCanExecuteRemotely = true;
-
-			Log.TraceVerbose("     Linking: " + LinkAction.StatusDescription);
-			Log.TraceVerbose("     Command: " + LinkAction.CommandArguments);
-
-			return OutputFile;
-		}
-
-		private void ExportObjectFilePaths(LinkEnvironment LinkEnvironment, string FileName)
-		{
-			// Write the list of object file directories
-			HashSet<DirectoryReference> ObjectFileDirectories = new HashSet<DirectoryReference>();
-			foreach (FileItem InputFile in LinkEnvironment.InputFiles)
-			{
-				ObjectFileDirectories.Add(InputFile.Location.Directory);
-			}
-			foreach (FileReference Library in LinkEnvironment.Libraries)
-			{
-				ObjectFileDirectories.Add(Library.Directory);
-			}
-			foreach (DirectoryReference LibraryPath in LinkEnvironment.SystemLibraryPaths)
-			{
-				ObjectFileDirectories.Add(LibraryPath);
-			}
-			foreach (string LibraryPath in (Environment.GetEnvironmentVariable("LIB") ?? "").Split(new char[] { ';' }, StringSplitOptions.RemoveEmptyEntries))
-			{
-				ObjectFileDirectories.Add(new DirectoryReference(LibraryPath));
-			}
-			Directory.CreateDirectory(Path.GetDirectoryName(FileName));
-			File.WriteAllLines(FileName, ObjectFileDirectories.Select(x => x.FullName).OrderBy(x => x).ToArray());
-		}
-
-		/// <summary>
-		/// Gets the default include paths for the given platform.
-		/// </summary>
-		public static string? GetVCIncludePaths(UnrealTargetPlatform Platform, WindowsCompiler Compiler)
-		{
-			Debug.Assert(Platform == UnrealTargetPlatform.HoloLens);
-
-
-			// Also add any include paths from the INCLUDE environment variable.  MSVC is not necessarily running with an environment that
-			// matches what UBT extracted from the vcvars*.bat using SetEnvironmentVariablesFromBatchFile().  We'll use the variables we
-			// extracted to populate the project file's list of include paths
-			// @todo projectfiles: Should we only do this for VC++ platforms?
-			var IncludePaths = Environment.GetEnvironmentVariable("INCLUDE");
-			if (!String.IsNullOrEmpty(IncludePaths) && !IncludePaths.EndsWith(";"))
-			{
-				IncludePaths += ";";
-			}
-
-			return IncludePaths;
-		}
-
-		/** Formats compiler output from Clang so that it is clickable in Visual Studio */
-		protected static void ClangCompilerOutputFormatter(object sender, DataReceivedEventArgs e)
-		{
-			var Output = e.Data;
-			if (Output == null)
-			{
-				return;
-			}
-
-			// @todo clang: Convert relative includes to absolute files so they'll be clickable
-			Log.TraceInformation(Output);
-		}
-
 		private static DirectoryReference? CurrentWindowsSdkBinDir = null;
 		private static Version? CurrentWindowsSdkVersion;
 
@@ -1190,6 +167,8 @@ namespace UnrealBuildTool
 
 		public override void ModifyBuildProducts(ReadOnlyTargetRules Target, UEBuildBinary Binary, List<string> Libraries, List<UEBuildBundleResource> BundleResources, Dictionary<FileReference, BuildProductType> BuildProducts)
 		{
+			base.ModifyBuildProducts(Target, Binary, Libraries, BundleResources, BuildProducts);
+
 			DirectoryReference HoloLensBinaryDirectory = Binary.OutputFilePath.Directory;
 
 			AddBuildProductSafe(BuildProducts, FileReference.Combine(HoloLensBinaryDirectory, "AppxManifest_"+ Target.Architecture + ".xml"), BuildProductType.BuildResource);
diff --git a/Engine/Source/Programs/UnrealBuildTool/Platform/HoloLens/UEBuildHoloLens.cs b/Engine/Source/Programs/UnrealBuildTool/Platform/HoloLens/UEBuildHoloLens.cs
--- a/Engine/Source/Programs/UnrealBuildTool/Platform/HoloLens/UEBuildHoloLens.cs
+++ b/Engine/Source/Programs/UnrealBuildTool/Platform/HoloLens/UEBuildHoloLens.cs
@@ -149,15 +149,12 @@ namespace UnrealBuildTool
 		#endregion
 	}
 
-	class HoloLens : UEBuildPlatform
+	class HoloLensPlatform : WindowsPlatform
 	{
 		public static readonly Version MinimumSDKVersionRecommended = new Version(10, 0, 17763, 0);
 		public static readonly Version MaximumSDKVersionTested = new Version(10, 0, 18362, int.MaxValue);
-		public static readonly Version MaximumSDKVersionForVS2015 = new Version(10, 0, 14393, int.MaxValue);
-		public static readonly Version MinimumSDKVersionForD3D12RHI = new Version(10, 0, 15063, 0);
 
-
-		public HoloLens(MicrosoftPlatformSDK InSDK) 
+		public HoloLensPlatform(MicrosoftPlatformSDK InSDK) 
 			: base(UnrealTargetPlatform.HoloLens, InSDK)
 		{
 		}
@@ -208,6 +205,7 @@ namespace UnrealBuildTool
 			Target.WindowsPlatform.Architecture = Target.HoloLensPlatform.Architecture;
 			Target.WindowsPlatform.bPixProfilingEnabled = Target.HoloLensPlatform.bPixProfilingEnabled;
 			Target.WindowsPlatform.bUseWindowsSDK10 = true;
+			Target.WindowsPlatform.bStrictConformanceMode = false;
 
 			Target.bDeployAfterCompile = true;
 			Target.bCompileNvCloth = false;      // requires CUDA
@@ -227,23 +225,18 @@ namespace UnrealBuildTool
 			// Be resilient to SDKs being uninstalled but still referenced in the INI file
 			VersionNumber? SelectedWindowsSdkVersion;
 			DirectoryReference? SelectedWindowsSdkDir;
-			if (!WindowsPlatform.TryGetWindowsSdkDir(Target.HoloLensPlatform.Win10SDKVersionString, out SelectedWindowsSdkVersion, out SelectedWindowsSdkDir))
+			if (WindowsPlatform.TryGetWindowsSdkDir(Target.HoloLensPlatform.Win10SDKVersionString, out SelectedWindowsSdkVersion, out SelectedWindowsSdkDir))
 			{
-				Target.HoloLensPlatform.Win10SDKVersionString = "Latest";
+				Target.WindowsPlatform.WindowsSdkVersion = Target.HoloLensPlatform.Win10SDKVersionString;
 			}
 
-			// Initialize the VC environment for the target, and set all the version numbers to the concrete values we chose.
-			VCEnvironment Environment = VCEnvironment.Create(Target.WindowsPlatform.Compiler, Platform, Target.WindowsPlatform.Architecture, Target.WindowsPlatform.CompilerVersion, Target.HoloLensPlatform.Win10SDKVersionString, null);
-			Target.WindowsPlatform.Environment = Environment;
-			Target.WindowsPlatform.Compiler = Environment.Compiler;
-			Target.WindowsPlatform.CompilerVersion = Environment.CompilerVersion.ToString();
-			Target.WindowsPlatform.WindowsSdkVersion = Environment.WindowsSdkVersion.ToString();
+			base.ValidateTarget(Target);
 
 			// Windows 10 SDK version
 			// Auto-detect latest compatible by default (recommended), allow for explicit override if necessary
 			// Validate that the SDK isn't too old, and that the combination of VS and SDK is supported.
 
-			Target.HoloLensPlatform.Win10SDKVersion = new Version(Environment.WindowsSdkVersion.ToString());
+			Target.HoloLensPlatform.Win10SDKVersion = new Version(Target.WindowsPlatform.Environment!.WindowsSdkVersion.ToString());
 
 			if(!Target.bGenerateProjectFiles)
 			{
@@ -272,26 +265,6 @@ namespace UnrealBuildTool
 			return WindowsArchitecture.x64.ToString();
 		}
 
-		/// <summary>
-		/// Determines if the given name is a build product for a target.
-		/// </summary>
-		/// <param name="FileName">The name to check</param>
-		/// <param name="NamePrefixes">Target or application names that may appear at the start of the build product name (eg. "UnrealEditor", "ShooterGameEditor")</param>
-		/// <param name="NameSuffixes">Suffixes which may appear at the end of the build product name</param>
-		/// <returns>True if the string matches the name of a build product, false otherwise</returns>
-		public override bool IsBuildProduct(string FileName, string[] NamePrefixes, string[] NameSuffixes)
-		{
-			return IsBuildProductName(FileName, NamePrefixes, NameSuffixes, ".exe")
-				|| IsBuildProductName(FileName, NamePrefixes, NameSuffixes, ".dll")
-				|| IsBuildProductName(FileName, NamePrefixes, NameSuffixes, ".dll.response")
-				|| IsBuildProductName(FileName, NamePrefixes, NameSuffixes, ".lib")
-				|| IsBuildProductName(FileName, NamePrefixes, NameSuffixes, ".pdb")
-				|| IsBuildProductName(FileName, NamePrefixes, NameSuffixes, ".exp")
-				|| IsBuildProductName(FileName, NamePrefixes, NameSuffixes, ".obj")
-				|| IsBuildProductName(FileName, NamePrefixes, NameSuffixes, ".map")
-				|| IsBuildProductName(FileName, NamePrefixes, NameSuffixes, ".objpaths");
-		}
-
 		/// <summary>
 		/// Get the extension to use for the given binary type
 		/// </summary>
@@ -301,12 +274,6 @@ namespace UnrealBuildTool
 		{
 			switch (InBinaryType)
 			{
-				case UEBuildBinaryType.DynamicLinkLibrary:
-					return ".dll";
-				case UEBuildBinaryType.Executable:
-					return ".exe";
-				case UEBuildBinaryType.StaticLibrary:
-					return ".lib";
 				case UEBuildBinaryType.Object:
 					return ".obj";
 				case UEBuildBinaryType.PrecompiledHeader:
@@ -315,23 +282,6 @@ namespace UnrealBuildTool
 			return base.GetBinaryExtension(InBinaryType);
 		}
 
-		/// <summary>
-		/// Get the extension to use for debug info for the given binary type
-		/// </summary>
-		/// <param name="Target">The target being built</param>
-		/// <param name="InBinaryType"> The binary type being built</param>
-		/// <returns>string	The debug info extension (i.e. 'pdb')</returns>
-		public override string[] GetDebugInfoExtensions(ReadOnlyTargetRules Target, UEBuildBinaryType InBinaryType)
-		{
-			switch (InBinaryType)
-			{
-				case UEBuildBinaryType.DynamicLinkLibrary:
-				case UEBuildBinaryType.Executable:
-					return new string[] { ".pdb" };
-			}
-			return new string[] { "" };
-		}
-
 		internal static DirectoryReference? GetCppCXMetadataLocation(WindowsCompiler Compiler, string CompilerVersion)
 		{
 			VersionNumber? SelectedToolChainVersion;
@@ -393,7 +343,6 @@ namespace UnrealBuildTool
 			if (DirectoryReference.Exists(ReferenceDir))
 			{
 				// Prefer a contract from a suitable SDK-versioned subdir of the references folder when available (starts with 15063 SDK)
-				//Version WindowsSDKVersionMaxForToolchain = Compiler <= WindowsCompiler.VisualStudio2019 ? HoloLens.MaximumSDKVersionForVS2015 : null;
 				DirectoryReference SDKVersionedReferenceDir = DirectoryReference.Combine(ReferenceDir, SDKVersion.ToString());
 				DirectoryReference ContractDir = DirectoryReference.Combine(SDKVersionedReferenceDir, ApiContract);
 				Version? ContractLatestVersion = null;
@@ -518,7 +467,7 @@ namespace UnrealBuildTool
 			}
 		}
 
-		private void ExpandWinMDReferences(ReadOnlyTargetRules Target, string SDKFolder, string SDKVersion, ref List<string> WinMDReferences)
+		internal static void ExpandWinMDReferences(ReadOnlyTargetRules Target, string SDKFolder, string SDKVersion, ref List<string> WinMDReferences)
 		{
 			// Code below will fail when not using the Win10 SDK.  Early out to avoid warning spam.
 			if (!Target.WindowsPlatform.bUseWindowsSDK10)
@@ -591,7 +540,7 @@ namespace UnrealBuildTool
 			{
 				WinMDReferenceArguments.AppendFormat(@" /FU""{0}""", WinMDReference);
 			}
-			CompileEnvironment.AdditionalArguments += WinMDReferenceArguments;
+			CompileEnvironment.AdditionalArguments += WinMDReferenceArguments.ToString().Trim();
 
 			CompileEnvironment.Definitions.Add("EXCEPTIONS_DISABLED=0");
 
@@ -645,7 +594,11 @@ namespace UnrealBuildTool
 				LinkEnvironment.SystemLibraries.Add("msvcprt.lib");
 			}
 			LinkEnvironment.SystemLibraries.Add("legacy_stdio_wide_specifiers.lib");
-			LinkEnvironment.SystemLibraries.Add("uuid.lib"); 
+			LinkEnvironment.SystemLibraries.Add("uuid.lib");
+
+			// Static CRT not supported for HoloLens
+			CompileEnvironment.bUseStaticCRT = false;
+			CompileEnvironment.bEnableExceptions = true;
 		}
 
 		/// <summary>
@@ -749,7 +702,7 @@ namespace UnrealBuildTool
 			// for GetValidSoftwareVersionRange reasons we probably want a HoloLensePlatformSDK class
 			MicrosoftPlatformSDK SDK = new MicrosoftPlatformSDK();
 
-			UEBuildPlatform.RegisterBuildPlatform(new HoloLens(SDK));
+			UEBuildPlatform.RegisterBuildPlatform(new HoloLensPlatform(SDK));
 			UEBuildPlatform.RegisterPlatformWithGroup(UnrealTargetPlatform.HoloLens, UnrealPlatformGroup.Microsoft);
 			UEBuildPlatform.RegisterPlatformWithGroup(UnrealTargetPlatform.HoloLens, UnrealPlatformGroup.HoloLens);
 		}
diff --git a/Engine/Source/Programs/UnrealBuildTool/Platform/Windows/VCToolChain.cs b/Engine/Source/Programs/UnrealBuildTool/Platform/Windows/VCToolChain.cs
--- a/Engine/Source/Programs/UnrealBuildTool/Platform/Windows/VCToolChain.cs
+++ b/Engine/Source/Programs/UnrealBuildTool/Platform/Windows/VCToolChain.cs
@@ -967,7 +967,7 @@ namespace UnrealBuildTool
 			//
 			//	PC
 			//
-			if (LinkEnvironment.Platform.IsInGroup(UnrealPlatformGroup.Windows))
+			if (LinkEnvironment.Platform.IsInGroup(UnrealPlatformGroup.Windows) || LinkEnvironment.Platform.IsInGroup(UnrealPlatformGroup.HoloLens))
 			{
 				Arguments.Add(string.Format("/MACHINE:{0}", WindowsExports.GetArchitectureSubpath(Target.WindowsPlatform.Architecture)));
 
@@ -1100,7 +1100,7 @@ namespace UnrealBuildTool
 			Arguments.Add("/ignore:4099");      // warning LNK4099: PDB '<file>' was not found with '<file>'
 		}
 
-		void AppendLibArguments(LinkEnvironment LinkEnvironment, List<string> Arguments)
+		protected virtual void AppendLibArguments(LinkEnvironment LinkEnvironment, List<string> Arguments)
 		{
 			// Prevents the linker from displaying its logo for each invocation.
 			Arguments.Add("/NOLOGO");
@@ -1111,7 +1111,7 @@ namespace UnrealBuildTool
 			//
 			//	PC
 			//
-			if (LinkEnvironment.Platform.IsInGroup(UnrealPlatformGroup.Windows))
+			if (LinkEnvironment.Platform.IsInGroup(UnrealPlatformGroup.Windows) || LinkEnvironment.Platform.IsInGroup(UnrealPlatformGroup.HoloLens))
 			{
 				Arguments.Add(string.Format("/MACHINE:{0}", WindowsExports.GetArchitectureSubpath(Target.WindowsPlatform.Architecture)));
 
@@ -1306,6 +1306,32 @@ namespace UnrealBuildTool
 					CompileAction.Arguments.Add(CompileEnvironment.AdditionalArguments);
 				}
 
+				if (CompileEnvironment.Platform.IsInGroup(UnrealPlatformGroup.HoloLens) && Target.HoloLensPlatform.bRunNativeCodeAnalysis)
+				{
+					// Add the analysis log to the produced item list.
+					FileItem AnalysisLogFile = FileItem.GetItemByFileReference(
+						FileReference.Combine(
+							OutputDir,
+							Path.GetFileName(SourceFile.AbsolutePath) + ".nativecodeanalysis.xml"
+							)
+						); ;
+					CompileAction.AdditionalProducedItems.Add(AnalysisLogFile);
+					// Peform code analysis with results in a log file
+					CompileAction.Arguments.AddFormat("/analyze:log \"{0}\"", AnalysisLogFile.AbsolutePath);
+					// Suppress code analysis output
+					CompileAction.Arguments.Add("/analyze:quiet");
+					string? rulesetFile = Target.HoloLensPlatform.NativeCodeAnalysisRuleset;
+					if (!String.IsNullOrEmpty(rulesetFile))
+					{
+						if (!Path.IsPathRooted(rulesetFile))
+						{
+							rulesetFile = FileReference.Combine(Target.ProjectFile!.Directory, rulesetFile).FullName;
+						}
+						// A non default ruleset was specified
+						CompileAction.Arguments.AddFormat("/analyze:ruleset \"{0}\"", rulesetFile);
+					}
+				}
+
 				if (SourceFile.HasExtension(".ixx"))
 				{
 					FileItem IfcFile = FileItem.GetItemByFileReference(FileReference.Combine(GetModuleInterfaceDir(OutputDir), SourceFile.Location.ChangeExtension(".ifc").GetFileName()));
diff --git a/Engine/Source/Programs/UnrealBuildTool/ProjectFiles/VisualStudio/VCProject.cs b/Engine/Source/Programs/UnrealBuildTool/ProjectFiles/VisualStudio/VCProject.cs
--- a/Engine/Source/Programs/UnrealBuildTool/ProjectFiles/VisualStudio/VCProject.cs
+++ b/Engine/Source/Programs/UnrealBuildTool/ProjectFiles/VisualStudio/VCProject.cs
@@ -894,7 +894,7 @@ namespace UnrealBuildTool
 				}
 				else if (InPlatforms.Contains(UnrealTargetPlatform.HoloLens))
 				{
-					SharedIncludeSearchPaths.Append(HoloLensToolChain.GetVCIncludePaths(UnrealTargetPlatform.HoloLens, GetCompilerForIntellisense()) + ";");
+					SharedIncludeSearchPaths.Append(VCToolChain.GetVCIncludePaths(UnrealTargetPlatform.HoloLens, GetCompilerForIntellisense(), null) + ";");
 				}
 			}
 
@@ -1776,7 +1776,7 @@ namespace UnrealBuildTool
 				{
 					TargetRules TargetRulesObject = Combination.ProjectTarget.TargetRules!;
 
-					if (Combination.Platform!.Value.IsInGroup(UnrealPlatformGroup.Windows) || (Combination.Platform == UnrealTargetPlatform.HoloLens))
+					if (Combination.Platform!.Value.IsInGroup(UnrealPlatformGroup.Windows) || Combination.Platform!.Value.IsInGroup(UnrealPlatformGroup.HoloLens))
 					{
 						VCUserFileContent.AppendLine("  <PropertyGroup {0}>", ConditionString);
 						if (TargetRulesObject.Type != TargetType.Game)
diff --git a/Engine/Source/Runtime/Online/HTTP/Private/HttpRetrySystem.cpp b/Engine/Source/Runtime/Online/HTTP/Private/HttpRetrySystem.cpp
--- a/Engine/Source/Runtime/Online/HTTP/Private/HttpRetrySystem.cpp
+++ b/Engine/Source/Runtime/Online/HTTP/Private/HttpRetrySystem.cpp
@@ -298,7 +298,7 @@ float FHttpRetrySystem::FManager::GetLockoutPeriodSeconds(const FHttpRetryReques
 	return LockoutPeriod;
 }
 
-static FRandomStream temp(4435261);
+static FRandomStream TempRandomStream(4435261);
 
 bool FHttpRetrySystem::FManager::Update(uint32* FileCount, uint32* FailingCount, uint32* FailedCount, uint32* CompletedCount)
 {
@@ -386,7 +386,7 @@ bool FHttpRetrySystem::FManager::Update(uint32* FileCount, uint32* FailingCount,
 					// Code to simulate request failure
 					if (RequestStatus == EHttpRequestStatus::Succeeded && RandomFailureRate.IsSet())
 					{
-						float random = temp.GetFraction();
+						float random = TempRandomStream.GetFraction();
 						if (random < RandomFailureRate.GetValue())
 						{
 							forceFail = true;
