UnrealBuildTool: Relative pathing for VCToolchain linking and resource building
Write .response file for resources so the command line is easier to find (however these are unused currently as rc.exe doesn't accept them)

#rnx

diff --git a/Engine/Source/Programs/UnrealBuildTool/Platform/Windows/VCToolChain.cs b/Engine/Source/Programs/UnrealBuildTool/Platform/Windows/VCToolChain.cs
--- a/Engine/Source/Programs/UnrealBuildTool/Platform/Windows/VCToolChain.cs
+++ b/Engine/Source/Programs/UnrealBuildTool/Platform/Windows/VCToolChain.cs
@@ -913,7 +913,8 @@ namespace UnrealBuildTool
 			else
 			{
 				Arguments.Add("/MANIFEST:EMBED");
-				Arguments.Add(String.Format("/MANIFESTINPUT:{0}", Utils.MakePathSafeToUseWithCommandLine(Target.WindowsPlatform.ManifestFile)));
+				FileItem ManifestFile = FileItem.GetItemByPath(Target.WindowsPlatform.ManifestFile);
+				Arguments.Add(String.Format("/MANIFESTINPUT:\"{0}\"", NormalizeCommandLinePath(ManifestFile, Target.WindowsPlatform.Compiler, false)));
 			}
 
 			// Prevents the linker from displaying its logo for each invocation.
@@ -1547,17 +1548,17 @@ namespace UnrealBuildTool
 				// Include paths. Don't use AddIncludePath() here, since it uses the full path and exceeds the max command line length.
 				foreach (DirectoryReference IncludePath in CompileEnvironment.UserIncludePaths)
 				{
-					Arguments.Add(String.Format("/I \"{0}\"", IncludePath));
+					Arguments.Add(String.Format("/I \"{0}\"", NormalizeCommandLinePath(IncludePath, Target.WindowsPlatform.Compiler, false)));
 				}
 
 				// System include paths.
 				foreach (DirectoryReference SystemIncludePath in CompileEnvironment.SystemIncludePaths)
 				{
-					Arguments.Add(String.Format("/I \"{0}\"", SystemIncludePath));
+					Arguments.Add(String.Format("/I \"{0}\"", NormalizeCommandLinePath(SystemIncludePath, Target.WindowsPlatform.Compiler, false)));
 				}
 				foreach (DirectoryReference SystemIncludePath in EnvVars.IncludePaths)
 				{
-					Arguments.Add(String.Format("/I \"{0}\"", SystemIncludePath));
+					Arguments.Add(String.Format("/I \"{0}\"", NormalizeCommandLinePath(SystemIncludePath, Target.WindowsPlatform.Compiler, false)));
 				}
 
 				// Preprocessor definitions.
@@ -1582,7 +1583,7 @@ namespace UnrealBuildTool
 				CompileAction.PrerequisiteItems.Add(FileItem.GetItemByFileReference(IconFile));
 
 				// Setup the compile environment, setting the icon to use via a macro. This is used in Default.rc2.
-				AddDefinition(Arguments, String.Format("BUILD_ICON_FILE_NAME=\"\\\"{0}\\\"\"", IconFile.FullName.Replace("\\", "\\\\")));
+				AddDefinition(Arguments, String.Format("BUILD_ICON_FILE_NAME=\"\\\"{0}\\\"\"", NormalizeCommandLinePath(IconFile, Target.WindowsPlatform.Compiler, false).Replace("\\", "\\\\")));
 
 				// Apply the target settings for the resources
 				if(!CompileEnvironment.bUseSharedBuildEnvironment)
@@ -1616,12 +1617,28 @@ namespace UnrealBuildTool
 						)
 					);
 				CompileAction.ProducedItems.Add(CompiledResourceFile);
-				Arguments.Add(String.Format("/fo \"{0}\"", CompiledResourceFile.AbsolutePath));
+				Arguments.Add(String.Format("/fo \"{0}\"", NormalizeCommandLinePath(CompiledResourceFile, Target.WindowsPlatform.Compiler, false)));
 				Result.ObjectFiles.Add(CompiledResourceFile);
 
 				// Add the RC file as a prerequisite of the action.
-				Arguments.Add(String.Format(" \"{0}\"", RCFile.AbsolutePath));
+				Arguments.Add(String.Format("\"{0}\"", NormalizeCommandLinePath(RCFile, Target.WindowsPlatform.Compiler, false)));
 
+				// Create a response file for the resource compilier
+				FileItem ResponseFile = FileItem.GetItemByPath(CompiledResourceFile.FullName + ".response");
+				Graph.CreateIntermediateTextFile(ResponseFile, Arguments);
+				CompileAction.PrerequisiteItems.Add(ResponseFile);
+
+				/* rc.exe currently errors when using a response file
+				string ResponseFileString = NormalizeCommandLinePath(ResponseFile, Target.WindowsPlatform.Compiler, false);
+
+				// cl.exe can't handle response files with a path longer than 260 characters, and relative paths can push it over the limit
+				if (!System.IO.Path.IsPathRooted(ResponseFileString) && System.IO.Path.Combine(CompileAction.WorkingDirectory.FullName, ResponseFileString).Length > 260)
+				{
+					ResponseFileString = ResponseFile.FullName;
+				}
+
+				CompileAction.CommandArguments = String.Format("@{0}", Utils.MakePathSafeToUseWithCommandLine(ResponseFileString));
+				*/
 				CompileAction.CommandArguments = String.Join(" ", Arguments);
 
 				// Add the C++ source file and its included files to the prerequisite item list.
@@ -1783,11 +1800,11 @@ namespace UnrealBuildTool
 				// Add the library paths to the argument list.
 				foreach (DirectoryReference LibraryPath in LinkEnvironment.SystemLibraryPaths)
 				{
-					Arguments.Add(String.Format("/LIBPATH:\"{0}\"", LibraryPath));
+					Arguments.Add(String.Format("/LIBPATH:\"{0}\"", NormalizeCommandLinePath(LibraryPath, Target.WindowsPlatform.Compiler, false)));
 				}
 				foreach (DirectoryReference LibraryPath in EnvVars.LibraryPaths)
 				{
-					Arguments.Add(String.Format("/LIBPATH:\"{0}\"", LibraryPath));
+					Arguments.Add(String.Format("/LIBPATH:\"{0}\"", NormalizeCommandLinePath(LibraryPath, Target.WindowsPlatform.Compiler, false)));
 				}
 
 				// Add the excluded default libraries to the argument list.
@@ -1882,7 +1899,7 @@ namespace UnrealBuildTool
 			List<string> InputFileNames = new List<string>();
 			foreach (FileItem InputFile in LinkEnvironment.InputFiles)
 			{
-				InputFileNames.Add(string.Format("\"{0}\"", InputFile.AbsolutePath));
+				InputFileNames.Add(string.Format("\"{0}\"", NormalizeCommandLinePath(InputFile, Target.WindowsPlatform.Compiler, false)));
 				PrerequisiteItems.Add(InputFile);
 			}
 
@@ -1890,7 +1907,7 @@ namespace UnrealBuildTool
 			{
 				foreach (FileReference Library in LinkEnvironment.Libraries)
 				{
-					InputFileNames.Add(string.Format("\"{0}\"", Library));
+					InputFileNames.Add(string.Format("\"{0}\"", NormalizeCommandLinePath(Library, Target.WindowsPlatform.Compiler, false)));
 					PrerequisiteItems.Add(FileItem.GetItemByFileReference(Library));
 				}
 				foreach (string SystemLibrary in LinkEnvironment.SystemLibraries)
@@ -1902,7 +1919,7 @@ namespace UnrealBuildTool
 			Arguments.AddRange(InputFileNames);
 
 			// Add the output file to the command-line.
-			Arguments.Add(String.Format("/OUT:\"{0}\"", OutputFile.AbsolutePath));
+			Arguments.Add(String.Format("/OUT:\"{0}\"", NormalizeCommandLinePath(OutputFile, Target.WindowsPlatform.Compiler, false)));
 
 			// For import libraries and exports generated by cross-referenced builds, we don't track output files. VS 15.3+ doesn't touch timestamps for libs
 			// and exp files with no modifications, breaking our dependency checking, but incremental linking will fall back to a full link if we delete it.
@@ -1921,7 +1938,7 @@ namespace UnrealBuildTool
 				{
 					// Write the import library to the output directory for nFringe support.
 					FileItem ImportLibraryFile = FileItem.GetItemByFileReference(ImportLibraryFilePath);
-					Arguments.Add(String.Format("/IMPLIB:\"{0}\"", ImportLibraryFilePath));
+					Arguments.Add(String.Format("/IMPLIB:\"{0}\"", NormalizeCommandLinePath(ImportLibraryFilePath, Target.WindowsPlatform.Compiler, false)));
 
 					// Like the export file above, don't add the import library as a produced item when it's cross referenced.
 					if(!LinkEnvironment.bIsCrossReferenced)
@@ -1936,7 +1953,7 @@ namespace UnrealBuildTool
 					{
 						FileReference PDBFilePath = FileReference.Combine(LinkEnvironment.OutputDirectory!, Path.GetFileNameWithoutExtension(OutputFile.AbsolutePath) + ".pdb");
 						FileItem PDBFile = FileItem.GetItemByFileReference(PDBFilePath);
-						Arguments.Add(String.Format("/PDB:\"{0}\"", PDBFilePath));
+						Arguments.Add(String.Format("/PDB:\"{0}\"", NormalizeCommandLinePath(PDBFilePath, Target.WindowsPlatform.Compiler, false)));
 						ProducedItems.Add(PDBFile);
 					}
 
@@ -1945,7 +1962,7 @@ namespace UnrealBuildTool
 					{
 						FileReference MAPFilePath = FileReference.Combine(LinkEnvironment.OutputDirectory!, Path.GetFileNameWithoutExtension(OutputFile.AbsolutePath) + ".map");
 						FileItem MAPFile = FileItem.GetItemByFileReference(MAPFilePath);
-						Arguments.Add(String.Format("/MAP:\"{0}\"", MAPFilePath));
+						Arguments.Add(String.Format("/MAP:\"{0}\"", NormalizeCommandLinePath(MAPFilePath, Target.WindowsPlatform.Compiler, false)));
 						ProducedItems.Add(MAPFile);
 
 						// Export a list of object file paths, so we can locate the object files referenced by the map file
