Set minimum Visual Studio version to 2019 v16.11.5, toolchain 14.29.31033
Set minimum Clang version to 10.0.0

This will enable the use of newer compiler flags such as increasing the max PDB size past 4GBs by default, using compiler generated dependency lists, and to clean up now unneeded older version checks before releasing UE5.

diff --git a/Engine/Programs/UnrealGameSync/UnrealGameSync.ini b/Engine/Programs/UnrealGameSync/UnrealGameSync.ini
--- a/Engine/Programs/UnrealGameSync/UnrealGameSync.ini
+++ b/Engine/Programs/UnrealGameSync/UnrealGameSync.ini
@@ -47,7 +47,7 @@ SdkInstallerDir=
 
 ; All the required SDKs for the current version of the engine
 +SdkInfo=(Category="Android", Description="NDK r21", Browse="$(SdkInstallerDir)\\Android")
-+SdkInfo=(Category="Windows", Description="Visual Studio 2019 v16.4")
++SdkInfo=(Category="Windows", Description="Visual Studio 2019 v16.11.5")
 +SdkInfo=(Category="Windows", Description="Windows SDK 10.0.16299.0")
 +SdkInfo=(Category="Linux", Description="v19_clang-11.0.1-centos7")
 
diff --git a/Engine/Source/Programs/UnrealBuildTool/Platform/Windows/UEBuildWindows.cs b/Engine/Source/Programs/UnrealBuildTool/Platform/Windows/UEBuildWindows.cs
--- a/Engine/Source/Programs/UnrealBuildTool/Platform/Windows/UEBuildWindows.cs
+++ b/Engine/Source/Programs/UnrealBuildTool/Platform/Windows/UEBuildWindows.cs
@@ -822,6 +822,8 @@ namespace UnrealBuildTool
 			VersionNumberRange.Parse("11.0.0", "11.999"),
 			VersionNumberRange.Parse("10.0.0", "10.999")
 		};
+		
+		static readonly VersionNumber MinimumClangVersion = new VersionNumber(10, 0, 0);
 
 		/// <summary>
 		/// Ranges of tested compiler toolchains to be used, in order of preference. If multiple toolchains in a range are present, the latest version will be preferred.
@@ -829,9 +831,11 @@ namespace UnrealBuildTool
 		/// </summary>
 		static readonly VersionNumberRange[] PreferredVisualCppVersions = new VersionNumberRange[]
 		{
-			VersionNumberRange.Parse("14.24.27905", "14.29.31033") // VS2019 16.2.3 -> 19.1
+			VersionNumberRange.Parse("14.29.30133", "14.29.30136"), // VS2019 16.11.5
 		};
 
+		static readonly VersionNumber MinimumVisualCppVersion = new VersionNumber(14, 29, 30133);
+
 		/// <summary>
 		/// Visual Studio installations
 		/// </summary>
@@ -1376,8 +1380,15 @@ namespace UnrealBuildTool
 
 				int Rank = PreferredClangVersions.TakeWhile(x => !x.Contains(Version)).Count();
 				bool Is64Bit = Is64BitExecutable(CompilerFile);
+
+				string? Error = null;
+				if (Version < MinimumClangVersion)
+				{
+					Error = $"UnrealBuildTool requires at minimum the Clang {MinimumClangVersion} toolchain. Please install a later toolchain from LLVM.";
+				}
+
 				Log.TraceLog("Found Clang toolchain: {0} (Version={1}, Is64Bit={2}, Rank={3})", ToolChainDir, Version, Is64Bit, Rank);
-				ToolChains.Add(new ToolChainInstallation(Version, Rank, Version, Is64Bit, false, null, ToolChainDir, null));
+				ToolChains.Add(new ToolChainInstallation(Version, Rank, Version, Is64Bit, false, Error, ToolChainDir, null));
 			}
 		}
 
@@ -1444,9 +1455,9 @@ namespace UnrealBuildTool
 			int FamilyRank = PreferredVisualCppVersions.TakeWhile(x => !x.Contains(Family)).Count();
 
 			string? Error = null;
-			if (Version >= new VersionNumber(14, 23, 0) && Version < new VersionNumber(14, 23, 28107))
+			if (Version < MinimumVisualCppVersion)
 			{
-				Error = String.Format("The Visual C++ 14.23 toolchain is known to have code-generation issues with UE4. Please install a later toolchain from the Visual Studio installer. See here for more information: https://developercommunity.visualstudio.com/content/problem/734585/msvc-142328019-compilation-bug.html");
+				Error = $"UnrealBuildTool requires at minimum the MSVC {MinimumVisualCppVersion} toolchain. Please install a later toolchain from the Visual Studio installer.";
 			}
 
 			Log.TraceLog("Found Visual Studio toolchain: {0} (Family={1}, FamilyRank={2}, Version={3}, Is64Bit={4}, Preview={5}, Error={6}, Redist={7})", ToolChainDir, Family, FamilyRank, Version, Is64Bit, bPreview, Error != null, RedistDir);
@@ -1522,37 +1533,6 @@ namespace UnrealBuildTool
 			return FindToolChainInstallations(Compiler).Count > 0;
 		}
 
-		/// <summary>
-		/// Checks if a given Visual C++ toolchain version is compatible with UE4
-		/// </summary>
-		/// <param name="Version">The version number to check</param>
-		/// <returns>True if the toolchain is compatible with UE4</returns>
-		static bool IsCompatibleVisualCppToolChain(VersionNumber Version)
-		{
-			string? Message;
-			return IsCompatibleVisualCppToolChain(Version, out Message);
-		}
-
-		/// <summary>
-		/// Checks if a given Visual C++ toolchain version is compatible with UE4
-		/// </summary>
-		/// <param name="Version">The version number to check</param>
-		/// <param name="Message">Receives a message describing why the toolchain is not compatible</param>
-		/// <returns>True if the toolchain is compatible with UE4</returns>
-		static bool IsCompatibleVisualCppToolChain(VersionNumber Version, [NotNullWhen(false)] out string? Message)
-		{
-			if (Version >= new VersionNumber(14, 23, 0) && Version < new VersionNumber(14, 23, 28107))
-			{
-				Message = String.Format("The Visual C++ 14.23 toolchain is known to have code-generation issues with UE4. Please install an earlier or later toolchain from the Visual Studio installer. See here for more information: https://developercommunity.visualstudio.com/content/problem/734585/msvc-142328019-compilation-bug.html");
-				return false;
-			}
-			else
-			{
-				Message = null;
-				return true;
-			}
-		}
-
 		/// <summary>
 		/// Select which toolchain to use, combining a custom preference with a default sort order
 		/// </summary>
diff --git a/Engine/Source/Programs/UnrealBuildTool/Platform/Windows/VCEnvironment.cs b/Engine/Source/Programs/UnrealBuildTool/Platform/Windows/VCEnvironment.cs
--- a/Engine/Source/Programs/UnrealBuildTool/Platform/Windows/VCEnvironment.cs
+++ b/Engine/Source/Programs/UnrealBuildTool/Platform/Windows/VCEnvironment.cs
@@ -477,12 +477,7 @@ namespace UnrealBuildTool
 			DirectoryReference? SelectedToolChainDir;
 			if(Compiler.IsClang())
 			{
-				const string DefaultClangVisualStudioVersion = "14.24.28315";
-				if (WindowsPlatform.TryGetToolChainDir(WindowsCompiler.VisualStudio2019, DefaultClangVisualStudioVersion, out SelectedToolChainVersion, out SelectedToolChainDir, out SelectedRedistDir))
-				{
-					ToolChain = WindowsCompiler.VisualStudio2019;
-				}
-				else if (WindowsPlatform.TryGetToolChainDir(WindowsCompiler.VisualStudio2019, null, out SelectedToolChainVersion, out SelectedToolChainDir, out SelectedRedistDir))
+				if (WindowsPlatform.TryGetToolChainDir(WindowsCompiler.VisualStudio2019, null, out SelectedToolChainVersion, out SelectedToolChainDir, out SelectedRedistDir))
 				{
 					ToolChain = WindowsCompiler.VisualStudio2019;
 				}
diff --git a/Engine/Source/Programs/UnrealBuildTool/Platform/Windows/VCToolChain.cs b/Engine/Source/Programs/UnrealBuildTool/Platform/Windows/VCToolChain.cs
--- a/Engine/Source/Programs/UnrealBuildTool/Platform/Windows/VCToolChain.cs
+++ b/Engine/Source/Programs/UnrealBuildTool/Platform/Windows/VCToolChain.cs
@@ -1174,10 +1174,7 @@ namespace UnrealBuildTool
 						BaseCompileAction.Arguments.Add("/Bt+ /d2cgsummary");
 					}
 
-					if(EnvVars.ToolChainVersion >= VersionNumber.Parse("14.14.26316"))
-					{
-						BaseCompileAction.Arguments.Add("/d1reportTime");
-					}
+					BaseCompileAction.Arguments.Add("/d1reportTime");
 				}
 			}
 
@@ -1374,7 +1371,7 @@ namespace UnrealBuildTool
 
 				if (CompileEnvironment.bGenerateDependenciesFile)
 				{
-					if (EnvVars.ToolChainVersion >= VersionNumber.Parse("14.27") && Target.WindowsPlatform.Compiler.IsMSVC() && !CompileAction.ForceClFilter)
+					if (Target.WindowsPlatform.Compiler.IsMSVC() && !CompileAction.ForceClFilter)
 					{
 						CompileAction.DependencyListFile = FileItem.GetItemByFileReference(FileReference.Combine(OutputDir, String.Format("{0}.json", SourceFile.Location.GetFileName())));
 					}
diff --git a/Engine/Source/Runtime/Core/Public/HoloLens/HoloLensPlatformCompilerSetup.h b/Engine/Source/Runtime/Core/Public/HoloLens/HoloLensPlatformCompilerSetup.h
--- a/Engine/Source/Runtime/Core/Public/HoloLens/HoloLensPlatformCompilerSetup.h
+++ b/Engine/Source/Runtime/Core/Public/HoloLens/HoloLensPlatformCompilerSetup.h
@@ -10,10 +10,10 @@
 	#error PLATFORM_HOLOLENS not defined
 #endif
 
-/**
-* We require at least Visual Studio 2015 to compile
-*/
-static_assert(_MSC_VER >= 1900, "Visual Studio 2015 or newer required.");
+//
+// We require at least Visual Studio 2019 to compile
+//
+static_assert(_MSC_VER >= 1929, "Visual Studio 2019 v16.11 or greater is required to build Unreal Engine.");
 
 
 #pragma warning (error:      4001 4002 4003 4004      4006 4007 4008 4009 4010 4011 4012 4013 4014 4015 4016 4017 4018 4019 4020 4021 4022 4023 4024 4025 4026 4027 4028 4029 4030 4031 4032 4033 4034 4035 4036 4037 4038 4039 4040 4041 4042 4043 4044 4045 4046 4047 4048 4049 4050 4051 4052 4053 4054 4055 4056 4057 4058 4059 4060                4064 4065 4066 4067 4068 4069 4070 4071 4072 4073 4074 4075 4076 4077 4078 4079 4080 4081 4082 4083 4084 4085 4086 4087 4088 4089 4090 4091 4092 4093 4094 4095 4096 4097 4098 4099)
diff --git a/Engine/Source/Runtime/Core/Public/Windows/WindowsPlatformCompilerSetup.h b/Engine/Source/Runtime/Core/Public/Windows/WindowsPlatformCompilerSetup.h
--- a/Engine/Source/Runtime/Core/Public/Windows/WindowsPlatformCompilerSetup.h
+++ b/Engine/Source/Runtime/Core/Public/Windows/WindowsPlatformCompilerSetup.h
@@ -16,7 +16,7 @@
 //
 // We require at least Visual Studio 2019 to compile
 //
-static_assert(_MSC_VER >= 1924, "Visual Studio 2019 v16.4 or greater is required to build Unreal Engine.");
+static_assert(_MSC_VER >= 1929, "Visual Studio 2019 v16.11 or greater is required to build Unreal Engine.");
 
 //
 // Manually enable all warnings as errors, except ones that are explicitly skipped.
