UnrealBuildTool: Experimental support for Intel OneAPI compiler

Current Issues:
* Compiles without pch: error: reference to 'align_val_t' is ambiguous
* Unresolved externals when building UnrealGame & Editor

#rnx

diff --git a/Engine/Source/Programs/UnrealBuildTool/Platform/Windows/UEBuildWindows.cs b/Engine/Source/Programs/UnrealBuildTool/Platform/Windows/UEBuildWindows.cs
--- a/Engine/Source/Programs/UnrealBuildTool/Platform/Windows/UEBuildWindows.cs
+++ b/Engine/Source/Programs/UnrealBuildTool/Platform/Windows/UEBuildWindows.cs
@@ -31,6 +31,11 @@ namespace UnrealBuildTool
 		/// </summary>
 		Clang,
 
+		/// <summary>
+		/// Use the Intel oneAPI C++ compiler
+		/// </summary>
+		Intel,
+
 		/// <summary>
 		/// Visual Studio 2019 (Visual C++ 16.0)
 		/// </summary>
@@ -55,7 +60,7 @@ namespace UnrealBuildTool
 		/// <returns>true if Clang based</returns>
 		public static bool IsClang(this WindowsCompiler Compiler)
 		{
-			return Compiler == WindowsCompiler.Clang;
+			return Compiler == WindowsCompiler.Clang || Compiler == WindowsCompiler.Intel;
 		}
 
 		/// <summary>
@@ -443,6 +448,7 @@ namespace UnrealBuildTool
 			switch (Compiler)
 			{
 				case WindowsCompiler.Clang:
+				case WindowsCompiler.Intel:
 				case WindowsCompiler.VisualStudio2019:
 				case WindowsCompiler.VisualStudio2022:
 					return "2015"; // VS2022 is backwards compatible with VS2015 compiler
@@ -836,6 +842,16 @@ namespace UnrealBuildTool
 
 		static readonly VersionNumber MinimumVisualCppVersion = new VersionNumber(14, 29, 30133);
 
+		/// <summary>
+		/// The default compiler version to be used, if installed. 
+		/// </summary>
+		static readonly VersionNumberRange[] PreferredIntelOneApiVersions =
+		{
+			VersionNumberRange.Parse("2022.0.0", "2022.9999"),
+		};
+		
+		static readonly VersionNumber MinimumIntelOneApiVersion = new VersionNumber(2021, 0, 0);
+
 		/// <summary>
 		/// Visual Studio installations
 		/// </summary>
@@ -861,6 +877,11 @@ namespace UnrealBuildTool
 		/// </summary>
 		public static readonly bool bAllowClangLinker = false;
 
+		/// <summary>
+		/// True if we should use the Intel linker (xilink\xilib) when we are compiling with Intel oneAPI, otherwise we use the MSVC linker
+		/// </summary>
+		public static readonly bool bAllowIntelLinker = true;
+		
 		MicrosoftPlatformSDK SDK;
 
 		/// <summary>
@@ -960,6 +981,12 @@ namespace UnrealBuildTool
 				Target.bUsePCHFiles = false;
 			}
 			
+			// @todo: Override PCH settings
+			if (Target.WindowsPlatform.Compiler == WindowsCompiler.Intel)
+			{
+				Target.bUseSharedPCHs = false;
+				Target.bUsePCHFiles = false;
+			}
 
 			// E&C support.
 			if (Target.bSupportEditAndContinue || Target.bAdaptiveUnityEnablesEditAndContinue)
@@ -1263,6 +1290,25 @@ namespace UnrealBuildTool
 							}
 						}
 					}
+					else if(Compiler == WindowsCompiler.Intel)
+					{
+						DirectoryReference InstallDir = DirectoryReference.Combine(DirectoryReference.GetSpecialFolder(Environment.SpecialFolder.ProgramFiles)!, "Intel", "oneAPI", "compiler");
+						FindIntelOneApiToolChains(InstallDir, ToolChains);
+
+						InstallDir = DirectoryReference.Combine(DirectoryReference.GetSpecialFolder(Environment.SpecialFolder.ProgramFilesX86)!, "Intel", "oneAPI", "compiler");
+						FindIntelOneApiToolChains(InstallDir, ToolChains);
+
+						// Check for AutoSDK paths
+						DirectoryReference? AutoSdkDir;
+						if (UEBuildPlatformSDK.TryGetHostPlatformAutoSDKDir(out AutoSdkDir))
+						{
+							DirectoryReference IntelBaseDir = DirectoryReference.Combine(AutoSdkDir, "Win64", "Intel");
+							if (DirectoryReference.Exists(IntelBaseDir))
+							{
+								FindIntelOneApiToolChains(IntelBaseDir, ToolChains);
+							}
+						}
+					}
 				    else if(Compiler.IsMSVC())
 				    {
 						// Enumerate all the manually installed toolchains
@@ -1329,6 +1375,22 @@ namespace UnrealBuildTool
 			}
 		}
 
+		/// <summary>
+		/// Finds all the valid Intel oneAPI toolchains under the given base directory
+		/// </summary>
+		/// <param name="BaseDir">Base directory to search</param>
+		/// <param name="ToolChains">Map of tool chain version to installation info</param>
+		static void FindIntelOneApiToolChains(DirectoryReference BaseDir, List<ToolChainInstallation> ToolChains)
+		{
+			if (DirectoryReference.Exists(BaseDir))
+			{
+				foreach (DirectoryReference ToolChainDir in DirectoryReference.EnumerateDirectories(BaseDir))
+				{
+					AddIntelOneApiToolChain(ToolChainDir, ToolChains);
+				}
+			}
+		}
+
 		/// <summary>
 		/// Finds the most appropriate redist directory for the given toolchain version
 		/// </summary>
@@ -1388,6 +1450,33 @@ namespace UnrealBuildTool
 				}
 
 				Log.TraceLog("Found Clang toolchain: {0} (Version={1}, Is64Bit={2}, Rank={3})", ToolChainDir, Version, Is64Bit, Rank);
+				ToolChains.Add(new ToolChainInstallation(Version, Rank, Version, Is64Bit, false, null, ToolChainDir, null));
+			}
+		}
+
+		/// <summary>
+		/// Add an Intel OneAPI toolchain
+		/// </summary>
+		/// <param name="ToolChainDir"></param>
+		/// <param name="ToolChains"></param>
+		static void AddIntelOneApiToolChain(DirectoryReference ToolChainDir, List<ToolChainInstallation> ToolChains)
+		{
+			FileReference CompilerFile = FileReference.Combine(ToolChainDir, "windows", "bin", "icx.exe");
+			if (FileReference.Exists(CompilerFile))
+			{
+				FileVersionInfo VersionInfo = FileVersionInfo.GetVersionInfo(CompilerFile.FullName);
+				VersionNumber Version = new VersionNumber(VersionInfo.FileMajorPart, VersionInfo.FileMinorPart, VersionInfo.FileBuildPart);
+
+				int Rank = PreferredIntelOneApiVersions.TakeWhile(x => !x.Contains(Version)).Count();
+				bool Is64Bit = Is64BitExecutable(CompilerFile);
+
+				string? Error = null;
+				if (Version < MinimumIntelOneApiVersion)
+				{
+					Error = $"UnrealBuildTool requires at minimum the Intel OneAPI {MinimumIntelOneApiVersion} toolchain. Please install a later toolchain from Intel.";
+				}
+
+				Log.TraceLog("Found Intel OneAPI toolchain: {0} (Version={1}, Is64Bit={2}, Rank={3})", ToolChainDir, Version, Is64Bit, Rank);
 				ToolChains.Add(new ToolChainInstallation(Version, Rank, Version, Is64Bit, false, Error, ToolChainDir, null));
 			}
 		}
diff --git a/Engine/Source/Programs/UnrealBuildTool/Platform/Windows/VCEnvironment.cs b/Engine/Source/Programs/UnrealBuildTool/Platform/Windows/VCEnvironment.cs
--- a/Engine/Source/Programs/UnrealBuildTool/Platform/Windows/VCEnvironment.cs
+++ b/Engine/Source/Programs/UnrealBuildTool/Platform/Windows/VCEnvironment.cs
@@ -226,10 +226,11 @@ namespace UnrealBuildTool
 			{
 				return FileReference.Combine(CompilerDir, "bin", "clang-cl.exe");
 			}
-			else
+			else if (Compiler == WindowsCompiler.Intel)
 			{
-				return FileReference.Combine(GetVCToolPath(Compiler, CompilerDir, Architecture), "cl.exe");
+				return FileReference.Combine(CompilerDir, "windows", "bin", "icx.exe");
 			}
+			return FileReference.Combine(GetVCToolPath(Compiler, CompilerDir, Architecture), "cl.exe");
 		}
 
 		/// <summary>
@@ -242,7 +243,10 @@ namespace UnrealBuildTool
 			{
 				return FileReference.Combine(CompilerDir, "bin", "lld-link.exe");
 			}
-
+			else if (Compiler == WindowsCompiler.Intel && WindowsPlatform.bAllowIntelLinker)
+			{
+				return FileReference.Combine(CompilerDir, "windows", "bin", "intel64", "xilink.exe");
+			}
 			return FileReference.Combine(DefaultLinkerDir, "link.exe");
 
 		}
@@ -258,7 +262,10 @@ namespace UnrealBuildTool
 				// @todo: lld-link is not currently working for building .lib
 				//return FileReference.Combine(CompilerDir, "bin", "lld-link.exe");
 			}
-
+			else if (Compiler == WindowsCompiler.Intel && WindowsPlatform.bAllowIntelLinker)
+			{
+				return FileReference.Combine(CompilerDir, "windows", "bin", "intel64", "xilib.exe");
+			}
 			return FileReference.Combine(DefaultLinkerDir, "lib.exe");
 		}
 
@@ -383,6 +390,13 @@ namespace UnrealBuildTool
 				DirectoryReference LibraryRootDir = DirectoryReference.Combine(WindowsSdkDir, "lib", "winv6.3");
 				LibraryPaths.Add(DirectoryReference.Combine(LibraryRootDir, "um", ArchFolder));
 			}
+
+			// Add path to Intel math libraries when using Intel oneAPI
+			if (Compiler == WindowsCompiler.Intel)
+			{
+				IncludePaths.Add(DirectoryReference.Combine(CompilerDir, "windows", "compiler", "include"));
+				LibraryPaths.Add(DirectoryReference.Combine(CompilerDir, "windows", "compiler", "lib", "intel64"));
+			}
 		}
 
 
diff --git a/Engine/Source/Programs/UnrealBuildTool/Platform/Windows/VCToolChain.cs b/Engine/Source/Programs/UnrealBuildTool/Platform/Windows/VCToolChain.cs
--- a/Engine/Source/Programs/UnrealBuildTool/Platform/Windows/VCToolChain.cs
+++ b/Engine/Source/Programs/UnrealBuildTool/Platform/Windows/VCToolChain.cs
@@ -621,7 +621,8 @@ namespace UnrealBuildTool
 				Arguments.Add("/fp:fast");
 			}
 
-			if (CompileEnvironment.bOptimizeCode)
+			// Intel oneAPI compiler does not support /Zo
+			if (CompileEnvironment.bOptimizeCode && Target.WindowsPlatform.Compiler != WindowsCompiler.Intel)
 			{
 				// Allow optimized code to be debugged more easily.  This makes PDBs a bit larger, but doesn't noticeably affect
 				// compile times.  The executable code is not affected at all by this switch, only the debugging information.
@@ -925,40 +926,20 @@ namespace UnrealBuildTool
 			if (LinkEnvironment.bCreateDebugInfo || Target.WindowsPlatform.bEnableAddressSanitizer)
 			{
 				// Output debug info for the linked executable.
-				Arguments.Add("/DEBUG");
+				// Beginning in Visual Studio 2017 /DEBUG defaults to /DEBUG:FASTLINK for debug builds
+				Arguments.Add("/DEBUG:FULL");
 			}
 
-			if (LinkEnvironment.bCreateDebugInfo)
+			if (LinkEnvironment.bCreateDebugInfo && LinkEnvironment.bUseFastPDBLinking)
 			{
 				// Allow partial PDBs for faster linking
-				if (LinkEnvironment.bUseFastPDBLinking)
+				if (Target.WindowsPlatform.Compiler == WindowsCompiler.Clang && WindowsPlatform.bAllowClangLinker)
 				{
-					switch (Target.WindowsPlatform.Compiler)
-					{
-						case WindowsCompiler.Default:
-							break;
-						case WindowsCompiler.Clang:
-							if (WindowsPlatform.bAllowClangLinker)
-							{
-								Arguments[Arguments.Count - 1] += ":GHASH";
-							}
-							break;
-						case WindowsCompiler.VisualStudio2019:
-						case WindowsCompiler.VisualStudio2022:
-							Arguments[Arguments.Count - 1] += ":FASTLINK";
-							break;
-					}
+					Arguments[Arguments.Count - 1] = "/DEBUG:GHASH";
 				}
 				else
 				{
-					// Beginning in Visual Studio 2017 /DEBUG defaults to /DEBUG:FASTLINK for debug builds
-					switch (Target.WindowsPlatform.Compiler)
-					{
-						case WindowsCompiler.VisualStudio2019:
-						case WindowsCompiler.VisualStudio2022:
-							Arguments[Arguments.Count - 1] += ":FULL";
-							break;
-					}
+					Arguments[Arguments.Count - 1] = "/DEBUG:FASTLINK";
 				}
 			}
 
