UnrealBuildTool: Pass $(AdditionalOptions) instead of the c++ version along with /Yu in project, as that contains the version already

#rnx

diff --git a/Engine/Source/Programs/UnrealBuildTool/ProjectFiles/VisualStudio/VCProject.cs b/Engine/Source/Programs/UnrealBuildTool/ProjectFiles/VisualStudio/VCProject.cs
--- a/Engine/Source/Programs/UnrealBuildTool/ProjectFiles/VisualStudio/VCProject.cs
+++ b/Engine/Source/Programs/UnrealBuildTool/ProjectFiles/VisualStudio/VCProject.cs
@@ -607,6 +607,7 @@ namespace UnrealBuildTool
 					return "/std:c++14";
 				case CppStandardVersion.Cpp17:
 					return "/std:c++17";
+				case CppStandardVersion.Cpp20:
 				case CppStandardVersion.Latest:
 					return "/std:c++latest";
 				// Will be added when MSVC is feature-complete.
@@ -1246,8 +1247,7 @@ namespace UnrealBuildTool
 							VCProjectFileContent.AppendLine("      <ForcedIncludeFiles>{0}</ForcedIncludeFiles>", ForceIncludePaths);
 							if (PchHeaderFile != null && ProjectFileFormat >= VCProjectFileFormat.VisualStudio2022)
 							{
-								VCProjectFileContent.AppendLine("      <AdditionalOptions>{0} /Yu\"{1}\"</AdditionalOptions>",
-									GetCppStandardCompileArgument(GetIntelliSenseCppVersion()),
+								VCProjectFileContent.AppendLine("      <AdditionalOptions>$(AdditionalOptions) /Yu\"{0}\"</AdditionalOptions>",
 									PchHeaderFile);
 							}
 						}
