Dsyms only created when needed

Now when before calling dsymutil we will check the dsym on disk and only call dsymutil when:
1)  The dsym does not already exist or
2)  The mach-o has a file modification time later than the dsym

diff --git a/Engine/Source/Programs/UnrealBuildTool/Platform/Mac/MacToolChain.cs b/Engine/Source/Programs/UnrealBuildTool/Platform/Mac/MacToolChain.cs
--- a/Engine/Source/Programs/UnrealBuildTool/Platform/Mac/MacToolChain.cs
+++ b/Engine/Source/Programs/UnrealBuildTool/Platform/Mac/MacToolChain.cs
@@ -1238,20 +1238,30 @@ namespace UnrealBuildTool
 			string ExtraOptions;
 			string DsymutilPath = GetDsymutilPath(out ExtraOptions, bIsForLTOBuild: false);
 
-			FileReference DsymShellFile = FileReference.Combine(LinkEnvironment.LocalShadowDirectory!, Path.GetFileNameWithoutExtension(MachOBinary.FullName) + "GenerateDSym.sh");
-			FileItem DsymShellFileItem = FileItem.GetItemByFileReference(DsymShellFile);
-
-			string Contents = String.Format(
-				"#!/bin/sh\nrm -rf \"{2}\"\nfor i in {{1..30}}\ndo \nif [ -f \"{1}\" ] \nthen break\nelse \necho \"Waiting for {1} before generating dSYM file.\"\n\tsleep 1\nfi\ndone\n\"{0}\" {3} -f \"{1}\" -o \"{2}\"",
-				DsymutilPath,
-				MachOBinary.AbsolutePath,
-				OutputFile.AbsolutePath,
-				ExtraOptions
-			);
-			Contents = Contents.Replace("\n", Environment.NewLine);
-
-			Utils.WriteFileIfChanged(DsymShellFile, Contents);
-			GenDebugAction.CommandArguments = string.Format("-c \"chmod +x \\\"{0}\\\"; \\\"{0}\\\"\"", DsymShellFileItem.AbsolutePath);
+			string ArgumentString = "-c \"";
+			ArgumentString += string.Format("for i in {{1..30}}; ");
+				ArgumentString += string.Format("do if [ -f \\\"{0}\\\" ] ; ", MachOBinary.AbsolutePath);
+				ArgumentString += string.Format("then ");
+					ArgumentString += string.Format("break; ");
+				ArgumentString += string.Format("else ");
+					ArgumentString += string.Format("echo \\\"Waiting for {0} before generating dSYM file.\\\"; ", MachOBinary.AbsolutePath);
+					ArgumentString += string.Format("sleep 1; ");
+				ArgumentString += string.Format("fi; ");
+			ArgumentString += string.Format("done; ");
+
+			ArgumentString += string.Format("if [ ! -f \\\"{1}\\\" ] || [ \\\"{0}\\\" -nt \\\"{1}\\\" ] ; ", MachOBinary.AbsolutePath, OutputFile.AbsolutePath);
+			ArgumentString += string.Format("then ");
+				ArgumentString += string.Format("echo \\\"Calling dsymutil for: {0}\\\"; ", OutputFile.AbsolutePath);
+				ArgumentString += string.Format("rm -rf \\\"{0}\\\"; ", OutputFile.AbsolutePath);
+				ArgumentString += string.Format(" \\\"{0}\\\" {3} -f \\\"{1}\\\" -o \\\"{2}\\\"; ",
+					DsymutilPath,
+					MachOBinary.AbsolutePath,
+					OutputFile.AbsolutePath,
+					ExtraOptions);
+			ArgumentString += string.Format("fi; ");
+			ArgumentString += "\"";
+
+			GenDebugAction.CommandArguments = ArgumentString;
 
 			if (LinkEnvironment.bIsCrossReferenced)
 			{
