UBT: Revert CL 18260416. We need to update the timestamp on the engine version file to prevent the update action running multiple times.

The action should not run if no engine files have been modified. If new engine files are being added to the manifest, the version file can be written without having any changes. Installed builds will not run the action due to the path check.

diff --git a/Engine/Source/Programs/UnrealBuildTool/Modes/WriteMetadataMode.cs b/Engine/Source/Programs/UnrealBuildTool/Modes/WriteMetadataMode.cs
--- a/Engine/Source/Programs/UnrealBuildTool/Modes/WriteMetadataMode.cs
+++ b/Engine/Source/Programs/UnrealBuildTool/Modes/WriteMetadataMode.cs
@@ -191,7 +191,7 @@ namespace UnrealBuildTool
 			{
 				DirectoryReference.CreateDirectory(TargetInfo.VersionFile.Directory);
 				TargetInfo.Version.BuildId = BuildId;
-				TargetInfo.Version.WriteIfModified(TargetInfo.VersionFile);
+				TargetInfo.Version.Write(TargetInfo.VersionFile);
 			}
 
 			// Write out the receipt
