Allows builders to target IOS and TVOS from Windows without iTunes installed

diff --git a/Engine/Source/Programs/UnrealBuildTool/Platform/Mac/ApplePlatformSDK.cs b/Engine/Source/Programs/UnrealBuildTool/Platform/Mac/ApplePlatformSDK.cs
--- a/Engine/Source/Programs/UnrealBuildTool/Platform/Mac/ApplePlatformSDK.cs
+++ b/Engine/Source/Programs/UnrealBuildTool/Platform/Mac/ApplePlatformSDK.cs
@@ -62,5 +62,18 @@ namespace UnrealBuildTool
 		{
 			return UnrealBuildBase.ApplePlatformSDK.InstalledSDKVersion;
 		}
+
+		protected override SDKStatus HasRequiredManualSDKInternal()
+		{
+			SDKStatus Status = base.HasRequiredManualSDKInternal();
+
+			// iTunes is technically only need to deploy to and run on connected devices.
+			// This code removes requirement for Windows builders to have Xcode installed.
+			if (Status == SDKStatus.Invalid && !RuntimePlatform.IsMac && Environment.GetEnvironmentVariable("IsBuildMachine") == "1")
+            {
+				Status = SDKStatus.Valid;
+            }
+			return Status;
+		}
 	}
-}
+}
\ No newline at end of file
