PR #8782: Fix for selecting VS2019 with a bad compiler version when we have VS2022 (Contributed by [REDACTED_PERSON])

#rnx

diff --git a/Engine/Source/Programs/UnrealBuildTool/Platform/Windows/UEBuildWindows.cs b/Engine/Source/Programs/UnrealBuildTool/Platform/Windows/UEBuildWindows.cs
--- a/Engine/Source/Programs/UnrealBuildTool/Platform/Windows/UEBuildWindows.cs
+++ b/Engine/Source/Programs/UnrealBuildTool/Platform/Windows/UEBuildWindows.cs
@@ -1060,7 +1060,6 @@ namespace UnrealBuildTool
 					}
 				} 
 			}
-
 			// Also check the default format for the Visual Studio project generator
 			object? ProjectFormatObject;
 			if (XmlConfig.TryGetValue(typeof(VCProjectFileGenerator), "Version", out ProjectFormatObject))
@@ -1090,15 +1089,16 @@ namespace UnrealBuildTool
 				}
 			}
 
-			// Second, default based on what's installed, test for 2015 first
-			if (HasCompiler(WindowsCompiler.VisualStudio2019))
+			// Second, default based on what's installed, test for 2019 first
+			if (FindToolChainInstallations(WindowsCompiler.VisualStudio2019).Where(x => x.Version >= MinimumVisualCppVersion).Count() > 0)
 			{
 				return WindowsCompiler.VisualStudio2019;
 			}
-			else if (HasCompiler(WindowsCompiler.VisualStudio2022))
+			else if (FindToolChainInstallations(WindowsCompiler.VisualStudio2022).Where(x => x.Version >= MinimumVisualCppVersion).Count() > 0)
 			{
 				return WindowsCompiler.VisualStudio2022;
 			}
+
 			// If we do have a Visual Studio installation, but we're missing just the C++ parts, warn about that.
 			if (TryGetVSInstallDirs(WindowsCompiler.VisualStudio2019) != null)
 			{
