Visual Studio project generators can request a distinct configuration name for each platform target.

This is to support cases where there may be more than one UnrealTargetPlatform requiring the same VS Solution Platform. Without a distinct separation these platforms will overwrite each other's properties.

#rnx

diff --git a/Engine/Source/Programs/UnrealBuildTool/ProjectFiles/PlatformProjectGenerator.cs b/Engine/Source/Programs/UnrealBuildTool/ProjectFiles/PlatformProjectGenerator.cs
--- a/Engine/Source/Programs/UnrealBuildTool/ProjectFiles/PlatformProjectGenerator.cs
+++ b/Engine/Source/Programs/UnrealBuildTool/ProjectFiles/PlatformProjectGenerator.cs
@@ -85,6 +85,18 @@ namespace UnrealBuildTool
 			return GetVisualStudioPlatformName(InPlatform, InConfiguration);
 		}
 
+		/// <summary>
+		/// Return whether we need a distinct platform name - for cases where there are two targets of the same VS solution platform
+		/// </summary>
+		/// <param name="InPlatform">  The UnrealTargetPlatform being built</param>
+		/// <param name="InConfiguration"> The UnrealTargetConfiguration being built</param>
+		/// <param name="InProjectDir">  The root directory for the current project</param>
+		/// <returns>bool</returns>
+		public virtual bool RequiresDistinctVisualStudioConfigurationName(UnrealTargetPlatform InPlatform, UnrealTargetConfiguration InConfiguration, DirectoryReference InProjectDir)
+		{
+			return false;
+		}
+
 		/// <summary>
 		/// Return project configuration settings that must be included before the default props file
 		/// </summary>
diff --git a/Engine/Source/Programs/UnrealBuildTool/ProjectFiles/VisualStudio/VCProject.cs b/Engine/Source/Programs/UnrealBuildTool/ProjectFiles/VisualStudio/VCProject.cs
--- a/Engine/Source/Programs/UnrealBuildTool/ProjectFiles/VisualStudio/VCProject.cs
+++ b/Engine/Source/Programs/UnrealBuildTool/ProjectFiles/VisualStudio/VCProject.cs
@@ -343,6 +343,12 @@ namespace UnrealBuildTool
 				// UnrealBuildTool calls it (e.g. "Win64" -> "x64".) GetVisualStudioPlatformName() will figure this out.
 				ProjectConfigurationName = Configuration.ToString();
 				ProjectPlatformName = PlatformProjectGenerator!.GetVisualStudioPlatformName(Platform, Configuration, BaseDir);
+
+				// The project generator may require a distinct configuration name - typically when two UnrealTargetPlatforms need the same ProjectPlatformName - otherwise the properties overwrite each oter.
+				if (PlatformProjectGenerator!.RequiresDistinctVisualStudioConfigurationName(Platform, Configuration, BaseDir))
+				{
+					ProjectConfigurationName = Platform.ToString() + "_" + Configuration.ToString();
+				}
 			}
 			else
 			{
