UnrealBuildTool: Add -NoSNDBS command line flag to match other -No{executor} flags

#rnx

diff --git a/Engine/Source/Programs/UnrealBuildTool/Configuration/BuildConfiguration.cs b/Engine/Source/Programs/UnrealBuildTool/Configuration/BuildConfiguration.cs
--- a/Engine/Source/Programs/UnrealBuildTool/Configuration/BuildConfiguration.cs
+++ b/Engine/Source/Programs/UnrealBuildTool/Configuration/BuildConfiguration.cs
@@ -60,6 +60,7 @@ namespace UnrealBuildTool
 		/// Whether SN-DBS may be used.
 		/// </summary>
 		[XmlConfigFile]
+		[CommandLine("-NoSNDBS", Value = "false")]
 		public bool bAllowSNDBS = true;
 
 		/// <summary>
