UBT: Fix engine version file being invalidated by changes to target files.

diff --git a/Engine/Source/Programs/UnrealBuildTool/Configuration/UEBuildTarget.cs b/Engine/Source/Programs/UnrealBuildTool/Configuration/UEBuildTarget.cs
--- a/Engine/Source/Programs/UnrealBuildTool/Configuration/UEBuildTarget.cs
+++ b/Engine/Source/Programs/UnrealBuildTool/Configuration/UEBuildTarget.cs
@@ -1996,7 +1996,7 @@ namespace UnrealBuildTool
 				{
 					WriteMetadataTargetInfo TargetInfo = new WriteMetadataTargetInfo(ProjectFile, null, null, ReceiptFileName, Receipt, FileNameToModuleManifest);
 					FileReference TargetInfoFile = FileReference.Combine(ProjectIntermediateDirectory, "TargetMetadata.dat");
-					CreateWriteMetadataAction(Makefile, ReceiptFileName.GetFileName(), TargetInfoFile, TargetInfo);
+					CreateWriteMetadataAction(Makefile, ReceiptFileName.GetFileName(), TargetInfoFile, TargetInfo, Makefile.OutputItems);
 				}
 				else
 				{
@@ -2015,13 +2015,28 @@ namespace UnrealBuildTool
 						}
 					}
 
+					List<FileItem> EnginePrereqItems = new List<FileItem>();
+					List<FileItem> TargetPrereqItems = new List<FileItem>();
+
+					foreach (FileItem OutputItem in Makefile.OutputItems)
+					{
+						if (OutputItem.Location.IsUnderDirectory(Unreal.EngineDirectory))
+						{
+							EnginePrereqItems.Add(OutputItem);
+						}
+						else
+						{
+							TargetPrereqItems.Add(OutputItem);
+						}
+					}
+
 					WriteMetadataTargetInfo EngineInfo = new WriteMetadataTargetInfo(null, VersionFile, Receipt.Version, null, null, EngineFileToManifest);
 					FileReference EngineInfoFile = FileReference.Combine(ProjectIntermediateDirectory, "EngineMetadata.dat");
-					CreateWriteMetadataAction(Makefile, VersionFile.GetFileName(), EngineInfoFile, EngineInfo);
+					CreateWriteMetadataAction(Makefile, VersionFile.GetFileName(), EngineInfoFile, EngineInfo, EnginePrereqItems);
 
 					WriteMetadataTargetInfo TargetInfo = new WriteMetadataTargetInfo(ProjectFile, VersionFile, null, ReceiptFileName, Receipt, TargetFileToManifest);
 					FileReference TargetInfoFile = FileReference.Combine(ProjectIntermediateDirectory, "TargetMetadata.dat");
-					CreateWriteMetadataAction(Makefile, ReceiptFileName.GetFileName(), TargetInfoFile, TargetInfo);
+					CreateWriteMetadataAction(Makefile, ReceiptFileName.GetFileName(), TargetInfoFile, TargetInfo, TargetPrereqItems);
 				}
 
 				// Create actions to run the post build steps
@@ -2151,7 +2166,7 @@ namespace UnrealBuildTool
 			return Makefile;
 		}
 
-		void CreateWriteMetadataAction(TargetMakefile Makefile, string StatusDescription, FileReference InfoFile, WriteMetadataTargetInfo Info)
+		void CreateWriteMetadataAction(TargetMakefile Makefile, string StatusDescription, FileReference InfoFile, WriteMetadataTargetInfo Info, IEnumerable<FileItem> PrerequisiteItems)
 		{
 			List<FileReference> ProducedItems = new List<FileReference>(Info.FileToManifest.Keys);
 			if (Info.Version != null && Info.VersionFile != null)
@@ -2190,7 +2205,7 @@ namespace UnrealBuildTool
 				WriteMetadataAction.bCanExecuteRemotely = false;
 				WriteMetadataAction.bUseActionHistory = false; // Different files for each target; do not want to invalidate based on this.
 				WriteMetadataAction.ProducedItems.AddRange(ProducedItems.Select(x => FileItem.GetItemByFileReference(x)));
-				WriteMetadataAction.PrerequisiteItems.AddRange(Makefile.OutputItems);
+				WriteMetadataAction.PrerequisiteItems.AddRange(PrerequisiteItems);
 
 				if (Info.Version == null && Info.VersionFile != null)
 				{
