UBT: Fix receipt not getting regenerated when only the project file has changed (eg. enabling/disabling plugins).

diff --git a/Engine/Source/Programs/UnrealBuildTool/Configuration/UEBuildTarget.cs b/Engine/Source/Programs/UnrealBuildTool/Configuration/UEBuildTarget.cs
--- a/Engine/Source/Programs/UnrealBuildTool/Configuration/UEBuildTarget.cs
+++ b/Engine/Source/Programs/UnrealBuildTool/Configuration/UEBuildTarget.cs
@@ -2180,12 +2180,13 @@ namespace UnrealBuildTool
 
 			if (!ProducedItems.Any(x => UnrealBuildTool.IsFileInstalled(x)))
 			{
-				// Save the info file, but deliberately do not add a prerequisite from the action onto it. It is always beneath the
-				// project intermediate directory (even for the engine) and will cause -NoEngineChanges to block the build when
-				// updated. The behavior we want is just to depend on the engine DLL timestamps. Instead, we add a makefile dependency
-				// on it, causing it to be regenerated if missing.
+				// Save the info file, but deliberately do not add a prerequisite from the action onto it unless we're updating the
+				// target file. Since the outputs are not target specific, but the input file is always beneath the project intermediate
+				// directory (even for the engine), it will cause -NoEngineChanges to block the build when updated. The behavior we want
+				// is just to depend on the engine DLL timestamps. Instead, we add a makefile dependency on it, causing it to be
+				// regenerated if missing, and only add it as a prereq if we're updating the .target file.
 				FileItem InfoFileItem = FileItem.GetItemByFileReference(InfoFile);
-				BinaryFormatterUtils.Save(InfoFile, Info);
+				BinaryFormatterUtils.SaveIfDifferent(InfoFile, Info);
 				InfoFileItem.ResetCachedInfo();
 				Makefile.InternalDependencies.Add(InfoFileItem);
 
@@ -2212,6 +2213,11 @@ namespace UnrealBuildTool
 					WriteMetadataAction.PrerequisiteItems.Add(FileItem.GetItemByFileReference(Info.VersionFile));
 				}
 
+				if (Info.ReceiptFile != null)
+				{
+					WriteMetadataAction.PrerequisiteItems.Add(InfoFileItem);
+				}
+
 				// Add the produced items as leaf nodes to be built
 				Makefile.OutputItems.AddRange(WriteMetadataAction.ProducedItems);
 			}
