UnrealBuildTool: Add out of memory compiler errors to XGE AutoRecover

#rnx

diff --git a/Engine/Source/Programs/UnrealBuildTool/Executors/XGE.cs b/Engine/Source/Programs/UnrealBuildTool/Executors/XGE.cs
--- a/Engine/Source/Programs/UnrealBuildTool/Executors/XGE.cs
+++ b/Engine/Source/Programs/UnrealBuildTool/Executors/XGE.cs
@@ -506,7 +506,14 @@ namespace UnrealBuildTool
 						)
 					);
 
-				if(Action.ActionType == ActionType.Link)
+				if (Action.ActionType == ActionType.Compile)
+				{
+					// C1060: compiler is out of heap space
+					// C1076: compiler limit: internal heap limit reached
+					// C3859: Failed to create virtual memory for PCH
+					ToolElement.SetAttribute("AutoRecover", "C1060,C1076,C3859");
+				}
+				else if (Action.ActionType == ActionType.Link)
 				{
 					ToolElement.SetAttribute("AutoRecover", "Unexpected PDB error; OK (0)");
 				}
