UnrealBuildTool: Set minimum WindowsSdk to 10.0.18362.0

#rnx

diff --git a/Engine/Programs/UnrealGameSync/UnrealGameSync.ini b/Engine/Programs/UnrealGameSync/UnrealGameSync.ini
--- a/Engine/Programs/UnrealGameSync/UnrealGameSync.ini
+++ b/Engine/Programs/UnrealGameSync/UnrealGameSync.ini
@@ -43,7 +43,8 @@ SdkInstallerDir=
 ; All the required SDKs for the current version of the engine
 +SdkInfo=(Category="Android", Description="NDK r21", Browse="$(SdkInstallerDir)\\Android")
 +SdkInfo=(Category="Windows", Description="Visual Studio 2019 v16.11.5")
-+SdkInfo=(Category="Windows", Description="Windows SDK 10.0.16299.0")
++SdkInfo=(Category="Windows", Description="Windows SDK 10.0.18362.0")
++SdkInfo=(Category="Windows", Description="LLVM Clang 11.1.0")
 +SdkInfo=(Category="Linux", Description="v19_clang-11.0.1-centos7")
 
 [//UE5/Release-5.0-EarlyAccess/UE4Games.uprojectdirs]
diff --git a/Engine/Source/Programs/UnrealBuildTool/Platform/Windows/MicrosoftPlatformSDK.cs b/Engine/Source/Programs/UnrealBuildTool/Platform/Windows/MicrosoftPlatformSDK.cs
--- a/Engine/Source/Programs/UnrealBuildTool/Platform/Windows/MicrosoftPlatformSDK.cs
+++ b/Engine/Source/Programs/UnrealBuildTool/Platform/Windows/MicrosoftPlatformSDK.cs
@@ -18,14 +18,13 @@ namespace UnrealBuildTool
 		/// </summary>
 		static readonly VersionNumber[] PreferredWindowsSdkVersions = new VersionNumber[]
 		{
-			VersionNumber.Parse("10.0.18362.0"),
-			VersionNumber.Parse("10.0.16299.0")
+			VersionNumber.Parse("10.0.18362.0")
 		};
 
 		public override string GetMainVersion()
 		{
 			// preferred/main version is the top of the Preferred list - 
-			return PreferredWindowsSdkVersions[0].ToString();
+			return PreferredWindowsSdkVersions.First().ToString();
 		}
 
 		public override void GetValidVersionRange(out string MinVersion, out string MaxVersion)
@@ -36,8 +35,8 @@ namespace UnrealBuildTool
 
 		public override void GetValidSoftwareVersionRange(out string? MinVersion, out string? MaxVersion)
 		{
-			// for now, allow any version of Windows 10
-			MinVersion = "10.0.0.0";
+			// minimum version is the oldest version in the Preferred list -
+			MinVersion = PreferredWindowsSdkVersions.Min().ToString();
 			MaxVersion = null;
 		}
 
