UnrealBuildTool: Remove GetVCToolPath logic to find cl.exe for old versions of visual studio

#rnx

diff --git a/Engine/Source/Programs/UnrealBuildTool/Platform/Windows/VCEnvironment.cs b/Engine/Source/Programs/UnrealBuildTool/Platform/Windows/VCEnvironment.cs
--- a/Engine/Source/Programs/UnrealBuildTool/Platform/Windows/VCEnvironment.cs
+++ b/Engine/Source/Programs/UnrealBuildTool/Platform/Windows/VCEnvironment.cs
@@ -180,38 +180,16 @@ namespace UnrealBuildTool
 		/// <returns>Directory containing the 32-bit toolchain binaries</returns>
 		public static DirectoryReference GetVCToolPath(WindowsCompiler Compiler, DirectoryReference VCToolChainDir, WindowsArchitecture Architecture)
 		{
-			if (Compiler.IsMSVC())
+			FileReference NativeCompilerPath = FileReference.Combine(VCToolChainDir, "bin", "HostX64", WindowsExports.GetArchitectureSubpath(Architecture), "cl.exe");
+			if (FileReference.Exists(NativeCompilerPath))
 			{
-				FileReference NativeCompilerPath = FileReference.Combine(VCToolChainDir, "bin", "HostX64", WindowsExports.GetArchitectureSubpath(Architecture), "cl.exe");
-				if (FileReference.Exists(NativeCompilerPath))
-				{
-					return NativeCompilerPath.Directory;
-				}
-
-				FileReference CrossCompilerPath = FileReference.Combine(VCToolChainDir, "bin", "HostX86", WindowsExports.GetArchitectureSubpath(Architecture), "cl.exe");
-				if (FileReference.Exists(CrossCompilerPath))
-				{
-					return CrossCompilerPath.Directory;
-				}
+				return NativeCompilerPath.Directory;
 			}
-			else
+
+			FileReference CrossCompilerPath = FileReference.Combine(VCToolChainDir, "bin", "HostX86", WindowsExports.GetArchitectureSubpath(Architecture), "cl.exe");
+			if (FileReference.Exists(CrossCompilerPath))
 			{
-				if (Architecture == WindowsArchitecture.x64)
-				{
-					// Use the native 64-bit compiler if present
-					FileReference NativeCompilerPath = FileReference.Combine(VCToolChainDir, "bin", "amd64", "cl.exe");
-					if (FileReference.Exists(NativeCompilerPath))
-					{
-						return NativeCompilerPath.Directory;
-					}
-
-					// Otherwise use the amd64-on-x86 compiler. VS2012 Express only includes the latter.
-					FileReference CrossCompilerPath = FileReference.Combine(VCToolChainDir, "bin", "x86_amd64", "cl.exe");
-					if (FileReference.Exists(CrossCompilerPath))
-					{
-						return CrossCompilerPath.Directory;
-					}
-				}
+				return CrossCompilerPath.Directory;
 			}
 
 			throw new BuildException("No required compiler toolchain found in {0}", VCToolChainDir);
