Fix for makepri logging being displayed incorrectly & simplify code.

#rnx

diff --git a/Engine/Source/Programs/UnrealBuildTool/Platform/Windows/VCManifestGenerator.cs b/Engine/Source/Programs/UnrealBuildTool/Platform/Windows/VCManifestGenerator.cs
--- a/Engine/Source/Programs/UnrealBuildTool/Platform/Windows/VCManifestGenerator.cs
+++ b/Engine/Source/Programs/UnrealBuildTool/Platform/Windows/VCManifestGenerator.cs
@@ -130,20 +130,25 @@ namespace UnrealBuildTool
 
 
         /// <summary>
-        /// Runs external program. Blocking.
+        /// Runs Makepri. Blocking
         /// </summary>
-        /// <param name="Executable">  Executable</param>
-        /// <param name="CommandLine">  Commandline</param>
+        /// <param name="CommandLine">Commandline</param>
         /// <returns>bool    Application ran successfully</returns>
-        protected bool RunExternalProgram(string Executable, string CommandLine)
+        protected bool RunMakePri(string CommandLine)
         {
-			if (File.Exists(Executable) == false)
+			string PriExecutable = GetMakePriBinaryPath();
+
+			if (File.Exists(PriExecutable) == false)
 			{
-				throw new BuildException("BUILD FAILED: Couldn't find the executable to Run: {0}", Executable);
+				throw new BuildException("BUILD FAILED: Couldn't find the makepri executable: {0}", PriExecutable);
 			}
 
-			int ExitCode;
-			string StdOutString = Utils.RunLocalProcessAndReturnStdOut(Executable, CommandLine, out ExitCode, (Log.OutputLevel >= LogEventType.Verbose) );
+			ProcessStartInfo StartInfo = new ProcessStartInfo(PriExecutable, CommandLine);			
+			StartInfo.UseShellExecute = false;
+			StartInfo.CreateNoWindow = true;
+			StartInfo.StandardOutputEncoding = Encoding.Unicode;
+			StartInfo.StandardErrorEncoding = Encoding.Unicode;
+			int ExitCode = Utils.RunLocalProcessAndLogOutput(StartInfo);
 
 			if (ExitCode == 0)
 			{
@@ -151,7 +156,7 @@ namespace UnrealBuildTool
 			}
 			else
 			{
-				Log.TraceError(Path.GetFileName(Executable) + " returned an error.\nApplication output:\n" + StdOutString);
+				Log.TraceError(Path.GetFileName(PriExecutable) + " returned an error.\nExit code:" + ExitCode );
 				return false;
 			}
         }
@@ -862,7 +867,6 @@ namespace UnrealBuildTool
 			if (UpdatedFilePaths.Count > 0)
 			{
 				// Create resource index configuration
-				string PriExecutable = GetMakePriBinaryPath();
 				string ResourceConfigFile = Path.Combine(IntermediatePath, "priconfig.xml");
 				bool bEnableAutoResourcePacks = false;
 				EngineIni.GetBool(IniSection_PlatformTargetSettings, "bEnableAutoResourcePacks", out bEnableAutoResourcePacks);
@@ -871,11 +875,11 @@ namespace UnrealBuildTool
 				// If the game is not going to support language resource packs then merge the culture qualifiers.
 				if (bEnableAutoResourcePacks || CulturesToStage.Count <= 1)
 				{
-					RunExternalProgram(PriExecutable, "createconfig /cf \"" + ResourceConfigFile + "\" /dq " + DefaultCulture + " /o " + GetMakePriExtraCommandLine() );
+					RunMakePri("createconfig /cf \"" + ResourceConfigFile + "\" /dq " + DefaultCulture + " /o " + GetMakePriExtraCommandLine());
 				}
 				else
 				{
-					RunExternalProgram(PriExecutable, "createconfig /cf \"" + ResourceConfigFile + "\" /dq " + String.Join("_", CulturesToStage) + " /o " + GetMakePriExtraCommandLine() );
+					RunMakePri("createconfig /cf \"" + ResourceConfigFile + "\" /dq " + String.Join("_", CulturesToStage) + " /o " + GetMakePriExtraCommandLine());
 				}
 
 				// Modify configuration to restrict indexing to the Resources directory (saves time and space)
@@ -936,7 +940,7 @@ namespace UnrealBuildTool
 				{
 					MakePriCommandLine += " /indexName \"" + IdentityName + "\"";
 				}
-				RunExternalProgram(PriExecutable, MakePriCommandLine);
+				RunMakePri(MakePriCommandLine);
 
 				// Remove any existing pri target files that were not generated by this latest update
 				IEnumerable<string> NewPriFiles = Directory.EnumerateFiles(IntermediatePath, "*.pri");
