UnrealBuildTool: Replace ParallelExecutor with (no longer experimental) TaskExecutor

#rnx

diff --git a/Engine/Source/Programs/UnrealBuildTool/Configuration/BuildConfiguration.cs b/Engine/Source/Programs/UnrealBuildTool/Configuration/BuildConfiguration.cs
--- a/Engine/Source/Programs/UnrealBuildTool/Configuration/BuildConfiguration.cs
+++ b/Engine/Source/Programs/UnrealBuildTool/Configuration/BuildConfiguration.cs
@@ -63,12 +63,6 @@ namespace UnrealBuildTool
 		[CommandLine("-NoSNDBS", Value = "false")]
 		public bool bAllowSNDBS = true;
 
-		/// <summary>
-		/// Whether the experimental async TaskExecutor may be used.
-		/// </summary>
-		[XmlConfigFile]
-		public bool bAllowTaskExecutor = false;
-
 		/// <summary>
 		/// Enables support for very fast iterative builds by caching target data. Turning this on causes Unreal Build Tool to emit
 		/// 'UBT Makefiles' for targets when they are built the first time. Subsequent builds will load these Makefiles and begin
diff --git a/Engine/Source/Programs/UnrealBuildTool/Configuration/TargetRules.cs b/Engine/Source/Programs/UnrealBuildTool/Configuration/TargetRules.cs
--- a/Engine/Source/Programs/UnrealBuildTool/Configuration/TargetRules.cs
+++ b/Engine/Source/Programs/UnrealBuildTool/Configuration/TargetRules.cs
@@ -1531,7 +1531,7 @@ namespace UnrealBuildTool
 		public string? AdditionalLinkerArguments;
 
 		/// <summary>
-		/// Max amount of memory that each compile action may require. Used by ParallelExecutor and TaskExecutor to decide the maximum 
+		/// Max amount of memory that each compile action may require. Used by ParallelExecutor to decide the maximum 
 		/// number of parallel actions to start at one time.
 		/// </summary>
 		public double MemoryPerActionGB = 0.0;
diff --git a/Engine/Source/Programs/UnrealBuildTool/Executors/Experimental/TaskExecutor.cs b/Engine/Source/Programs/UnrealBuildTool/Executors/Experimental/TaskExecutor.cs
deleted file mode 100644
diff --git a/Engine/Source/Programs/UnrealBuildTool/Executors/ParallelExecutor.cs b/Engine/Source/Programs/UnrealBuildTool/Executors/ParallelExecutor.cs
--- a/Engine/Source/Programs/UnrealBuildTool/Executors/ParallelExecutor.cs
+++ b/Engine/Source/Programs/UnrealBuildTool/Executors/ParallelExecutor.cs
@@ -4,45 +4,23 @@ using EpicGames.Core;
 using System;
 using System.Collections.Generic;
 using System.Diagnostics;
+using System.IO;
 using System.Linq;
 using System.Threading;
+using System.Threading.Tasks;
 
 namespace UnrealBuildTool
 {
 	/// <summary>
-	/// This executor is similar to LocalExecutor, but uses p/invoke on Windows to ensure that child processes are started at a lower priority and are terminated when the parent process terminates.
+	/// This executor uses async Tasks to process the action graph
 	/// </summary>
 	class ParallelExecutor : ActionExecutor
 	{
-		[DebuggerDisplay("{Inner}")]
-		class BuildAction
-		{
-			public int SortIndex;
-			public LinkedAction Inner;
-
-			public HashSet<BuildAction> Dependencies = new HashSet<BuildAction>();
-			public int MissingDependencyCount;
-
-			public HashSet<BuildAction> Dependants = new HashSet<BuildAction>();
-			public int TotalDependantCount;
-
-			public List<string> LogLines = new List<string>();
-			public int ExitCode = -1;
-
-			public TimeSpan ExecutionTime;
-			public TimeSpan ProcessorTime;
-
-			public BuildAction(LinkedAction Inner)
-			{
-				this.Inner = Inner;
-			}
-		}
-
 		/// <summary>
 		/// Maximum processor count for local execution. 
 		/// </summary>
 		[XmlConfigFile]
-		public static int MaxProcessorCount = int.MaxValue;
+		private static int MaxProcessorCount = int.MaxValue;
 
 		/// <summary>
 		/// Processor count multiplier for local execution. Can be below 1 to reserve CPU for other tasks.
@@ -51,50 +29,52 @@ namespace UnrealBuildTool
 		/// This value is ignored if the CPU does not support hyper-threading.
 		/// </summary>
 		[XmlConfigFile]
-		public static double ProcessorCountMultiplier = 1.0;
+		private static double ProcessorCountMultiplier = 1.0;
 
 		/// <summary>
 		/// Free memory per action in bytes, used to limit the number of parallel actions if the machine is memory starved.
 		/// Set to 0 to disable free memory checking.
 		/// </summary>
 		[XmlConfigFile]
-		static double MemoryPerActionBytes = 1.5 * 1024 * 1024 * 1024;
+		private static double MemoryPerActionBytes = 1.5 * 1024 * 1024 * 1024;
 
 		/// <summary>
 		/// When enabled, will stop compiling targets after a compile error occurs.
 		/// </summary>
 		[XmlConfigFile]
-		bool bStopCompilationAfterErrors = false;
+		private static bool bStopCompilationAfterErrors = false;
 
 		/// <summary>
 		/// Whether to show compilation times along with worst offenders or not.
 		/// </summary>
 		[XmlConfigFile]
-		bool bShowCompilationTimes = false;
+		private static bool bShowCompilationTimes = false;
 
 		/// <summary>
 		/// Whether to show compilation times for each executed action
 		/// </summary>
 		[XmlConfigFile]
-		bool bShowPerActionCompilationTimes = false;
+		private static bool bShowPerActionCompilationTimes = false;
 
 		/// <summary>
 		/// Whether to log command lines for actions being executed
 		/// </summary>
 		[XmlConfigFile]
-		bool bLogActionCommandLines = false;
+		private static bool bLogActionCommandLines = false;
 
 		/// <summary>
 		/// Add target names for each action executed
 		/// </summary>
 		[XmlConfigFile]
-		bool bPrintActionTargetNames = false;
+		private static bool bPrintActionTargetNames = false;
 
 		/// <summary>
 		/// How many processes that will be executed in parallel
 		/// </summary>
 		public int NumParallelProcesses { get; private set; }
 
+		private static readonly char[] LineEndingSplit = new char[] { '\n', '\r' };
+
 		public static int GetDefaultNumParallelProcesses()
 		{
 			double MemoryPerActionBytesComputed = Math.Max(MemoryPerActionBytes, MemoryPerActionBytesOverride);
@@ -135,15 +115,35 @@ namespace UnrealBuildTool
 		}
 
 		/// <summary>
-		/// Checks whether the parallel executor can be used
+		/// Checks whether the task executor can be used
 		/// </summary>
-		/// <returns>True if the parallel executor can be used</returns>
+		/// <returns>True if the task executor can be used</returns>
 		public static bool IsAvailable()
 		{
 			return true;
 		}
 
-		Dictionary<TargetDescriptor, string> TargetPrefixes = new Dictionary<TargetDescriptor, string>();
+		private class ExecuteResults
+		{
+			public List<string> LogLines { get; private set; }
+			public int ExitCode { get; private set; }
+			public TimeSpan ExecutionTime { get; private set; }
+			public TimeSpan ProcessorTime { get; private set; }
+
+			public ExecuteResults(List<string> LogLines, int ExitCode, TimeSpan ExecutionTime, TimeSpan ProcessorTime)
+			{
+				this.LogLines = LogLines;
+				this.ExitCode = ExitCode;
+				this.ProcessorTime = ProcessorTime;
+				this.ExecutionTime = ExecutionTime;
+			}
+			public ExecuteResults(List<string> LogLines, int ExitCode)
+			{
+				this.LogLines = LogLines;
+				this.ExitCode = ExitCode;
+			}
+		}
+
 
 		/// <summary>
 		/// Executes the specified actions locally.
@@ -151,290 +151,216 @@ namespace UnrealBuildTool
 		/// <returns>True if all the tasks successfully executed, or false if any of them failed.</returns>
 		public override bool ExecuteActions(List<LinkedAction> InputActions)
 		{
-			int ActualNumParallelProcesses = Math.Min(InputActions.Count, NumParallelProcesses);
-			Log.TraceInformation("Building {0} {1} with {2} {3}...", InputActions.Count, (InputActions.Count == 1) ? "action" : "actions", ActualNumParallelProcesses, (ActualNumParallelProcesses == 1) ? "process" : "processes");
+			int NumCompletedActions = 0;
+			int TotalActions = InputActions.Count;
+			int ActualNumParallelProcesses = Math.Min(TotalActions, NumParallelProcesses);
 
-			TargetPrefixes.Clear();
+			using ManagedProcessGroup ProcessGroup = new ManagedProcessGroup();
+			using SemaphoreSlim MaxProcessSemaphore = new SemaphoreSlim(ActualNumParallelProcesses, ActualNumParallelProcesses);
+			using ProgressWriter ProgressWriter = new ProgressWriter("Compiling C++ source code...", false);
 
-			// Create actions with all our internal metadata
-			List<BuildAction> Actions = new List<BuildAction>();
-			for (int Idx = 0; Idx < InputActions.Count; Idx++)
-			{
-				BuildAction Action = new BuildAction(InputActions[Idx]);
-				Action.SortIndex = Idx;
+			Log.TraceInformation("Building {0} {1} with {2} {3}...", TotalActions, (TotalActions == 1) ? "action" : "actions", ActualNumParallelProcesses, (ActualNumParallelProcesses == 1) ? "process" : "processes");
 
-				if (!Action.Inner.StatusDescription.EndsWith(".ispc"))
-				{
-					Action.SortIndex += 10000;
-				}
+			Dictionary<LinkedAction, Task<ExecuteResults>> Tasks = new Dictionary<LinkedAction, Task<ExecuteResults>>();
+			List<Task> AllTasks = new List<Task>();
 
-				Actions.Add(Action);
+			using LogIndentScope Indent = new LogIndentScope("  ");
 
-				TargetDescriptor? Target = Action.Inner.Target;
-				if (bPrintActionTargetNames && Target != null && !TargetPrefixes.ContainsKey(Target))
-				{
-					string TargetPrefix = $"{(1 + TargetPrefixes.Count()).ToString()}>";
-					TargetPrefixes.Add(Target, TargetPrefix);
-					Log.TraceInformation($"{TargetPrefix}------ Unreal Build started: {Target.Name} {Target.Configuration} {Target.Platform} {Target.Architecture} ------");
-				}
-			}
+			CancellationTokenSource CancellationTokenSource = new CancellationTokenSource();
+			CancellationToken CancellationToken = CancellationTokenSource.Token;
 
-			// Update all the actions with all their dependencies
-			Dictionary<LinkedAction, BuildAction> LinkedActionToBuildAction = Actions.ToDictionary(x => x.Inner, x => x);
-			foreach (BuildAction Action in Actions)
+			// Create a task for every action
+			foreach (LinkedAction Action in InputActions)
 			{
-				foreach (LinkedAction PrerequisiteAction in Action.Inner.PrerequisiteActions)
-				{
-					BuildAction? Dependency;
-					if (LinkedActionToBuildAction.TryGetValue(PrerequisiteAction, out Dependency))
-					{
-						Action.Dependencies.Add(Dependency);
-						Dependency.Dependants.Add(Action);
-					}
-				}
-			}
+				Task<ExecuteResults> ExecuteTask = ExecuteAction(Action, Tasks, ProcessGroup, MaxProcessSemaphore, CancellationToken);
+				Task LogTask = ExecuteTask.ContinueWith(antecedent => LogCompletedAction(Action, antecedent, CancellationTokenSource, ProgressWriter, TotalActions, ref NumCompletedActions), CancellationToken);
 
-			// Figure out the recursive dependency count
-			HashSet<BuildAction> VisitedActions = new HashSet<BuildAction>();
-			foreach (BuildAction Action in Actions)
-			{
-				Action.MissingDependencyCount = Action.Dependencies.Count;
-				RecursiveIncDependents(Action, VisitedActions);
+				Tasks.Add(Action, ExecuteTask);
+				AllTasks.Add(ExecuteTask);
+				AllTasks.Add(LogTask);
 			}
 
-			// Create the list of things to process
-			List<BuildAction> QueuedActions = new List<BuildAction>();
-			foreach (BuildAction Action in Actions)
-			{
-				if (Action.MissingDependencyCount == 0)
-				{
-					QueuedActions.Add(Action);
-				}
-			}
-
-			List<BuildAction> AllCompletedActions = new List<BuildAction>();
+			// Wait for all tasks to complete
+			Task.WaitAll(AllTasks.ToArray(), CancellationToken);
 
-			// Execute the actions
-			using (LogIndentScope Indent = new LogIndentScope("  "))
+			if (bShowCompilationTimes)
 			{
-				// Create a job object for all the child processes
-				bool bResult = true;
-				Dictionary<BuildAction, Thread> ExecutingActions = new Dictionary<BuildAction, Thread>();
-				List<BuildAction> CompletedActions = new List<BuildAction>();
-
-				TimeSpan TotalProcessingTime;
-				using (ManagedProcessGroup ProcessGroup = new ManagedProcessGroup())
+				Log.TraceInformation("");
+				if (ProcessGroup.TotalProcessorTime.Ticks > 0)
 				{
-					using (AutoResetEvent CompletedEvent = new AutoResetEvent(false))
-					{
-						int NumCompletedActions = 0;
-						using (ProgressWriter ProgressWriter = new ProgressWriter("Compiling C++ source code...", false))
-						{
-							while (QueuedActions.Count > 0 || ExecutingActions.Count > 0)
-							{
-								// Sort the actions by the number of things dependent on them
-								QueuedActions.Sort((A, B) => (A.TotalDependantCount == B.TotalDependantCount) ? (B.SortIndex - A.SortIndex) : (B.TotalDependantCount - A.TotalDependantCount));
-
-								// Create threads up to the maximum number of actions
-								while (ExecutingActions.Count < ActualNumParallelProcesses && QueuedActions.Count > 0)
-								{
-									BuildAction Action = QueuedActions[QueuedActions.Count - 1];
-									QueuedActions.RemoveAt(QueuedActions.Count - 1);
-
-									Thread ExecutingThread = new Thread(() => { ExecuteAction(ProcessGroup, Action, CompletedActions, CompletedEvent); });
-									string Description = $"{(Action.Inner.CommandDescription != null ? Action.Inner.CommandDescription : Action.Inner.CommandPath.GetFileName())} {Action.Inner.StatusDescription}".Trim();
-									ExecutingThread.Name = String.Format("Build:{0}", Description);
-									ExecutingThread.Start();
-
-									ExecutingActions.Add(Action, ExecutingThread);
-								}
-
-								// Wait for something to finish
-								CompletedEvent.WaitOne();
-
-								// Wait for something to finish and flush it to the log
-								lock (CompletedActions)
-								{
-									foreach (BuildAction CompletedAction in CompletedActions)
-									{
-										// Join the thread
-										Thread CompletedThread = ExecutingActions[CompletedAction];
-										CompletedThread.Join();
-										ExecutingActions.Remove(CompletedAction);
-
-										// Update the progress
-										NumCompletedActions++;
-										ProgressWriter.Write(NumCompletedActions, InputActions.Count);
-
-										string Description = string.Empty;
-
-										string TargetPrefix = "";
-										TargetDescriptor? Target = CompletedAction.Inner.Target;
-										if (bPrintActionTargetNames)
-										{
-											if (Target != null)
-											{
-												TargetPrefix = TargetPrefixes[Target];
-											}
-											else
-											{
-												TargetPrefix = "->";
-											}
-										}
-
-										if (bLogActionCommandLines)
-										{
-											Log.TraceLog($"{TargetPrefix}[{NumCompletedActions}/{InputActions.Count}] Command: {CompletedAction.Inner.CommandPath} {CompletedAction.Inner.CommandArguments}");
-										}
-
-										// Write it to the log
-										if (CompletedAction.Inner.bShouldOutputStatusDescription || CompletedAction.LogLines.Count == 0)
-										{
-											Description = $"{(CompletedAction.Inner.CommandDescription ?? CompletedAction.Inner.CommandPath.GetFileNameWithoutExtension())} {CompletedAction.Inner.StatusDescription}".Trim();
-										}
-										else if (CompletedAction.LogLines.Count > 0)
-										{
-											Description = $"{(CompletedAction.Inner.CommandDescription ?? CompletedAction.Inner.CommandPath.GetFileNameWithoutExtension())} {CompletedAction.LogLines[0]}".Trim();
-										}
-
-										string CompilationTimes = "";
-										if (bShowPerActionCompilationTimes)
-										{
-											CompilationTimes = $" (Wall: {CompletedAction.ExecutionTime.TotalSeconds:0.00}s CPU: {CompletedAction.ProcessorTime.TotalSeconds:0.00}s)";
-										}
-											
-										Log.TraceInformation("{0}[{1}/{2}] {3}{4}", TargetPrefix, NumCompletedActions, InputActions.Count, Description, CompilationTimes);
-										foreach (string Line in CompletedAction.LogLines.Skip(CompletedAction.Inner.bShouldOutputStatusDescription ? 0 : 1))
-										{
-											Log.TraceInformation(Line);
-										}
-
-										AllCompletedActions.Add(CompletedAction);
-
-										// Check the exit code
-										if (CompletedAction.ExitCode == 0)
-										{
-											// Mark all the dependents as done
-											foreach (BuildAction DependantAction in CompletedAction.Dependants)
-											{
-												if (--DependantAction.MissingDependencyCount == 0)
-												{
-													QueuedActions.Add(DependantAction);
-												}
-											}
-										}
-										else
-										{
-											// BEGIN TEMPORARY TO CATCH PVS-STUDIO ISSUES
-											if (CompletedAction.LogLines.Count == 0)
-											{
-												Log.TraceInformation("{0}[{1}/{2}]{3} - Error but no output", TargetPrefix, NumCompletedActions, InputActions.Count, Description);
-												Log.TraceInformation("{0}[{1}/{2}]{3} - {4} {5} {6} {7}", TargetPrefix, NumCompletedActions, InputActions.Count, Description, CompletedAction.ExitCode,
-													CompletedAction.Inner.WorkingDirectory, CompletedAction.Inner.CommandPath, CompletedAction.Inner.CommandArguments);
-											}
-											// END TEMPORARY
-											// Update the exit code if it's not already set
-											if (bResult && CompletedAction.ExitCode != 0)
-											{
-												bResult = false;
-											}
-										}
-									}
-									CompletedActions.Clear();
-								}
-
-								// If we've already got a non-zero exit code, clear out the list of queued actions so nothing else will run
-								if (!bResult && bStopCompilationAfterErrors)
-								{
-									QueuedActions.Clear();
-								}
-							}
-						}
-					}
-
-					TotalProcessingTime = ProcessGroup.TotalProcessorTime;
+					Log.TraceInformation("Total CPU Time: {0} s", ProcessGroup.TotalProcessorTime.TotalSeconds);
+					Log.TraceInformation("");
 				}
 
-				if (bShowCompilationTimes)
+				if (Tasks.Count > 0)
 				{
+					Log.TraceInformation("Compilation Time Top {0}", Math.Min(20, Tasks.Count));
 					Log.TraceInformation("");
-					if (TotalProcessingTime.Ticks > 0)
+					foreach (var Pair in Tasks.OrderByDescending(x => x.Value.Result.ExecutionTime).Take(20))
 					{
-						Log.TraceInformation("Total CPU Time: {0:0.00} s", TotalProcessingTime.TotalSeconds);
-						Log.TraceInformation("");
-					}
-
-					if (AllCompletedActions.Count > 0)
-					{
-						Log.TraceInformation("Compilation Time Top {0}", Math.Min(20, AllCompletedActions.Count));
-						Log.TraceInformation("");
-						foreach (BuildAction Action in AllCompletedActions.OrderByDescending(x => x.ExecutionTime).Take(20))
+						string Description = $"{(Pair.Key.Inner.CommandDescription ?? Pair.Key.Inner.CommandPath.GetFileNameWithoutExtension())} {Pair.Key.Inner.StatusDescription}".Trim();
+						if (Pair.Value.Result.ProcessorTime.Ticks > 0)
 						{
-							string Description = $"{(Action.Inner.CommandDescription != null ? Action.Inner.CommandDescription : Action.Inner.CommandPath.GetFileName())} {Action.Inner.StatusDescription}".Trim();
-							if (Action.ProcessorTime.Ticks > 0)
-							{
-								Log.TraceInformation("{0} [ Wall Time {1:0.00} s / CPU Time {2:0.00} s ]", Description, Action.ExecutionTime.TotalSeconds, Action.ProcessorTime.TotalSeconds);
-							}
-							else
-							{
-								Log.TraceInformation("{0} [ Time {1:0.00} s ]", Description, Action.ExecutionTime.TotalSeconds);
-							}
+							Log.TraceInformation("{0} [ Wall Time {1:0.00} s / CPU Time {2:0.00} s ]", Description, Pair.Value.Result.ExecutionTime.TotalSeconds, Pair.Value.Result.ProcessorTime.TotalSeconds);
 						}
-						Log.TraceInformation("");
+						else
+						{
+							Log.TraceInformation("{0} [ Time {1:0.00} s ]", Description, Pair.Value.Result.ExecutionTime.TotalSeconds);
+						}
+
 					}
+					Log.TraceInformation("");
 				}
-
-				return bResult;
 			}
+
+			// Return if all tasks succeeded
+			return Tasks.Values.All(x => x.Result.ExitCode == 0);
 		}
 
-		/// <summary>
-		/// Execute an individual action
-		/// </summary>
-		/// <param name="ProcessGroup">The process group</param>
-		/// <param name="Action">The action to execute</param>
-		/// <param name="CompletedActions">On completion, the list to add the completed action to</param>
-		/// <param name="CompletedEvent">Event to set once an event is complete</param>
-		static void ExecuteAction(ManagedProcessGroup ProcessGroup, BuildAction Action, List<BuildAction> CompletedActions, AutoResetEvent CompletedEvent)
+		private async Task<ExecuteResults> ExecuteAction(LinkedAction Action, Dictionary<LinkedAction, Task<ExecuteResults>> Tasks, ManagedProcessGroup ProcessGroup, SemaphoreSlim MaxProcessSemaphore, CancellationToken CancellationToken)
 		{
+			Task? SemaphoreTask = null;
 			try
 			{
-				using (ManagedProcess Process = new ManagedProcess(ProcessGroup, Action.Inner.CommandPath.FullName, Action.Inner.CommandArguments, Action.Inner.WorkingDirectory.FullName, null, null, ProcessPriorityClass.BelowNormal))
+				// Wait for Tasks list to be populated with any PrerequisiteActions
+				while (!Action.PrerequisiteActions.All(x => Tasks.ContainsKey(x)))
 				{
-					Action.LogLines.AddRange(Process.ReadAllLines());
-					Action.ExitCode = Process.ExitCode;
-					Action.ProcessorTime = Process.TotalProcessorTime;
-					Action.ExecutionTime = Process.ExitTime - Process.StartTime;
+					await Task.Delay(100, CancellationToken);
+					CancellationToken.ThrowIfCancellationRequested();
 				}
+
+				// Wait for all PrerequisiteActions to complete
+				ExecuteResults[] Results = await Task.WhenAll(Action.PrerequisiteActions.Select(x => Tasks[x]).ToArray());
+
+				// Cancel this task if any PrerequisiteActions fail (or were cancelled)
+				if (Results.Any(x => x.ExitCode != 0))
+				{
+					throw new OperationCanceledException();
+				}
+
+				CancellationToken.ThrowIfCancellationRequested();
+
+				// Limit the number of concurrent processes that will run in parallel
+				SemaphoreTask = MaxProcessSemaphore.WaitAsync(CancellationToken);
+				await SemaphoreTask;
+
+				CancellationToken.ThrowIfCancellationRequested();
+
+				using ManagedProcess Process = new ManagedProcess(ProcessGroup, Action.CommandPath.FullName, Action.CommandArguments, Action.WorkingDirectory.FullName, null, null, ProcessPriorityClass.BelowNormal);
+
+				MemoryStream StdOutStream = new MemoryStream();
+				await Process.CopyToAsync(StdOutStream, CancellationToken);
+
+				CancellationToken.ThrowIfCancellationRequested();
+
+				Process.WaitForExit();
+
+				CancellationToken.ThrowIfCancellationRequested();
+
+				List<string> LogLines = Console.OutputEncoding.GetString(StdOutStream.GetBuffer(), 0, Convert.ToInt32(StdOutStream.Length)).Split(LineEndingSplit, StringSplitOptions.RemoveEmptyEntries).ToList();
+				int ExitCode = Process.ExitCode;
+				TimeSpan ProcessorTime = Process.TotalProcessorTime;
+				TimeSpan ExecutionTime = Process.ExitTime - Process.StartTime;
+				return new ExecuteResults(LogLines, ExitCode, ExecutionTime, ProcessorTime);
+			}
+			catch (OperationCanceledException)
+			{
+				return new ExecuteResults(new List<string>(), int.MaxValue);
 			}
 			catch (Exception Ex)
 			{
 				Log.WriteException(Ex, null);
-				Action.ExitCode = 1;
+				return new ExecuteResults(new List<string>(), int.MaxValue);
 			}
-
-			lock (CompletedActions)
+			finally
 			{
-				CompletedActions.Add(Action);
+				if (SemaphoreTask != null && SemaphoreTask.Status == TaskStatus.RanToCompletion)
+				{
+					MaxProcessSemaphore.Release();
+				}
 			}
-
-			CompletedEvent.Set();
 		}
 
-		/// <summary>
-		/// Increment the number of dependants of an action, recursively
-		/// </summary>
-		/// <param name="Action">The action to update</param>
-		/// <param name="VisitedActions">Set of visited actions</param>
-		private static void RecursiveIncDependents(BuildAction Action, HashSet<BuildAction> VisitedActions)
+		private void LogCompletedAction(LinkedAction Action, Task<ExecuteResults> ExecuteTask, CancellationTokenSource CancellationTokenSource, ProgressWriter ProgressWriter, int TotalActions, ref int NumCompletedActions)
 		{
-			foreach (BuildAction Dependency in Action.Dependants)
+			List<string> LogLines = new List<string>();
+			int ExitCode = int.MaxValue;
+			TimeSpan ExecutionTime = TimeSpan.Zero;
+			TimeSpan ProcessorTime = TimeSpan.Zero;
+			if (ExecuteTask.Status == TaskStatus.RanToCompletion)
+			{
+				ExecuteResults ExecuteTaskResult = ExecuteTask.Result;
+				LogLines = ExecuteTaskResult.LogLines;
+				ExitCode = ExecuteTaskResult.ExitCode;
+				ExecutionTime = ExecuteTaskResult.ExecutionTime;
+				ProcessorTime = ExecuteTaskResult.ProcessorTime;
+			}
+
+			// Write it to the log
+			string Description = string.Empty;
+			if (Action.bShouldOutputStatusDescription || LogLines.Count == 0)
+			{
+				Description = $"{(Action.CommandDescription ?? Action.CommandPath.GetFileNameWithoutExtension())} {Action.StatusDescription}".Trim();
+			}
+			else if (LogLines.Count > 0)
+			{
+				Description = $"{(Action.CommandDescription ?? Action.CommandPath.GetFileNameWithoutExtension())} {LogLines[0]}".Trim();
+			}
+
+			lock (ProgressWriter)
 			{
-				if (!VisitedActions.Contains(Action))
+				int CompletedActions;
+				CompletedActions = Interlocked.Increment(ref NumCompletedActions);
+				ProgressWriter.Write(CompletedActions, TotalActions);
+
+				// Cancelled
+				if (ExitCode == int.MaxValue)
+				{
+					Log.TraceInformation("[{0}/{1}] {2} cancelled", CompletedActions, TotalActions, Description);
+					return;
+				}
+
+				string TargetDetails = "";
+				TargetDescriptor? Target = Action.Target;
+				if (bPrintActionTargetNames && Target != null)
+				{
+					TargetDetails = $"[{Target.Name} {Target.Platform} {Target.Configuration}]";
+				}
+
+				if (bLogActionCommandLines)
+				{
+					Log.TraceLog($"[{CompletedActions}/{TotalActions}]{TargetDetails} Command: {Action.CommandPath} {Action.CommandArguments}");
+				}
+
+				string CompilationTimes = "";
+
+				if (bShowPerActionCompilationTimes)
 				{
-					VisitedActions.Add(Action);
-					Dependency.TotalDependantCount++;
-					RecursiveIncDependents(Dependency, VisitedActions);
+					CompilationTimes = $" (Wall: {ExecutionTime.TotalSeconds:0.00}s CPU: {ProcessorTime.TotalSeconds:0.00}s)";
+				}
+
+				Log.TraceInformation("[{0}/{1}]{2}{3} {4}", CompletedActions, TotalActions, TargetDetails, CompilationTimes, Description);
+				foreach (string Line in LogLines.Skip(Action.bShouldOutputStatusDescription ? 0 : 1))
+				{
+					Log.TraceInformation(Line);
+				}
+
+				if (ExitCode != 0)
+				{
+					// BEGIN TEMPORARY TO CATCH PVS-STUDIO ISSUES
+					if (LogLines.Count == 0)
+					{
+						Log.TraceInformation("[{0}/{1}]{2} {3} - Error but no output", NumCompletedActions, TotalActions, TargetDetails, Description);
+						Log.TraceInformation("[{0}/{1}]{2} {3} - {4} {5} {6} {7}", NumCompletedActions, TotalActions, TargetDetails, Description, ExitCode,
+							Action.WorkingDirectory, Action.CommandPath, Action.CommandArguments);
+					}
+					// END TEMPORARY
+
+					// Cancel all other pending tasks
+					if (bStopCompilationAfterErrors)
+					{
+						CancellationTokenSource.Cancel();
+					}
 				}
 			}
 		}
diff --git a/Engine/Source/Programs/UnrealBuildTool/System/ActionGraph.cs b/Engine/Source/Programs/UnrealBuildTool/System/ActionGraph.cs
--- a/Engine/Source/Programs/UnrealBuildTool/System/ActionGraph.cs
+++ b/Engine/Source/Programs/UnrealBuildTool/System/ActionGraph.cs
@@ -356,10 +356,6 @@ namespace UnrealBuildTool
 				{
 					Executor = new SNDBS();
 				}
-				else if (BuildConfiguration.bAllowTaskExecutor && TaskExecutor.IsAvailable())
-				{
-					Executor = new TaskExecutor(BuildConfiguration.MaxParallelActions);
-				}
 				else
 				{
 					Executor = new ParallelExecutor(BuildConfiguration.MaxParallelActions);
