Fix arm64 not using lld, causing a linker error

diff --git a/Engine/Source/Programs/UnrealBuildTool/Platform/Linux/LinuxToolChain.cs b/Engine/Source/Programs/UnrealBuildTool/Platform/Linux/LinuxToolChain.cs
--- a/Engine/Source/Programs/UnrealBuildTool/Platform/Linux/LinuxToolChain.cs
+++ b/Engine/Source/Programs/UnrealBuildTool/Platform/Linux/LinuxToolChain.cs
@@ -1352,7 +1352,7 @@ namespace UnrealBuildTool
 
 		protected virtual bool UsingLld(string Architecture)
 		{
-			return Architecture.StartsWith("x86_64");
+			return true;
 		}
 
 		/// <summary>
