add support for querying min & max supported Windows SDK

#rnx

diff --git a/Engine/Source/Programs/UnrealBuildTool/Platform/Windows/MicrosoftPlatformSDK.cs b/Engine/Source/Programs/UnrealBuildTool/Platform/Windows/MicrosoftPlatformSDK.cs
--- a/Engine/Source/Programs/UnrealBuildTool/Platform/Windows/MicrosoftPlatformSDK.cs
+++ b/Engine/Source/Programs/UnrealBuildTool/Platform/Windows/MicrosoftPlatformSDK.cs
@@ -180,8 +180,10 @@ namespace UnrealBuildTool
 		/// <param name="DesiredVersion">The desired Windows SDK version. This may be "Latest", a specific version number, or null. If null, the function will look for DefaultWindowsSdkVersion. Failing that, it will return the latest version.</param>
 		/// <param name="OutSdkVersion">Receives the version number of the selected Windows SDK</param>
 		/// <param name="OutSdkDir">Receives the root directory for the selected SDK</param>
+		/// <param name="MinVersion">Optional minimum required version. Ignored if DesiredVesrion is specified</param>
+		/// <param name="MaxVersion">Optional maximum required version. Ignored if DesiredVesrion is specified</param>
 		/// <returns>True if the toolchain directory was found correctly</returns>
-		public static bool TryGetWindowsSdkDir(string? DesiredVersion, [NotNullWhen(true)] out VersionNumber? OutSdkVersion, [NotNullWhen(true)] out DirectoryReference? OutSdkDir)
+		public static bool TryGetWindowsSdkDir(string? DesiredVersion, [NotNullWhen(true)] out VersionNumber? OutSdkVersion, [NotNullWhen(true)] out DirectoryReference? OutSdkDir, VersionNumber? MinVersion = null, VersionNumber? MaxVersion = null)
 		{
 			// Get a map of Windows SDK versions to their root directories
 			/*IReadOnlyDictionary<VersionNumber, DirectoryReference> WindowsSdkDirs =*/
@@ -200,7 +202,7 @@ namespace UnrealBuildTool
 					throw new BuildException("Unable to find requested Windows SDK; '{0}' is an invalid version", DesiredVersion);
 				}
 			}
-			else
+			else if (MinVersion == null && MaxVersion == null)
 			{
 				WindowsSdkVersion = PreferredWindowsSdkVersions.FirstOrDefault(x => CachedWindowsSdkDirs!.ContainsKey(x));
 				if (WindowsSdkVersion == null && CachedWindowsSdkDirs!.Count > 0)
@@ -208,6 +210,13 @@ namespace UnrealBuildTool
 					WindowsSdkVersion = CachedWindowsSdkDirs.OrderBy(x => x.Key).Last().Key;
 				}
 			}
+			else if (CachedWindowsSdkDirs!.Count > 0)
+			{
+				WindowsSdkVersion = CachedWindowsSdkDirs.OrderBy(x => x.Key).Where( 
+					x =>
+					(MinVersion == null || x.Key >= MinVersion) &&
+					(MaxVersion == null || x.Key <= MaxVersion)).Last().Key;
+			}
 
 			// Get the actual directory for this version
 			DirectoryReference? SdkDir;
