PR #8797: add missing '$' in the log so that values are formatted correctly. (Contributed by [REDACTED_PERSON])

diff --git a/Engine/Source/Programs/UnrealBuildTool/System/ModuleDescriptor.cs b/Engine/Source/Programs/UnrealBuildTool/System/ModuleDescriptor.cs
--- a/Engine/Source/Programs/UnrealBuildTool/System/ModuleDescriptor.cs
+++ b/Engine/Source/Programs/UnrealBuildTool/System/ModuleDescriptor.cs
@@ -281,7 +281,7 @@ namespace UnrealBuildTool
 						}
 						else
 						{
-							Log.TraceWarningTask(JsonFilePath, "Unknown platform {TargetPlatformName} while parsing deny list for module descriptor {Module.Name}");
+							Log.TraceWarningTask(JsonFilePath, $"Unknown platform {TargetPlatformName} while parsing deny list for module descriptor {Module.Name}");
 						}
 					}
 				}
