UBT TestMode that allows building a test executable for an associated target.
- Each module can now have a Tests folder at its root containing low level tests written with Catch2
- Used like this: UnrealBuildTool.exe -Mode=Test TargetName Platform Configuration
- It creates a test executable that's associated with the specified build Target and it includes all the tests for the Target AND all its dependencies as well
- When building without -Mode=Test it will omit the Tests folder contents from compilation: in theory we never allowed folders named Tests at the root of modules, but there might be isolated exceptions where we didn't validate folder names - in this case we will need to move that Tests folder one level up for example
- Code paths for test and non-test builds are entirely separated

diff --git a/Engine/Source/Programs/UnrealBuildTool/Configuration/ModuleRules.cs b/Engine/Source/Programs/UnrealBuildTool/Configuration/ModuleRules.cs
--- a/Engine/Source/Programs/UnrealBuildTool/Configuration/ModuleRules.cs
+++ b/Engine/Source/Programs/UnrealBuildTool/Configuration/ModuleRules.cs
@@ -1041,6 +1041,17 @@ namespace UnrealBuildTool
 			}
 		}
 
+		/// <summary>
+		/// Returns module's low level tests directory "Tests".
+		/// </summary>
+		public string TestsDirectory
+		{
+			get
+			{
+				return Path.Combine(Directory.FullName, "Tests");
+			}
+		}
+
 #nullable disable
 		/// <summary>
 		/// Constructor. For backwards compatibility while the parameterless constructor is being phased out, initialization which would happen here is done by 
diff --git a/Engine/Source/Programs/UnrealBuildTool/Configuration/TargetDescriptor.cs b/Engine/Source/Programs/UnrealBuildTool/Configuration/TargetDescriptor.cs
--- a/Engine/Source/Programs/UnrealBuildTool/Configuration/TargetDescriptor.cs
+++ b/Engine/Source/Programs/UnrealBuildTool/Configuration/TargetDescriptor.cs
@@ -16,12 +16,20 @@ namespace UnrealBuildTool
 	/// </summary>
 	class TargetDescriptor
 	{
+		public static string TEST_TARGETS_SUFFIX = "Tests";
+
 		public FileReference? ProjectFile;
 		public string Name;
 		public UnrealTargetPlatform Platform;
 		public UnrealTargetConfiguration Configuration;
 		public string Architecture;
 		public CommandLineArguments AdditionalArguments;
+		public bool IsTestsTarget = false;
+
+		public static string GetTestedTargetName(string Name)
+		{
+			return Name.Substring(0, Name.Length - TEST_TARGETS_SUFFIX.Length);
+		}
 
 		/// <summary>
 		/// Foreign plugin to compile against this target
@@ -160,6 +168,11 @@ namespace UnrealBuildTool
 			this.AdditionalArguments = new CommandLineArguments(AdditionalArguments.ToArray());
 		}
 
+		public TargetDescriptor Copy()
+		{
+			return (TargetDescriptor)this.MemberwiseClone();
+		}
+
 		public static TargetDescriptor FromTargetInfo(TargetInfo Info)
 		{
 			return new TargetDescriptor(Info.ProjectFile, Info.Name, Info.Platform, Info.Configuration, Info.Architecture, Info.Arguments);
diff --git a/Engine/Source/Programs/UnrealBuildTool/Configuration/TargetRules.cs b/Engine/Source/Programs/UnrealBuildTool/Configuration/TargetRules.cs
--- a/Engine/Source/Programs/UnrealBuildTool/Configuration/TargetRules.cs
+++ b/Engine/Source/Programs/UnrealBuildTool/Configuration/TargetRules.cs
@@ -215,6 +215,21 @@ namespace UnrealBuildTool
 
 		private readonly string DefaultName;
 
+		private TestTargetRules? TestTargetRules;
+
+		/// <summary>
+		/// Create and return test target executable on demand for this target and its dependencies' tests.
+		/// </summary>
+		public TestTargetRules CreateOrGetTestTarget()
+		{
+			if (TestTargetRules == null)
+			{
+				TargetInfo TestsTargetInfo = new TargetInfo(Name + TargetDescriptor.TEST_TARGETS_SUFFIX, Platform, Configuration, Architecture, ProjectFile, null);
+				TestTargetRules = new TestTargetRules(this, TestsTargetInfo);
+			}
+			return TestTargetRules;	
+		}
+
 		/// <summary>
 		/// File containing the general type for this target (not including platform/group)
 		/// </summary>
@@ -1414,7 +1429,7 @@ namespace UnrealBuildTool
 		{
 			get
 			{
-				return (LaunchModuleNamePrivate == null && Type != global::UnrealBuildTool.TargetType.Program)? "Launch" : LaunchModuleNamePrivate;
+				return (LaunchModuleNamePrivate == null && Type != global::UnrealBuildTool.TargetType.Program) ? "Launch" : LaunchModuleNamePrivate;
 			}
 			set
 			{
@@ -1785,6 +1800,11 @@ namespace UnrealBuildTool
 		/// <returns>Array of platforms that the target supports</returns>
 		internal UnrealTargetPlatform[] GetSupportedPlatforms()
 		{
+			if (this is TestTargetRules TestTarget)
+			{
+				return TestTarget.TestedTarget.GetSupportedPlatforms();
+			}
+
 			// Otherwise take the SupportedPlatformsAttribute from the first type in the inheritance chain that supports it
 			for (Type? CurrentType = GetType(); CurrentType != null; CurrentType = CurrentType.BaseType)
 			{
@@ -1964,6 +1984,19 @@ namespace UnrealBuildTool
 			get { return Inner.Name; }
 		}
 
+		public TestTargetRules TestTarget { get { return Inner.CreateOrGetTestTarget(); } }
+
+		public TargetRules TestedTarget {
+			get
+			{
+				if (Inner is TestTargetRules)
+				{
+					return ((TestTargetRules)Inner).TestedTarget;
+				}
+				throw new Exception("Not a test target.");
+			}
+		}
+
 		internal FileReference File
 		{
 			get { return Inner.File!; }
@@ -2940,6 +2973,11 @@ namespace UnrealBuildTool
 			get { return "../Binaries/ThirdParty/"; }
 		}
 
+		public bool IsTestTarget()
+		{
+			return Inner is TestTargetRules;
+		}
+
 		/// <summary>
 		/// Checks if current platform is part of a given platform group
 		/// </summary>
diff --git a/Engine/Source/Programs/UnrealBuildTool/Configuration/TestModuleRules.cs b/Engine/Source/Programs/UnrealBuildTool/Configuration/TestModuleRules.cs
new file mode 100644
--- /dev/null
+++ b/Engine/Source/Programs/UnrealBuildTool/Configuration/TestModuleRules.cs
@@ -0,0 +1,75 @@
+// Copyright Epic Games, Inc. All Rights Reserved.
+
+using EpicGames.Core;
+using System.IO;
+
+namespace UnrealBuildTool
+{
+	/// <summary>
+	/// ModuleRules extension for low level tests.
+	/// </summary>
+	public class TestModuleRules : ModuleRules
+	{
+		/// <summary>
+		/// Associated tested module of this test module.
+		/// </summary>
+		public ModuleRules TestedModule { get; private set; }
+
+		/// <summary>
+		/// Constructs a TestModuleRules object with an associated tested module.
+		/// </summary>
+		public TestModuleRules(ModuleRules TestedModule) : base(new ReadOnlyTargetRules(TestedModule.Target.TestTarget))
+		{
+			this.TestedModule = TestedModule;
+
+			Name = TestedModule.Name + "Tests";
+			if (!string.IsNullOrEmpty(TestedModule.ShortName))
+			{
+				ShortName = TestedModule.ShortName + "Tests";
+			}
+
+			File = FileReference.Combine(UnrealBuildTool.EngineSourceDirectory, "UnrealBuildTool", "Configuration", "TestModuleRules.cs");
+			Directory = DirectoryReference.Combine(TestedModule.Directory, "Tests");
+
+			Context = TestedModule.Context;
+
+			PCHUsage = PCHUsageMode.NoPCHs;
+			PrecompileForTargets = PrecompileTargetsType.None;
+
+			if (Target.Configuration == UnrealTargetConfiguration.Debug && Target.Platform == UnrealTargetPlatform.Linux)
+			{
+				OptimizeCode = CodeOptimization.Never;
+			}
+
+			bAllowConfidentialPlatformDefines = true;
+			bLegalToDistributeObjectCode = true;
+
+			// Required false for catch.hpp
+			bUseUnity = false;
+
+			// Disable exception handling so that tests can assert for exceptions
+			bEnableObjCExceptions = false;
+			bEnableExceptions = false;
+
+			PrivateDependencyModuleNames.AddRange(
+				new string[] {
+					"Core",
+					"Projects",
+					"LowLevelTestsRunner"
+				});
+
+			PrivateIncludePaths.Add("Runtime/Launch/Private");
+			PrivateIncludePathModuleNames.Add("Launch");
+
+			// Tests can refer to tested module's Public and Private paths
+			PublicIncludePaths.Add(Path.Combine(TestedModule.ModuleDirectory, "Public"));
+			PrivateIncludePaths.Add(Path.Combine(TestedModule.ModuleDirectory, "Private"));
+
+			// Platforms specific setup
+			if (Target.Platform == UnrealTargetPlatform.Android)
+			{
+				PublicDefinitions.Add("CATCH_CONFIG_NOSTDOUT");
+			}
+		}
+	}
+}
diff --git a/Engine/Source/Programs/UnrealBuildTool/Configuration/TestTargetRules.cs b/Engine/Source/Programs/UnrealBuildTool/Configuration/TestTargetRules.cs
new file mode 100644
--- /dev/null
+++ b/Engine/Source/Programs/UnrealBuildTool/Configuration/TestTargetRules.cs
@@ -0,0 +1,100 @@
+// Copyright Epic Games, Inc. All Rights Reserved.
+
+using EpicGames.Core;
+using System.IO;
+
+namespace UnrealBuildTool
+{
+	/// <summary>
+	/// TargetRules extension for low level tests.
+	/// </summary>
+	public class TestTargetRules : TargetRules
+	{
+		/// <summary>
+		/// Associated tested target of this test target.
+		/// </summary>
+		public TargetRules TestedTarget { get; private set; }
+
+		/// <summary>
+		/// TestTargetRules is setup as a program and is linked monolithically .
+		/// It removes a lot of default compilation behavior in order to produce a minimal test environment.
+		/// </summary>
+		public TestTargetRules(TargetRules TestedTarget, TargetInfo Target) : base(Target)
+		{
+			if (TestedTarget is TestTargetRules)
+			{
+				throw new BuildException("TestedTarget can't be of type TestTargetRules.");
+			}
+
+			this.TestedTarget = TestedTarget;
+
+			DefaultBuildSettings = BuildSettingsVersion.V2;
+
+			ExeBinariesSubFolder = Name = TestedTarget.Name + "Tests";
+			TargetSourceFile = File = FileReference.Combine(UnrealBuildTool.EngineSourceDirectory, "UnrealBuildTool", "Configuration", "TestTargetRules.cs");
+			LaunchModuleName = TestedTarget.LaunchModuleName + "Tests";
+
+			WindowsPlatform = TestedTarget.WindowsPlatform;
+
+			Type = TargetType.Program;
+			LinkType = TargetLinkType.Monolithic;
+
+			bBuildInSolutionByDefault = false;
+
+			bDeployAfterCompile = false;
+			bIsBuildingConsoleApplication = true;
+
+			// Disabling default true flags that aren't necessary for tests
+
+			// Lean and Mean mode
+			bBuildDeveloperTools = false;
+
+			// No localization
+			bCompileICU = false;
+
+			// No need for shaders by default
+			bForceBuildShaderFormats = false;
+
+			// Do not link against the engine, no Chromium Embedded Framework etc.
+			bCompileAgainstEngine = false;
+			bCompileCEF3 = false;
+			bCompileAgainstCoreUObject = false;
+			bCompileAgainstApplicationCore = false;
+			bUseLoggingInShipping = true;
+
+			bool bDebugOrDevelopment = Target.Configuration == UnrealTargetConfiguration.Debug || Target.Configuration == UnrealTargetConfiguration.Development;
+			bBuildWithEditorOnlyData = Target.Platform.IsInGroup(UnrealPlatformGroup.Desktop) && bDebugOrDevelopment;
+
+			// Disable malloc profiling in tests
+			bUseMallocProfiler = false;
+
+			// Useful for debugging test failures
+			if (Target.Configuration == UnrealTargetConfiguration.Debug)
+			{
+				bDebugBuildsActuallyUseDebugCRT = true;
+			}
+
+			GlobalDefinitions.AddRange(TestedTarget.GlobalDefinitions);
+
+			GlobalDefinitions.Add("STATS=0");
+
+			// Platform specific setup
+			if (Target.Platform == UnrealTargetPlatform.Android)
+			{
+				UndecoratedConfiguration = Target.Configuration;
+
+				string VersionScriptFile = Path.Combine(Directory.GetCurrentDirectory(), @"Developer\LowLevelTestsRunner\Private\Platform\Android\HideSymbols.ldscript");
+				AdditionalLinkerArguments = " -Wl,--version-script=" + VersionScriptFile;
+
+				GlobalDefinitions.Add("USE_ANDROID_INPUT=0");
+				GlobalDefinitions.Add("USE_ANDROID_OPENGL=0");
+				GlobalDefinitions.Add("USE_ANDROID_LAUNCH=0");
+				GlobalDefinitions.Add("USE_ANDROID_JNI=0");
+			}
+			else if (Target.Platform == UnrealTargetPlatform.IOS)
+			{
+				GlobalDefinitions.Add("HAS_METAL=0");
+			}
+		}
+	}
+}
diff --git a/Engine/Source/Programs/UnrealBuildTool/Configuration/UEBuildModule.cs b/Engine/Source/Programs/UnrealBuildTool/Configuration/UEBuildModule.cs
--- a/Engine/Source/Programs/UnrealBuildTool/Configuration/UEBuildModule.cs
+++ b/Engine/Source/Programs/UnrealBuildTool/Configuration/UEBuildModule.cs
@@ -940,7 +940,7 @@ namespace UnrealBuildTool
 			}
 		}
 
-		public delegate UEBuildModule CreateModuleDelegate(string Name, string ReferenceChain);
+		public delegate UEBuildModule CreateModuleDelegate(string Name, string ReferenceChain, bool IsTestModule = false);
 
 		/// <summary>
 		/// Public entry point to recursively create a module and all its dependencies
diff --git a/Engine/Source/Programs/UnrealBuildTool/Configuration/UEBuildModuleCPP.cs b/Engine/Source/Programs/UnrealBuildTool/Configuration/UEBuildModuleCPP.cs
--- a/Engine/Source/Programs/UnrealBuildTool/Configuration/UEBuildModuleCPP.cs
+++ b/Engine/Source/Programs/UnrealBuildTool/Configuration/UEBuildModuleCPP.cs
@@ -392,7 +392,7 @@ namespace UnrealBuildTool
 
 			// Find all the input files
 			Dictionary<DirectoryItem, FileItem[]> DirectoryToSourceFiles = new Dictionary<DirectoryItem, FileItem[]>();
-			InputFileCollection InputFiles = FindInputFiles(Target.Platform, DirectoryToSourceFiles);
+			InputFileCollection InputFiles = FindInputFiles(Target.Platform, DirectoryToSourceFiles, Target.IsTestTarget());
 
 			foreach (KeyValuePair<DirectoryItem, FileItem[]> Pair in DirectoryToSourceFiles)
 			{
@@ -1490,8 +1490,9 @@ namespace UnrealBuildTool
 		/// </summary>
 		/// <param name="Platform">The platform the module is being built for</param>
 		/// <param name="DirectoryToSourceFiles">Map of directory to source files inside it</param>
+		/// <param name="IncludeTests">Whether tests are included in the file list</param>
 		/// <returns>Set of source files that should be built</returns>
-		public InputFileCollection FindInputFiles(UnrealTargetPlatform Platform, Dictionary<DirectoryItem, FileItem[]> DirectoryToSourceFiles)
+		public InputFileCollection FindInputFiles(UnrealTargetPlatform Platform, Dictionary<DirectoryItem, FileItem[]> DirectoryToSourceFiles, bool IncludeTests = false)
 		{
 			ReadOnlyHashSet<string> ExcludedNames = UEBuildPlatform.GetBuildPlatform(Platform).GetExcludedFolderNames();
 
@@ -1501,7 +1502,32 @@ namespace UnrealBuildTool
 			foreach (DirectoryReference Dir in ModuleDirectories)
 			{
 				DirectoryItem ModuleDirectoryItem = DirectoryItem.GetItemByDirectoryReference(Dir);
-				FindInputFilesFromDirectoryRecursive(ModuleDirectoryItem, ExcludedNames, SourceDirectories, DirectoryToSourceFiles, InputFiles);
+				if (!IncludeTests)
+				{
+					// Skip Tests folder at the base of the source folder if not including tests
+					foreach (DirectoryItem FirstLevelDir in ModuleDirectoryItem.EnumerateDirectories())
+					{
+						if (FirstLevelDir.Name == "Tests")
+						{
+							continue;
+						}
+						else
+						{
+							FindInputFilesFromDirectoryRecursive(FirstLevelDir, ExcludedNames, SourceDirectories, DirectoryToSourceFiles, InputFiles);
+						}
+					}
+					// Add remaining source files at the base of the source folder
+					FileItem[] SourceFiles = FindInputFilesFromDirectory(ModuleDirectoryItem, InputFiles);
+					if (SourceFiles.Length > 0)
+					{
+						SourceDirectories.Add(ModuleDirectoryItem.Location);
+					}
+					DirectoryToSourceFiles.Add(ModuleDirectoryItem, SourceFiles);
+				}
+				else
+				{
+					FindInputFilesFromDirectoryRecursive(ModuleDirectoryItem, ExcludedNames, SourceDirectories, DirectoryToSourceFiles, InputFiles);
+				}
 			}
 
 			return InputFiles;
@@ -1519,7 +1545,7 @@ namespace UnrealBuildTool
 		{
 			foreach(DirectoryItem SubDirectory in BaseDirectory.EnumerateDirectories())
 			{
-				if(!ExcludedNames.Contains(SubDirectory.Name))
+				if (!ExcludedNames.Contains(SubDirectory.Name))
 				{
 					FindInputFilesFromDirectoryRecursive(SubDirectory, ExcludedNames, SourceDirectories, DirectoryToSourceFiles, InputFiles);
 				}
diff --git a/Engine/Source/Programs/UnrealBuildTool/Configuration/UEBuildTarget.cs b/Engine/Source/Programs/UnrealBuildTool/Configuration/UEBuildTarget.cs
--- a/Engine/Source/Programs/UnrealBuildTool/Configuration/UEBuildTarget.cs
+++ b/Engine/Source/Programs/UnrealBuildTool/Configuration/UEBuildTarget.cs
@@ -708,7 +708,15 @@ namespace UnrealBuildTool
 			TargetRules RulesObject;
 			using (GlobalTracer.Instance.BuildSpan("RulesAssembly.CreateTargetRules()").StartActive())
 			{
-				RulesObject = RulesAssembly.CreateTargetRules(Descriptor.Name, Descriptor.Platform, Descriptor.Configuration, Descriptor.Architecture, Descriptor.ProjectFile, Descriptor.AdditionalArguments);
+				if (!Descriptor.IsTestsTarget) {
+					RulesObject = RulesAssembly.CreateTargetRules(Descriptor.Name, Descriptor.Platform, Descriptor.Configuration, Descriptor.Architecture, Descriptor.ProjectFile, Descriptor.AdditionalArguments);
+				}
+				else
+				{
+					// Get testable target first, then use its tests executable target instead
+					RulesObject = RulesAssembly.CreateTargetRules(TargetDescriptor.GetTestedTargetName(Descriptor.Name), Descriptor.Platform, Descriptor.Configuration, Descriptor.Architecture, Descriptor.ProjectFile, Descriptor.AdditionalArguments);
+					RulesObject = RulesObject.CreateOrGetTestTarget();
+				}
 			}
 			if ((ProjectFileGenerator.bGenerateProjectFiles == false) && !RulesObject.GetSupportedPlatforms().Contains(Descriptor.Platform))
 			{
@@ -3557,7 +3565,7 @@ namespace UnrealBuildTool
 			}
 
 			// Create the launch module
-			UEBuildModuleCPP LaunchModule = FindOrCreateCppModuleByName(Rules.LaunchModuleName, TargetRulesFile.GetFileName());
+			UEBuildModuleCPP LaunchModule = FindOrCreateCppModuleByName(Rules.LaunchModuleName, TargetRulesFile.GetFileName(), Rules.IsTestTarget());
 
 			// Get the intermediate directory for the launch module directory. This can differ from the standard engine intermediate directory because it is always configuration-specific.
 			DirectoryReference IntermediateDirectory;
@@ -4016,10 +4024,19 @@ namespace UnrealBuildTool
 		/// <summary>
 		/// Create a rules object for the given module, and set any default values for this target
 		/// </summary>
-		private ModuleRules CreateModuleRulesAndSetDefaults(string ModuleName, string ReferenceChain)
+		private ModuleRules CreateModuleRulesAndSetDefaults(string ModuleName, string ReferenceChain, bool IsTestModule = false)
 		{
 			// Create the rules from the assembly
-			ModuleRules RulesObject = RulesAssembly.CreateModuleRules(ModuleName, Rules, ReferenceChain);
+			ModuleRules RulesObject;
+			if (IsTestModule)
+			{
+				RulesObject = RulesAssembly.CreateModuleRules(TargetDescriptor.GetTestedTargetName(ModuleName), new ReadOnlyTargetRules(Rules.TestedTarget), ReferenceChain);
+				RulesObject = new TestModuleRules(RulesObject);
+			}
+			else
+			{
+				RulesObject = RulesAssembly.CreateModuleRules(ModuleName, Rules, ReferenceChain);
+			}
 
 			// Set whether the module requires an IMPLEMENT_MODULE macro
 			if(!RulesObject.bRequiresImplementModule.HasValue)
@@ -4115,18 +4132,26 @@ namespace UnrealBuildTool
 		/// </summary>
 		/// <param name="ModuleName">Name of the module</param>
 		/// <param name="ReferenceChain">Chain of references causing this module to be instantiated, for display in error messages</param>
-		public UEBuildModule FindOrCreateModuleByName(string ModuleName, string ReferenceChain)
+		/// <param name="IsTestModule">Whether it's a test module or not</param>
+		public UEBuildModule FindOrCreateModuleByName(string ModuleName, string ReferenceChain, bool IsTestModule = false)
 		{
 			UEBuildModule? Module;
 			if (!Modules.TryGetValue(ModuleName, out Module))
 			{
 				// @todo projectfiles: Cross-platform modules can appear here during project generation, but they may have already
 				//   been filtered out by the project generator.  This causes the projects to not be added to directories properly.
-				ModuleRules RulesObject = CreateModuleRulesAndSetDefaults(ModuleName, ReferenceChain);
+				ModuleRules RulesObject = CreateModuleRulesAndSetDefaults(ModuleName, ReferenceChain, IsTestModule);
 				DirectoryReference ModuleDirectory = RulesObject.File.Directory;
 
+				// If we're building a test executable, include test harness dependency
+				if (!IsTestModule && ReferenceChain.Contains(typeof(TestModuleRules).Name) && Directory.Exists(RulesObject.TestsDirectory))
+				{
+					// TODO: If explicit Tests module rules exists, use that one instead
+					RulesObject.PrivateDependencyModuleNames.Add("LowLevelTestsRunner");
+				}
+
 				// Clear the bUsePrecompiled flag if we're compiling a foreign plugin; since it's treated like an engine module, it will default to true in an installed build.
-				if(RulesObject.Plugin != null && RulesObject.Plugin.File == ForeignPlugin)
+				if (RulesObject.Plugin != null && RulesObject.Plugin.File == ForeignPlugin)
 				{
 					RulesObject.bPrecompile = true;
 					RulesObject.bUsePrecompiled = false;
@@ -4164,7 +4189,7 @@ namespace UnrealBuildTool
 					GeneratedCodeDirectory = DirectoryReference.Combine(GeneratedCodeDirectory, PlatformIntermediateFolder, AppName, "Inc");
 
 					// Append the binaries subfolder, if present. We rely on this to ensure that build products can be filtered correctly.
-					if(RulesObject.BinariesSubFolder != null)
+					if (RulesObject.BinariesSubFolder != null)
 					{
 						GeneratedCodeDirectory = DirectoryReference.Combine(GeneratedCodeDirectory, RulesObject.BinariesSubFolder);
 					}
@@ -4190,7 +4215,7 @@ namespace UnrealBuildTool
 					// If it's a game module (plugin or otherwise), add the root source directory to the include paths.
 					if (RulesObject.File.IsUnderDirectory(TargetRulesFile.Directory) || (RulesObject.Plugin != null && RulesObject.Plugin.LoadedFrom == PluginLoadedFrom.Project))
 					{
-						if(DirectoryReference.Exists(BaseSourceDirectory))
+						if (DirectoryReference.Exists(BaseSourceDirectory))
 						{
 							RulesObject.PublicIncludePaths.Add(NormalizeIncludePath(BaseSourceDirectory));
 						}
@@ -4241,10 +4266,11 @@ namespace UnrealBuildTool
 		/// </summary>
 		/// <param name="ModuleName">Name of the module</param>
 		/// <param name="ReferenceChain">Chain of references causing this module to be instantiated, for display in error messages</param>
+		/// <param name="IsTestModule">Whether it's a test module.</param>
 		/// <returns>New C++ module</returns>
-		public UEBuildModuleCPP FindOrCreateCppModuleByName(string ModuleName, string ReferenceChain)
+		public UEBuildModuleCPP FindOrCreateCppModuleByName(string ModuleName, string ReferenceChain, bool IsTestModule = false)
 		{
-			UEBuildModuleCPP? CppModule = FindOrCreateModuleByName(ModuleName, ReferenceChain) as UEBuildModuleCPP;
+			UEBuildModuleCPP? CppModule = FindOrCreateModuleByName(ModuleName, ReferenceChain, IsTestModule) as UEBuildModuleCPP;
 			if(CppModule == null)
 			{
 				throw new BuildException("'{0}' is not a C++ module (referenced via {1})", ModuleName, ReferenceChain);
diff --git a/Engine/Source/Programs/UnrealBuildTool/Modes/BuildMode.cs b/Engine/Source/Programs/UnrealBuildTool/Modes/BuildMode.cs
--- a/Engine/Source/Programs/UnrealBuildTool/Modes/BuildMode.cs
+++ b/Engine/Source/Programs/UnrealBuildTool/Modes/BuildMode.cs
@@ -3,13 +3,9 @@
 using System;
 using System.Collections.Generic;
 using System.Diagnostics;
-using System.IO;
 using System.Linq;
 using System.Text;
-using System.Threading;
-using System.Threading.Tasks;
 using EpicGames.Core;
-using OpenTracing;
 using OpenTracing.Util;
 using UnrealBuildBase;
 
@@ -259,7 +255,6 @@ namespace UnrealBuildTool
 		/// <returns>Result from the compilation</returns>
 		public static void Build(List<TargetDescriptor> TargetDescriptors, BuildConfiguration BuildConfiguration, ISourceFileWorkingSet WorkingSet, BuildOptions Options, FileReference? WriteOutdatedActionsFile, bool bSkipPreBuildTargets = false)
 		{
-
 			List<TargetMakefile> TargetMakefiles = new List<TargetMakefile>();
 
 			for (int Idx = 0; Idx < TargetDescriptors.Count; ++Idx)
@@ -571,7 +566,7 @@ namespace UnrealBuildTool
 						XGE.ExportActions(MergedActionsToExecute);
 					}
 				}
-				else if(WriteOutdatedActionsFile != null)
+				else if (WriteOutdatedActionsFile != null)
 				{
 					OutputToolchainInfo(TargetDescriptors, Makefiles);
 
diff --git a/Engine/Source/Programs/UnrealBuildTool/Modes/TestMode.cs b/Engine/Source/Programs/UnrealBuildTool/Modes/TestMode.cs
new file mode 100644
--- /dev/null
+++ b/Engine/Source/Programs/UnrealBuildTool/Modes/TestMode.cs
@@ -0,0 +1,65 @@
+// Copyright Epic Games, Inc. All Rights Reserved.
+
+using System.Collections.Generic;
+using EpicGames.Core;
+using UnrealBuildBase;
+
+namespace UnrealBuildTool
+{
+	/// <summary>
+	/// Builds low level tests on one or more targets.
+	/// </summary>
+	[ToolMode("Test", ToolModeOptions.XmlConfig | ToolModeOptions.BuildPlatforms | ToolModeOptions.SingleInstance | ToolModeOptions.ShowExecutionTime)]
+	class TestMode : ToolMode
+	{
+		/// <summary>
+		/// Main entry point
+		/// </summary>
+		/// <param name="Arguments">Command-line arguments</param>
+		public override int Execute(CommandLineArguments Arguments)
+		{
+			Arguments.ApplyTo(this);
+
+			// Create the build configuration object, and read the settings
+			BuildConfiguration BuildConfiguration = new BuildConfiguration();
+			Arguments.ApplyTo(BuildConfiguration);
+			XmlConfig.ApplyTo(BuildConfiguration);
+
+			// Parse all the targets being built
+			List<TargetDescriptor> TargetDescriptors = TargetDescriptor.ParseCommandLine(Arguments, false, false, false);
+
+			BuildTests(TargetDescriptors, BuildConfiguration);
+
+			return 0;
+		}
+
+		/// <summary>
+		/// Build tests for a list of targets.
+		/// It generates artificial test target descriptors to build a target AND its dependencies's tests in one monolithic executable.
+		/// Each module containing a "Tests" folder is included.
+		/// The target that we generate the tests executable for must include a main.cpp file and Setup and Teardown methods.
+		/// The generated executable is the target name + "Tests", e.g. UnrealEditorTests.exe
+		/// Passing UnrealEditor here would build all the tests in all the modules.
+		/// </summary>
+		/// <param name="TargetDescriptors">Target descriptors</param>
+		/// <param name="BuildConfiguration">Current build configuration</param>
+		public static void BuildTests(List<TargetDescriptor> TargetDescriptors, BuildConfiguration BuildConfiguration)
+		{
+			List<TargetDescriptor> TestTargetDescriptors = new List<TargetDescriptor>();
+
+			for (int Idx = 0; Idx < TargetDescriptors.Count; ++Idx)
+			{
+				TargetDescriptor TestsTargetDescriptor = TargetDescriptors[Idx].Copy();
+				TestsTargetDescriptor.Name = TargetDescriptors[Idx].Name + "Tests";
+				TestsTargetDescriptor.IsTestsTarget = true;
+				TestTargetDescriptors.Add(TestsTargetDescriptor);
+			}
+
+			using (ISourceFileWorkingSet WorkingSet = SourceFileWorkingSet.Create(Unreal.RootDirectory, new HashSet<DirectoryReference>()))
+			{
+				BuildMode.Build(TestTargetDescriptors, BuildConfiguration, WorkingSet, BuildOptions.None, null);
+			}
+		}
+	}
+}
+
