UBT: Add minimal C++ coroutine support

diff --git a/Engine/Source/Programs/UnrealBuildTool/Configuration/TargetRules.cs b/Engine/Source/Programs/UnrealBuildTool/Configuration/TargetRules.cs
--- a/Engine/Source/Programs/UnrealBuildTool/Configuration/TargetRules.cs
+++ b/Engine/Source/Programs/UnrealBuildTool/Configuration/TargetRules.cs
@@ -1204,6 +1204,13 @@ namespace UnrealBuildTool
 		/// </summary>
 		public bool bEnableCppModules = false;
 
+		/// <summary>
+		/// Whether to enable support for C++20 coroutines
+		/// This option is provided to facilitate evaluation of the feature.
+		/// Expect the name of the option to change. Use of coroutines in with UE is untested and unsupported.
+		/// </summary>
+		public bool bEnableCppCoroutinesForEvaluation = false;
+
 		/// <summary>
 		/// If true, then enable memory profiling in the build (defines USE_MALLOC_PROFILER=1 and forces bOmitFramePointers=false).
 		/// </summary>
@@ -2659,6 +2666,11 @@ namespace UnrealBuildTool
 			get { return Inner.bEnableCppModules; }
 		}
 
+		public bool bEnableCppCoroutinesForEvaluation
+		{
+			get { return Inner.bEnableCppCoroutinesForEvaluation; }
+		}
+
 		public bool bUseMallocProfiler
 		{
 			get { return Inner.bUseMallocProfiler; }
diff --git a/Engine/Source/Programs/UnrealBuildTool/Configuration/UEBuildTarget.cs b/Engine/Source/Programs/UnrealBuildTool/Configuration/UEBuildTarget.cs
--- a/Engine/Source/Programs/UnrealBuildTool/Configuration/UEBuildTarget.cs
+++ b/Engine/Source/Programs/UnrealBuildTool/Configuration/UEBuildTarget.cs
@@ -3949,6 +3949,17 @@ namespace UnrealBuildTool
 				GlobalCompileEnvironment.Definitions.Add("WITH_CPP_MODULES=0");
 			}
 
+			// Whether C++20 coroutines are enabled
+			if (Rules.bEnableCppCoroutinesForEvaluation)
+			{
+				GlobalCompileEnvironment.Definitions.Add("WITH_CPP_COROUTINES=1");
+				GlobalCompileEnvironment.bEnableCoroutines = true;
+			}
+			else
+			{
+				GlobalCompileEnvironment.Definitions.Add("WITH_CPP_COROUTINES=0");
+			}
+
 			// Define the custom config if specified, this should not be set in an editor build
 			if (!bUseSharedBuildEnvironment && !String.IsNullOrEmpty(Rules.CustomConfig))
 			{
diff --git a/Engine/Source/Programs/UnrealBuildTool/Platform/Android/AndroidToolChain.cs b/Engine/Source/Programs/UnrealBuildTool/Platform/Android/AndroidToolChain.cs
--- a/Engine/Source/Programs/UnrealBuildTool/Platform/Android/AndroidToolChain.cs
+++ b/Engine/Source/Programs/UnrealBuildTool/Platform/Android/AndroidToolChain.cs
@@ -845,6 +845,11 @@ namespace UnrealBuildTool
 			Result += " -x c++";
 			Result += GetCppStandardCompileArgument(CompileEnvironment.CppStandard);
 
+			if (CompileEnvironment.bEnableCoroutines)
+			{
+				Result += " -fcoroutines-ts";
+			}
+
 			return Result;
 		}
 
@@ -861,6 +866,11 @@ namespace UnrealBuildTool
 			Result += " -x c++-header";
 			Result += GetCppStandardCompileArgument(CompileEnvironment.CppStandard);
 
+			if (CompileEnvironment.bEnableCoroutines)
+			{
+				Result += " -fcoroutines-ts";
+			}
+
 			return Result;
 		}
 
diff --git a/Engine/Source/Programs/UnrealBuildTool/Platform/IOS/IOSToolChain.cs b/Engine/Source/Programs/UnrealBuildTool/Platform/IOS/IOSToolChain.cs
--- a/Engine/Source/Programs/UnrealBuildTool/Platform/IOS/IOSToolChain.cs
+++ b/Engine/Source/Programs/UnrealBuildTool/Platform/IOS/IOSToolChain.cs
@@ -376,6 +376,12 @@ namespace UnrealBuildTool
 			Result += " -x objective-c++";
 			Result += GetCppStandardCompileArgument(CompileEnvironment.CppStandard);
 			Result += " -stdlib=libc++";
+
+			if (CompileEnvironment.bEnableCoroutines)
+			{
+				Result += " -fcoroutines-ts";
+			}
+
 			return Result;
 		}
 
@@ -410,6 +416,12 @@ namespace UnrealBuildTool
 			Result += " -x objective-c++-header";
 			Result += GetCppStandardCompileArgument(CompileEnvironment.CppStandard);
 			Result += " -stdlib=libc++";
+
+			if (CompileEnvironment.bEnableCoroutines)
+			{
+				Result += " -fcoroutines-ts";
+			}
+
 			return Result;
 		}
 
diff --git a/Engine/Source/Programs/UnrealBuildTool/Platform/Linux/LinuxToolChain.cs b/Engine/Source/Programs/UnrealBuildTool/Platform/Linux/LinuxToolChain.cs
--- a/Engine/Source/Programs/UnrealBuildTool/Platform/Linux/LinuxToolChain.cs
+++ b/Engine/Source/Programs/UnrealBuildTool/Platform/Linux/LinuxToolChain.cs
@@ -831,6 +831,12 @@ namespace UnrealBuildTool
 			string Result = "";
 			Result += " -x c++";
 			Result += GetCompilerStandardVersion_CPP(CompileEnvironment.CppStandard);
+
+			if (CompileEnvironment.bEnableCoroutines)
+			{
+				Result += " -fcoroutines-ts";
+			}
+
 			return Result;
 		}
 
@@ -885,6 +891,12 @@ namespace UnrealBuildTool
 			string Result = "";
 			Result += " -x c++-header";
 			Result += GetCompilerStandardVersion_CPP(CompileEnvironment.CppStandard);
+
+			if (CompileEnvironment.bEnableCoroutines)
+			{
+				Result += " -fcoroutines-ts";
+			}
+
 			return Result;
 		}
 
diff --git a/Engine/Source/Programs/UnrealBuildTool/Platform/Mac/MacToolChain.cs b/Engine/Source/Programs/UnrealBuildTool/Platform/Mac/MacToolChain.cs
--- a/Engine/Source/Programs/UnrealBuildTool/Platform/Mac/MacToolChain.cs
+++ b/Engine/Source/Programs/UnrealBuildTool/Platform/Mac/MacToolChain.cs
@@ -303,6 +303,12 @@ namespace UnrealBuildTool
 			Result += " -x objective-c++";
 			Result += GetCppStandardCompileArgument(CompileEnvironment.CppStandard);
 			Result += " -stdlib=libc++";
+
+			if (CompileEnvironment.bEnableCoroutines)
+			{
+				Result += " -fcoroutines-ts";
+			}
+
 			return Result;
 		}
 
@@ -336,6 +342,12 @@ namespace UnrealBuildTool
 			Result += " -x objective-c++-header";
 			Result += GetCppStandardCompileArgument(CompileEnvironment.CppStandard);
 			Result += " -stdlib=libc++";
+
+			if (CompileEnvironment.bEnableCoroutines)
+			{
+				Result += " -fcoroutines-ts";
+			}
+
 			return Result;
 		}
 
diff --git a/Engine/Source/Programs/UnrealBuildTool/Platform/Windows/VCToolChain.cs b/Engine/Source/Programs/UnrealBuildTool/Platform/Windows/VCToolChain.cs
--- a/Engine/Source/Programs/UnrealBuildTool/Platform/Windows/VCToolChain.cs
+++ b/Engine/Source/Programs/UnrealBuildTool/Platform/Windows/VCToolChain.cs
@@ -754,7 +754,6 @@ namespace UnrealBuildTool
 				Arguments.Add("/WX");
 			}
 
-			// Note: Should be shared with all Windows-like platforms.
 			switch (CompileEnvironment.CppStandard)
 			{
 				case CppStandardVersion.Cpp14:
@@ -785,6 +784,18 @@ namespace UnrealBuildTool
 					throw new BuildException($"Unsupported C++ standard type set: {CompileEnvironment.CppStandard}");
 			}
 
+			if (CompileEnvironment.bEnableCoroutines)
+			{
+				if (Target.WindowsPlatform.Compiler.IsMSVC())
+				{
+					Arguments.Add("/await:strict");
+				}
+				else
+				{
+					Arguments.Add("-fcoroutines-ts");
+				}
+			}
+
 			if (Target.WindowsPlatform.Compiler.IsClang())
 			{
 				// Enable codeview ghash for faster lld links
@@ -877,7 +888,6 @@ namespace UnrealBuildTool
 				{
 					Arguments.Add("-Wno-ordered-compare-function-pointers");
 				}
-
 			}
 		}
 
diff --git a/Engine/Source/Programs/UnrealBuildTool/System/CppCompileEnvironment.cs b/Engine/Source/Programs/UnrealBuildTool/System/CppCompileEnvironment.cs
--- a/Engine/Source/Programs/UnrealBuildTool/System/CppCompileEnvironment.cs
+++ b/Engine/Source/Programs/UnrealBuildTool/System/CppCompileEnvironment.cs
@@ -403,6 +403,13 @@ namespace UnrealBuildTool
 		/// </summary>
 		public int AnalyzeStackSizeWarning = 300000;
 
+		/// <summary>
+		/// Enable C++ coroutine support. 
+		/// For MSVC, adds "/await:strict" to the command line. Program should #include &lt;coroutine&gt;
+		/// For Clang, adds "-fcoroutines-ts" to the command line. Program should #include &lt;experimental/coroutine&gt; (not supported in every clang toolchain)
+		/// </summary>
+		public bool bEnableCoroutines = false;
+
 		/// <summary>
 		/// Default constructor.
 		/// </summary>
@@ -484,6 +491,7 @@ namespace UnrealBuildTool
 			bHackHeaderGenerator = Other.bHackHeaderGenerator;
 			bHideSymbolsByDefault = Other.bHideSymbolsByDefault;
 			CppStandard = Other.CppStandard;
+			bEnableCoroutines = Other.bEnableCoroutines;
 		}
 	}
 }
