UE-136919 //UE5/Main - Compile UnrealGame HoloLens - No required compiler toolchain found
-The build machine was trying to use a locally installed toolchain that did not have the arm64 toolchain installed.
-We now identify toolchains with arm64 support and select from only those for arm64 builds.

diff --git a/Engine/Source/Programs/UnrealBuildTool/Executors/Experimental/FASTBuild.cs b/Engine/Source/Programs/UnrealBuildTool/Executors/Experimental/FASTBuild.cs
--- a/Engine/Source/Programs/UnrealBuildTool/Executors/Experimental/FASTBuild.cs
+++ b/Engine/Source/Programs/UnrealBuildTool/Executors/Experimental/FASTBuild.cs
@@ -826,7 +826,7 @@ namespace UnrealBuildTool
 				// it probably means we are building for another platform.
                 if(BuildType == FBBuildType.Windows)
                 {
-					VCEnv = VCEnvironment.Create(WindowsPlatform.GetDefaultCompiler(null), UnrealTargetPlatform.Win64, WindowsArchitecture.x64, null, null, null);
+					VCEnv = VCEnvironment.Create(WindowsPlatform.GetDefaultCompiler(null, WindowsArchitecture.x64), UnrealTargetPlatform.Win64, WindowsArchitecture.x64, null, null, null);
 				}
             }
 			catch (Exception)
diff --git a/Engine/Source/Programs/UnrealBuildTool/Platform/HoloLens/HoloLensProjectGenerator.cs b/Engine/Source/Programs/UnrealBuildTool/Platform/HoloLens/HoloLensProjectGenerator.cs
--- a/Engine/Source/Programs/UnrealBuildTool/Platform/HoloLens/HoloLensProjectGenerator.cs
+++ b/Engine/Source/Programs/UnrealBuildTool/Platform/HoloLens/HoloLensProjectGenerator.cs
@@ -78,8 +78,8 @@ namespace UnrealBuildTool
 				VCProjectFileContent.Append("		<WindowsTargetPlatformVersion>" + MaxTestedVersion + "</WindowsTargetPlatformVersion>" + ProjectFileGenerator.NewLine);
 			}
 
-			WindowsCompiler Compiler = WindowsPlatform.GetDefaultCompiler(TargetFilePath);
-			DirectoryReference? PlatformWinMDLocation = HoloLensPlatform.GetCppCXMetadataLocation(Compiler, "Latest");
+			WindowsCompiler Compiler = WindowsPlatform.GetDefaultCompiler(TargetFilePath, WindowsArchitecture.x64);
+			DirectoryReference? PlatformWinMDLocation = HoloLensPlatform.GetCppCXMetadataLocation(Compiler, "Latest", WindowsArchitecture.x64);
 			if (PlatformWinMDLocation == null || !FileReference.Exists(FileReference.Combine(PlatformWinMDLocation, "platform.winmd")))
 			{
 				throw new BuildException("Unable to find platform.winmd for {0} toolchain; '{1}' is an invalid version", WindowsPlatform.GetCompilerName(Compiler), "Latest");
diff --git a/Engine/Source/Programs/UnrealBuildTool/Platform/HoloLens/UEBuildHoloLens.cs b/Engine/Source/Programs/UnrealBuildTool/Platform/HoloLens/UEBuildHoloLens.cs
--- a/Engine/Source/Programs/UnrealBuildTool/Platform/HoloLens/UEBuildHoloLens.cs
+++ b/Engine/Source/Programs/UnrealBuildTool/Platform/HoloLens/UEBuildHoloLens.cs
@@ -192,11 +192,11 @@ namespace UnrealBuildTool
 				}
 				else
 				{
-					Target.HoloLensPlatform.Compiler = WindowsPlatform.GetDefaultCompiler(Target.ProjectFile);
+					Target.HoloLensPlatform.Compiler = WindowsPlatform.GetDefaultCompiler(Target.ProjectFile, Target.HoloLensPlatform.Architecture);
 				}
 			}
 
-			if(!Target.bGenerateProjectFiles)
+			if (!Target.bGenerateProjectFiles)
 			{
 				Log.TraceInformationOnce("Using {0} architecture for deploying to HoloLens device", Target.HoloLensPlatform.Architecture);
 			}
@@ -240,7 +240,7 @@ namespace UnrealBuildTool
 
 			if(!Target.bGenerateProjectFiles)
 			{
-				Log.TraceInformationOnce("Building using Windows SDK version {0} for HoloLens", Target.HoloLensPlatform.Win10SDKVersion);
+				Log.TraceInformationOnce("Building using Windows SDK version {0} for HoloLens and compiler {1} {2}", Target.HoloLensPlatform.Win10SDKVersion, Target.HoloLensPlatform.Compiler, Target.WindowsPlatform.CompilerVersion);
 
 				if (Target.HoloLensPlatform.Win10SDKVersion < MinimumSDKVersionRecommended)
 				{
@@ -282,12 +282,12 @@ namespace UnrealBuildTool
 			return base.GetBinaryExtension(InBinaryType);
 		}
 
-		internal static DirectoryReference? GetCppCXMetadataLocation(WindowsCompiler Compiler, string CompilerVersion)
+		internal static DirectoryReference? GetCppCXMetadataLocation(WindowsCompiler Compiler, string CompilerVersion, WindowsArchitecture Architecture)
 		{
 			VersionNumber? SelectedToolChainVersion;
 			DirectoryReference? SelectedToolChainDir;
 			DirectoryReference? SelectedRedistDir;
-			if (!WindowsPlatform.TryGetToolChainDir(Compiler, CompilerVersion, out SelectedToolChainVersion, out SelectedToolChainDir, out SelectedRedistDir))
+			if (!WindowsPlatform.TryGetToolChainDir(Compiler, CompilerVersion, Architecture, out SelectedToolChainVersion, out SelectedToolChainDir, out SelectedRedistDir))
 			{
 				return null;
 			}
diff --git a/Engine/Source/Programs/UnrealBuildTool/Platform/Windows/UEBuildWindows.cs b/Engine/Source/Programs/UnrealBuildTool/Platform/Windows/UEBuildWindows.cs
--- a/Engine/Source/Programs/UnrealBuildTool/Platform/Windows/UEBuildWindows.cs
+++ b/Engine/Source/Programs/UnrealBuildTool/Platform/Windows/UEBuildWindows.cs
@@ -781,6 +781,11 @@ namespace UnrealBuildTool
 			/// </summary>
 			public bool IsPreview { get; }
 
+			/// <summary>
+			/// The architecture of this ToolChainInstallation (multiple ToolChainInstallation instances may be created, one per architecture).
+			/// </summary>
+			public WindowsArchitecture Architecture { get; }
+
 			/// <summary>
 			/// Reason for this toolchain not being compatible
 			/// </summary>
@@ -804,16 +809,18 @@ namespace UnrealBuildTool
 			/// <param name="Version"></param>
 			/// <param name="Is64Bit"></param>
 			/// <param name="IsPreview">Whether it's a pre-release version of the toolchain</param>
+			/// <param name="Architecture"></param>
 			/// <param name="Error"></param>
 			/// <param name="BaseDir">Base directory for the toolchain</param>
 			/// <param name="RedistDir">Optional directory for redistributable components (DLLs etc)</param>
-			public ToolChainInstallation(VersionNumber Family, int FamilyRank, VersionNumber Version, bool Is64Bit, bool IsPreview, string? Error, DirectoryReference BaseDir, DirectoryReference? RedistDir)
+			public ToolChainInstallation(VersionNumber Family, int FamilyRank, VersionNumber Version, bool Is64Bit, bool IsPreview, WindowsArchitecture Architecture, string? Error, DirectoryReference BaseDir, DirectoryReference? RedistDir)
 			{
 				this.Family = Family;
 				this.FamilyRank = FamilyRank;
 				this.Version = Version;
 				this.Is64Bit = Is64Bit;
 				this.IsPreview = IsPreview;
+				this.Architecture = Architecture;
 				this.Error = Error;
 				this.BaseDir = BaseDir;
 				this.RedistDir = RedistDir;
@@ -922,7 +929,6 @@ namespace UnrealBuildTool
 			if (Platform == UnrealTargetPlatform.HoloLens && Target.Architecture.ToLower() == "arm64")
 			{
 				Target.WindowsPlatform.Architecture = WindowsArchitecture.ARM64;
-				Log.TraceInformation("Using Windows ARM64 architecture");
 			}
 			else if (Platform == UnrealTargetPlatform.Win64)
 			{
@@ -951,17 +957,17 @@ namespace UnrealBuildTool
 			// Set the compiler version if necessary
 			if (Target.WindowsPlatform.Compiler == WindowsCompiler.Default)
 			{
-				if (Target.WindowsPlatform.StaticAnalyzer == WindowsStaticAnalyzer.PVSStudio && HasCompiler(WindowsCompiler.VisualStudio2019))
+				if (Target.WindowsPlatform.StaticAnalyzer == WindowsStaticAnalyzer.PVSStudio && HasCompiler(WindowsCompiler.VisualStudio2019, Target.WindowsPlatform.Architecture))
 				{
 					Target.WindowsPlatform.Compiler = WindowsCompiler.VisualStudio2019;
 				}
-				if (Target.WindowsPlatform.StaticAnalyzer == WindowsStaticAnalyzer.PVSStudio && HasCompiler(WindowsCompiler.VisualStudio2022))
+				if (Target.WindowsPlatform.StaticAnalyzer == WindowsStaticAnalyzer.PVSStudio && HasCompiler(WindowsCompiler.VisualStudio2022, Target.WindowsPlatform.Architecture))
 				{
 					Target.WindowsPlatform.Compiler = WindowsCompiler.VisualStudio2022;
 				}
 				else
 				{
-					Target.WindowsPlatform.Compiler = GetDefaultCompiler(Target.ProjectFile);
+					Target.WindowsPlatform.Compiler = GetDefaultCompiler(Target.ProjectFile, Target.WindowsPlatform.Architecture);
 				}
 			}
 
@@ -1043,10 +1049,10 @@ namespace UnrealBuildTool
 		/// Gets the default compiler which should be used, if it's not set explicitly by the target, command line, or config file.
 		/// </summary>
 		/// <returns>The default compiler version</returns>
-		internal static WindowsCompiler GetDefaultCompiler(FileReference? ProjectFile)
+		internal static WindowsCompiler GetDefaultCompiler(FileReference? ProjectFile, WindowsArchitecture Architecture)
 		{
 			// If there's no specific compiler set, try to pick the matching compiler for the selected IDE
-			if(ProjectFileGeneratorSettings.Format != null)
+			if (ProjectFileGeneratorSettings.Format != null)
 			{
 				foreach(ProjectFileFormat Format in ProjectFileGeneratorSettings.ParseFormatList(ProjectFileGeneratorSettings.Format))
 				{
@@ -1090,11 +1096,11 @@ namespace UnrealBuildTool
 			}
 
 			// Second, default based on what's installed, test for 2019 first
-			if (FindToolChainInstallations(WindowsCompiler.VisualStudio2019).Where(x => x.Version >= MinimumVisualCppVersion).Count() > 0)
+			if (FindToolChainInstallations(WindowsCompiler.VisualStudio2019).Where(x => x.Version >= MinimumVisualCppVersion && x.Architecture == Architecture).Count() > 0)
 			{
 				return WindowsCompiler.VisualStudio2019;
 			}
-			else if (FindToolChainInstallations(WindowsCompiler.VisualStudio2022).Where(x => x.Version >= MinimumVisualCppVersion).Count() > 0)
+			else if (FindToolChainInstallations(WindowsCompiler.VisualStudio2022).Where(x => x.Version >= MinimumVisualCppVersion && x.Architecture == Architecture).Count() > 0)
 			{
 				return WindowsCompiler.VisualStudio2022;
 			}
@@ -1102,10 +1108,16 @@ namespace UnrealBuildTool
 			// If we do have a Visual Studio installation, but we're missing just the C++ parts, warn about that.
 			if (TryGetVSInstallDirs(WindowsCompiler.VisualStudio2019) != null)
 			{
-				Log.TraceWarning("Visual Studio 2019 is installed, but is missing the C++ toolchain. Please verify that the \"VC++ 2019 toolset\" component is selected in the Visual Studio 2019 installation options.");
+				string ToolSetWarning = Architecture == WindowsArchitecture.ARM64 ?
+					"Please verify that the \"VC++ 2019 toolset\" component is selected in the Visual Studio 2019 installation options." :
+					"Please verify that the \"VS 2019 C++ ARM64 build tools\" component is selected in the Visual Studio 2019 installation options.";
+				Log.TraceWarning("Visual Studio 2019 is installed, but is missing the C++ toolchain. {0}", ToolSetWarning);
 			}
 			else if (TryGetVSInstallDirs(WindowsCompiler.VisualStudio2022) != null)
 			{
+				string ToolSetWarning = Architecture == WindowsArchitecture.ARM64 ?
+					"Please verify that the \"VS 2022 C++ x64/x86 build tools (Latest)\" component is selected in the Visual Studio 2022 installation options." :
+					"Please verify that the \"VS 2022 C++ ARM64 build tools\" component is selected in the Visual Studio 2022 installation options.";
 				Log.TraceWarning("Visual Studio 2022 is installed, but is missing the C++ toolchain. Please verify that the \"VS 2022 C++ x64/x86 build tools (Latest)\" component is selected in the Visual Studio 2022 installation options.");
 			}
 			else
@@ -1450,7 +1462,7 @@ namespace UnrealBuildTool
 				}
 
 				Log.TraceLog("Found Clang toolchain: {0} (Version={1}, Is64Bit={2}, Rank={3})", ToolChainDir, Version, Is64Bit, Rank);
-				ToolChains.Add(new ToolChainInstallation(Version, Rank, Version, Is64Bit, false, null, ToolChainDir, null));
+				ToolChains.Add(new ToolChainInstallation(Version, Rank, Version, Is64Bit, false, WindowsArchitecture.x64, null, ToolChainDir, null));
 			}
 		}
 
@@ -1477,7 +1489,7 @@ namespace UnrealBuildTool
 				}
 
 				Log.TraceLog("Found Intel OneAPI toolchain: {0} (Version={1}, Is64Bit={2}, Rank={3})", ToolChainDir, Version, Is64Bit, Rank);
-				ToolChains.Add(new ToolChainInstallation(Version, Rank, Version, Is64Bit, false, Error, ToolChainDir, null));
+				ToolChains.Add(new ToolChainInstallation(Version, Rank, Version, Is64Bit, false, WindowsArchitecture.x64, Error, ToolChainDir, null));
 			}
 		}
 
@@ -1549,8 +1561,15 @@ namespace UnrealBuildTool
 				Error = $"UnrealBuildTool requires at minimum the MSVC {MinimumVisualCppVersion} toolchain. Please install a later toolchain from the Visual Studio installer.";
 			}
 
-			Log.TraceLog("Found Visual Studio toolchain: {0} (Family={1}, FamilyRank={2}, Version={3}, Is64Bit={4}, Preview={5}, Error={6}, Redist={7})", ToolChainDir, Family, FamilyRank, Version, Is64Bit, bPreview, Error != null, RedistDir);
-			ToolChains.Add(new ToolChainInstallation(Family, FamilyRank, Version, Is64Bit, bPreview, Error, ToolChainDir, RedistDir));
+			Log.TraceLog("Found Visual Studio toolchain: {0} (Family={1}, FamilyRank={2}, Version={3}, Is64Bit={4}, Preview={5}, Architecture={6}, Error={7}, Redist={8})", ToolChainDir, Family, FamilyRank, Version, Is64Bit, bPreview, WindowsArchitecture.x64.ToString(), Error != null, RedistDir);
+			ToolChains.Add(new ToolChainInstallation(Family, FamilyRank, Version, Is64Bit, bPreview, WindowsArchitecture.x64, Error, ToolChainDir, RedistDir));
+
+			bool HasArm64 = HasArm64ToolChain(ToolChainDir);
+			if (HasArm64)
+			{
+				Log.TraceLog("Found Visual Studio toolchain: {0} (Family={1}, FamilyRank={2}, Version={3}, Is64Bit={4}, Preview={5}, Architecture={6}, Error={7}, Redist={8})", ToolChainDir, Family, FamilyRank, Version, Is64Bit, bPreview, WindowsArchitecture.ARM64.ToString(), Error != null, RedistDir);
+				ToolChains.Add(new ToolChainInstallation(Family, FamilyRank, Version, Is64Bit, bPreview, WindowsArchitecture.ARM64, Error, ToolChainDir, RedistDir));
+			}
 		}
 
 		/// <summary>
@@ -1602,6 +1621,16 @@ namespace UnrealBuildTool
 			return FileReference.Exists(FileReference.Combine(ToolChainDir, "bin", "amd64", "cl.exe")) || FileReference.Exists(FileReference.Combine(ToolChainDir, "bin", "Hostx64", "x64", "cl.exe"));
 		}
 
+		/// <summary>
+		/// Checks if the given directory contains a arm64 toolchain.  Used to require arm64, which is an optional install item, when that is our target architecture.
+		/// </summary>
+		/// <param name="ToolChainDir">Directory to check</param>
+		/// <returns>True if the given directory contains the arm64 toolchain</returns>
+		static bool HasArm64ToolChain(DirectoryReference ToolChainDir)
+		{
+			return FileReference.Exists(FileReference.Combine(ToolChainDir, "bin", "Hostx64", "arm64", "cl.exe"));
+		}
+
 		/// <summary>
 		/// Determines if an IDE for the given compiler is installed.
 		/// </summary>
@@ -1616,21 +1645,24 @@ namespace UnrealBuildTool
 		/// Determines if a given compiler is installed
 		/// </summary>
 		/// <param name="Compiler">Compiler to check for</param>
+		/// <param name="Architecture">Architecture the compiler must support</param>
 		/// <returns>True if the given compiler is installed</returns>
-		public static bool HasCompiler(WindowsCompiler Compiler)
+		public static bool HasCompiler(WindowsCompiler Compiler, WindowsArchitecture Architecture)
 		{
-			return FindToolChainInstallations(Compiler).Count > 0;
+			return FindToolChainInstallations(Compiler).Where(x => (x.Architecture == Architecture)).Count() > 0;
 		}
 
 		/// <summary>
 		/// Select which toolchain to use, combining a custom preference with a default sort order
 		/// </summary>
 		/// <param name="ToolChains"></param>
-		/// <param name="Preference"></param>
+		/// <param name="Preference">Ordering function</param>
+		/// <param name="Architecture">Architecture that must be supported</param>
 		/// <returns></returns>
-		static ToolChainInstallation? SelectToolChain(IEnumerable<ToolChainInstallation> ToolChains, Func<IOrderedEnumerable<ToolChainInstallation>, IOrderedEnumerable<ToolChainInstallation>> Preference)
+		static ToolChainInstallation? SelectToolChain(IEnumerable<ToolChainInstallation> ToolChains, Func<IOrderedEnumerable<ToolChainInstallation>, IOrderedEnumerable<ToolChainInstallation>> Preference, WindowsArchitecture Architecture)
 		{
-			ToolChainInstallation? ToolChain = Preference(ToolChains.OrderByDescending(x => x.Error == null))
+			ToolChainInstallation? ToolChain = Preference(ToolChains.Where(x => x.Architecture == Architecture)
+				.OrderByDescending(x => x.Error == null))
 				.ThenByDescending(x => x.Is64Bit)
 				.ThenBy(x => x.IsPreview)
 				.ThenBy(x => x.FamilyRank)
@@ -1650,11 +1682,12 @@ namespace UnrealBuildTool
 		/// </summary>
 		/// <param name="Compiler">Major version of the compiler to use</param>
 		/// <param name="CompilerVersion">The minimum compiler version to use</param>
+		/// <param name="Architecture">Architecture that is required</param>
 		/// <param name="OutToolChainVersion">Receives the chosen toolchain version</param>
 		/// <param name="OutToolChainDir">Receives the directory containing the toolchain</param>
 		/// <param name="OutRedistDir">Receives the optional directory containing redistributable components</param>
 		/// <returns>True if the toolchain directory was found correctly</returns>
-		public static bool TryGetToolChainDir(WindowsCompiler Compiler, string? CompilerVersion, [NotNullWhen(true)] out VersionNumber? OutToolChainVersion, [NotNullWhen(true)] out DirectoryReference? OutToolChainDir, out DirectoryReference? OutRedistDir)
+		public static bool TryGetToolChainDir(WindowsCompiler Compiler, string? CompilerVersion, WindowsArchitecture Architecture, [NotNullWhen(true)] out VersionNumber? OutToolChainVersion, [NotNullWhen(true)] out DirectoryReference? OutToolChainDir, out DirectoryReference? OutRedistDir)
 		{
 			// Find all the installed toolchains
 			List<ToolChainInstallation> ToolChains = FindToolChainInstallations(Compiler);
@@ -1663,7 +1696,7 @@ namespace UnrealBuildTool
 			ToolChainInstallation? ToolChain = null;
 			if (CompilerVersion == null)
 			{
-				ToolChain = SelectToolChain(ToolChains, x => x);
+				ToolChain = SelectToolChain(ToolChains, x => x, Architecture);
 				if (ToolChain == null)
 				{
 					OutToolChainVersion = null;
@@ -1674,26 +1707,26 @@ namespace UnrealBuildTool
 			}
 			else if (String.Compare(CompilerVersion, "Latest", StringComparison.InvariantCultureIgnoreCase) == 0)
 			{
-				ToolChain = SelectToolChain(ToolChains, x => x.ThenByDescending(x => x.Version));
+				ToolChain = SelectToolChain(ToolChains, x => x.ThenByDescending(x => x.Version), Architecture);
 				if (ToolChain == null)
 				{
-					throw new BuildException("Unable to find C++ toolchain for {0}", Compiler);
+					throw new BuildException("Unable to find C++ toolchain for {0} {1}", Compiler, Architecture.ToString());
 				}
 			}
 			else if (String.Compare(CompilerVersion, "Preview", StringComparison.InvariantCultureIgnoreCase) == 0)
 			{
-				ToolChain = SelectToolChain(ToolChains, x => x.ThenByDescending(x => x.IsPreview));
+				ToolChain = SelectToolChain(ToolChains, x => x.ThenByDescending(x => x.IsPreview), Architecture);
 				if (ToolChain == null || !ToolChain.IsPreview)
 				{
-					throw new BuildException("Unable to find preview toolchain for {0}", Compiler);
+					throw new BuildException("Unable to find preview toolchain for {0} {1}", Compiler, Architecture.ToString());
 				}
 			}
 			else if (VersionNumber.TryParse(CompilerVersion, out VersionNumber? ToolChainVersion))
 			{
-				ToolChain = SelectToolChain(ToolChains, x => x.ThenByDescending(x => x.Version == ToolChainVersion).ThenByDescending(x => x.Family == ToolChainVersion));
+				ToolChain = SelectToolChain(ToolChains, x => x.ThenByDescending(x => x.Version == ToolChainVersion).ThenByDescending(x => x.Family == ToolChainVersion), Architecture);
 				if (ToolChain == null)
 				{
-					throw new BuildException("Unable to find {0} toolchain for {1}", ToolChainVersion, Compiler);
+					throw new BuildException("Unable to find {0} toolchain for {1} {2}", ToolChainVersion, Compiler, Architecture.ToString());
 				}
 			}
 			else
diff --git a/Engine/Source/Programs/UnrealBuildTool/Platform/Windows/VCEnvironment.cs b/Engine/Source/Programs/UnrealBuildTool/Platform/Windows/VCEnvironment.cs
--- a/Engine/Source/Programs/UnrealBuildTool/Platform/Windows/VCEnvironment.cs
+++ b/Engine/Source/Programs/UnrealBuildTool/Platform/Windows/VCEnvironment.cs
@@ -458,9 +458,9 @@ namespace UnrealBuildTool
 			VersionNumber? SelectedCompilerVersion;
 			DirectoryReference? SelectedCompilerDir;
 			DirectoryReference? SelectedRedistDir;
-			if(!WindowsPlatform.TryGetToolChainDir(Compiler, CompilerVersion, out SelectedCompilerVersion, out SelectedCompilerDir, out SelectedRedistDir))
+			if(!WindowsPlatform.TryGetToolChainDir(Compiler, CompilerVersion, Architecture, out SelectedCompilerVersion, out SelectedCompilerDir, out SelectedRedistDir))
 			{
-				throw new BuildException("{0}{1} must be installed in order to build this target.", WindowsPlatform.GetCompilerName(Compiler), String.IsNullOrEmpty(CompilerVersion)? "" : String.Format(" ({0})", CompilerVersion));
+				throw new BuildException("{0}{1} {2} must be installed in order to build this target.", WindowsPlatform.GetCompilerName(Compiler), String.IsNullOrEmpty(CompilerVersion)? "" : String.Format(" ({0})", CompilerVersion, Architecture.ToString()));
 			}
 
 			// Get the toolchain info
@@ -469,11 +469,11 @@ namespace UnrealBuildTool
 			DirectoryReference? SelectedToolChainDir;
 			if(Compiler.IsClang())
 			{
-				if (WindowsPlatform.TryGetToolChainDir(WindowsCompiler.VisualStudio2019, null, out SelectedToolChainVersion, out SelectedToolChainDir, out SelectedRedistDir))
+				if (WindowsPlatform.TryGetToolChainDir(WindowsCompiler.VisualStudio2019, null, Architecture, out SelectedToolChainVersion, out SelectedToolChainDir, out SelectedRedistDir))
 				{
 					ToolChain = WindowsCompiler.VisualStudio2019;
 				}
-				else if (WindowsPlatform.TryGetToolChainDir(WindowsCompiler.VisualStudio2022, null, out SelectedToolChainVersion, out SelectedToolChainDir, out SelectedRedistDir))
+				else if (WindowsPlatform.TryGetToolChainDir(WindowsCompiler.VisualStudio2022, null, Architecture, out SelectedToolChainVersion, out SelectedToolChainDir, out SelectedRedistDir))
 				{
 					ToolChain = WindowsCompiler.VisualStudio2022;
 				}
diff --git a/Engine/Source/Programs/UnrealBuildTool/ProjectFiles/VisualStudio/VCProjectFileGenerator.cs b/Engine/Source/Programs/UnrealBuildTool/ProjectFiles/VisualStudio/VCProjectFileGenerator.cs
--- a/Engine/Source/Programs/UnrealBuildTool/ProjectFiles/VisualStudio/VCProjectFileGenerator.cs
+++ b/Engine/Source/Programs/UnrealBuildTool/ProjectFiles/VisualStudio/VCProjectFileGenerator.cs
@@ -258,7 +258,7 @@ namespace UnrealBuildTool
 			// If we have a non-default setting for visual studio, check the compiler exists. If not, revert to the default.
 			if (Settings.ProjectFileFormat == VCProjectFileFormat.VisualStudio2019)
 			{
-				if (!WindowsPlatform.HasCompiler(WindowsCompiler.VisualStudio2019))
+				if (!WindowsPlatform.HasCompiler(WindowsCompiler.VisualStudio2019, WindowsArchitecture.x64))
 				{
 					Log.TraceWarning("Visual Studio C++ 2019 installation not found - ignoring preferred project file format.");
 					Settings.ProjectFileFormat = VCProjectFileFormat.Default;
@@ -266,7 +266,7 @@ namespace UnrealBuildTool
 			}
 			else if (Settings.ProjectFileFormat == VCProjectFileFormat.VisualStudio2022)
 			{
-				if (!WindowsPlatform.HasCompiler(WindowsCompiler.VisualStudio2022))
+				if (!WindowsPlatform.HasCompiler(WindowsCompiler.VisualStudio2022, WindowsArchitecture.x64))
 				{
 					Log.TraceWarning("Visual Studio C++ 2022 installation not found - ignoring preferred project file format.");
 					Settings.ProjectFileFormat = VCProjectFileFormat.Default;
@@ -279,7 +279,7 @@ namespace UnrealBuildTool
 			if (Settings.ProjectFileFormat == VCProjectFileFormat.Default)
 			{
 				// Enumerate all the valid installations. This list is already sorted by preference.
-				List<VisualStudioInstallation> Installations = WindowsPlatform.VisualStudioInstallations.Where(x => WindowsPlatform.HasCompiler(x.Compiler)).ToList();
+				List<VisualStudioInstallation> Installations = WindowsPlatform.VisualStudioInstallations.Where(x => WindowsPlatform.HasCompiler(x.Compiler, WindowsArchitecture.x64)).ToList();
 
 				// Get the corresponding project file format
 				VCProjectFileFormat Format = VCProjectFileFormat.VisualStudio2019;
diff --git a/Engine/Source/Programs/UnrealBuildTool/ProjectFiles/VisualStudioCode/VSCodeProjectFileGenerator.cs b/Engine/Source/Programs/UnrealBuildTool/ProjectFiles/VisualStudioCode/VSCodeProjectFileGenerator.cs
--- a/Engine/Source/Programs/UnrealBuildTool/ProjectFiles/VisualStudioCode/VSCodeProjectFileGenerator.cs
+++ b/Engine/Source/Programs/UnrealBuildTool/ProjectFiles/VisualStudioCode/VSCodeProjectFileGenerator.cs
@@ -322,7 +322,7 @@ namespace UnrealBuildTool
 			FileReference? CompilerPath;
 			if (HostPlatform == UnrealTargetPlatform.Win64)
 			{
-				VCEnvironment Environment = VCEnvironment.Create(WindowsPlatform.GetDefaultCompiler(null), Target.Platform, Target.Rules.WindowsPlatform.Architecture, null, Target.Rules.WindowsPlatform.WindowsSdkVersion, null);
+				VCEnvironment Environment = VCEnvironment.Create(WindowsPlatform.GetDefaultCompiler(null, Target.Rules.WindowsPlatform.Architecture), Target.Platform, Target.Rules.WindowsPlatform.Architecture, null, Target.Rules.WindowsPlatform.WindowsSdkVersion, null);
 				CompilerPath = FileReference.FromString(Environment.CompilerPath.FullName);
 				UsingClang = false;
 			}
