Remove extraneous log message when generating dSYMs on Mac.

#rnx

diff --git a/Engine/Source/Programs/UnrealBuildTool/Platform/Mac/MacToolChain.cs b/Engine/Source/Programs/UnrealBuildTool/Platform/Mac/MacToolChain.cs
--- a/Engine/Source/Programs/UnrealBuildTool/Platform/Mac/MacToolChain.cs
+++ b/Engine/Source/Programs/UnrealBuildTool/Platform/Mac/MacToolChain.cs
@@ -1256,14 +1256,12 @@ namespace UnrealBuildTool
 				ArgumentString += string.Format("then ");
 					ArgumentString += string.Format("break; ");
 				ArgumentString += string.Format("else ");
-					ArgumentString += string.Format("echo \\\"Waiting for {0} before generating dSYM file.\\\"; ", MachOBinary.AbsolutePath);
 					ArgumentString += string.Format("sleep 1; ");
 				ArgumentString += string.Format("fi; ");
 			ArgumentString += string.Format("done; ");
 
 			ArgumentString += string.Format("if [ ! -f \\\"{1}\\\" ] || [ \\\"{0}\\\" -nt \\\"{1}\\\" ] ; ", MachOBinary.AbsolutePath, OutputFile.AbsolutePath);
 			ArgumentString += string.Format("then ");
-				ArgumentString += string.Format("echo \\\"Calling dsymutil for: {0}\\\"; ", OutputFile.AbsolutePath);
 				ArgumentString += string.Format("rm -rf \\\"{0}\\\"; ", OutputFile.AbsolutePath);
 				ArgumentString += string.Format(" \\\"{0}\\\" {3} -f \\\"{1}\\\" -o \\\"{2}\\\"; ",
 					DsymutilPath,
