UnrealBuildTool: Fix BuildException string format when the MSVC arm64 toolchain is needed and isn't installed

#rnx

diff --git a/Engine/Source/Programs/UnrealBuildTool/Platform/Windows/VCEnvironment.cs b/Engine/Source/Programs/UnrealBuildTool/Platform/Windows/VCEnvironment.cs
--- a/Engine/Source/Programs/UnrealBuildTool/Platform/Windows/VCEnvironment.cs
+++ b/Engine/Source/Programs/UnrealBuildTool/Platform/Windows/VCEnvironment.cs
@@ -460,7 +460,7 @@ namespace UnrealBuildTool
 			DirectoryReference? SelectedRedistDir;
 			if(!WindowsPlatform.TryGetToolChainDir(Compiler, CompilerVersion, Architecture, out SelectedCompilerVersion, out SelectedCompilerDir, out SelectedRedistDir))
 			{
-				throw new BuildException("{0}{1} {2} must be installed in order to build this target.", WindowsPlatform.GetCompilerName(Compiler), String.IsNullOrEmpty(CompilerVersion)? "" : String.Format(" ({0})", CompilerVersion, Architecture.ToString()));
+				throw new BuildException("{0}{1} {2} must be installed in order to build this target.", WindowsPlatform.GetCompilerName(Compiler), String.IsNullOrEmpty(CompilerVersion)? "" : String.Format(" ({0})", CompilerVersion), Architecture.ToString());
 			}
 
 			// Get the toolchain info
