UnrealBuildTool: Ensure name suffixes are added if the TargetName is replaced in UEBuildTarget (manually merge due to conflict)

#rnx

diff --git a/Engine/Source/Programs/UnrealBuildTool/Configuration/TargetRules.cs b/Engine/Source/Programs/UnrealBuildTool/Configuration/TargetRules.cs
--- a/Engine/Source/Programs/UnrealBuildTool/Configuration/TargetRules.cs
+++ b/Engine/Source/Programs/UnrealBuildTool/Configuration/TargetRules.cs
@@ -189,11 +189,20 @@ namespace UnrealBuildTool
 		{
 			get
 			{
-				if (!String.IsNullOrEmpty(NameOverride))
+				if (!string.IsNullOrEmpty(NameOverride))
 				{
+					if (NameSuffixes.Count > 0)
+					{
+						return $"{NameOverride}-{NameSuffix}";
+					}
 					return NameOverride;
 				}
 
+				if (NameSuffixes.Count > 0)
+				{
+					return $"{DefaultName}-{NameSuffix}";
+				}
+
 				return DefaultName;
 			}
 			set
@@ -205,14 +214,36 @@ namespace UnrealBuildTool
 		/// <summary>
 		/// If the Name of this target has been overriden
 		/// </summary>
-		public bool IsNameOverriden() { return !String.IsNullOrEmpty(NameOverride); }
+		public bool IsNameOverriden() { return !string.IsNullOrEmpty(NameOverride) || NameSuffixes.Count > 0; }
+
+		/// <summary>
+		/// Add an optional suffix to append to Name
+		/// </summary>
+		/// <param name="Suffix">The string to append</param>
+		public void AddNameSuffix(string Suffix) { NameSuffixes.Add(Suffix); }
+
+		/// <summary>
+		/// Optional suffix to append to Name
+		/// </summary>
+		public string? NameSuffix
+		{
+			get
+			{
+				return NameSuffixes.Count > 0 ? string.Join("-", NameSuffixes) : null;
+			}
+		}
+
+		private HashSet<string> NameSuffixes = new HashSet<string>();
 
 		/// <summary>
 		/// Override the name used for this target
 		/// </summary>
 		[CommandLine("-TargetNameOverride=")]
-		private string? NameOverride;
+		private string? NameOverride { get; set; }
 
+		/// <summary>
+		/// The default name for this target, 
+		/// </summary>
 		private readonly string DefaultName;
 
 		private TestTargetRules? TestTargetRules;
@@ -1991,6 +2022,11 @@ namespace UnrealBuildTool
 			get { return Inner.Name; }
 		}
 
+		public string? NameSuffix
+		{
+			get { return Inner.NameSuffix; }
+		}
+
 		public TestTargetRules TestTarget { get { return Inner.CreateOrGetTestTarget(); } }
 
 		public TargetRules TestedTarget {
diff --git a/Engine/Source/Programs/UnrealBuildTool/Configuration/UEBuildTarget.cs b/Engine/Source/Programs/UnrealBuildTool/Configuration/UEBuildTarget.cs
--- a/Engine/Source/Programs/UnrealBuildTool/Configuration/UEBuildTarget.cs
+++ b/Engine/Source/Programs/UnrealBuildTool/Configuration/UEBuildTarget.cs
@@ -1156,6 +1156,10 @@ namespace UnrealBuildTool
 				else
 				{
 					AppName = GetAppNameForTargetType(Rules.Type);
+					if (!string.IsNullOrEmpty(Rules.NameSuffix))
+					{
+						AppName = $"{AppName}-{Rules.NameSuffix}";
+					}
 				}
 			}
 
diff --git a/Engine/Source/Programs/UnrealBuildTool/Platform/Linux/UEBuildLinux.cs b/Engine/Source/Programs/UnrealBuildTool/Platform/Linux/UEBuildLinux.cs
--- a/Engine/Source/Programs/UnrealBuildTool/Platform/Linux/UEBuildLinux.cs
+++ b/Engine/Source/Programs/UnrealBuildTool/Platform/Linux/UEBuildLinux.cs
@@ -200,31 +200,27 @@ namespace UnrealBuildTool
 
 		public override void ValidateTarget(TargetRules Target)
 		{
-			if (!Target.IsNameOverriden())
+			HashSet<string> Sanitizers = new HashSet<string>();
+			if (Target.LinuxPlatform.bEnableAddressSanitizer)
 			{
-				string? SanitizerSuffix = null;
-
-				if (Target.LinuxPlatform.bEnableAddressSanitizer)
-				{
-					SanitizerSuffix = "ASan";
-				}
-				else if (Target.LinuxPlatform.bEnableThreadSanitizer)
-				{
-					SanitizerSuffix = "TSan";
-				}
-				else if (Target.LinuxPlatform.bEnableUndefinedBehaviorSanitizer)
-				{
-					SanitizerSuffix = "UBSan";
-				}
-				else if (Target.LinuxPlatform.bEnableMemorySanitizer)
-				{
-					SanitizerSuffix = "MSan";
-				}
+				Sanitizers.Add("A");
+			}
+			if (Target.LinuxPlatform.bEnableThreadSanitizer)
+			{
+				Sanitizers.Add("T");
+			}
+			if (Target.LinuxPlatform.bEnableUndefinedBehaviorSanitizer)
+			{
+				Sanitizers.Add("UB");
+			}
+			if (Target.LinuxPlatform.bEnableMemorySanitizer)
+			{
+				Sanitizers.Add("M");
+			}
 
-				if (!String.IsNullOrEmpty(SanitizerSuffix))
-				{
-					Target.Name = Target.Name + "-" + SanitizerSuffix;
-				}
+			if (Sanitizers.Count > 0)
+			{
+				Target.AddNameSuffix($"{string.Join(string.Empty, Sanitizers)}San");
 			}
 
 			if (Target.bAllowLTCG && Target.LinkType != TargetLinkType.Monolithic)
