MakePri is given only one default culture.
 - other cultures will be included in the pri automatically.
 - remove the "auto resource packs" option entirely because no plaforms support it now.
 - cleaned up some comments.

#rnx

diff --git a/Engine/Source/Programs/UnrealBuildTool/Platform/Windows/VCManifestGenerator.cs b/Engine/Source/Programs/UnrealBuildTool/Platform/Windows/VCManifestGenerator.cs
--- a/Engine/Source/Programs/UnrealBuildTool/Platform/Windows/VCManifestGenerator.cs
+++ b/Engine/Source/Programs/UnrealBuildTool/Platform/Windows/VCManifestGenerator.cs
@@ -22,7 +22,6 @@ namespace UnrealBuildTool
 
 		protected virtual string IniSection_PlatformTargetSettings { get { return string.Format( "/Script/{0}PlatformEditor.{0}TargetSettings", Platform.ToString() ); } }
 		protected virtual string IniSection_GeneralProjectSettings { get { return "/Script/EngineSettings.GeneralProjectSettings"; } }
-		protected virtual bool bPlatformSupportsPriResourcePacks { get { return false; } }
 
 		protected const string BuildResourceSubPath = "Resources";
 		protected const string EngineResourceSubPath = "DefaultImages";
@@ -868,39 +867,23 @@ namespace UnrealBuildTool
 			{
 				// Create resource index configuration
 				string ResourceConfigFile = Path.Combine(IntermediatePath, "priconfig.xml");
-				bool bEnableAutoResourcePacks = false;
-				EngineIni.GetBool(IniSection_PlatformTargetSettings, "bEnableAutoResourcePacks", out bEnableAutoResourcePacks);
-				bEnableAutoResourcePacks &= bPlatformSupportsPriResourcePacks;
+				RunMakePri("createconfig /cf \"" + ResourceConfigFile + "\" /dq " + DefaultCulture + " /o " + GetMakePriExtraCommandLine());
 
-				// If the game is not going to support language resource packs then merge the culture qualifiers.
-				if (bEnableAutoResourcePacks || CulturesToStage.Count <= 1)
-				{
-					RunMakePri("createconfig /cf \"" + ResourceConfigFile + "\" /dq " + DefaultCulture + " /o " + GetMakePriExtraCommandLine());
-				}
-				else
-				{
-					RunMakePri("createconfig /cf \"" + ResourceConfigFile + "\" /dq " + String.Join("_", CulturesToStage) + " /o " + GetMakePriExtraCommandLine());
-				}
-
-				// Modify configuration to restrict indexing to the Resources directory (saves time and space)
+				// Load the new resource index configuration
 				XmlDocument PriConfig = new XmlDocument();
 				PriConfig.Load(ResourceConfigFile);
 
-				// If the game is not going to support resource packs then remove the autoResourcePackages.
-				if (!bEnableAutoResourcePacks)
-				{
-					XmlNode PackagingNode = PriConfig.SelectSingleNode("/resources/packaging");
-					PackagingNode.ParentNode.RemoveChild(PackagingNode);
-				}
+				// remove the packaging node - we do not want to split the pri & only want one .pri file
+				XmlNode PackagingNode = PriConfig.SelectSingleNode("/resources/packaging");
+				PackagingNode.ParentNode.RemoveChild(PackagingNode);
 
-				// The previous implementation using startIndexAt="Resources" did not produce the expected ResourceMapSubtree hierarchy, so this manually specifies all resources in a .resfiles instead.
+				// all required resources are explicitly listed in resources.resfiles, rather than relying on makepri to discover them
 				string ResourcesResFile = Path.Combine(IntermediatePath, "resources.resfiles");
-
 				XmlNode PriIndexNode = PriConfig.SelectSingleNode("/resources/index");
 				XmlAttribute PriStartIndex = PriIndexNode.Attributes["startIndexAt"];
 				PriStartIndex.Value = ResourcesResFile;
 
-				// Swap the default folder indexer-config to a RESFILES indexer-config.
+				// swap the folder indexer-config to a RESFILES indexer-config.
 				XmlElement FolderIndexerConfigNode = (XmlElement)PriConfig.SelectSingleNode("/resources/index/indexer-config[@type='folder']");
 				FolderIndexerConfigNode.SetAttribute("type", "RESFILES");
 				FolderIndexerConfigNode.RemoveAttribute("foldernameAsQualifier");
@@ -908,6 +891,7 @@ namespace UnrealBuildTool
 
 				PriConfig.Save(ResourceConfigFile);
 
+				// generate resources.resfiles
 				IEnumerable<string> Resources = Directory.EnumerateFiles(Path.Combine(OutputPath, BuildResourceSubPath), "*.*", SearchOption.AllDirectories);
 				System.Text.StringBuilder ResourcesList = new System.Text.StringBuilder();
 				foreach (string Resource in Resources)
