Test Mode: boolean flag on ModuleRules class that allows including Tests folder into non-test builds.

#rnx

diff --git a/Engine/Source/Programs/UnrealBuildTool/Configuration/ModuleRules.cs b/Engine/Source/Programs/UnrealBuildTool/Configuration/ModuleRules.cs
--- a/Engine/Source/Programs/UnrealBuildTool/Configuration/ModuleRules.cs
+++ b/Engine/Source/Programs/UnrealBuildTool/Configuration/ModuleRules.cs
@@ -714,6 +714,20 @@ namespace UnrealBuildTool
 			}
 		}
 
+		/// <summary>
+		/// Decides whether to include the "Tests" folder at the root of the module as part of the compilation.
+		/// By default this folder is excluded from the module's compilation.
+		/// The "Tests" folder should be reserved for low level tests but there are exceptions to this rule.
+		/// </summary>
+		public bool bBuildIncludeTestsFolder
+		{
+			get { return bBuildIncludeTestsFolderOverride ?? false; }
+		}
+		/// <summary>
+		/// Set this override to true in derived target classes to include the "Tests" folder at the root of the module in the compilation.
+		/// </summary>
+		protected bool? bBuildIncludeTestsFolderOverride;
+
 
 		/// <summary>
 		/// The number of source files in this module before unity build will be activated for that module.  If set to
diff --git a/Engine/Source/Programs/UnrealBuildTool/Configuration/UEBuildBinary.cs b/Engine/Source/Programs/UnrealBuildTool/Configuration/UEBuildBinary.cs
--- a/Engine/Source/Programs/UnrealBuildTool/Configuration/UEBuildBinary.cs
+++ b/Engine/Source/Programs/UnrealBuildTool/Configuration/UEBuildBinary.cs
@@ -718,7 +718,7 @@ namespace UnrealBuildTool
 				{
 					// Compile each module.
 					Log.TraceVerbose("Compile module: " + Module.Name);
-					LinkInputFiles = Module.Compile(Target, ToolChain, BinaryCompileEnvironment, SpecificFilesToCompile, WorkingSet, Graph);
+					LinkInputFiles = Module.Compile(Target, ToolChain, BinaryCompileEnvironment, SpecificFilesToCompile, WorkingSet, Graph, Module.Rules.bBuildIncludeTestsFolder);
 
 					// Save the module outputs. In monolithic builds, this is just the object files.
 					if (Target.LinkType == TargetLinkType.Monolithic)
diff --git a/Engine/Source/Programs/UnrealBuildTool/Configuration/UEBuildModule.cs b/Engine/Source/Programs/UnrealBuildTool/Configuration/UEBuildModule.cs
--- a/Engine/Source/Programs/UnrealBuildTool/Configuration/UEBuildModule.cs
+++ b/Engine/Source/Programs/UnrealBuildTool/Configuration/UEBuildModule.cs
@@ -871,7 +871,7 @@ namespace UnrealBuildTool
 		/// <summary>
 		/// Compiles the module, and returns a list of files output by the compiler.
 		/// </summary>
-		public virtual List<FileItem> Compile(ReadOnlyTargetRules Target, UEToolChain ToolChain, CppCompileEnvironment CompileEnvironment, List<FileReference> SpecificFilesToCompile, ISourceFileWorkingSet WorkingSet, IActionGraphBuilder Graph)
+		public virtual List<FileItem> Compile(ReadOnlyTargetRules Target, UEToolChain ToolChain, CppCompileEnvironment CompileEnvironment, List<FileReference> SpecificFilesToCompile, ISourceFileWorkingSet WorkingSet, IActionGraphBuilder Graph, bool bBuildIncludeTestsFolder = false)
 		{
 			// Generate type libraries for Windows
 			foreach(ModuleRules.TypeLibrary TypeLibrary in Rules.TypeLibraries)
diff --git a/Engine/Source/Programs/UnrealBuildTool/Configuration/UEBuildModuleCPP.cs b/Engine/Source/Programs/UnrealBuildTool/Configuration/UEBuildModuleCPP.cs
--- a/Engine/Source/Programs/UnrealBuildTool/Configuration/UEBuildModuleCPP.cs
+++ b/Engine/Source/Programs/UnrealBuildTool/Configuration/UEBuildModuleCPP.cs
@@ -353,7 +353,7 @@ namespace UnrealBuildTool
 		}
 
 		// UEBuildModule interface.
-		public override List<FileItem> Compile(ReadOnlyTargetRules Target, UEToolChain ToolChain, CppCompileEnvironment BinaryCompileEnvironment, List<FileReference> SpecificFilesToCompile, ISourceFileWorkingSet WorkingSet, IActionGraphBuilder Graph)
+		public override List<FileItem> Compile(ReadOnlyTargetRules Target, UEToolChain ToolChain, CppCompileEnvironment BinaryCompileEnvironment, List<FileReference> SpecificFilesToCompile, ISourceFileWorkingSet WorkingSet, IActionGraphBuilder Graph, bool bBuildIncludeTestsFolder = false)
 		{
 			//UEBuildPlatform BuildPlatform = UEBuildPlatform.GetBuildPlatform(BinaryCompileEnvironment.Platform);
 
@@ -392,7 +392,7 @@ namespace UnrealBuildTool
 
 			// Find all the input files
 			Dictionary<DirectoryItem, FileItem[]> DirectoryToSourceFiles = new Dictionary<DirectoryItem, FileItem[]>();
-			InputFileCollection InputFiles = FindInputFiles(Target.Platform, DirectoryToSourceFiles, Target.IsTestTarget());
+			InputFileCollection InputFiles = FindInputFiles(Target.Platform, DirectoryToSourceFiles, Target.IsTestTarget() || bBuildIncludeTestsFolder);
 
 			foreach (KeyValuePair<DirectoryItem, FileItem[]> Pair in DirectoryToSourceFiles)
 			{
