REL - Consolidate Coroutine compiler switches

diff --git a/Engine/Source/Programs/UnrealBuildTool/Modes/GenerateClangDatabase.cs b/Engine/Source/Programs/UnrealBuildTool/Modes/GenerateClangDatabase.cs
--- a/Engine/Source/Programs/UnrealBuildTool/Modes/GenerateClangDatabase.cs
+++ b/Engine/Source/Programs/UnrealBuildTool/Modes/GenerateClangDatabase.cs
@@ -111,6 +111,15 @@ namespace UnrealBuildTool
 										throw new BuildException($"Unsupported C++ standard type set: {ModuleCompileEnvironment.CppStandard}");
 								}
 
+								if (ModuleCompileEnvironment.bEnableCoroutines)
+								{
+									CommandBuilder.AppendFormat(" -fcoroutines-ts");
+									if (!ModuleCompileEnvironment.bEnableExceptions)
+									{
+										CommandBuilder.AppendFormat(" -Wno-coroutine-missing-unhandled-exception");
+									}
+								}
+
 								foreach (FileItem ForceIncludeFile in ModuleCompileEnvironment.ForceIncludeFiles)
 								{
 									CommandBuilder.AppendFormat(" -include \"{0}\"", ForceIncludeFile.FullName);
diff --git a/Engine/Source/Programs/UnrealBuildTool/Platform/Android/AndroidToolChain.cs b/Engine/Source/Programs/UnrealBuildTool/Platform/Android/AndroidToolChain.cs
--- a/Engine/Source/Programs/UnrealBuildTool/Platform/Android/AndroidToolChain.cs
+++ b/Engine/Source/Programs/UnrealBuildTool/Platform/Android/AndroidToolChain.cs
@@ -711,10 +711,12 @@ namespace UnrealBuildTool
 			if (CompileEnvironment.bEnableExceptions)
 			{
 				Result += " -fexceptions";
+				Result += " -DPLATFORM_EXCEPTIONS_DISABLED=0";
 			}
 			else
 			{
 				Result += " -fno-exceptions";
+				Result += " -DPLATFORM_EXCEPTIONS_DISABLED=1";
 			}
 
 			// Conditionally enable (default disabled) generation of information about every class with virtual functions for use by the C++ runtime type identification features
@@ -822,20 +824,35 @@ namespace UnrealBuildTool
 			return Result;
 		}
 
-		static string GetCppStandardCompileArgument(CppStandardVersion Version)
+		static string GetCppStandardCompileArgument(CppCompileEnvironment CompileEnvironment)
 		{
-			switch (Version)
+			string Result;
+			switch (CompileEnvironment.CppStandard)
 			{
 				case CppStandardVersion.Cpp14:
-					return " -std=c++14";
+					Result = " -std=c++14";
+					break;
 				case CppStandardVersion.Latest:
 				case CppStandardVersion.Cpp17:
-					return " -std=c++17";
+					Result = " -std=c++17";
+					break;
 				case CppStandardVersion.Cpp20:
-					return " -std=c++20";
+					Result = " -std=c++20";
+					break;
 				default:
-					throw new BuildException($"Unsupported C++ standard type set: {Version}");
+					throw new BuildException($"Unsupported C++ standard type set: {CompileEnvironment.CppStandard}");
 			}
+
+			if (CompileEnvironment.bEnableCoroutines)
+			{
+				Result += " -fcoroutines-ts";
+				if (!CompileEnvironment.bEnableExceptions)
+				{
+					Result += " -Wno-coroutine-missing-unhandled-exception";
+				}
+			}
+
+			return Result;
 		}
 
 		static string GetCompileArguments_CPP(CppCompileEnvironment CompileEnvironment)
@@ -843,12 +860,7 @@ namespace UnrealBuildTool
 			string Result = "";
 
 			Result += " -x c++";
-			Result += GetCppStandardCompileArgument(CompileEnvironment.CppStandard);
-
-			if (CompileEnvironment.bEnableCoroutines)
-			{
-				Result += " -fcoroutines-ts";
-			}
+			Result += GetCppStandardCompileArgument(CompileEnvironment);
 
 			return Result;
 		}
@@ -864,12 +876,7 @@ namespace UnrealBuildTool
 			string Result = "";
 
 			Result += " -x c++-header";
-			Result += GetCppStandardCompileArgument(CompileEnvironment.CppStandard);
-
-			if (CompileEnvironment.bEnableCoroutines)
-			{
-				Result += " -fcoroutines-ts";
-			}
+			Result += GetCppStandardCompileArgument(CompileEnvironment);
 
 			return Result;
 		}
diff --git a/Engine/Source/Programs/UnrealBuildTool/Platform/IOS/IOSToolChain.cs b/Engine/Source/Programs/UnrealBuildTool/Platform/IOS/IOSToolChain.cs
--- a/Engine/Source/Programs/UnrealBuildTool/Platform/IOS/IOSToolChain.cs
+++ b/Engine/Source/Programs/UnrealBuildTool/Platform/IOS/IOSToolChain.cs
@@ -243,10 +243,12 @@ namespace UnrealBuildTool
 			if (CompileEnvironment.bEnableExceptions)
 			{
 				Result += " -fexceptions";
+				Result += " -DPLATFORM_EXCEPTIONS_DISABLED=0";
 			}
 			else
 			{
 				Result += " -fno-exceptions";
+				Result += " -DPLATFORM_EXCEPTIONS_DISABLED=1";
 			}
 
 			string? SanitizerMode = Environment.GetEnvironmentVariable("ENABLE_ADDRESS_SANITIZER");
@@ -374,14 +376,9 @@ namespace UnrealBuildTool
 		{
 			string Result = "";
 			Result += " -x objective-c++";
-			Result += GetCppStandardCompileArgument(CompileEnvironment.CppStandard);
+			Result += GetCppStandardCompileArgument(CompileEnvironment);
 			Result += " -stdlib=libc++";
 
-			if (CompileEnvironment.bEnableCoroutines)
-			{
-				Result += " -fcoroutines-ts";
-			}
-
 			return Result;
 		}
 
@@ -389,7 +386,7 @@ namespace UnrealBuildTool
 		{
 			string Result = "";
 			Result += " -x objective-c++";
-			Result += GetCppStandardCompileArgument(CompileEnvironment.CppStandard);
+			Result += GetCppStandardCompileArgument(CompileEnvironment);
 			Result += " -stdlib=libc++";
 			return Result;
 		}
@@ -398,7 +395,7 @@ namespace UnrealBuildTool
 		{
 			string Result = "";
 			Result += " -x objective-c";
-			Result += GetCppStandardCompileArgument(CompileEnvironment.CppStandard);
+			Result += GetCppStandardCompileArgument(CompileEnvironment);
 			Result += " -stdlib=libc++";
 			return Result;
 		}
@@ -414,14 +411,9 @@ namespace UnrealBuildTool
 		{
 			string Result = "";
 			Result += " -x objective-c++-header";
-			Result += GetCppStandardCompileArgument(CompileEnvironment.CppStandard);
+			Result += GetCppStandardCompileArgument(CompileEnvironment);
 			Result += " -stdlib=libc++";
 
-			if (CompileEnvironment.bEnableCoroutines)
-			{
-				Result += " -fcoroutines-ts";
-			}
-
 			return Result;
 		}
 
diff --git a/Engine/Source/Programs/UnrealBuildTool/Platform/Linux/LinuxToolChain.cs b/Engine/Source/Programs/UnrealBuildTool/Platform/Linux/LinuxToolChain.cs
--- a/Engine/Source/Programs/UnrealBuildTool/Platform/Linux/LinuxToolChain.cs
+++ b/Engine/Source/Programs/UnrealBuildTool/Platform/Linux/LinuxToolChain.cs
@@ -810,36 +810,45 @@ namespace UnrealBuildTool
 				: string.Format("{0}={1}", myKey, myValue);
 		}
 
-		static string GetCompilerStandardVersion_CPP(CppStandardVersion Version)
+		static string GetCompilerStandardVersion_CPP(CppCompileEnvironment CompileEnvironment)
 		{
-			switch (Version)
+			string Result;
+			switch (CompileEnvironment.CppStandard)
 			{
 				case CppStandardVersion.Cpp14:
-					return " -std=c++14";
-				case CppStandardVersion.Cpp17:
-					return " -std=c++17";
+					Result = " -std=c++14";
+					break;
 				case CppStandardVersion.Latest:
+				case CppStandardVersion.Cpp17:
+					Result = " -std=c++17";
+					break;
 				case CppStandardVersion.Cpp20:
-					return " -std=c++20";
+					Result = " -std=c++20";
+					break;
 				default:
-					throw new BuildException($"Unsupported C++ standard type set: {Version}");
+					throw new BuildException($"Unsupported C++ standard type set: {CompileEnvironment.CppStandard}");
 			}
-		}
-
-		static string GetCompileArguments_CPP(CppCompileEnvironment CompileEnvironment)
-		{
-			string Result = "";
-			Result += " -x c++";
-			Result += GetCompilerStandardVersion_CPP(CompileEnvironment.CppStandard);
 
 			if (CompileEnvironment.bEnableCoroutines)
 			{
 				Result += " -fcoroutines-ts";
+				if (!CompileEnvironment.bEnableExceptions)
+				{
+					Result += " -Wno-coroutine-missing-unhandled-exception";
+				}
 			}
 
 			return Result;
 		}
 
+		static string GetCompileArguments_CPP(CppCompileEnvironment CompileEnvironment)
+		{
+			string Result = "";
+			Result += " -x c++";
+			Result += GetCompilerStandardVersion_CPP(CompileEnvironment);
+			return Result;
+		}
+
 		static string GetCompileArguments_C()
 		{
 			string Result = "";
@@ -853,7 +862,7 @@ namespace UnrealBuildTool
 			Result += " -x objective-c++";
 			Result += " -fobjc-abi-version=2";
 			Result += " -fobjc-legacy-dispatch";
-			Result += GetCompilerStandardVersion_CPP(CompileEnvironment.CppStandard);
+			Result += GetCompilerStandardVersion_CPP(CompileEnvironment);
 			return Result;
 		}
 
@@ -882,7 +891,7 @@ namespace UnrealBuildTool
 			Result += " -x objective-c";
 			Result += " -fobjc-abi-version=2";
 			Result += " -fobjc-legacy-dispatch";
-			Result += GetCompilerStandardVersion_CPP(CompileEnvironment.CppStandard);
+			Result += GetCompilerStandardVersion_CPP(CompileEnvironment);
 			return Result;
 		}
 
@@ -890,13 +899,7 @@ namespace UnrealBuildTool
 		{
 			string Result = "";
 			Result += " -x c++-header";
-			Result += GetCompilerStandardVersion_CPP(CompileEnvironment.CppStandard);
-
-			if (CompileEnvironment.bEnableCoroutines)
-			{
-				Result += " -fcoroutines-ts";
-			}
-
+			Result += GetCompilerStandardVersion_CPP(CompileEnvironment);
 			return Result;
 		}
 
diff --git a/Engine/Source/Programs/UnrealBuildTool/Platform/Mac/MacToolChain.cs b/Engine/Source/Programs/UnrealBuildTool/Platform/Mac/MacToolChain.cs
--- a/Engine/Source/Programs/UnrealBuildTool/Platform/Mac/MacToolChain.cs
+++ b/Engine/Source/Programs/UnrealBuildTool/Platform/Mac/MacToolChain.cs
@@ -196,6 +196,8 @@ namespace UnrealBuildTool
 			Result += " -fpascal-strings";
 
 			Result += " -fexceptions";
+			Result += " -DPLATFORM_EXCEPTIONS_DISABLED=0";
+
 			Result += " -fasm-blocks";
 
 			if(CompileEnvironment.bHideSymbolsByDefault)
@@ -301,14 +303,9 @@ namespace UnrealBuildTool
 		{
 			string Result = "";
 			Result += " -x objective-c++";
-			Result += GetCppStandardCompileArgument(CompileEnvironment.CppStandard);
+			Result += GetCppStandardCompileArgument(CompileEnvironment);
 			Result += " -stdlib=libc++";
 
-			if (CompileEnvironment.bEnableCoroutines)
-			{
-				Result += " -fcoroutines-ts";
-			}
-
 			return Result;
 		}
 
@@ -316,7 +313,7 @@ namespace UnrealBuildTool
 		{
 			string Result = "";
 			Result += " -x objective-c++";
-			Result += GetCppStandardCompileArgument(CompileEnvironment.CppStandard);
+			Result += GetCppStandardCompileArgument(CompileEnvironment);
 			Result += " -stdlib=libc++";
 			return Result;
 		}
@@ -340,14 +337,9 @@ namespace UnrealBuildTool
 		{
 			string Result = "";
 			Result += " -x objective-c++-header";
-			Result += GetCppStandardCompileArgument(CompileEnvironment.CppStandard);
+			Result += GetCppStandardCompileArgument(CompileEnvironment);
 			Result += " -stdlib=libc++";
 
-			if (CompileEnvironment.bEnableCoroutines)
-			{
-				Result += " -fcoroutines-ts";
-			}
-
 			return Result;
 		}
 
diff --git a/Engine/Source/Programs/UnrealBuildTool/Platform/Windows/VCToolChain.cs b/Engine/Source/Programs/UnrealBuildTool/Platform/Windows/VCToolChain.cs
--- a/Engine/Source/Programs/UnrealBuildTool/Platform/Windows/VCToolChain.cs
+++ b/Engine/Source/Programs/UnrealBuildTool/Platform/Windows/VCToolChain.cs
@@ -538,11 +538,13 @@ namespace UnrealBuildTool
 			{
 				// Enable C++ exception handling, but not C exceptions.
 				Arguments.Add("/EHsc");
+				Arguments.Add("/DPLATFORM_EXCEPTIONS_DISABLED=0");
 			}
 			else
 			{
 				// This is required to disable exception handling in VC platform headers.
 				AddDefinition(Arguments, "_HAS_EXCEPTIONS=0");
+				Arguments.Add("/DPLATFORM_EXCEPTIONS_DISABLED=1");
 			}
 
 			// If enabled, create debug information.
diff --git a/Engine/Source/Programs/UnrealBuildTool/ToolChain/AppleToolChain.cs b/Engine/Source/Programs/UnrealBuildTool/ToolChain/AppleToolChain.cs
--- a/Engine/Source/Programs/UnrealBuildTool/ToolChain/AppleToolChain.cs
+++ b/Engine/Source/Programs/UnrealBuildTool/ToolChain/AppleToolChain.cs
@@ -200,20 +200,35 @@ namespace UnrealBuildTool
 			return FullClangVersion;
 		}
 
-		protected static string GetCppStandardCompileArgument(CppStandardVersion Version)
+		protected static string GetCppStandardCompileArgument(CppCompileEnvironment CompileEnvironment)
 		{
-			switch (Version)
+			string Result;
+			switch (CompileEnvironment.CppStandard)
 			{
 				case CppStandardVersion.Cpp14:
-					return " -std=c++14";
-				case CppStandardVersion.Cpp17:
-					return " -std=c++17";
+					Result = " -std=c++14";
+					break;
 				case CppStandardVersion.Latest:
+				case CppStandardVersion.Cpp17:
+					Result = " -std=c++17";
+					break;
 				case CppStandardVersion.Cpp20:
-					return " -std=c++20";
+					Result = " -std=c++20";
+					break;
 				default:
-					throw new BuildException($"Unsupported C++ standard type set: {Version}");
+					throw new BuildException($"Unsupported C++ standard type set: {CompileEnvironment.CppStandard}");
+			}
+
+			if (CompileEnvironment.bEnableCoroutines)
+			{
+				Result += " -fcoroutines-ts";
+				if (!CompileEnvironment.bEnableExceptions)
+				{
+					Result += " -Wno-coroutine-missing-unhandled-exception";
+				}
 			}
+
+			return Result;
 		}
 
 		protected string GetDsymutilPath(out string ExtraOptions, bool bIsForLTOBuild=false)
diff --git a/Engine/Source/Runtime/Core/Public/Windows/WindowsPlatformCompilerSetup.h b/Engine/Source/Runtime/Core/Public/Windows/WindowsPlatformCompilerSetup.h
--- a/Engine/Source/Runtime/Core/Public/Windows/WindowsPlatformCompilerSetup.h
+++ b/Engine/Source/Runtime/Core/Public/Windows/WindowsPlatformCompilerSetup.h
@@ -228,7 +228,7 @@ static_assert(_MSC_VER >= 1929, "Visual Studio 2019 v16.11 or greater is require
 #pragma warning(default: 4906) // string literal cast to 'LPWSTR'																	https://docs.microsoft.com/en-us/cpp/error-messages/compiler-warnings/compiler-warning-level-1-c4906
 #pragma warning(default: 4928) // illegal copy-initialization; more than one user-defined conversion has been implicitly applied	https://docs.microsoft.com/en-us/cpp/error-messages/compiler-warnings/compiler-warning-level-1-c4928
 #pragma warning(default: 4931) // we are assuming the type library was built for number-bit pointers								https://docs.microsoft.com/en-us/cpp/error-messages/compiler-warnings/compiler-warning-level-4-c4931
-#pragma warning(default: 4946) // reinterpret_cast used between related classes: 'class1' and 'class2'								https://docs.microsoft.com/en-us/cpp/error-messages/compiler-warnings/compiler-warning-level-1-c4946
+#pragma warning(disable: 4946) // reinterpret_cast used between related classes: 'class1' and 'class2'								https://docs.microsoft.com/en-us/cpp/error-messages/compiler-warnings/compiler-warning-level-1-c4946
 #pragma warning(default: 5038) // data member 'A::y' will be initialized after data member 'A::x'									https://docs.microsoft.com/en-us/cpp/error-messages/compiler-warnings/c5038
 #pragma warning(disable: 4984) // 'if constexpr' is a C++17 language extension
 
