LWC: Force enable LWC - override target.cs and build.xml LWC settings in preparation for LWC toggle removal

diff --git a/Engine/Source/Programs/UnrealBuildTool/Configuration/TargetRules.cs b/Engine/Source/Programs/UnrealBuildTool/Configuration/TargetRules.cs
--- a/Engine/Source/Programs/UnrealBuildTool/Configuration/TargetRules.cs
+++ b/Engine/Source/Programs/UnrealBuildTool/Configuration/TargetRules.cs
@@ -2255,12 +2255,12 @@ namespace UnrealBuildTool
 
 		public bool bCompileISPC
 		{
-			get { return Inner.bCompileISPC && GlobalDefinitions.Contains("UE_LARGE_WORLD_COORDINATES_DISABLED=1"); }   // LWC_TODO: Temporarily disable ISPC when LWC is turned on. To be removed when double support is added to ISPC.
+			get { return false; /*Inner.bCompileISPC && GlobalDefinitions.Contains("UE_LARGE_WORLD_COORDINATES_DISABLED=1");*/ }   // LWC_TODO: Temporarily disable ISPC when LWC is turned on (note this is now forced - we ignore LARGE_WORLD_COORDINATES_DISABLED).
 		}
 
 		public bool bLWCDisabled 
 		{
-			get { return GlobalDefinitions.Contains("UE_LARGE_WORLD_COORDINATES_DISABLED=1"); }
+			get { return false; /*GlobalDefinitions.Contains("UE_LARGE_WORLD_COORDINATES_DISABLED=1");*/ }
 		}
 		
 		public bool bCompilePython
diff --git a/Engine/Source/Programs/UnrealBuildTool/Configuration/UEBuildTarget.cs b/Engine/Source/Programs/UnrealBuildTool/Configuration/UEBuildTarget.cs
--- a/Engine/Source/Programs/UnrealBuildTool/Configuration/UEBuildTarget.cs
+++ b/Engine/Source/Programs/UnrealBuildTool/Configuration/UEBuildTarget.cs
@@ -3704,7 +3704,14 @@ namespace UnrealBuildTool
 
 			// Add the 'Engine/Source' path as a global include path for all modules
 			GlobalCompileEnvironment.UserIncludePaths.Add(UnrealBuildTool.EngineSourceDirectory);
-
+						
+			// LWC_TODO: Very temp! Enforce no LWC toggling.
+			int PreviousDefinition = GlobalCompileEnvironment.Definitions.FindIndex(s => s.Contains("UE_LARGE_WORLD_COORDINATES_DISABLED"));
+			if (PreviousDefinition >= 0)
+			{
+				GlobalCompileEnvironment.Definitions.RemoveAt(PreviousDefinition);
+			}
+			
 			[REDACTED_EMAIL]: Do any platform specific tool chain initialization here if required
 
 			UnrealTargetConfiguration EngineTargetConfiguration = Configuration == UnrealTargetConfiguration.DebugGame ? UnrealTargetConfiguration.Development : Configuration;
diff --git a/Engine/Source/Runtime/Core/Public/Misc/LargeWorldCoordinates.h b/Engine/Source/Runtime/Core/Public/Misc/LargeWorldCoordinates.h
--- a/Engine/Source/Runtime/Core/Public/Misc/LargeWorldCoordinates.h
+++ b/Engine/Source/Runtime/Core/Public/Misc/LargeWorldCoordinates.h
@@ -2,6 +2,8 @@
 
 #pragma once
 
+// NOTE: UE_LARGE_WORLD_COORDINATES_DISABLED has temporarily been overridden in the build scripts.
+
 // Use float based world coordinates throughout. Note: This is unsupported!
 #ifndef UE_LARGE_WORLD_COORDINATES_DISABLED
 #define UE_LARGE_WORLD_COORDINATES_DISABLED 0
