UnrealBuildTool: Fix incorrect warning when c++ build tools aren't installed

#rnx

diff --git a/Engine/Source/Programs/UnrealBuildTool/Platform/Windows/UEBuildWindows.cs b/Engine/Source/Programs/UnrealBuildTool/Platform/Windows/UEBuildWindows.cs
--- a/Engine/Source/Programs/UnrealBuildTool/Platform/Windows/UEBuildWindows.cs
+++ b/Engine/Source/Programs/UnrealBuildTool/Platform/Windows/UEBuildWindows.cs
@@ -1108,21 +1108,21 @@ namespace UnrealBuildTool
 			// If we do have a Visual Studio installation, but we're missing just the C++ parts, warn about that.
 			if (TryGetVSInstallDirs(WindowsCompiler.VisualStudio2019) != null)
 			{
-				string ToolSetWarning = Architecture == WindowsArchitecture.ARM64 ?
-					"Please verify that the \"VC++ 2019 toolset\" component is selected in the Visual Studio 2019 installation options." :
-					"Please verify that the \"VS 2019 C++ ARM64 build tools\" component is selected in the Visual Studio 2019 installation options.";
-				Log.TraceWarning("Visual Studio 2019 is installed, but is missing the C++ toolchain. {0}", ToolSetWarning);
+				string ToolSetWarning = Architecture == WindowsArchitecture.x64 ?
+					"MSVC v142 - VS 2019 C++ x64/x86 build tools (Latest)" :
+					"MSVC v142 - VS 2019 C++ ARM64 build tools (Latest)";
+				Log.TraceWarning("Visual Studio 2019 is installed, but is missing the C++ toolchain. Please verify that the \"{0}\" component is selected in the Visual Studio 2019 installation options.", ToolSetWarning);
 			}
 			else if (TryGetVSInstallDirs(WindowsCompiler.VisualStudio2022) != null)
 			{
-				string ToolSetWarning = Architecture == WindowsArchitecture.ARM64 ?
-					"Please verify that the \"VS 2022 C++ x64/x86 build tools (Latest)\" component is selected in the Visual Studio 2022 installation options." :
-					"Please verify that the \"VS 2022 C++ ARM64 build tools\" component is selected in the Visual Studio 2022 installation options.";
-				Log.TraceWarning("Visual Studio 2022 is installed, but is missing the C++ toolchain. Please verify that the \"VS 2022 C++ x64/x86 build tools (Latest)\" component is selected in the Visual Studio 2022 installation options.");
+				string ToolSetWarning = Architecture == WindowsArchitecture.x64 ?
+					"MSVC v143 - VS 2022 C++ x64/x86 build tools (Latest)" :
+					"MSVC v143 - VS 2022 C++ ARM64 build tools (Latest)";
+				Log.TraceWarning("Visual Studio 2022 is installed, but is missing the C++ toolchain. Please verify that the \"{0}\" component is selected in the Visual Studio 2022 installation options.", ToolSetWarning);
 			}
 			else
 			{
-				Log.TraceWarning("No Visual C++ installation was found. Please download and install Visual Studio 2019 (or newer) with C++ components.");
+				Log.TraceWarning("No Visual C++ installation was found. Please download and install Visual Studio 2019 or 2022 with C++ components.");
 			}
 
 			// Finally, default to VS2019 anyway
