Modify UBT to not check for live coding unless there are actions to execute.

diff --git a/Engine/Source/Programs/UnrealBuildTool/Modes/BuildMode.cs b/Engine/Source/Programs/UnrealBuildTool/Modes/BuildMode.cs
--- a/Engine/Source/Programs/UnrealBuildTool/Modes/BuildMode.cs
+++ b/Engine/Source/Programs/UnrealBuildTool/Modes/BuildMode.cs
@@ -506,6 +506,10 @@ namespace UnrealBuildTool
 						}
 						HotReloadTargetIdx = Idx;
 					}
+					else if (MergedActionsToExecute.Count > 0)
+					{
+						HotReload.CheckForLiveCodingSessionActive(TargetDescriptors[Idx], Makefiles[Idx], BuildConfiguration);
+					}
 				}	
 
 				if (HotReloadTargetIdx != -1)
diff --git a/Engine/Source/Programs/UnrealBuildTool/System/HotReload.cs b/Engine/Source/Programs/UnrealBuildTool/System/HotReload.cs
--- a/Engine/Source/Programs/UnrealBuildTool/System/HotReload.cs
+++ b/Engine/Source/Programs/UnrealBuildTool/System/HotReload.cs
@@ -181,12 +181,6 @@ namespace UnrealBuildTool
 				TargetDescriptor.HotReloadMode = GetDefaultMode(TargetDescriptor, BuildConfiguration);
 			}
 
-			// Guard against a live coding session for this target being active
-			if (BuildConfiguration.bAllowHotReloadFromIDE && TargetDescriptor.HotReloadMode != HotReloadMode.LiveCoding && TargetDescriptor.ForeignPlugin == null && HotReload.IsLiveCodingSessionActive(Makefile))
-			{
-				throw new BuildException("Unable to start regular build while Live Coding is active. Press Ctrl+Alt+F11 to trigger a Live Coding compile.");
-			}
-
 			// Apply the previous hot reload state
 			if (TargetDescriptor.HotReloadMode == HotReloadMode.Disabled)
 			{
@@ -217,6 +211,17 @@ namespace UnrealBuildTool
 			return PatchedOldLocationToNewLocation;
 		}
 
+		public static void CheckForLiveCodingSessionActive(TargetDescriptor TargetDescriptor, TargetMakefile Makefile, BuildConfiguration BuildConfiguration)
+		{
+			// Guard against a live coding session for this target being active
+			if (BuildConfiguration.bAllowHotReloadFromIDE && TargetDescriptor.HotReloadMode != HotReloadMode.LiveCoding && TargetDescriptor.ForeignPlugin == null && HotReload.IsLiveCodingSessionActive(Makefile))
+			{
+				throw new BuildException("Unable to start regular build while Live Coding is active. Press Ctrl+Alt+F11 to trigger a Live Coding compile."
+					+" If you are attempting to package a build, then there are pending project changes that need to be compiled."
+					+" Exit the editor and build the project prior to attempting to package a build");
+			}
+		}
+
 		/// <summary>
 		/// Checks whether a live coding session is currently active for a target. If so, we don't want to allow modifying any object files before they're loaded.
 		/// </summary>
