PR #8802: UBT: GenerateClangDatabase: Allow custom output directory using -OutputDir= (Contributed by [REDACTED_PERSON])

diff --git a/Engine/Source/Programs/UnrealBuildTool/Modes/GenerateClangDatabase.cs b/Engine/Source/Programs/UnrealBuildTool/Modes/GenerateClangDatabase.cs
--- a/Engine/Source/Programs/UnrealBuildTool/Modes/GenerateClangDatabase.cs
+++ b/Engine/Source/Programs/UnrealBuildTool/Modes/GenerateClangDatabase.cs
@@ -150,7 +150,8 @@ namespace UnrealBuildTool
 				}
 
 				// Write the compile database
-				FileReference DatabaseFile = FileReference.Combine(Unreal.RootDirectory, "compile_commands.json");
+				DirectoryReference DatabaseDirectory = Arguments.GetDirectoryReferenceOrDefault("-OutputDir=", Unreal.RootDirectory);
+				FileReference DatabaseFile = FileReference.Combine(DatabaseDirectory, "compile_commands.json");
 				using (JsonWriter Writer = new JsonWriter(DatabaseFile))
 				{
 					Writer.WriteArrayStart();
