PR #8761: Fix `UE_5_1_OR_LATER` not being defined by UnrealBuildTool (Contributed by [REDACTED_PERSON])

diff --git a/Engine/Source/Programs/UnrealBuildTool/System/RulesAssembly.cs b/Engine/Source/Programs/UnrealBuildTool/System/RulesAssembly.cs
--- a/Engine/Source/Programs/UnrealBuildTool/System/RulesAssembly.cs
+++ b/Engine/Source/Programs/UnrealBuildTool/System/RulesAssembly.cs
@@ -224,6 +224,7 @@ namespace UnrealBuildTool
 			PreprocessorDefines.Add("WITH_FORWARDED_TARGET_RULES_CTOR");
 
 			// Define macros for the Unreal engine version, starting with 4.17
+			// Assumes the current MajorVersion is 5
 			BuildVersion? Version;
 			if (BuildVersion.TryRead(BuildVersion.GetDefaultFileName(), out Version))
 			{
@@ -231,7 +232,7 @@ namespace UnrealBuildTool
 				{
 					PreprocessorDefines.Add(String.Format("UE_4_{0}_OR_LATER", MinorVersion));
 				}
-				for (int MinorVersion = 0; MinorVersion <= 0; MinorVersion++)
+				for (int MinorVersion = 0; MinorVersion <= Version.MinorVersion; MinorVersion++)
 				{
 					PreprocessorDefines.Add(String.Format("UE_5_{0}_OR_LATER", MinorVersion));
 				}
