PR #8800: UBT: GenerateClangDatabase: Escape quotes in definitions (Contributed by [REDACTED_PERSON])

diff --git a/Engine/Source/Programs/UnrealBuildTool/Modes/GenerateClangDatabase.cs b/Engine/Source/Programs/UnrealBuildTool/Modes/GenerateClangDatabase.cs
--- a/Engine/Source/Programs/UnrealBuildTool/Modes/GenerateClangDatabase.cs
+++ b/Engine/Source/Programs/UnrealBuildTool/Modes/GenerateClangDatabase.cs
@@ -126,7 +126,7 @@ namespace UnrealBuildTool
 								}
 								foreach (string Definition in ModuleCompileEnvironment.Definitions)
 								{
-									CommandBuilder.AppendFormat(" -D\"{0}\"", Definition);
+									CommandBuilder.AppendFormat(" -D\"{0}\"", Definition.Replace("\"", "\\\""));
 								}
 								foreach (DirectoryReference IncludePath in ModuleCompileEnvironment.UserIncludePaths)
 								{
