UnrealBuildTool: Log any use of bEnableCppCoroutinesForEvaluation

#trivial

diff --git a/Engine/Source/Programs/UnrealBuildTool/Configuration/UEBuildTarget.cs b/Engine/Source/Programs/UnrealBuildTool/Configuration/UEBuildTarget.cs
--- a/Engine/Source/Programs/UnrealBuildTool/Configuration/UEBuildTarget.cs
+++ b/Engine/Source/Programs/UnrealBuildTool/Configuration/UEBuildTarget.cs
@@ -3963,6 +3963,7 @@ namespace UnrealBuildTool
 			// Whether C++20 coroutines are enabled
 			if (Rules.bEnableCppCoroutinesForEvaluation)
 			{
+				Log.TraceInformationOnce($"NOTE: C++ coroutine support is considered experimental and should be used for evaluation purposes only ({nameof(TargetRules.bEnableCppCoroutinesForEvaluation)})");
 				GlobalCompileEnvironment.Definitions.Add("WITH_CPP_COROUTINES=1");
 				GlobalCompileEnvironment.bEnableCoroutines = true;
 			}
