Disable ISPC on HoloLens - no support for Windows-ARM64 currently in ISPC

diff --git a/Engine/Source/Programs/UnrealBuildTool/Platform/HoloLens/UEBuildHoloLens.cs b/Engine/Source/Programs/UnrealBuildTool/Platform/HoloLens/UEBuildHoloLens.cs
--- a/Engine/Source/Programs/UnrealBuildTool/Platform/HoloLens/UEBuildHoloLens.cs
+++ b/Engine/Source/Programs/UnrealBuildTool/Platform/HoloLens/UEBuildHoloLens.cs
@@ -253,6 +253,9 @@ namespace UnrealBuildTool
 			}
 
 			HoloLensExports.InitWindowsSdkToolPath(Target.HoloLensPlatform.Win10SDKVersion.ToString());
+
+			// ISPC currently doesn't support Windows-AArch64
+			Target.bCompileISPC = false;
 		}
 
 		/// <summary>
