Fixed universal character warnings from the ISPC compiler.

UBT defines some macros which embed the path of the user's sync. UE_LIVE_CODING_ENGINE_DIR being one such example. If the user's sync path contains a path element that starts with a 'u' followed by valid hexadecimal characters, such as 'd:\p4\ue5-release' the '\ue5' gets interpreted as an escaped universal character. ISPC compiler doesn't like this.

#rnx

diff --git a/Engine/Source/Programs/UnrealBuildTool/ToolChain/UEToolChain.cs b/Engine/Source/Programs/UnrealBuildTool/ToolChain/UEToolChain.cs
--- a/Engine/Source/Programs/UnrealBuildTool/ToolChain/UEToolChain.cs
+++ b/Engine/Source/Programs/UnrealBuildTool/ToolChain/UEToolChain.cs
@@ -496,7 +496,7 @@ namespace UnrealBuildTool
 				foreach (string Definition in CompileEnvironment.Definitions)
 				{
 					// TODO: Causes ISPC compiler to generate a spurious warning about the universal character set
-					if(!Definition.Contains("\\\\U"))
+					if (!Definition.Contains("\\\\U") && !Definition.Contains("\\\\u"))
 					{
 						Arguments.Add(String.Format("-D\"{0}\"", Definition));
 					}
@@ -700,7 +700,7 @@ namespace UnrealBuildTool
 				foreach (string Definition in CompileEnvironment.Definitions)
 				{
 					// TODO: Causes ISPC compiler to generate a spurious warning about the universal character set
-					if(!Definition.Contains("\\\\U"))
+					if (!Definition.Contains("\\\\U") && !Definition.Contains("\\\\u"))
 					{
 						Arguments.Add(String.Format("-D\"{0}\"", Definition));
 					}
