Fix iOS C++ projects from Windows installed builds
Use Windows UBT on remote Mac when using remote toolchain

diff --git a/Engine/Build/Rsync/RsyncEngine.txt b/Engine/Build/Rsync/RsyncEngine.txt
--- a/Engine/Build/Rsync/RsyncEngine.txt
+++ b/Engine/Build/Rsync/RsyncEngine.txt
@@ -45,7 +45,8 @@
 + /Binaries/DotNET/Ionic.Zip.Reduced.dll
 + /Binaries/DotNET/IOS/
 + /Binaries/DotNET/IOS/**
-protect /Binaries/DotNET/UnrealBuildTool/**
++ /Binaries/DotNET/UnrealBuildTool/
++ /Binaries/DotNET/UnrealBuildTool/**
 
 # Binaries/Mac: Include DsymExporter, ignore everything else
 + /Binaries/Mac/
@@ -74,7 +75,6 @@ protect /Binaries/TVOS/
 
 # Build
 + /Build/
-- /Build/InstalledBuild.txt
 + /Build/*
 + /Build/BatchFiles/
 + /Build/BatchFiles/**
diff --git a/Engine/Source/Programs/Shared/EpicGames.Build/Unreal.cs b/Engine/Source/Programs/Shared/EpicGames.Build/Unreal.cs
--- a/Engine/Source/Programs/Shared/EpicGames.Build/Unreal.cs
+++ b/Engine/Source/Programs/Shared/EpicGames.Build/Unreal.cs
@@ -116,6 +116,11 @@ namespace UnrealBuildBase
 		/// </summary>
 		public static readonly FileReference UnrealBuildToolPath = FindUnrealBuildTool();
 
+		/// <summary>
+		/// The path to UBT
+		/// </summary>
+		public static readonly FileReference UnrealBuildToolDllPath = FindUnrealBuildTool().ChangeExtension(".dll");
+
 		/// <summary>
 		/// The directory containing the bundled .NET installation
 		/// </summary>
diff --git a/Engine/Source/Programs/UnrealBuildTool/System/ActionGraphBuilder.cs b/Engine/Source/Programs/UnrealBuildTool/System/ActionGraphBuilder.cs
--- a/Engine/Source/Programs/UnrealBuildTool/System/ActionGraphBuilder.cs
+++ b/Engine/Source/Programs/UnrealBuildTool/System/ActionGraphBuilder.cs
@@ -284,8 +284,8 @@ namespace UnrealBuildTool
 			}
 
 			Action NewAction = Graph.CreateAction(Type);
-			NewAction.CommandPath = UnrealBuildTool.GetUBTPath();
-			NewAction.CommandArguments = String.Format("-Mode={0} {1}", Attribute.Name, Arguments);
+			NewAction.CommandPath = Unreal.DotnetPath;
+			NewAction.CommandArguments = $"{Unreal.UnrealBuildToolDllPath} -Mode={Attribute.Name} {Arguments}";
 			return NewAction;
 		}
 
