Live Coding error message in UBT simplified and clearer.

diff --git a/Engine/Source/Programs/UnrealBuildTool/System/HotReload.cs b/Engine/Source/Programs/UnrealBuildTool/System/HotReload.cs
--- a/Engine/Source/Programs/UnrealBuildTool/System/HotReload.cs
+++ b/Engine/Source/Programs/UnrealBuildTool/System/HotReload.cs
@@ -216,9 +216,7 @@ namespace UnrealBuildTool
 			// Guard against a live coding session for this target being active
 			if (BuildConfiguration.bAllowHotReloadFromIDE && TargetDescriptor.HotReloadMode != HotReloadMode.LiveCoding && TargetDescriptor.ForeignPlugin == null && HotReload.IsLiveCodingSessionActive(Makefile))
 			{
-				throw new BuildException("Unable to start regular build while Live Coding is active. Press Ctrl+Alt+F11 to trigger a Live Coding compile."
-					+" If you are attempting to package a build, then there are pending project changes that need to be compiled."
-					+" Exit the editor and build the project prior to attempting to package a build");
+				throw new BuildException("Unable to build while Live Coding is active. Exit the editor and game, or press Ctrl+Alt+F11 if iterating on code in the editor or game");
 			}
 		}
 
