Option to skip .dsym generation for up to 20x faster builds on Mac.

Touching one .cpp file and building QAGame takes 8m00s on average (Xeon, 16 logical core). 7m30s of that is runing dsymutil (and sometimes it doesn't complete because UBT's timeout protection kicks in). As these .dsym files do not seem to be required for the build-run development loop, using "-NoDSYM" to skip them them greatly improves the iteration experience.

#rnx

diff --git a/Engine/Source/Programs/UnrealBuildTool/Platform/Mac/UEBuildMac.cs b/Engine/Source/Programs/UnrealBuildTool/Platform/Mac/UEBuildMac.cs
--- a/Engine/Source/Programs/UnrealBuildTool/Platform/Mac/UEBuildMac.cs
+++ b/Engine/Source/Programs/UnrealBuildTool/Platform/Mac/UEBuildMac.cs
@@ -33,6 +33,13 @@ namespace UnrealBuildTool
 		[CommandLine("-EnableUBSan")]
 		[XmlConfigFile(Category = "BuildConfiguration", Name = "bEnableUndefinedBehaviorSanitizer")]
 		public bool bEnableUndefinedBehaviorSanitizer = false;
+
+		/// <summary>
+		/// Enables the generation of .dsym files. This can be disabled to enable faster iteration times during development.
+		/// </summary>
+		[CommandLine("-NoDSYM", Value = "false")]
+		[XmlConfigFile(Category = "BuildConfiguration", Name = "bUseDSYMFiles")]
+		public bool bUseDSYMFiles = true;
 	}
 
 	/// <summary>
@@ -129,6 +136,7 @@ namespace UnrealBuildTool
 			Target.GlobalDefinitions.Add("GL_SILENCE_DEPRECATION=1");
 
 			Target.bUsePDBFiles = !Target.bDisableDebugInfo && ShouldCreateDebugInfo(new ReadOnlyTargetRules(Target));
+			Target.bUsePDBFiles &= Target.MacPlatform.bUseDSYMFiles;
 
 			// we always deploy - the build machines need to be able to copy the files back, which needs the full bundle
 			Target.bDeployAfterCompile = true;
