PR #8836: Fix for uploading Android 12 shipping builds (Contributed by [REDACTED_PERSON])

#android

diff --git a/Engine/Source/Programs/UnrealBuildTool/Platform/Android/UEDeployAndroid.cs b/Engine/Source/Programs/UnrealBuildTool/Platform/Android/UEDeployAndroid.cs
--- a/Engine/Source/Programs/UnrealBuildTool/Platform/Android/UEDeployAndroid.cs
+++ b/Engine/Source/Programs/UnrealBuildTool/Platform/Android/UEDeployAndroid.cs
@@ -2662,6 +2662,7 @@ namespace UnrealBuildTool
 			{
 				// normal application settings
 				Text.AppendLine("\t\t<activity android:name=\"com.epicgames.unreal.SplashActivity\"");
+				Text.AppendLine("\t\t          android:exported=\"true\"");
 				Text.AppendLine("\t\t          android:label=\"@string/app_name\"");
 				Text.AppendLine("\t\t          android:theme=\"@style/UnrealSplashTheme\"");
 				Text.AppendLine("\t\t          android:launchMode=\"singleTask\"");
@@ -2685,6 +2686,7 @@ namespace UnrealBuildTool
 			else
 			{
 				Text.AppendLine("\t\t<activity android:name=\"com.epicgames.unreal.GameActivity\"");
+				Text.AppendLine("\t\t          android:exported=\"true\"");
 				Text.AppendLine("\t\t          android:label=\"@string/app_name\"");
 				Text.AppendLine("\t\t          android:theme=\"@android:style/Theme.Black.NoTitleBar.Fullscreen\"");
 				Text.AppendLine(bAddDensity ? "\t\t          android:configChanges=\"mcc|mnc|uiMode|density|screenSize|smallestScreenSize|screenLayout|orientation|keyboardHidden|keyboard\""
