PR #8841: Fix incorrect format string in MacToolChain (Contributed by [REDACTED_PERSON])

diff --git a/Engine/Source/Programs/UnrealBuildTool/Platform/Mac/MacToolChain.cs b/Engine/Source/Programs/UnrealBuildTool/Platform/Mac/MacToolChain.cs
--- a/Engine/Source/Programs/UnrealBuildTool/Platform/Mac/MacToolChain.cs
+++ b/Engine/Source/Programs/UnrealBuildTool/Platform/Mac/MacToolChain.cs
@@ -748,7 +748,7 @@ namespace UnrealBuildTool
 					}
 					else
 					{
-						Log.TraceWarning("Unexpected third party dylib location when generating RPATH entries: {1}. Skipping.", LibraryFullPath);
+						Log.TraceWarning("Unexpected third party dylib location when generating RPATH entries: {0}. Skipping.", LibraryFullPath);
 					}
 
 					// For staged code-based games we need additional entry if the game is not stored directly in the engine's root directory
