Mark UnrealBuildConfiguration to be serialized to json as a string.
It's used in some build outputs so we want it to be forward and backward compatible.

diff --git a/Engine/Source/Programs/UnrealBuildTool/Configuration/UEBuildTarget.cs b/Engine/Source/Programs/UnrealBuildTool/Configuration/UEBuildTarget.cs
--- a/Engine/Source/Programs/UnrealBuildTool/Configuration/UEBuildTarget.cs
+++ b/Engine/Source/Programs/UnrealBuildTool/Configuration/UEBuildTarget.cs
@@ -10,6 +10,7 @@ using System.Globalization;
 using System.IO;
 using System.Xml;
 using System.Runtime.Serialization;
+using System.Text.Json.Serialization;
 using System.Text.RegularExpressions;
 using EpicGames.Core;
 using System.Reflection;
@@ -595,6 +596,7 @@ namespace UnrealBuildTool
 	/// <summary>
 	/// The type of configuration a target can be built for.  Roughly order by optimization level. 
 	/// </summary>
+	[JsonConverter(typeof(JsonStringEnumConverter))]
 	public enum UnrealTargetConfiguration
 	{
 		/// <summary>
