Adding bNeedsExtraShaderFormats to TargetRules to protect hidden platforms from exposing their DataDrivenPlatformInfo defines to other hidden platforms.

diff --git a/Engine/Source/Programs/UnrealBuildTool/Configuration/TargetRules.cs b/Engine/Source/Programs/UnrealBuildTool/Configuration/TargetRules.cs
--- a/Engine/Source/Programs/UnrealBuildTool/Configuration/TargetRules.cs
+++ b/Engine/Source/Programs/UnrealBuildTool/Configuration/TargetRules.cs
@@ -632,6 +632,20 @@ namespace UnrealBuildTool
 		/// </summary>
 		public bool bForceBuildShaderFormats = false;
 
+		/// <summary>
+		/// Override for including extra shader formats
+		/// </summary>
+		public bool? bNeedsExtraShaderFormatsOverride;
+
+		/// <summary>
+		/// Whether we should include any extra shader formats. By default this is only enabled for Program and Editor targets.
+		/// </summary>
+		public bool bNeedsExtraShaderFormats
+		{
+			set { bNeedsExtraShaderFormatsOverride = value; }
+			get { return bNeedsExtraShaderFormatsOverride ?? (bForceBuildShaderFormats || bBuildDeveloperTools) && (Type == TargetType.Editor || Type == TargetType.Program); }
+		}
+
 		/// <summary>
 		/// Whether we should compile SQLite using the custom "Unreal" platform (true), or using the native platform (false).
 		/// </summary>
@@ -2315,6 +2329,11 @@ namespace UnrealBuildTool
 			get { return Inner.bForceBuildShaderFormats; }
 		}
 
+		public bool bNeedsExtraShaderFormats
+		{
+			get { return Inner.bNeedsExtraShaderFormats; }
+		}
+
 		public bool bCompileCustomSQLitePlatform
 		{
 			get { return Inner.bCompileCustomSQLitePlatform; }
