UnrealBuildTool: Relative paths result in a longer path length, leading to build failures

diff --git a/Engine/Source/Programs/UnrealBuildTool/Platform/Windows/VCToolChain.cs b/Engine/Source/Programs/UnrealBuildTool/Platform/Windows/VCToolChain.cs
--- a/Engine/Source/Programs/UnrealBuildTool/Platform/Windows/VCToolChain.cs
+++ b/Engine/Source/Programs/UnrealBuildTool/Platform/Windows/VCToolChain.cs
@@ -117,7 +117,7 @@ namespace UnrealBuildTool
 		public static string NormalizeCommandLinePath(FileSystemReference Reference, WindowsCompiler Compiler, bool bPreprocessOnly)
 		{
 			// Try to use a relative path to shorten command line length. Always need the full path when preprocessing because the output file will be in a different place, where include paths cannot be relative.
-			if (Reference.IsUnderDirectory(Unreal.RootDirectory) && !bPreprocessOnly)
+			if (Reference.IsUnderDirectory(Unreal.EngineDirectory) && !bPreprocessOnly)
 			{
 				return Reference.MakeRelativeTo(UnrealBuildTool.EngineSourceDirectory);
 			}
