UnrealBuildTool: Quote dll path when launching ubt from ubt

#rnx

diff --git a/Engine/Source/Programs/UnrealBuildTool/System/ActionGraphBuilder.cs b/Engine/Source/Programs/UnrealBuildTool/System/ActionGraphBuilder.cs
--- a/Engine/Source/Programs/UnrealBuildTool/System/ActionGraphBuilder.cs
+++ b/Engine/Source/Programs/UnrealBuildTool/System/ActionGraphBuilder.cs
@@ -285,7 +285,8 @@ namespace UnrealBuildTool
 
 			Action NewAction = Graph.CreateAction(Type);
 			NewAction.CommandPath = Unreal.DotnetPath;
-			NewAction.CommandArguments = $"{Unreal.UnrealBuildToolDllPath} -Mode={Attribute.Name} {Arguments}";
+			NewAction.CommandArguments = $"\"{Unreal.UnrealBuildToolDllPath}\" -Mode={Attribute.Name} {Arguments}";
+			NewAction.CommandDescription = Attribute.Name;
 			return NewAction;
 		}
 
