UnrealBuildTool: Minor cleanup of HybridExecutor

#rnx

diff --git a/Engine/Source/Programs/UnrealBuildTool/Executors/HybridExecutor.cs b/Engine/Source/Programs/UnrealBuildTool/Executors/HybridExecutor.cs
--- a/Engine/Source/Programs/UnrealBuildTool/Executors/HybridExecutor.cs
+++ b/Engine/Source/Programs/UnrealBuildTool/Executors/HybridExecutor.cs
@@ -1,16 +1,13 @@
 // Copyright Epic Games, Inc. All Rights Reserved.
 
-using System;
 using System.Collections.Generic;
 using System.Linq;
-using System.Text;
-using System.Threading.Tasks;
 
 namespace UnrealBuildTool
 {
 	/// <summary>
 	/// Executor which distributes large sets of parallel actions through a remote executor, and any leaf serial actions through a local executor. Currently uses
-	/// XGE for remote actions and ParallelExecutor for local actions.
+	/// XGE or SNDBS for remote actions and ParallelExecutor for local actions.
 	/// </summary>
 	class HybridExecutor : ActionExecutor
 	{
@@ -23,25 +20,25 @@ namespace UnrealBuildTool
 		/// <summary>
 		/// Executor to use for remote actions
 		/// </summary>
-		ActionExecutor RemoteExecutor;
+		readonly ActionExecutor RemoteExecutor;
 
 		/// <summary>
 		/// Executor to use for local actions
 		/// </summary>
-		ActionExecutor LocalExecutor;
+		readonly ActionExecutor LocalExecutor;
 
 		/// <summary>
 		/// Constructor
 		/// </summary>
-		public HybridExecutor(int InMaxLocalActions)
+		public HybridExecutor(int InMaxLocalActions, ActionExecutor? InLocalExecutor = null, ActionExecutor? InRemoteExecutor = null)
 		{
 			MaxLocalActions = InMaxLocalActions;
-			this.LocalExecutor = new ParallelExecutor(MaxLocalActions);
-			this.RemoteExecutor = new XGE();
+			LocalExecutor = InLocalExecutor ?? new ParallelExecutor(MaxLocalActions);
+			RemoteExecutor = InRemoteExecutor ?? (XGE.IsAvailable() ? (ActionExecutor)new XGE() : new SNDBS());
 
 			XmlConfig.ApplyTo(this);
 
-			if(MaxLocalActions == 0)
+			if (MaxLocalActions == 0)
 			{
 				MaxLocalActions = Utils.GetPhysicalProcessorCount();
 			}
@@ -53,16 +50,13 @@ namespace UnrealBuildTool
 		/// <returns>True if the executor may be used</returns>
 		public static bool IsAvailable()
 		{
-			return XGE.IsAvailable() && ParallelExecutor.IsAvailable();
+			return (XGE.IsAvailable() || SNDBS.IsAvailable()) && ParallelExecutor.IsAvailable();
 		}
 
 		/// <summary>
 		/// Name of this executor for diagnostic output
 		/// </summary>
-		public override string Name
-		{
-			get { return String.Format("hybrid ({0}+{1})", LocalExecutor.Name, RemoteExecutor.Name); }
-		}
+		public override string Name => $"hybrid ({LocalExecutor.Name}+{RemoteExecutor.Name})";
 
 		/// <summary>
 		/// Execute the given actions
@@ -73,9 +67,9 @@ namespace UnrealBuildTool
 		{
 			// Find the number of dependants for each action
 			Dictionary<LinkedAction, int> ActionToNumDependents = ActionsToExecute.ToDictionary(x => x, x => 0);
-			foreach(LinkedAction Action in ActionsToExecute)
+			foreach (LinkedAction Action in ActionsToExecute)
 			{
-				foreach(LinkedAction PrerequisiteAction in Action.PrerequisiteActions)
+				foreach (LinkedAction PrerequisiteAction in Action.PrerequisiteActions)
 				{
 					ActionToNumDependents[PrerequisiteAction]++;
 				}
@@ -83,20 +77,20 @@ namespace UnrealBuildTool
 
 			// Build up a set of leaf actions in several iterations, ensuring that the number of leaf actions in each 
 			HashSet<LinkedAction> LeafActions = new HashSet<LinkedAction>();
-			for(;;)
+			for (;;)
 			{
 				// Find all the leaf actions in the graph
 				List<LinkedAction> NewLeafActions = new List<LinkedAction>();
-				foreach(LinkedAction Action in ActionsToExecute)
+				foreach (LinkedAction Action in ActionsToExecute)
 				{
-					if(ActionToNumDependents[Action] == 0 && !LeafActions.Contains(Action))
+					if (ActionToNumDependents[Action] == 0 && !LeafActions.Contains(Action))
 					{
 						NewLeafActions.Add(Action);
 					}
 				}
 
 				// Exit once we can't prune any more layers from the tree
-				if(NewLeafActions.Count == 0 || NewLeafActions.Count >= MaxLocalActions)
+				if (NewLeafActions.Count == 0 || NewLeafActions.Count >= MaxLocalActions)
 				{
 					break;
 				}
@@ -105,9 +99,9 @@ namespace UnrealBuildTool
 				LeafActions.UnionWith(NewLeafActions);
 
 				// Decrement the dependent counts for any of their prerequisites, so we can try and remove those from the tree in another iteration
-				foreach(LinkedAction NewLeafAction in NewLeafActions)
+				foreach (LinkedAction NewLeafAction in NewLeafActions)
 				{
-					foreach(LinkedAction PrerequisiteAction in NewLeafAction.PrerequisiteActions)
+					foreach (LinkedAction PrerequisiteAction in NewLeafAction.PrerequisiteActions)
 					{
 						ActionToNumDependents[PrerequisiteAction]--;
 					}
@@ -117,9 +111,9 @@ namespace UnrealBuildTool
 			// Split the list of actions into those which should be executed locally and remotely
 			List<LinkedAction> LocalActionsToExecute = new List<LinkedAction>(LeafActions.Count);
 			List<LinkedAction> RemoteActionsToExecute = new List<LinkedAction>(ActionsToExecute.Count - LeafActions.Count);
-			foreach(LinkedAction ActionToExecute in ActionsToExecute)
+			foreach (LinkedAction ActionToExecute in ActionsToExecute)
 			{
-				if(LeafActions.Contains(ActionToExecute))
+				if (LeafActions.Contains(ActionToExecute))
 				{
 					LocalActionsToExecute.Add(ActionToExecute);
 				}
@@ -130,21 +124,15 @@ namespace UnrealBuildTool
 			}
 
 			// Execute the remote actions
-			if(RemoteActionsToExecute.Count > 0)
+			if (RemoteActionsToExecute.Count > 0 && !RemoteExecutor.ExecuteActions(RemoteActionsToExecute))
 			{
-				if (!RemoteExecutor.ExecuteActions(RemoteActionsToExecute))
-				{
-					return false;
-				}
+				return false;
 			}
 
 			// Pass all the local actions through to the parallel executor
-			if(LocalActionsToExecute.Count > 0)
+			if (LocalActionsToExecute.Count > 0 && !LocalExecutor.ExecuteActions(LocalActionsToExecute))
 			{
-				if(!LocalExecutor.ExecuteActions(LocalActionsToExecute))
-				{
-					return false;
-				}
+				return false;
 			}
 
 			return true;
