UnrealBuildTool: vscode project generation fixes

#rnx

diff --git a/Engine/Source/Programs/UnrealBuildTool/ProjectFiles/VisualStudioCode/VSCodeProjectFileGenerator.cs b/Engine/Source/Programs/UnrealBuildTool/ProjectFiles/VisualStudioCode/VSCodeProjectFileGenerator.cs
--- a/Engine/Source/Programs/UnrealBuildTool/ProjectFiles/VisualStudioCode/VSCodeProjectFileGenerator.cs
+++ b/Engine/Source/Programs/UnrealBuildTool/ProjectFiles/VisualStudioCode/VSCodeProjectFileGenerator.cs
@@ -295,16 +295,20 @@ namespace UnrealBuildTool
 			public readonly TargetType Type;
 			public readonly UnrealTargetPlatform Platform;
 			public readonly UnrealTargetConfiguration Configuration;
+			public readonly CppStandardVersion CppStandard;
 			public readonly FileReference? CompilerPath;
+			public readonly DirectoryReference? SysRootPath;
 			public readonly Dictionary<DirectoryReference, string> ModuleCommandLines;
 
-			public BuildTarget(string InName, TargetType InType, UnrealTargetPlatform InPlatform, UnrealTargetConfiguration InConfiguration, FileReference? InCompilerPath, Dictionary<DirectoryReference, string> InModulesCommandLines)
+			public BuildTarget(string InName, TargetType InType, UnrealTargetPlatform InPlatform, UnrealTargetConfiguration InConfiguration, CppStandardVersion InCppStandard, FileReference? InCompilerPath, DirectoryReference? InSysRootPath, Dictionary<DirectoryReference, string> InModulesCommandLines)
 			{
 				Name = InName;
 				Type = InType;
 				Platform = InPlatform;
 				Configuration = InConfiguration;
+				CppStandard = InCppStandard;
 				CompilerPath = InCompilerPath;
+				SysRootPath = InSysRootPath;
 				ModuleCommandLines = InModulesCommandLines;
 			}
 
@@ -319,7 +323,8 @@ namespace UnrealBuildTool
 			base.AddTargetForIntellisense(Target);
 
 			bool UsingClang = true;
-			FileReference? CompilerPath;
+			FileReference? CompilerPath = null;
+			DirectoryReference? SysRootPath = null;
 			if (HostPlatform == UnrealTargetPlatform.Win64)
 			{
 				VCEnvironment Environment = VCEnvironment.Create(WindowsPlatform.GetDefaultCompiler(null, Target.Rules.WindowsPlatform.Architecture), Target.Platform, Target.Rules.WindowsPlatform.Architecture, null, Target.Rules.WindowsPlatform.WindowsSdkVersion, null);
@@ -329,11 +334,17 @@ namespace UnrealBuildTool
 			else if (HostPlatform == UnrealTargetPlatform.Linux)
 			{
 				CompilerPath = FileReference.FromString(LinuxCommon.WhichClang());
+				string? InternalSDKPath = UEBuildPlatform.GetSDK(UnrealTargetPlatform.Linux)?.GetInternalSDKPath();
+				if (!string.IsNullOrEmpty(InternalSDKPath))
+				{
+					SysRootPath = DirectoryReference.FromString(InternalSDKPath);
+				}
 			}
 			else if (HostPlatform == UnrealTargetPlatform.Mac)
 			{
 				MacToolChainSettings Settings = new MacToolChainSettings(false);
 				CompilerPath = FileReference.FromString(Settings.ToolchainDir + "clang++");
+				SysRootPath = DirectoryReference.FromString(Settings.BaseSDKDir + "/MacOSX" + Settings.MacOSSDKVersion + ".sdk");
 			}
 			else
 			{
@@ -381,7 +392,7 @@ namespace UnrealBuildTool
 				}
 			}
 
-			BuildTargets.Add(new BuildTarget(Target.TargetName, Target.TargetType, Target.Platform, Target.Configuration, CompilerPath, ModuleDirectoryToCompileCommand));
+			BuildTargets.Add(new BuildTarget(Target.TargetName, Target.TargetType, Target.Platform, Target.Configuration, GlobalCompileEnvironment.CppStandard, CompilerPath, SysRootPath, ModuleDirectoryToCompileCommand));
 		}
 
 		private class ProjectData
@@ -588,9 +599,13 @@ namespace UnrealBuildTool
 					HashSet<FileReference> AllSourceFiles = new HashSet<FileReference>();
 					Dictionary<DirectoryReference, string> AllModuleCommandLines = new Dictionary<DirectoryReference, string>();
 					FileReference? CompilerPath = null;
+					DirectoryReference? SysRootPath = null;
+					CppStandardVersion CppStandard = CppStandardVersion.Default;
 					
 					foreach (ProjectData.Project Project in Projects.AllProjects)
 					{
+						AllSourceFiles.UnionWith(Project.SourceProject.SourceFiles.Select(x => x.Reference));
+
 						foreach (ProjectData.Target ProjectTarget in Project.Targets)
 						{
 							BuildTarget BuildTarget = BuildTargets.FirstOrDefault(Target => Target.Name == ProjectTarget.Name);
@@ -600,11 +615,10 @@ namespace UnrealBuildTool
 								continue;
 
 							string Name = string.Format("{0} {1} {2} {3} ({4})", ProjectTarget.Name, ProjectTarget.Type, BuildTarget.Platform, BuildTarget.Configuration, Project.Name);
-							WriteConfiguration(Name, Project.Name, Project.SourceProject.SourceFiles.Select(x => x.Reference), BuildTarget.CompilerPath!, BuildTarget.ModuleCommandLines, OutFile, OutputDirectory);
-
-							AllSourceFiles.UnionWith(Project.SourceProject.SourceFiles.Select(x => x.Reference));
+							WriteConfiguration(Name, Project.Name, Project.SourceProject.SourceFiles.Select(x => x.Reference), BuildTarget.CppStandard, BuildTarget.CompilerPath!, BuildTarget.SysRootPath, BuildTarget.ModuleCommandLines, OutFile, OutputDirectory);
 
 							CompilerPath = BuildTarget.CompilerPath;
+
 							foreach (KeyValuePair<DirectoryReference, string> Pair in BuildTarget.ModuleCommandLines)
 							{
 								if(!AllModuleCommandLines.ContainsKey(Pair.Key))
@@ -612,6 +626,15 @@ namespace UnrealBuildTool
 									AllModuleCommandLines[Pair.Key] = Pair.Value;
 								}
 							}
+
+							if (BuildTarget.CppStandard > CppStandard)
+							{
+								CppStandard = BuildTarget.CppStandard;
+							}
+							if (BuildTarget.SysRootPath != null)
+							{
+								SysRootPath = BuildTarget.SysRootPath;
+							}
 						}
 					}
 
@@ -629,7 +652,7 @@ namespace UnrealBuildTool
 						DefaultConfigName = "Win32";
 					}
 
-					WriteConfiguration(DefaultConfigName, "Default", AllSourceFiles, CompilerPath!, AllModuleCommandLines, OutFile, OutputDirectory);
+					WriteConfiguration(DefaultConfigName, "Default", AllSourceFiles, CppStandard, CompilerPath!, SysRootPath, AllModuleCommandLines, OutFile, OutputDirectory);
 				}
 				OutFile.EndArray();
 			}
@@ -638,11 +661,57 @@ namespace UnrealBuildTool
 			OutFile.Write(FileReference.Combine(OutputDirectory, "c_cpp_properties.json"));
 		}
 
-		private void WriteConfiguration(string Name, string ProjectName, IEnumerable<FileReference> SourceFiles, FileReference CompilerPath, Dictionary<DirectoryReference, string> ModuleCommandLines, JsonFile OutFile, DirectoryReference OutputDirectory)
+		private void WriteConfiguration(string Name, string ProjectName, IEnumerable<FileReference> SourceFiles, CppStandardVersion CppStandard, FileReference CompilerPath, DirectoryReference? SysRootPath, Dictionary<DirectoryReference, string> ModuleCommandLines, JsonFile OutFile, DirectoryReference OutputDirectory)
 		{
 			OutFile.BeginObject();
 
 			OutFile.AddField("name", Name);
+			OutFile.AddField("compilerPath", CompilerPath.FullName);
+
+			if (HostPlatform == UnrealTargetPlatform.Mac)
+			{
+				string SysRoot = SysRootPath != null ? SysRootPath.FullName : "/Applications/Xcode.app/Contents/Developer/Platforms/MacOSX.platform/Developer/SDKs/MacOSX.sdk";
+				OutFile.BeginArray("compilerArgs");
+				{
+					OutFile.AddUnnamedField("-isysroot");
+					OutFile.AddUnnamedField(SysRoot);
+				}
+				OutFile.EndArray();
+
+				OutFile.BeginArray("macFrameworkPath");
+				{
+					OutFile.AddUnnamedField(SysRoot + "/System/Library/Frameworks");
+				}
+				OutFile.EndArray();
+			}
+			else if (SysRootPath != null)
+			{
+				OutFile.BeginArray("compilerArgs");
+				{
+					OutFile.AddUnnamedField("-isysroot");
+					OutFile.AddUnnamedField(SysRootPath.FullName);
+				}
+				OutFile.EndArray();
+			}
+
+			switch (CppStandard)
+			{
+				case CppStandardVersion.Cpp14:
+					OutFile.AddField("cStandard", "c11");
+					OutFile.AddField("cppStandard", "c++14");
+					break;
+				case CppStandardVersion.Cpp17:
+					OutFile.AddField("cStandard", "c17");
+					OutFile.AddField("cppStandard", "c++17");
+					break;
+				case CppStandardVersion.Cpp20:
+				case CppStandardVersion.Latest:
+					OutFile.AddField("cStandard", "c17");
+					OutFile.AddField("cppStandard", "c++20");
+					break;
+				default:
+					throw new BuildException($"Unsupported C++ standard type set: {CppStandard}");
+			}
 
 			if (HostPlatform == UnrealTargetPlatform.Win64)
 			{
@@ -653,16 +722,6 @@ namespace UnrealBuildTool
 				OutFile.AddField("intelliSenseMode", "clang-x64");
 			}
 
-			if (HostPlatform == UnrealTargetPlatform.Mac)
-			{
-				OutFile.BeginArray("macFrameworkPath");
-				{
-					OutFile.AddUnnamedField("/System/Library/Frameworks");
-					OutFile.AddUnnamedField("/Library/Frameworks");
-				}
-				OutFile.EndArray();
-			}
-
 			FileReference CompileCommands = FileReference.Combine(OutputDirectory, string.Format("compileCommands_{0}.json", ProjectName));
 			WriteCompileCommands(CompileCommands, SourceFiles, CompilerPath, ModuleCommandLines);
 			OutFile.AddField("compileCommands", MakePathString(CompileCommands, bInAbsolute: true, bForceSkipQuotes: true));
@@ -817,7 +876,10 @@ namespace UnrealBuildTool
 
 					Writer.WriteObjectStart();
 					Writer.WriteValue("file", MakePathString(File, bInAbsolute: true, bForceSkipQuotes: true));
-					Writer.WriteValue("command", String.Format("{0} @{1}", MakePathString(CompilerPath, bInAbsolute: true), MakePathString(ResponseFile, bInAbsolute: true)));
+					Writer.WriteArrayStart("arguments");
+					Writer.WriteValue(MakePathString(CompilerPath, bInAbsolute: true, bForceSkipQuotes: true));
+					Writer.WriteValue($"@{MakePathString(ResponseFile, bInAbsolute: true, bForceSkipQuotes: true)}");
+					Writer.WriteArrayEnd();
 					Writer.WriteValue("directory", UnrealBuildTool.EngineSourceDirectory.ToString());
 					Writer.WriteObjectEnd();
 				}
