VCToolChain: preprocessing no longer needs special handling when normalizing command line paths

#rn trivial
#rnx

diff --git a/Engine/Source/Programs/UnrealBuildTool/Platform/Windows/VCCompileAction.cs b/Engine/Source/Programs/UnrealBuildTool/Platform/Windows/VCCompileAction.cs
--- a/Engine/Source/Programs/UnrealBuildTool/Platform/Windows/VCCompileAction.cs
+++ b/Engine/Source/Programs/UnrealBuildTool/Platform/Windows/VCCompileAction.cs
@@ -396,17 +396,17 @@ namespace UnrealBuildTool
 			List<string> Arguments = new List<string>();
 			if (SourceFile != null)
 			{
-				VCToolChain.AddSourceFile(Arguments, SourceFile, CompilerType, PreprocessedFile != null);
+				VCToolChain.AddSourceFile(Arguments, SourceFile);
 			}
 
 			foreach (DirectoryReference IncludePath in IncludePaths)
 			{
-				VCToolChain.AddIncludePath(Arguments, IncludePath, CompilerType, PreprocessedFile != null);
+				VCToolChain.AddIncludePath(Arguments, IncludePath, CompilerType);
 			}
 
 			foreach (DirectoryReference SystemIncludePath in SystemIncludePaths)
 			{
-				VCToolChain.AddSystemIncludePath(Arguments, SystemIncludePath, CompilerType, PreprocessedFile != null);
+				VCToolChain.AddSystemIncludePath(Arguments, SystemIncludePath, CompilerType);
 			}
 
 			foreach (string Definition in Definitions)
@@ -418,22 +418,22 @@ namespace UnrealBuildTool
 
 			foreach (FileItem ForceIncludeFile in ForceIncludeFiles)
 			{
-				VCToolChain.AddForceIncludeFile(Arguments, ForceIncludeFile, CompilerType, PreprocessedFile != null);
+				VCToolChain.AddForceIncludeFile(Arguments, ForceIncludeFile);
 			}
 
 			if (CreatePchFile != null)
 			{
-				VCToolChain.AddCreatePchFile(Arguments, PchThroughHeaderFile!, CreatePchFile, CompilerType, PreprocessedFile != null);
+				VCToolChain.AddCreatePchFile(Arguments, PchThroughHeaderFile!, CreatePchFile);
 			}
 
 			if (UsingPchFile != null && CompilerType.IsMSVC())
 			{
-				VCToolChain.AddUsingPchFile(Arguments, PchThroughHeaderFile!, UsingPchFile, CompilerType, PreprocessedFile != null);
+				VCToolChain.AddUsingPchFile(Arguments, PchThroughHeaderFile!, UsingPchFile);
 			}
 
 			if (PreprocessedFile != null)
 			{
-				VCToolChain.AddPreprocessedFile(Arguments, PreprocessedFile, CompilerType, PreprocessedFile != null);
+				VCToolChain.AddPreprocessedFile(Arguments, PreprocessedFile);
 
 				// this is parsed by external tools wishing to open this file directly.
 				Log.TraceInformation("PreProcessPath: " + PreprocessedFile);
@@ -441,18 +441,18 @@ namespace UnrealBuildTool
 
 			if (ObjectFile != null)
 			{
-				VCToolChain.AddObjectFile(Arguments, ObjectFile, CompilerType, PreprocessedFile != null);
+				VCToolChain.AddObjectFile(Arguments, ObjectFile);
 			}
 
 			// A better way to express this? .json is used as output for /sourceDependencies), but .md.json is used as output for /sourceDependencies:directives)
 			if (DependencyListFile != null && DependencyListFile.HasExtension(".json") && !DependencyListFile.HasExtension(".md.json"))
 			{
-				VCToolChain.AddSourceDependenciesFile(Arguments, DependencyListFile, CompilerType, PreprocessedFile != null);
+				VCToolChain.AddSourceDependenciesFile(Arguments, DependencyListFile);
 			}
 
 			if (DependencyListFile != null && DependencyListFile.HasExtension(".d"))
 			{
-				VCToolChain.AddSourceDependsFile(Arguments, DependencyListFile, CompilerType, PreprocessedFile != null);
+				VCToolChain.AddSourceDependsFile(Arguments, DependencyListFile);
 			}
 
 			Arguments.AddRange(this.Arguments);
@@ -467,7 +467,7 @@ namespace UnrealBuildTool
 			}
 			else
 			{
-				string ResponseFileString = VCToolChain.NormalizeCommandLinePath(ResponseFile, CompilerType, PreprocessedFile != null);
+				string ResponseFileString = VCToolChain.NormalizeCommandLinePath(ResponseFile);
 
 				// cl.exe can't handle response files with a path longer than 260 characters, and relative paths can push it over the limit
 				if (!System.IO.Path.IsPathRooted(ResponseFileString) && System.IO.Path.Combine(WorkingDirectory.FullName, ResponseFileString).Length > 260)
@@ -481,12 +481,12 @@ namespace UnrealBuildTool
 		string GetClFilterArguments()
 		{
 			List<string> Arguments = new List<string>();
-			string DependencyListFileString = VCToolChain.NormalizeCommandLinePath(DependencyListFile!, CompilerType, PreprocessedFile != null);
+			string DependencyListFileString = VCToolChain.NormalizeCommandLinePath(DependencyListFile!);
 			Arguments.Add(String.Format("-dependencies={0}", Utils.MakePathSafeToUseWithCommandLine(DependencyListFileString)));
 
 			if (TimingFile != null)
 			{
-				string TimingFileString = VCToolChain.NormalizeCommandLinePath(TimingFile, CompilerType, PreprocessedFile != null);
+				string TimingFileString = VCToolChain.NormalizeCommandLinePath(TimingFile);
 				Arguments.Add(String.Format("-timing={0}", Utils.MakePathSafeToUseWithCommandLine(TimingFileString)));
 			}
 			if (bShowIncludes)
diff --git a/Engine/Source/Programs/UnrealBuildTool/Platform/Windows/VCToolChain.cs b/Engine/Source/Programs/UnrealBuildTool/Platform/Windows/VCToolChain.cs
--- a/Engine/Source/Programs/UnrealBuildTool/Platform/Windows/VCToolChain.cs
+++ b/Engine/Source/Programs/UnrealBuildTool/Platform/Windows/VCToolChain.cs
@@ -68,11 +68,11 @@ namespace UnrealBuildTool
 		{
 			if(EnvVars.Compiler == EnvVars.ToolChain)
 			{
-				Lines.Add(String.Format("Using {0} {1} toolchain ({2}) and Windows {3} SDK ({4}).", WindowsPlatform.GetCompilerName(EnvVars.Compiler), EnvVars.ToolChainVersion, EnvVars.ToolChainDir, EnvVars.WindowsSdkVersion, EnvVars.WindowsSdkDir));
+				Lines.Add($"Using {WindowsPlatform.GetCompilerName(EnvVars.Compiler)} {EnvVars.ToolChainVersion} toolchain ({EnvVars.ToolChainDir}) and Windows {EnvVars.WindowsSdkVersion} SDK ({EnvVars.WindowsSdkDir}).");
 			}
 			else
 			{
-				Lines.Add(String.Format("Using {0} {1} compiler ({2}) with {3} {4} runtime ({5}) and Windows {6} SDK ({7}).", WindowsPlatform.GetCompilerName(EnvVars.Compiler), EnvVars.CompilerVersion, EnvVars.CompilerDir, WindowsPlatform.GetCompilerName(EnvVars.ToolChain), EnvVars.ToolChainVersion, EnvVars.ToolChainDir, EnvVars.WindowsSdkVersion, EnvVars.WindowsSdkDir));
+				Lines.Add($"Using {WindowsPlatform.GetCompilerName(EnvVars.Compiler)} {EnvVars.CompilerVersion} compiler ({EnvVars.CompilerDir}) with {WindowsPlatform.GetCompilerName(EnvVars.ToolChain)} {EnvVars.ToolChainVersion} runtime ({EnvVars.ToolChainDir}) and Windows {EnvVars.WindowsSdkVersion} SDK ({EnvVars.WindowsSdkDir}).");
 			}
 		}
 
@@ -114,10 +114,10 @@ namespace UnrealBuildTool
 			}
 		}
 
-		public static string NormalizeCommandLinePath(FileSystemReference Reference, WindowsCompiler Compiler, bool bPreprocessOnly)
+		public static string NormalizeCommandLinePath(FileSystemReference Reference)
 		{
-			// Try to use a relative path to shorten command line length. Always need the full path when preprocessing because the output file will be in a different place, where include paths cannot be relative.
-			if (Reference.IsUnderDirectory(Unreal.EngineDirectory) && !bPreprocessOnly)
+			// Try to use a relative path to shorten command line length.
+			if (Reference.IsUnderDirectory(Unreal.EngineDirectory))
 			{
 				return Reference.MakeRelativeTo(UnrealBuildTool.EngineSourceDirectory);
 			}
@@ -125,21 +125,21 @@ namespace UnrealBuildTool
 			return Reference.FullName;
 		}
 
-		public static string NormalizeCommandLinePath(FileItem Item, WindowsCompiler Compiler, bool bPreprocessOnly)
+		public static string NormalizeCommandLinePath(FileItem Item)
 		{
-			return NormalizeCommandLinePath(Item.Location, Compiler, bPreprocessOnly);
+			return NormalizeCommandLinePath(Item.Location);
 		}
 
-		public static void AddSourceFile(List<string> Arguments, FileItem SourceFile, WindowsCompiler Compiler, bool bPreprocessOnly)
+		public static void AddSourceFile(List<string> Arguments, FileItem SourceFile)
 		{
-			string SourceFileString = NormalizeCommandLinePath(SourceFile, Compiler, bPreprocessOnly);
+			string SourceFileString = NormalizeCommandLinePath(SourceFile);
 			Arguments.Add(Utils.MakePathSafeToUseWithCommandLine(SourceFileString));
 		}
 
-		public static void AddIncludePath(List<string> Arguments, DirectoryReference IncludePath, WindowsCompiler Compiler, bool bPreprocessOnly, bool bSystemInclude = false)
+		private static void AddIncludePath(List<string> Arguments, DirectoryReference IncludePath, WindowsCompiler Compiler, bool bSystemInclude)
 		{
 			// Try to use a relative path to shorten command line length. Always need the full path when preprocessing because the output file will be in a different place, where include paths cannot be relative.
-			string IncludePathString = NormalizeCommandLinePath(IncludePath, Compiler, bPreprocessOnly);
+			string IncludePathString = NormalizeCommandLinePath(IncludePath);
 
 			if (Compiler.IsClang() && bSystemInclude)
 			{
@@ -153,36 +153,41 @@ namespace UnrealBuildTool
 			}
 		}
 
-		public static void AddSystemIncludePath(List<string> Arguments, DirectoryReference IncludePath, WindowsCompiler Compiler, bool bPreprocessOnly)
+		public static void AddIncludePath(List<string> Arguments, DirectoryReference IncludePath, WindowsCompiler Compiler)
+		{
+			AddIncludePath(Arguments, IncludePath, Compiler, false);
+		}
+
+		public static void AddSystemIncludePath(List<string> Arguments, DirectoryReference IncludePath, WindowsCompiler Compiler)
 		{
-			AddIncludePath(Arguments, IncludePath, Compiler, bPreprocessOnly, true);
+			AddIncludePath(Arguments, IncludePath, Compiler, true);
 		}
 
-		public static void AddForceIncludeFile(List<string> Arguments, FileItem ForceIncludeFile, WindowsCompiler Compiler, bool bPreprocessOnly)
+		public static void AddForceIncludeFile(List<string> Arguments, FileItem ForceIncludeFile)
 		{
-			string ForceIncludeFileString = NormalizeCommandLinePath(ForceIncludeFile, Compiler, bPreprocessOnly);
+			string ForceIncludeFileString = NormalizeCommandLinePath(ForceIncludeFile);
 			Arguments.Add($"/FI\"{ForceIncludeFileString}\"");
 		}
 
-		public static void AddCreatePchFile(List<string> Arguments, FileItem PchThroughHeaderFile, FileItem CreatePchFile, WindowsCompiler Compiler, bool bPreprocessOnly)
+		public static void AddCreatePchFile(List<string> Arguments, FileItem PchThroughHeaderFile, FileItem CreatePchFile)
 		{
-			string PchThroughHeaderFilePath = NormalizeCommandLinePath(PchThroughHeaderFile, Compiler, bPreprocessOnly);
-			string CreatePchFilePath = NormalizeCommandLinePath(CreatePchFile, Compiler, bPreprocessOnly);
+			string PchThroughHeaderFilePath = NormalizeCommandLinePath(PchThroughHeaderFile);
+			string CreatePchFilePath = NormalizeCommandLinePath(CreatePchFile);
 			Arguments.Add($"/Yc\"{PchThroughHeaderFilePath}\"");
 			Arguments.Add($"/Fp\"{CreatePchFilePath}\"");
 		}
 
-		public static void AddUsingPchFile(List<string> Arguments, FileItem PchThroughHeaderFile, FileItem UsingPchFile, WindowsCompiler Compiler, bool bPreprocessOnly)
+		public static void AddUsingPchFile(List<string> Arguments, FileItem PchThroughHeaderFile, FileItem UsingPchFile)
 		{
-			string PchThroughHeaderFilePath = NormalizeCommandLinePath(PchThroughHeaderFile, Compiler, bPreprocessOnly);
-			string UsingPchFilePath = NormalizeCommandLinePath(UsingPchFile, Compiler, bPreprocessOnly);
+			string PchThroughHeaderFilePath = NormalizeCommandLinePath(PchThroughHeaderFile);
+			string UsingPchFilePath = NormalizeCommandLinePath(UsingPchFile);
 			Arguments.Add($"/Yu\"{PchThroughHeaderFilePath}\"");
 			Arguments.Add($"/Fp\"{UsingPchFilePath}\"");
 		}
 
-		public static void AddPreprocessedFile(List<string> Arguments, FileItem PreprocessedFile, WindowsCompiler Compiler, bool bPreprocessOnly)
+		public static void AddPreprocessedFile(List<string> Arguments, FileItem PreprocessedFile)
 		{
-			string PreprocessedFileString = NormalizeCommandLinePath(PreprocessedFile, Compiler, bPreprocessOnly);
+			string PreprocessedFileString = NormalizeCommandLinePath(PreprocessedFile);
 			Arguments.Add("/P"); // Preprocess
 			Arguments.Add("/C"); // Preserve comments when preprocessing
 			Arguments.Add($"/Fi\"{PreprocessedFileString}\""); // Preprocess to a file
@@ -191,21 +196,21 @@ namespace UnrealBuildTool
 			Log.TraceInformation("PreProcessPath: " + PreprocessedFile);
 		}
 
-		public static void AddObjectFile(List<string> Arguments, FileItem ObjectFile, WindowsCompiler Compiler, bool bPreprocessOnly)
+		public static void AddObjectFile(List<string> Arguments, FileItem ObjectFile)
 		{
-			string ObjectFileString = NormalizeCommandLinePath(ObjectFile, Compiler, bPreprocessOnly);
+			string ObjectFileString = NormalizeCommandLinePath(ObjectFile);
 			Arguments.Add($"/Fo\"{ObjectFileString}\"");
 		}
 
-		public static void AddSourceDependenciesFile(List<string> Arguments, FileItem SourceDependenciesFile, WindowsCompiler Compiler, bool bPreprocessOnly)
+		public static void AddSourceDependenciesFile(List<string> Arguments, FileItem SourceDependenciesFile)
 		{
-			string SourceDependenciesFileString = NormalizeCommandLinePath(SourceDependenciesFile, Compiler, bPreprocessOnly);
+			string SourceDependenciesFileString = NormalizeCommandLinePath(SourceDependenciesFile);
 			Arguments.Add("/sourceDependencies " + Utils.MakePathSafeToUseWithCommandLine(SourceDependenciesFileString));
 		}
 
-		public static void AddSourceDependsFile(List<string> Arguments, FileItem SourceDependsFile, WindowsCompiler Compiler, bool bPreprocessOnly)
+		public static void AddSourceDependsFile(List<string> Arguments, FileItem SourceDependsFile)
 		{
-			string SourceDependsFileString = NormalizeCommandLinePath(SourceDependsFile, Compiler, bPreprocessOnly);
+			string SourceDependsFileString = NormalizeCommandLinePath(SourceDependsFile);
 			Arguments.Add($"/clang:-MD /clang:-MF\"{SourceDependsFileString}\"");
 		}
 
@@ -284,7 +289,7 @@ namespace UnrealBuildTool
 
 					foreach (String Checker in EnabledCheckers)
 					{
-						Arguments.Add(String.Format("-Xclang -analyzer-checker={0}", Checker));
+						Arguments.Add($"-Xclang -analyzer-checker={Checker}");
 					}
 				}
 			}
@@ -366,7 +371,7 @@ namespace UnrealBuildTool
 			{
 				if(Target.WindowsPlatform.PCHMemoryAllocationFactor > 0)
 				{
-					Arguments.Add(String.Format("/Zm{0}", Target.WindowsPlatform.PCHMemoryAllocationFactor));
+					Arguments.Add($"/Zm{Target.WindowsPlatform.PCHMemoryAllocationFactor}");
 				}
 			}
 
@@ -721,7 +726,7 @@ namespace UnrealBuildTool
 					FileSpecifier += "-header";
 				}
 
-				Arguments.Add(String.Format("-Xclang -x -Xclang \"{0}\"", FileSpecifier));
+				Arguments.Add($"-Xclang -x -Xclang \"{FileSpecifier}\"");
 			}
 
 
@@ -927,7 +932,7 @@ namespace UnrealBuildTool
 			{
 				Arguments.Add("/MANIFEST:EMBED");
 				FileItem ManifestFile = FileItem.GetItemByPath(Target.WindowsPlatform.ManifestFile);
-				Arguments.Add(String.Format("/MANIFESTINPUT:\"{0}\"", NormalizeCommandLinePath(ManifestFile, Target.WindowsPlatform.Compiler, false)));
+				Arguments.Add($"/MANIFESTINPUT:\"{NormalizeCommandLinePath(ManifestFile)}\"");
 			}
 
 			// Prevents the linker from displaying its logo for each invocation.
@@ -963,7 +968,7 @@ namespace UnrealBuildTool
 			//
 			if (LinkEnvironment.Platform.IsInGroup(UnrealPlatformGroup.Windows) || LinkEnvironment.Platform.IsInGroup(UnrealPlatformGroup.HoloLens))
 			{
-				Arguments.Add(string.Format("/MACHINE:{0}", WindowsExports.GetArchitectureSubpath(Target.WindowsPlatform.Architecture)));
+				Arguments.Add($"/MACHINE:{WindowsExports.GetArchitectureSubpath(Target.WindowsPlatform.Architecture)}");
 
 				{
 					if (LinkEnvironment.bIsBuildingConsoleApplication)
@@ -1107,7 +1112,7 @@ namespace UnrealBuildTool
 			//
 			if (LinkEnvironment.Platform.IsInGroup(UnrealPlatformGroup.Windows) || LinkEnvironment.Platform.IsInGroup(UnrealPlatformGroup.HoloLens))
 			{
-				Arguments.Add(string.Format("/MACHINE:{0}", WindowsExports.GetArchitectureSubpath(Target.WindowsPlatform.Architecture)));
+				Arguments.Add($"/MACHINE:{WindowsExports.GetArchitectureSubpath(Target.WindowsPlatform.Architecture)}");
 
 				{
 					if (LinkEnvironment.bIsBuildingConsoleApplication)
@@ -1183,8 +1188,8 @@ namespace UnrealBuildTool
 				if (CompileEnvironment.PrecompiledHeaderAction == PrecompiledHeaderAction.Create)
 				{
 					// Generate a CPP File that just includes the precompiled header.
-					string PrecompiledHeaderIncludeFilenameString = NormalizeCommandLinePath(CompileEnvironment.PrecompiledHeaderIncludeFilename!, Target.WindowsPlatform.Compiler, CompileEnvironment.bPreprocessOnly);
-					string PchCppFile = string.Format("// Compiler: {0}\n#include \"{1}\"\r\n", EnvVars.CompilerVersion, PrecompiledHeaderIncludeFilenameString.Replace('\\', '/'));
+					string PrecompiledHeaderIncludeFilenameString = NormalizeCommandLinePath(CompileEnvironment.PrecompiledHeaderIncludeFilename!);
+					string PchCppFile = $"// Compiler: {EnvVars.CompilerVersion}\n#include \"{PrecompiledHeaderIncludeFilenameString.Replace('\\', '/')}\"\r\n";
 					CompileAction.SourceFile = FileItem.GetItemByFileReference(CompileEnvironment.PrecompiledHeaderIncludeFilename!.ChangeExtension(".cpp"));
 					Graph.CreateIntermediateTextFile(CompileAction.SourceFile, PchCppFile);
 
@@ -1201,7 +1206,7 @@ namespace UnrealBuildTool
 						// NOTE: The symbol name we use here is arbitrary, and all that matters is that it is
 						// unique per PCH module used in our library
 						string FakeUniquePCHSymbolName = CompileEnvironment.PrecompiledHeaderIncludeFilename.GetFileNameWithoutExtension();
-						CompileAction.Arguments.Add(String.Format("/Yl{0}", FakeUniquePCHSymbolName));
+						CompileAction.Arguments.Add($"/Yl{FakeUniquePCHSymbolName}");
 					}
 				}
 
@@ -1223,7 +1228,7 @@ namespace UnrealBuildTool
 					{
 						using (StreamWriter Writer = File.AppendText(Target.WindowsPlatform.ObjSrcMapFile))
 						{
-							Writer.WriteLine(string.Format("\"{0}\" -> \"{1}\"", ObjectLeafFilename, SourceFile.AbsolutePath));
+							Writer.WriteLine($"\"{ObjectLeafFilename}\" -> \"{SourceFile.AbsolutePath}\"");
 						}
 					}
 
@@ -1273,7 +1278,7 @@ namespace UnrealBuildTool
 					}
 
 					// Specify the PDB file that the compiler should write to.
-					CompileAction.Arguments.Add(String.Format("/Fd\"{0}\"", PDBLocation));
+					CompileAction.Arguments.Add($"/Fd\"{PDBLocation}\"");
 
 					// Don't allow remote execution when PDB files are enabled; we need to modify the same files. XGE works around this by generating separate
 					// PDB files per agent, but this functionality is only available with the Visual C++ extension package (via the VCCompiler=true tool option).
@@ -1327,7 +1332,7 @@ namespace UnrealBuildTool
 				{
 					FileItem IfcFile = FileItem.GetItemByFileReference(FileReference.Combine(GetModuleInterfaceDir(OutputDir), SourceFile.Location.ChangeExtension(".ifc").GetFileName()));
 					CompileAction.Arguments.Add("/interface");
-					CompileAction.Arguments.Add(String.Format("/ifcOutput \"{0}\"", IfcFile.Location));
+					CompileAction.Arguments.Add($"/ifcOutput \"{IfcFile.Location}\"");
 					CompileAction.CompiledModuleInterfaceFile = IfcFile;
 
 					FileItem IfcDepsFile = FileItem.GetItemByFileReference(FileReference.Combine(OutputDir, SourceFile.Location.GetFileName() + ".md.json"));
@@ -1337,7 +1342,7 @@ namespace UnrealBuildTool
 					CompileDepsAction.ResponseFile = FileItem.GetItemByPath(IfcDepsFile + ".response");
 					CompileDepsAction.ObjectFile = null;
 					CompileDepsAction.DependencyListFile = IfcDepsFile;
-					CompileDepsAction.Arguments.Add(String.Format("/sourceDependencies:directives \"{0}\"", IfcDepsFile.Location));
+					CompileDepsAction.Arguments.Add($"/sourceDependencies:directives \"{IfcDepsFile.Location}\"");
 					CompileDepsAction.AdditionalPrerequisiteItems.Add(SourceFile);
 					CompileDepsAction.AdditionalPrerequisiteItems.AddRange(CompileEnvironment.ForceIncludeFiles);
 					CompileDepsAction.AdditionalPrerequisiteItems.AddRange(CompileEnvironment.AdditionalPrerequisites);
@@ -1358,7 +1363,7 @@ namespace UnrealBuildTool
 				if (Target.bPrintToolChainTimingInfo || Target.WindowsPlatform.bCompilerTrace)
 				{
 					CompileAction.ForceClFilter = true;
-					CompileAction.TimingFile = FileItem.GetItemByFileReference(FileReference.Combine(OutputDir, String.Format("{0}.timing", SourceFile.Location.GetFileName())));
+					CompileAction.TimingFile = FileItem.GetItemByFileReference(FileReference.Combine(OutputDir, $"{SourceFile.Location.GetFileName()}.timing"));
 					GenerateParseTimingInfoAction(SourceFile, CompileAction.TimingFile, Graph);
 				}
 
@@ -1366,15 +1371,15 @@ namespace UnrealBuildTool
 				{
 					if (Target.WindowsPlatform.Compiler.IsMSVC() && !CompileAction.ForceClFilter)
 					{
-						CompileAction.DependencyListFile = FileItem.GetItemByFileReference(FileReference.Combine(OutputDir, String.Format("{0}.json", SourceFile.Location.GetFileName())));
+						CompileAction.DependencyListFile = FileItem.GetItemByFileReference(FileReference.Combine(OutputDir, $"{SourceFile.Location.GetFileName()}.json"));
 					}
 					else if (Target.WindowsPlatform.Compiler.IsClang())
 					{
-						CompileAction.DependencyListFile = FileItem.GetItemByFileReference(FileReference.Combine(OutputDir, String.Format("{0}.d", SourceFile.Location.GetFileName())));
+						CompileAction.DependencyListFile = FileItem.GetItemByFileReference(FileReference.Combine(OutputDir, $"{SourceFile.Location.GetFileName()}.d"));
 					}
 					else
 					{
-						CompileAction.DependencyListFile = FileItem.GetItemByFileReference(FileReference.Combine(OutputDir, String.Format("{0}.txt", SourceFile.Location.GetFileName())));
+						CompileAction.DependencyListFile = FileItem.GetItemByFileReference(FileReference.Combine(OutputDir, $"{SourceFile.Location.GetFileName()}.txt"));
 						CompileAction.bShowIncludes = Target.WindowsPlatform.bShowIncludes;
 					}
 				}
@@ -1407,7 +1412,7 @@ namespace UnrealBuildTool
 		{
 			FileItem TimingJsonFile = FileItem.GetItemByPath(Path.ChangeExtension(TimingFile.AbsolutePath, ".cta"));
 
-			string ParseTimingArguments = String.Format("-TimingFile=\"{0}\"", TimingFile);
+			string ParseTimingArguments = $"-TimingFile=\"{TimingFile}\"";
 			if (Target.bParseTimingInfoForTracing)
 			{
 				ParseTimingArguments += " -Tracing";
@@ -1443,12 +1448,12 @@ namespace UnrealBuildTool
 					}
 					File.WriteAllLines(ManifestFile.FullName, TimingJsonFiles.Select(f => f.AbsolutePath));
 
-					FileReference ExpectedCompileTimeFile = FileReference.FromString(Path.Combine(Makefile.ProjectIntermediateDirectory.FullName, String.Format("{0}.json", Target.Name)));
+					FileReference ExpectedCompileTimeFile = FileReference.FromString(Path.Combine(Makefile.ProjectIntermediateDirectory.FullName, $"{Target.Name}.json"));
 					List<string> ActionArgs = new List<string>()
 					{
-						String.Format("-Name={0}", Target.Name),
-						String.Format("-ManifestFile={0}", ManifestFile.FullName),
-						String.Format("-CompileTimingFile={0}", ExpectedCompileTimeFile),
+						$"-Name={Target.Name}",
+						$"-ManifestFile={ManifestFile.FullName}",
+						$"-CompileTimingFile={ExpectedCompileTimeFile}",
 					};
 
 					Action AggregateTimingInfoAction = Makefile.CreateRecursiveAction<AggregateParsedTimingInfo>(ActionType.ParseTimingInfo, string.Join(" ", ActionArgs));
@@ -1484,8 +1489,8 @@ namespace UnrealBuildTool
 				}
 
 
-				String ASanRuntimeDLL = String.Format("clang_rt.asan_dynamic-{0}.dll", ASanArchSuffix);
-				String ASanDebugRuntimeDLL = String.Format("clang_rt.asan_dbg_dynamic-{0}.dll", ASanArchSuffix);
+				string ASanRuntimeDLL = $"clang_rt.asan_dynamic-{ASanArchSuffix}.dll";
+				string ASanDebugRuntimeDLL = $"clang_rt.asan_dbg_dynamic-{ASanArchSuffix}.dll";
 
 				RuntimeDependencies.Add(new RuntimeDependency(FileReference.Combine(ExeDir, ASanRuntimeDLL), StagedFileType.NonUFS));
 				TargetFileToSourceFile[FileReference.Combine(ExeDir, ASanRuntimeDLL)] = FileReference.Combine(ASanRuntimeDir, ASanRuntimeDLL);
@@ -1538,17 +1543,17 @@ namespace UnrealBuildTool
 				// Include paths. Don't use AddIncludePath() here, since it uses the full path and exceeds the max command line length.
 				foreach (DirectoryReference IncludePath in CompileEnvironment.UserIncludePaths)
 				{
-					Arguments.Add(String.Format("/I \"{0}\"", NormalizeCommandLinePath(IncludePath, Target.WindowsPlatform.Compiler, false)));
+					Arguments.Add($"/I \"{NormalizeCommandLinePath(IncludePath)}\"");
 				}
 
 				// System include paths.
 				foreach (DirectoryReference SystemIncludePath in CompileEnvironment.SystemIncludePaths)
 				{
-					Arguments.Add(String.Format("/I \"{0}\"", NormalizeCommandLinePath(SystemIncludePath, Target.WindowsPlatform.Compiler, false)));
+					Arguments.Add($"/I \"{NormalizeCommandLinePath(SystemIncludePath)}\"");
 				}
 				foreach (DirectoryReference SystemIncludePath in EnvVars.IncludePaths)
 				{
-					Arguments.Add(String.Format("/I \"{0}\"", NormalizeCommandLinePath(SystemIncludePath, Target.WindowsPlatform.Compiler, false)));
+					Arguments.Add($"/I \"{NormalizeCommandLinePath(SystemIncludePath)}\"");
 				}
 
 				// Preprocessor definitions.
@@ -1573,29 +1578,29 @@ namespace UnrealBuildTool
 				CompileAction.PrerequisiteItems.Add(FileItem.GetItemByFileReference(IconFile));
 
 				// Setup the compile environment, setting the icon to use via a macro. This is used in Default.rc2.
-				AddDefinition(Arguments, String.Format("BUILD_ICON_FILE_NAME=\"\\\"{0}\\\"\"", NormalizeCommandLinePath(IconFile, Target.WindowsPlatform.Compiler, false).Replace("\\", "\\\\")));
+				AddDefinition(Arguments, $"BUILD_ICON_FILE_NAME=\"\\\"{NormalizeCommandLinePath(IconFile).Replace("\\", "\\\\")}\\\"\"");
 
 				// Apply the target settings for the resources
 				if(!CompileEnvironment.bUseSharedBuildEnvironment)
 				{
 					if (!String.IsNullOrEmpty(Target.WindowsPlatform.CompanyName))
 					{
-						AddDefinition(Arguments, String.Format("PROJECT_COMPANY_NAME={0}", SanitizeMacroValue(Target.WindowsPlatform.CompanyName)));
+						AddDefinition(Arguments, $"PROJECT_COMPANY_NAME={SanitizeMacroValue(Target.WindowsPlatform.CompanyName)}");
 					}
 
 					if (!String.IsNullOrEmpty(Target.WindowsPlatform.CopyrightNotice))
 					{
-						AddDefinition(Arguments, String.Format("PROJECT_COPYRIGHT_STRING={0}", SanitizeMacroValue(Target.WindowsPlatform.CopyrightNotice)));
+						AddDefinition(Arguments, $"PROJECT_COPYRIGHT_STRING={SanitizeMacroValue(Target.WindowsPlatform.CopyrightNotice)}");
 					}
 
 					if (!String.IsNullOrEmpty(Target.WindowsPlatform.ProductName))
 					{
-						AddDefinition(Arguments, String.Format("PROJECT_PRODUCT_NAME={0}", SanitizeMacroValue(Target.WindowsPlatform.ProductName)));
+						AddDefinition(Arguments, $"PROJECT_PRODUCT_NAME={SanitizeMacroValue(Target.WindowsPlatform.ProductName)}");
 					}
 
 					if (Target.ProjectFile != null)
 					{
-						AddDefinition(Arguments, String.Format("PROJECT_PRODUCT_IDENTIFIER={0}", SanitizeMacroValue(Target.ProjectFile.GetFileNameWithoutExtension())));
+						AddDefinition(Arguments, $"PROJECT_PRODUCT_IDENTIFIER={SanitizeMacroValue(Target.ProjectFile.GetFileNameWithoutExtension())}");
 					}
 				}
 
@@ -1607,11 +1612,11 @@ namespace UnrealBuildTool
 						)
 					);
 				CompileAction.ProducedItems.Add(CompiledResourceFile);
-				Arguments.Add(String.Format("/fo \"{0}\"", NormalizeCommandLinePath(CompiledResourceFile, Target.WindowsPlatform.Compiler, false)));
+				Arguments.Add($"/fo \"{NormalizeCommandLinePath(CompiledResourceFile)}\"");
 				Result.ObjectFiles.Add(CompiledResourceFile);
 
 				// Add the RC file as a prerequisite of the action.
-				Arguments.Add(String.Format("\"{0}\"", NormalizeCommandLinePath(RCFile, Target.WindowsPlatform.Compiler, false)));
+				Arguments.Add($"\"{NormalizeCommandLinePath(RCFile)}\"");
 
 				// Create a response file for the resource compilier
 				FileItem ResponseFile = FileItem.GetItemByPath(CompiledResourceFile.FullName + ".response");
@@ -1619,7 +1624,7 @@ namespace UnrealBuildTool
 				CompileAction.PrerequisiteItems.Add(ResponseFile);
 
 				/* rc.exe currently errors when using a response file
-				string ResponseFileString = NormalizeCommandLinePath(ResponseFile, Target.WindowsPlatform.Compiler, false);
+				string ResponseFileString = NormalizeCommandLinePath(ResponseFile);
 
 				// cl.exe can't handle response files with a path longer than 260 characters, and relative paths can push it over the limit
 				if (!System.IO.Path.IsPathRooted(ResponseFileString) && System.IO.Path.Combine(CompileAction.WorkingDirectory.FullName, ResponseFileString).Length > 260)
@@ -1627,7 +1632,7 @@ namespace UnrealBuildTool
 					ResponseFileString = ResponseFile.FullName;
 				}
 
-				CompileAction.CommandArguments = String.Format("@{0}", Utils.MakePathSafeToUseWithCommandLine(ResponseFileString));
+				CompileAction.CommandArguments = $"@{Utils.MakePathSafeToUseWithCommandLine(ResponseFileString)}";
 				*/
 				CompileAction.CommandArguments = String.Join(" ", Arguments);
 
@@ -1660,24 +1665,24 @@ namespace UnrealBuildTool
 			FileItem ObjectFile = FileItem.GetItemByFileReference(OutputFile.ChangeExtension(".obj"));
 
 			List<string> Arguments = new List<string>();
-			Arguments.Add(String.Format("\"{0}\"", InputFile.Location));
+			Arguments.Add($"\"{InputFile.Location}\"");
 			Arguments.Add("/c");
 			Arguments.Add("/nologo");
-			Arguments.Add(String.Format("/Fo\"{0}\"", ObjectFile.Location));
+			Arguments.Add($"/Fo\"{ObjectFile.Location}\"");
 
 			foreach (DirectoryReference IncludePath in CompileEnvironment.UserIncludePaths)
 			{
-				AddIncludePath(Arguments, IncludePath, Target.WindowsPlatform.Compiler, CompileEnvironment.bPreprocessOnly);
+				AddIncludePath(Arguments, IncludePath, Target.WindowsPlatform.Compiler);
 			}
 
 			foreach (DirectoryReference IncludePath in CompileEnvironment.SystemIncludePaths)
 			{
-				AddSystemIncludePath(Arguments, IncludePath, Target.WindowsPlatform.Compiler, CompileEnvironment.bPreprocessOnly);
+				AddSystemIncludePath(Arguments, IncludePath, Target.WindowsPlatform.Compiler);
 			}
 
 			foreach (DirectoryReference IncludePath in EnvVars.IncludePaths)
 			{
-				AddSystemIncludePath(Arguments, IncludePath, Target.WindowsPlatform.Compiler, CompileEnvironment.bPreprocessOnly);
+				AddSystemIncludePath(Arguments, IncludePath, Target.WindowsPlatform.Compiler);
 			}
 
 			// Create the compile action. Only mark the object file as an output, because we need to touch the generated header afterwards.
@@ -1697,7 +1702,7 @@ namespace UnrealBuildTool
 			Action TouchAction = Graph.CreateAction(ActionType.BuildProject);
 			TouchAction.CommandDescription = "Touch";
 			TouchAction.CommandPath = BuildHostPlatform.Current.Shell;
-			TouchAction.CommandArguments = String.Format("/C \"copy /b \"{0}\"+,, \"{0}\" 1>nul:\"", OutputFile.FullName);
+			TouchAction.CommandArguments = $"/C \"copy /b \"{OutputFile.FullName}\"+,, \"{OutputFile.FullName}\" 1>nul:\"";
 			TouchAction.WorkingDirectory = UnrealBuildTool.EngineSourceDirectory;
 			TouchAction.PrerequisiteItems.Add(ObjectFile);
 			TouchAction.ProducedItems.Add(FileItem.GetItemByFileReference(OutputFile));
@@ -1756,7 +1761,7 @@ namespace UnrealBuildTool
 				Arguments.Add("/DEF");
 
 				// Ensure that the import library references the correct filename for the linked binary.
-				Arguments.Add(String.Format("/NAME:\"{0}\"", LinkEnvironment.OutputFilePath.GetFileName()));
+				Arguments.Add($"/NAME:\"{LinkEnvironment.OutputFilePath.GetFileName()}\"");
 
 				// Ignore warnings about object files with no public symbols.
 				Arguments.Add("/IGNORE:4221");
@@ -1768,13 +1773,13 @@ namespace UnrealBuildTool
 				// Delay-load these DLLs.
 				foreach (string DelayLoadDLL in LinkEnvironment.DelayLoadDLLs.Distinct())
 				{
-					Arguments.Add(String.Format("/DELAYLOAD:\"{0}\"", DelayLoadDLL));
+					Arguments.Add($"/DELAYLOAD:\"{DelayLoadDLL}\"");
 				}
 
 				// Pass the module definition file to the linker if we have one
 				if (LinkEnvironment.ModuleDefinitionFile != null && LinkEnvironment.ModuleDefinitionFile.Length > 0)
 				{
-					Arguments.Add(String.Format("/DEF:\"{0}\"", LinkEnvironment.ModuleDefinitionFile));
+					Arguments.Add($"/DEF:\"{LinkEnvironment.ModuleDefinitionFile}\"");
 				}
 			}
 
@@ -1790,17 +1795,17 @@ namespace UnrealBuildTool
 				// Add the library paths to the argument list.
 				foreach (DirectoryReference LibraryPath in LinkEnvironment.SystemLibraryPaths)
 				{
-					Arguments.Add(String.Format("/LIBPATH:\"{0}\"", NormalizeCommandLinePath(LibraryPath, Target.WindowsPlatform.Compiler, false)));
+					Arguments.Add($"/LIBPATH:\"{NormalizeCommandLinePath(LibraryPath)}\"");
 				}
 				foreach (DirectoryReference LibraryPath in EnvVars.LibraryPaths)
 				{
-					Arguments.Add(String.Format("/LIBPATH:\"{0}\"", NormalizeCommandLinePath(LibraryPath, Target.WindowsPlatform.Compiler, false)));
+					Arguments.Add($"/LIBPATH:\"{NormalizeCommandLinePath(LibraryPath)}\"");
 				}
 
 				// Add the excluded default libraries to the argument list.
 				foreach (string ExcludedLibrary in LinkEnvironment.ExcludedLibraries)
 				{
-					Arguments.Add(String.Format("/NODEFAULTLIB:\"{0}\"", ExcludedLibrary));
+					Arguments.Add($"/NODEFAULTLIB:\"{ExcludedLibrary}\"");
 				}
 			}
 
@@ -1830,18 +1835,18 @@ namespace UnrealBuildTool
 				{
 					if (LinkEnvironment.bIsBuildingDLL)
 					{
-						Arguments.Add(String.Format("/wholearchive:clang_rt.asan{0}_dll_thunk-{1}.lib", ASanDebugInfix, ASanArchSuffix));
+						Arguments.Add($"/wholearchive:clang_rt.asan{ASanDebugInfix}_dll_thunk-{ASanArchSuffix}.lib");
 					}
 					else
 					{
-						Arguments.Add(String.Format("/wholearchive:clang_rt.asan{0}-{1}.lib", ASanDebugInfix, ASanArchSuffix));
-						Arguments.Add(String.Format("/wholearchive:clang_rt.asan_cxx{0}-{1}.lib", ASanDebugInfix, ASanArchSuffix));
+						Arguments.Add($"/wholearchive:clang_rt.asan{ASanDebugInfix}-{ASanArchSuffix}.lib");
+						Arguments.Add($"/wholearchive:clang_rt.asan_cxx{ASanDebugInfix}-{ASanArchSuffix}.lib");
 					}
 				}
 				else
 				{
-					Arguments.Add(String.Format("/wholearchive:clang_rt.asan{0}_dynamic-{1}.lib", ASanDebugInfix, ASanArchSuffix));
-					Arguments.Add(String.Format("/wholearchive:clang_rt.asan{0}_dynamic_runtime_thunk-{1}.lib", ASanDebugInfix, ASanArchSuffix));
+					Arguments.Add($"/wholearchive:clang_rt.asan{ASanDebugInfix}_dynamic-{ASanArchSuffix}.lib");
+					Arguments.Add($"/wholearchive:clang_rt.asan{ASanDebugInfix}_dynamic_runtime_thunk-{ASanArchSuffix}.lib");
 				}
 			}
 
@@ -1889,7 +1894,7 @@ namespace UnrealBuildTool
 			List<string> InputFileNames = new List<string>();
 			foreach (FileItem InputFile in LinkEnvironment.InputFiles)
 			{
-				InputFileNames.Add(string.Format("\"{0}\"", NormalizeCommandLinePath(InputFile, Target.WindowsPlatform.Compiler, false)));
+				InputFileNames.Add(string.Format("\"{0}\"", NormalizeCommandLinePath(InputFile)));
 				PrerequisiteItems.Add(InputFile);
 			}
 
@@ -1897,7 +1902,7 @@ namespace UnrealBuildTool
 			{
 				foreach (FileReference Library in LinkEnvironment.Libraries)
 				{
-					InputFileNames.Add(string.Format("\"{0}\"", NormalizeCommandLinePath(Library, Target.WindowsPlatform.Compiler, false)));
+					InputFileNames.Add(string.Format("\"{0}\"", NormalizeCommandLinePath(Library)));
 					PrerequisiteItems.Add(FileItem.GetItemByFileReference(Library));
 				}
 				foreach (string SystemLibrary in LinkEnvironment.SystemLibraries)
@@ -1909,7 +1914,7 @@ namespace UnrealBuildTool
 			Arguments.AddRange(InputFileNames);
 
 			// Add the output file to the command-line.
-			Arguments.Add(String.Format("/OUT:\"{0}\"", NormalizeCommandLinePath(OutputFile, Target.WindowsPlatform.Compiler, false)));
+			Arguments.Add($"/OUT:\"{NormalizeCommandLinePath(OutputFile)}\"");
 
 			// For import libraries and exports generated by cross-referenced builds, we don't track output files. VS 15.3+ doesn't touch timestamps for libs
 			// and exp files with no modifications, breaking our dependency checking, but incremental linking will fall back to a full link if we delete it.
@@ -1928,7 +1933,7 @@ namespace UnrealBuildTool
 				{
 					// Write the import library to the output directory for nFringe support.
 					FileItem ImportLibraryFile = FileItem.GetItemByFileReference(ImportLibraryFilePath);
-					Arguments.Add(String.Format("/IMPLIB:\"{0}\"", NormalizeCommandLinePath(ImportLibraryFilePath, Target.WindowsPlatform.Compiler, false)));
+					Arguments.Add($"/IMPLIB:\"{NormalizeCommandLinePath(ImportLibraryFilePath)}\"");
 
 					// Like the export file above, don't add the import library as a produced item when it's cross referenced.
 					if(!LinkEnvironment.bIsCrossReferenced)
@@ -1943,7 +1948,7 @@ namespace UnrealBuildTool
 					{
 						FileReference PDBFilePath = FileReference.Combine(LinkEnvironment.OutputDirectory!, Path.GetFileNameWithoutExtension(OutputFile.AbsolutePath) + ".pdb");
 						FileItem PDBFile = FileItem.GetItemByFileReference(PDBFilePath);
-						Arguments.Add(String.Format("/PDB:\"{0}\"", NormalizeCommandLinePath(PDBFilePath, Target.WindowsPlatform.Compiler, false)));
+						Arguments.Add($"/PDB:\"{NormalizeCommandLinePath(PDBFilePath)}\"");
 						ProducedItems.Add(PDBFile);
 					}
 
@@ -1952,7 +1957,7 @@ namespace UnrealBuildTool
 					{
 						FileReference MAPFilePath = FileReference.Combine(LinkEnvironment.OutputDirectory!, Path.GetFileNameWithoutExtension(OutputFile.AbsolutePath) + ".map");
 						FileItem MAPFile = FileItem.GetItemByFileReference(MAPFilePath);
-						Arguments.Add(String.Format("/MAP:\"{0}\"", NormalizeCommandLinePath(MAPFilePath, Target.WindowsPlatform.Compiler, false)));
+						Arguments.Add($"/MAP:\"{NormalizeCommandLinePath(MAPFilePath)}\"");
 						ProducedItems.Add(MAPFile);
 
 						// Export a list of object file paths, so we can locate the object files referenced by the map file
@@ -1970,7 +1975,7 @@ namespace UnrealBuildTool
 			// Add any forced references to functions
 			foreach(string IncludeFunction in LinkEnvironment.IncludeFunctions)
 			{
-				Arguments.Add(String.Format("/INCLUDE:{0}", IncludeFunction));
+				Arguments.Add($"/INCLUDE:{IncludeFunction}");
 			}
 
 			// Allow the toolchain to adjust/process the link arguments
@@ -1996,7 +2001,7 @@ namespace UnrealBuildTool
 			{
 				LinkAction.CommandPath = EnvVars.LinkerPath;
 			}
-			LinkAction.CommandArguments = String.Format("@\"{0}\"", ResponseFileName);
+			LinkAction.CommandArguments = $"@\"{ResponseFileName}\"";
 			LinkAction.CommandVersion = EnvVars.ToolChainVersion.ToString();
 			LinkAction.ProducedItems.AddRange(ProducedItems);
 			LinkAction.PrerequisiteItems.AddRange(PrerequisiteItems);
