Always include a UE5 target when generating an Xcode project file for a blueprint project for codesigning

diff --git a/Engine/Source/Programs/UnrealBuildTool/ProjectFiles/Xcode/XcodeProjectFileGenerator.cs b/Engine/Source/Programs/UnrealBuildTool/ProjectFiles/Xcode/XcodeProjectFileGenerator.cs
--- a/Engine/Source/Programs/UnrealBuildTool/ProjectFiles/Xcode/XcodeProjectFileGenerator.cs
+++ b/Engine/Source/Programs/UnrealBuildTool/ProjectFiles/Xcode/XcodeProjectFileGenerator.cs
@@ -301,6 +301,11 @@ namespace UnrealBuildTool
 
 			if (bGeneratingGameProjectFiles)
 			{
+				if (bGeneratingRunIOSProject || bGeneratingRunTVOSProject)
+				{
+					// an Engine target is required in order to be able to get Xcode to sign blueprint projects
+					IncludeEnginePrograms = true;
+				}
 				bIncludeEngineSource = true;
 			}
 		}
