Fix installed SDK components and min and max buils-tools used for Android

#android

diff --git a/Engine/Config/BaseEngine.ini b/Engine/Config/BaseEngine.ini
--- a/Engine/Config/BaseEngine.ini
+++ b/Engine/Config/BaseEngine.ini
@@ -2561,6 +2561,7 @@ bEnableAdvertisingIdentifier=True
 [/Script/AndroidRuntimeSettings.AndroidRuntimeSettings]
 SDKAPILevelOverride=
 NDKAPILevelOverride=
+BuildToolsOverride=
 bEnableGooglePlaySupport=false
 bSupportAdMob=true
 bBuildForArm64=true
diff --git a/Engine/Extras/Android/SetupAndroid.bat b/Engine/Extras/Android/SetupAndroid.bat
--- a/Engine/Extras/Android/SetupAndroid.bat
+++ b/Engine/Extras/Android/SetupAndroid.bat
@@ -17,10 +17,10 @@ SET CMAKE_VERSION=%3
 SET NDK_VERSION=%4
 
 rem hardcoded versions for compatibility with non-Turnkey manual running
-if "%PLATFORMS_VERSION%" == "" SET PLATFORMS_VERSION=android-28
-if "%BUILDTOOLS_VERSION%" == "" SET BUILDTOOLS_VERSION=28.0.3
+if "%PLATFORMS_VERSION%" == "" SET PLATFORMS_VERSION=android-30
+if "%BUILDTOOLS_VERSION%" == "" SET BUILDTOOLS_VERSION=30.0.3
 if "%CMAKE_VERSION%" == "" SET CMAKE_VERSION=3.10.2.4988404
-if "%NDK_VERSION%" == "" SET NDK_VERSION=21.1.6352462
+if "%NDK_VERSION%" == "" SET NDK_VERSION=21.4.7075529
 
 
 FOR /F "tokens=2*" %%A IN ('REG.exe query "%KEY_NAME%" /v "%VALUE_NAME%"') DO (set STUDIO_PATH=%%B)
diff --git a/Engine/Extras/Android/SetupAndroid.command b/Engine/Extras/Android/SetupAndroid.command
--- a/Engine/Extras/Android/SetupAndroid.command
+++ b/Engine/Extras/Android/SetupAndroid.command
@@ -19,16 +19,16 @@ fi
 
 rem hardcoded versions for compatibility with non-Turnkey manual running
 if [[ -z "${PLATFORMS_VERSION}" ]]; then
-    PLATFORMS_VERSION="android-28"
+    PLATFORMS_VERSION="android-30"
 fi
 if [[ -z "${BUILDTOOLS_VERSION}" ]]; then
-    BUILDTOOLS_VERSION="28.0.3"
+    BUILDTOOLS_VERSION="30.0.3"
 fi
 if [[ -z "${CMAKE_VERSION}" ]]; then
     CMAKE_VERSION="3.10.2.4988404"
 fi
 if [[ -z "${NDK_VERSION}" ]]; then
-    NDK_VERSION="21.1.6352462"
+    NDK_VERSION="21.4.7075529"
 fi
 
 
diff --git a/Engine/Extras/Android/SetupAndroid.sh b/Engine/Extras/Android/SetupAndroid.sh
--- a/Engine/Extras/Android/SetupAndroid.sh
+++ b/Engine/Extras/Android/SetupAndroid.sh
@@ -27,16 +27,16 @@ fi
 
 # hardcoded versions for compatibility with non-Turnkey manual running
 if [[ -z "${PLATFORMS_VERSION}" ]]; then
-    PLATFORMS_VERSION="android-28"
+    PLATFORMS_VERSION="android-30"
 fi
 if [[ -z "${BUILDTOOLS_VERSION}" ]]; then
-    BUILDTOOLS_VERSION="28.0.3"
+    BUILDTOOLS_VERSION="30.0.3"
 fi
 if [[ -z "${CMAKE_VERSION}" ]]; then
     CMAKE_VERSION="3.10.2.4988404"
 fi
 if [[ -z "${NDK_VERSION}" ]]; then
-    NDK_VERSION="21.1.6352462"
+    NDK_VERSION="21.4.7075529"
 fi
 
 
diff --git a/Engine/Source/Programs/UnrealBuildTool/Platform/Android/UEDeployAndroid.cs b/Engine/Source/Programs/UnrealBuildTool/Platform/Android/UEDeployAndroid.cs
--- a/Engine/Source/Programs/UnrealBuildTool/Platform/Android/UEDeployAndroid.cs
+++ b/Engine/Source/Programs/UnrealBuildTool/Platform/Android/UEDeployAndroid.cs
@@ -282,6 +282,11 @@ namespace UnrealBuildTool
 
 		private uint GetRevisionValue(string VersionString)
 		{
+			if (VersionString == null)
+			{
+				return 0;
+			}
+
 			// read up to 4 sections (ie. 20.0.3.5), first section most significant
 			// each section assumed to be 0 to 255 range
 			uint Value = 0;
@@ -309,57 +314,73 @@ namespace UnrealBuildTool
 				return CachedBuildToolsVersion;
 			}
 
-			// get a list of the directories in build-tools.. may be more than one set installed (or none which is bad)
-			string[] Subdirs = Directory.GetDirectories(Path.Combine(HomePath, "build-tools"));
-			if (Subdirs.Length == 0)
-			{
-				throw new BuildException("Failed to find %ANDROID_HOME%/build-tools subdirectory. Run SDK manager and install build-tools.");
-			}
-
-			// valid directories will have a source.properties with the Pkg.Revision (there is no guarantee we can use the directory name as revision)
 			string? BestVersionString = null;
 			uint BestVersion = 0;
-			foreach (string CandidateDir in Subdirs)
+
+			ConfigHierarchy Ini = GetConfigCacheIni(ConfigHierarchyType.Engine);
+			Ini.GetString("/Script/AndroidRuntimeSettings.AndroidRuntimeSettings", "BuildToolsOverride", out BestVersionString);
+
+			if (BestVersionString == null || BestVersionString == "" || BestVersionString == "latest")
 			{
-				string AaptFilename = Path.Combine(CandidateDir, RuntimePlatform.IsWindows ? "aapt.exe" : "aapt");
-				string RevisionString = "";
-				uint RevisionValue = 0;
+				// get a list of the directories in build-tools.. may be more than one set installed (or none which is bad)
+				string[] Subdirs = Directory.GetDirectories(Path.Combine(HomePath, "build-tools"));
+				if (Subdirs.Length == 0)
+				{
+					throw new BuildException("Failed to find %ANDROID_HOME%/build-tools subdirectory. Run SDK manager and install build-tools.");
+				}
 
-				if (File.Exists(AaptFilename))
+				// valid directories will have a source.properties with the Pkg.Revision (there is no guarantee we can use the directory name as revision)
+				foreach (string CandidateDir in Subdirs)
 				{
-					string SourcePropFilename = Path.Combine(CandidateDir, "source.properties");
-					if (File.Exists(SourcePropFilename))
+					string AaptFilename = Path.Combine(CandidateDir, RuntimePlatform.IsWindows ? "aapt.exe" : "aapt");
+					string RevisionString = "";
+					uint RevisionValue = 0;
+
+					if (File.Exists(AaptFilename))
 					{
-						string[] PropertyContents = File.ReadAllLines(SourcePropFilename);
-						foreach (string PropertyLine in PropertyContents)
+						string SourcePropFilename = Path.Combine(CandidateDir, "source.properties");
+						if (File.Exists(SourcePropFilename))
 						{
-							if (PropertyLine.StartsWith("Pkg.Revision="))
+							string[] PropertyContents = File.ReadAllLines(SourcePropFilename);
+							foreach (string PropertyLine in PropertyContents)
 							{
-								RevisionString = PropertyLine.Substring(13);
-								RevisionValue = GetRevisionValue(RevisionString);
-								break;
+								if (PropertyLine.StartsWith("Pkg.Revision="))
+								{
+									RevisionString = PropertyLine.Substring(13);
+									RevisionValue = GetRevisionValue(RevisionString);
+									break;
+								}
 							}
 						}
 					}
-				}
 
-				// remember it if newer version or haven't found one yet
-				if (RevisionValue > BestVersion || BestVersionString == null)
-				{
-					BestVersion = RevisionValue;
-					BestVersionString = RevisionString;
+					// remember it if newer version or haven't found one yet
+					if (RevisionValue > BestVersion || BestVersionString == null)
+					{
+						BestVersion = RevisionValue;
+						BestVersionString = RevisionString;
+					}
 				}
 			}
 
 			if (BestVersionString == null)
 			{
-				throw new BuildException("Failed to find %ANDROID_HOME%/build-tools subdirectory with aapt. Run SDK manager and install build-tools.");
+				BestVersionString = "30.0.3";
+				Log.TraceWarning("Failed to find %ANDROID_HOME%/build-tools subdirectory. Will attempt to use {0}.", BestVersionString);
+			}
+
+			BestVersion = GetRevisionValue(BestVersionString);
+
+			// with Gradle enabled use at least 28.0.3 (will be installed by Gradle if missing)
+			if (BestVersion < ((28 << 24) | (0 << 16) | (3 << 8)))
+			{
+				BestVersionString = "28.0.3";
 			}
 
-			// with Gradle enabled use at least 24.0.2 (will be installed by Gradle if missing)
-			if (BestVersion < ((24 << 24) | (0 << 16) | (2 << 8)))
+			// don't allow higher than 30.0.3 for now (will be installed by Gradle if missing)
+			if (BestVersion > ((30 << 24) | (0 << 16) | (3 << 8)))
 			{
-				BestVersionString = "24.0.2";
+				BestVersionString = "30.0.3";
 			}
 
 			CachedBuildToolsVersion = BestVersionString;
diff --git a/Engine/Source/Runtime/Android/AndroidRuntimeSettings/Classes/AndroidRuntimeSettings.h b/Engine/Source/Runtime/Android/AndroidRuntimeSettings/Classes/AndroidRuntimeSettings.h
--- a/Engine/Source/Runtime/Android/AndroidRuntimeSettings/Classes/AndroidRuntimeSettings.h
+++ b/Engine/Source/Runtime/Android/AndroidRuntimeSettings/Classes/AndroidRuntimeSettings.h
@@ -607,6 +607,10 @@ public:
 	UPROPERTY(GlobalConfig, EditAnywhere, Category = "Project SDK Override", Meta = (DisplayName = "NDK API Level (specific version or 'latest' - see tooltip)"))
 	FString NDKAPILevelOverride;
 
+	// Which build-tools to package with (a specific version or (without quotes) 'latest' for latest version on disk).
+	UPROPERTY(GlobalConfig, EditAnywhere, Category = "Project SDK Override", Meta = (DisplayName = "Build-Tools Version (specific version or 'latest')"))
+	FString BuildToolsOverride;
+
 	/** Whether to enable LOD streaming for landscape visual meshes. Only supported on feature level ES3.1 or above. */
 	UPROPERTY(GlobalConfig, EditAnywhere, Category = "Misc", Meta = (DisplayName = "Stream landscape visual mesh LODs"))
 	bool bStreamLandscapeMeshLODs;
