UEBuildTarget: When explicitly building a plugin, retain the pre-18353193 behavior for build location

#rnx

diff --git a/Engine/Source/Programs/UnrealBuildTool/Configuration/UEBuildTarget.cs b/Engine/Source/Programs/UnrealBuildTool/Configuration/UEBuildTarget.cs
--- a/Engine/Source/Programs/UnrealBuildTool/Configuration/UEBuildTarget.cs
+++ b/Engine/Source/Programs/UnrealBuildTool/Configuration/UEBuildTarget.cs
@@ -3086,7 +3086,8 @@ namespace UnrealBuildTool
 		{
 			// Get the root output directory and base name (target name/app name) for this binary
 			DirectoryReference BaseOutputDirectory;
-			if(bUseSharedBuildEnvironment && ModuleRules.Plugin?.Type != PluginType.External)
+
+			if(bUseSharedBuildEnvironment && (ModuleRules.Plugin == null || (ModuleRules.Plugin.Type != PluginType.External || ModuleRules.Plugin.bExplicitPluginTarget)))
 			{
 				BaseOutputDirectory = ModuleRules.Context.DefaultOutputBaseDir;
 			}
@@ -3114,8 +3115,10 @@ namespace UnrealBuildTool
 				ModuleConfiguration = UnrealTargetConfiguration.Development;
 			}
 
+			bool bUseExternalFolder = ModuleRules.Plugin != null && (ModuleRules.Plugin.Type == PluginType.External && !ModuleRules.Plugin.bExplicitPluginTarget);
+
 			// Get the output and intermediate directories for this module
-			DirectoryReference IntermediateDirectory = DirectoryReference.Combine(BaseOutputDirectory, GetPlatformIntermediateFolder(Platform, Architecture, ModuleRules.Plugin?.Type == PluginType.External), AppName, ModuleConfiguration.ToString());
+			DirectoryReference IntermediateDirectory = DirectoryReference.Combine(BaseOutputDirectory, GetPlatformIntermediateFolder(Platform, Architecture, bUseExternalFolder), AppName, ModuleConfiguration.ToString());
 
 			// Append a subdirectory if the module rules specifies one
 			if (!String.IsNullOrEmpty(ModuleRules.BinariesSubFolder))
diff --git a/Engine/Source/Programs/UnrealBuildTool/System/Plugins.cs b/Engine/Source/Programs/UnrealBuildTool/System/Plugins.cs
--- a/Engine/Source/Programs/UnrealBuildTool/System/Plugins.cs
+++ b/Engine/Source/Programs/UnrealBuildTool/System/Plugins.cs
@@ -88,6 +88,11 @@ namespace UnrealBuildTool
 		/// </summary>
 		public PluginType Type;
 
+		/// <summary>
+		/// Used to indicate whether a plugin is being explicitly packaged via the -plugin command line
+		/// </summary>
+		public bool bExplicitPluginTarget = false;
+
 		/// <summary>
 		/// Constructs a PluginInfo object
 		/// </summary>
diff --git a/Engine/Source/Programs/UnrealBuildTool/System/RulesCompiler.cs b/Engine/Source/Programs/UnrealBuildTool/System/RulesCompiler.cs
--- a/Engine/Source/Programs/UnrealBuildTool/System/RulesCompiler.cs
+++ b/Engine/Source/Programs/UnrealBuildTool/System/RulesCompiler.cs
@@ -325,7 +325,11 @@ namespace UnrealBuildTool
 				List<PluginInfo> ForeignPlugins = new List<PluginInfo>();
 				if (!Parent.EnumeratePlugins().Any(x => x.File == PluginFileName && x.Type == PluginType.Engine))
 				{
-					ForeignPlugins.Add(new PluginInfo(PluginFileName, PluginType.External));
+					PluginInfo ForeignPluginInfo = new PluginInfo(PluginFileName, PluginType.External)
+					{
+						bExplicitPluginTarget = true
+					};
+					ForeignPlugins.Add(ForeignPluginInfo);
 				}
 
 				// Create a new scope for the plugin. It should not reference anything else.
