Fixed Verse generated headers getting destroyed erroneously

This is a temporary measure while we coordinate a more robust fix to clear out stale UObject headers while keeping VNI headers around

diff --git a/Engine/Source/Programs/UnrealBuildTool/System/ExternalExecution.cs b/Engine/Source/Programs/UnrealBuildTool/System/ExternalExecution.cs
--- a/Engine/Source/Programs/UnrealBuildTool/System/ExternalExecution.cs
+++ b/Engine/Source/Programs/UnrealBuildTool/System/ExternalExecution.cs
@@ -690,7 +690,8 @@ namespace UnrealBuildTool
 				else
 				{
 					// Remove any stale generated code directory
-					if(Module.GeneratedCodeDirectory != null && !Module.Rules.bUsePrecompiled)
+					// Don't destroy the directory if we have Verse
+					if(!Module.bHasVerse && Module.GeneratedCodeDirectory != null && !Module.Rules.bUsePrecompiled)
 					{
 						if (DirectoryReference.Exists(Module.GeneratedCodeDirectory))
 						{
