UnrealBuildTool: Support -Timing for Linux & Mac

#rnx

diff --git a/Engine/Source/Programs/UnrealBuildTool/Platform/Linux/LinuxToolChain.cs b/Engine/Source/Programs/UnrealBuildTool/Platform/Linux/LinuxToolChain.cs
--- a/Engine/Source/Programs/UnrealBuildTool/Platform/Linux/LinuxToolChain.cs
+++ b/Engine/Source/Programs/UnrealBuildTool/Platform/Linux/LinuxToolChain.cs
@@ -1329,6 +1329,14 @@ namespace UnrealBuildTool
 				// Add the source file path to the command-line.
 				FileArguments += string.Format(" \"{0}\"", SourceFile.AbsolutePath.Replace('\\', '/'));
 
+				// Generate the timing info
+				if (CompileEnvironment.bPrintTimingInfo)
+				{
+					FileItem TraceFile = FileItem.GetItemByFileReference(FileReference.Combine(OutputDir, Path.GetFileName(SourceFile.AbsolutePath) + ".json"));
+					FileArguments += " -ftime-trace";
+					CompileAction.ProducedItems.Add(TraceFile);
+				}
+
 				// Generate the included header dependency list
 				if(CompileEnvironment.bGenerateDependenciesFile)
 				{
diff --git a/Engine/Source/Programs/UnrealBuildTool/Platform/Mac/MacToolChain.cs b/Engine/Source/Programs/UnrealBuildTool/Platform/Mac/MacToolChain.cs
--- a/Engine/Source/Programs/UnrealBuildTool/Platform/Mac/MacToolChain.cs
+++ b/Engine/Source/Programs/UnrealBuildTool/Platform/Mac/MacToolChain.cs
@@ -555,6 +555,14 @@ namespace UnrealBuildTool
 				// Add the source file path to the command-line.
 				FileArguments += string.Format(" \"{0}\"", SourceFile.AbsolutePath);
 
+				// Generate the timing info
+				if (CompileEnvironment.bPrintTimingInfo)
+				{
+					FileItem TraceFile = FileItem.GetItemByFileReference(FileReference.Combine(OutputDir, Path.GetFileName(SourceFile.AbsolutePath) + ".json"));
+					FileArguments += " -ftime-trace";
+					CompileAction.ProducedItems.Add(TraceFile);
+				}
+
 				// Generate the included header dependency list
 				if(CompileEnvironment.bGenerateDependenciesFile)
 				{
