Modularized setup and teardown for LLTs + global setup preliminary.
Nightly LLTs preflight 6205328b4c6c33118fa85c3d

diff --git a/Engine/Source/Developer/LowLevelTestsRunner/LowLevelTestsRunner.Build.cs b/Engine/Source/Developer/LowLevelTestsRunner/LowLevelTestsRunner.Build.cs
--- a/Engine/Source/Developer/LowLevelTestsRunner/LowLevelTestsRunner.Build.cs
+++ b/Engine/Source/Developer/LowLevelTestsRunner/LowLevelTestsRunner.Build.cs
@@ -2,7 +2,7 @@
 
 using UnrealBuildTool;
 using System.IO;
-public class LowLevelTestsRunner: ModuleRules
+public class LowLevelTestsRunner : ModuleRules
 {
 	public LowLevelTestsRunner(ReadOnlyTargetRules Target) : base(Target)
 	{
@@ -14,21 +14,32 @@ public class LowLevelTestsRunner: ModuleRules
 		PublicIncludePaths.AddRange(
 			new string[]
 			{
-				Path.Combine(ModuleDirectory, "Public"),
+				"Runtime/Launch/Public",
 				Path.Combine(Target.UEThirdPartySourceDirectory, "Catch2")
 			}
 		);
 
 		PrivateIncludePaths.AddRange(
 			new string[] {
-				Path.Combine(ModuleDirectory, "Private")
+				"Runtime/Launch/Private"
 			}
 		);
 
-		PublicIncludePathModuleNames.AddRange(
-				new string[] {
-					"Core"
-				}
+		PrivateDependencyModuleNames.AddRange(
+			new string[] {
+				"Core",
+				"Projects"
+			}
 		);
+
+		if (Target.bCompileAgainstApplicationCore)
+		{
+			PrivateDependencyModuleNames.Add("ApplicationCore");
+		}
+
+		if (Target.bCompileAgainstCoreUObject)
+		{
+			PrivateDependencyModuleNames.Add("CoreUObject");
+		}
 	}
 }
\ No newline at end of file
diff --git a/Engine/Source/Developer/LowLevelTestsRunner/Private/TestMain.cpp b/Engine/Source/Developer/LowLevelTestsRunner/Private/TestMain.cpp
new file mode 100644
--- /dev/null
+++ b/Engine/Source/Developer/LowLevelTestsRunner/Private/TestMain.cpp
@@ -0,0 +1,7 @@
+// Copyright Epic Games, Inc. All Rights Reserved.
+
+#include "RequiredProgramMainCPPInclude.h"
+
+TCHAR GInternalProjectName[64] = TEXT("LowLevelTests");
+const TCHAR* GForeignEngineDir = nullptr;
+FEngineLoop GEngineLoop;
diff --git a/Engine/Source/Developer/LowLevelTestsRunner/Private/TestRunner.cpp b/Engine/Source/Developer/LowLevelTestsRunner/Private/TestRunner.cpp
--- a/Engine/Source/Developer/LowLevelTestsRunner/Private/TestRunner.cpp
+++ b/Engine/Source/Developer/LowLevelTestsRunner/Private/TestRunner.cpp
@@ -4,8 +4,49 @@
 
 #include "TestRunner.h"
 #include "TestHarness.h"
+#include "Modules/ModuleManager.h"
 #include "HAL/PlatformTLS.h"
 
+#include "CommonEngineInit.inl"
+
+// TODO: Implement global initialization via flags
+#if WITH_ENGINE || WITH_EDITOR || WITH_APPLICATION_CORE || WITH_COREUOBJECT
+#define USE_GLOBAL_ENGINE_SETUP
+#endif
+
+
+void GlobalSetup()
+{
+	if (bGAllowLogging)
+	{
+		FCommandLine::Set(TEXT(""));
+	}
+	else
+	{
+		FCommandLine::Set(TEXT(R"(-LogCmds="global off")"));
+		FLogSuppressionInterface::Get().ProcessConfigAndCommandLine();
+	}
+	InitThreadPool();
+	InitAsyncQueues();
+	InitTaskGraph();
+	InitOutputDevices();
+	InitRendering();
+	InitDerivedDataCache();
+	InitSlate();
+	InitForWithEditorOnlyData();
+	InitEditor();
+	InitCoreUObject();
+}
+
+void GlobalTeardown()
+{
+	CleanupCoreUObject();
+	CleanupThreadPool();
+	CleanupTaskGraph();
+	CleanupPlatform();
+}
+
+
 int RunTests(int argc, const char* argv[])
 {
 	// remember thread id of the main thread
@@ -17,7 +58,17 @@ int RunTests(int argc, const char* argv[])
 	std::this_thread::sleep_for(std::chrono::milliseconds(5000));
 #endif
 
-	Setup();
+// Global initialization
+#ifdef USE_GLOBAL_ENGINE_SETUP
+	GlobalSetup();
+#endif
+
+	TArray<FName> ModuleNames;
+	FModuleManager::Get().FindModules(TEXT("*LowLevelTests"), ModuleNames);
+	for (FName ModuleName : ModuleNames)
+	{
+		FModuleManager::Get().LoadModule(ModuleName);
+	}
 
 	//Read command-line from file (if any). Some platforms do this earlier.
 #ifndef PLATFORM_SKIP_ADDITIONAL_ARGS
@@ -66,17 +117,6 @@ int RunTests(int argc, const char* argv[])
 		}
 	}
 
-	// Clear after init so that in windows it is actually clear during the test.
-	if (bGAllowLogging)
-	{
-		FCommandLine::Set(TEXT(""));
-	}
-	else
-	{
-		FCommandLine::Set(TEXT(R"(-LogCmds="global off")"));
-		FLogSuppressionInterface::Get().ProcessConfigAndCommandLine();
-	}
-
 	int SessionResult = 0;
 	{
 		TGuardValue<bool> CatchRunning(bCatchIsRunning, true);
@@ -84,7 +124,15 @@ int RunTests(int argc, const char* argv[])
 		CatchArgv.Reset();
 	}
 
-	Teardown();
+	for (FName ModuleName : ModuleNames)
+	{
+		FModuleManager::Get().UnloadModule(ModuleName);
+	}
+
+// Global cleanup
+#if defined(USE_GLOBAL_ENGINE_SETUP)
+	GlobalTeardown();
+#endif
 
 #if PLATFORM_DESKTOP
 	if (bWaitForInputToTerminate)
diff --git a/Engine/Source/Developer/LowLevelTestsRunner/Public/CommonEngineInit.inl b/Engine/Source/Developer/LowLevelTestsRunner/Public/CommonEngineInit.inl
--- a/Engine/Source/Developer/LowLevelTestsRunner/Public/CommonEngineInit.inl
+++ b/Engine/Source/Developer/LowLevelTestsRunner/Public/CommonEngineInit.inl
@@ -2,18 +2,36 @@
 
 #pragma once
 
+#include "CoreMinimal.h"
+
+#include "Async/TaskGraphInterfaces.h"
+#include "HAL/PlatformFileManager.h"
+#include "HAL/PlatformProcess.h"
+#include "HAL/PlatformOutputDevices.h"
+#include "Misc/ConfigCacheIni.h"
+#include "Misc/CoreDelegates.h"
+#include "Misc/QueuedThreadPool.h"
+
+#if WITH_APPLICATION_CORE
 #include "HAL/PlatformApplicationMisc.h"
+#endif // WITH_APPLICATION_CORE
 
+#if WITH_ENGINE
+#include "DistanceFieldAtlas.h"
+#include "MeshCardRepresentation.h"
 #include "RenderUtils.h"
 #include "ShaderParameterMetadata.h"
+#endif // WITH_ENGINE
 
-#if WITH_EDITORONLY_DATA
+#if WITH_EDITORONLY_DATA && defined(DERIVEDDATACACHE_API)
 #include "DerivedDataBuild.h"
 #include "DerivedDataCache.h"
-#endif // WITH_EDITORONLY_DATA
+#endif // WITH_EDITORONLY_DATA && defined(DERIVEDDATACACHE_API)
 
-#include "DistanceFieldAtlas.h"
-#include "MeshCardRepresentation.h"
+#if WITH_COREUOBJECT
+#include "UObject/Package.h"
+#include "UObject/PackageResourceManager.h"
+#endif // WITH_COREUOBJECT
 
 /**
  * A lot of what's in this file was taken from LaunchEngineLoop.cpp.
@@ -22,6 +40,7 @@
 
 void InitThreadPool()
 {
+#if WITH_EDITOR
 	{
 		int32 NumThreadsInLargeThreadPool = FMath::Max( FPlatformMisc::NumberOfCoresIncludingHyperthreads() - 2, 2 );
 		int32 NumThreadsInThreadPool = FPlatformMisc::NumberOfWorkerThreadsToSpawn();
@@ -40,12 +59,12 @@ void InitThreadPool()
 		constexpr int32 StackSize = 128 * 1024;
 
 		// TaskGraph has it's HP threads slightly below normal, we want to be below the taskgraph HP threads to avoid interfering with the game-thread.
-		verify( GLargeThreadPool->Create( NumThreadsInLargeThreadPool, StackSize, TPri_BelowNormal, TEXT("LargeThreadPool") ) );
+		verify(GLargeThreadPool->Create(NumThreadsInLargeThreadPool, StackSize, TPri_BelowNormal, TEXT("LargeThreadPool")));
 
 		// GThreadPool will schedule on the LargeThreadPool but limit max concurrency to the given number.
 		GThreadPool = new FQueuedThreadPoolWrapper( GLargeThreadPool, NumThreadsInThreadPool );
 	}
-
+#endif // WITH_EDITOR
 	if (FPlatformProcess::SupportsMultithreading())
 	{
 		GIOThreadPool = FQueuedThreadPool::Allocate();
@@ -54,10 +73,56 @@ void InitThreadPool()
 	}
 }
 
+void CleanupThreadPool()
+{
+	if (GThreadPool != nullptr)
+	{
+		GThreadPool->Destroy();
+	}
+	if (GBackgroundPriorityThreadPool != nullptr)
+	{
+		GBackgroundPriorityThreadPool->Destroy();
+	}
+	if (GIOThreadPool != nullptr)
+	{
+		GIOThreadPool->Destroy();
+	}
+}
+
 void InitTaskGraph()
 {
-	FTaskGraphInterface::Startup( FPlatformMisc::NumberOfWorkerThreadsToSpawn() );
-	FTaskGraphInterface::Get().AttachToThread( ENamedThreads::GameThread );
+	FTaskGraphInterface::Startup(bMultithreaded ? FPlatformMisc::NumberOfWorkerThreadsToSpawn() : 1);
+	FTaskGraphInterface::Get().AttachToThread(ENamedThreads::GameThread);
+}
+
+void CleanupTaskGraph()
+{
+	FTaskGraphInterface::Shutdown();
+}
+
+void InitCoreUObject()
+{
+#if WITH_COREUOBJECT
+	if (!GetTransientPackage())
+	{
+		FModuleManager::Get().AddExtraBinarySearchPaths();
+		FConfigCacheIni::InitializeConfigSystem();
+		FPlatformFileManager::Get().InitializeNewAsyncIO();
+
+		FModuleManager::Get().LoadModule(TEXT("CoreUObject"));
+		FCoreDelegates::OnInit.Broadcast();
+		ProcessNewlyLoadedUObjects();
+	}
+#endif // WITH_COREUOBJECT
+}
+
+void CleanupCoreUObject()
+{
+#if WITH_COREUOBJECT
+	FCoreDelegates::OnPreExit.Broadcast();
+	MALLOC_PROFILER(GMalloc->Exec(nullptr, TEXT("MPROF STOP"), *GLog); );
+	FCoreDelegates::OnExit.Broadcast();
+#endif // WITH_COREUOBJECT
 }
 
 void InitOutputDevices()
@@ -80,6 +145,7 @@ void InitStats()
 
 void InitRendering()
 {
+#if WITH_ENGINE
 	FShaderParametersMetadata::InitializeAllUniformBufferStructs();
 
 	{
@@ -92,11 +158,12 @@ void InitRendering()
 		// One-time initialization of global variables based on engine configuration.
 		RenderUtilsInit();
 	}
+#endif // WITH_ENGINE
 }
 
-#if WITH_EDITORONLY_DATA
 void InitDerivedDataCache()
 {
+#if WITH_EDITORONLY_DATA && defined(DERIVEDDATACACHE_API)
 	if (!FPlatformProperties::RequiresCookedData())
 	{
 		// Ensure that DDC is initialized from the game thread.
@@ -104,20 +171,24 @@ void InitDerivedDataCache()
 		UE::DerivedData::GetCache();
 		GetDerivedDataCacheRef();
 	}
+#endif // WITH_EDITORONLY_DATA && defined(DERIVEDDATACACHE_API)
 }
-#endif // WITH_EDITORONLY_DATA
+
 
 void InitAsyncQueues()
 {
+#if WITH_ENGINE
 	check(!GDistanceFieldAsyncQueue);
 	GDistanceFieldAsyncQueue = new FDistanceFieldAsyncQueue();
 
 	check(!GCardRepresentationAsyncQueue);
 	GCardRepresentationAsyncQueue = new FCardRepresentationAsyncQueue();
+#endif // WITH_ENGINE
 }
 
 void InitSlate()
 {
+#if WITH_EDITOR
 	FSlateApplication::Create();
 
 	TSharedPtr<FSlateRenderer> SlateRenderer = FModuleManager::Get().LoadModuleChecked<ISlateNullRendererModule>("SlateNullRenderer").CreateSlateNullRenderer();
@@ -126,6 +197,7 @@ void InitSlate()
 	// If Slate is being used, initialize the renderer after RHIInit
 	FSlateApplication& CurrentSlateApp = FSlateApplication::Get();
 	CurrentSlateApp.InitializeRenderer(SlateRendererSharedRef);
+#endif // WITH_EDITOR
 }
 
 void InitForWithEditorOnlyData()
@@ -134,16 +206,16 @@ void InitForWithEditorOnlyData()
 	//Initialize the PackageResourceManager, which is needed to load any (non-script) Packages. It is first used in ProcessNewlyLoadedObjects (due to the loading of asset references in Class Default Objects)
 	// It has to be intialized after the AssetRegistryModule; the editor implementations of PackageResourceManager relies on it
 	IPackageResourceManager::Initialize();
-#endif
+#endif // WITH_COREUOBJECT
 #if WITH_EDITOR
 	// Initialize the BulkDataRegistry, which registers BulkData structs loaded from Packages for later building. It uses the same lifetime as IPackageResourceManager
 	IBulkDataRegistry::Initialize();
-#endif
+#endif // WITH_EDITOR
 }
 
-#if WITH_EDITOR
 void InitEditor()
 {
+#if WITH_EDITOR
 	FModuleManager::Get().LoadModuleChecked("UnrealEd");
 
 	GIsEditor = true;
@@ -151,5 +223,12 @@ void InitEditor()
 
 	GEngine->ParseCommandline();
 	GEditor->InitEditor(&GEngineLoop);
-}
 #endif // #if WITH_EDITOR
+}
+
+void CleanupPlatform()
+{
+	FPlatformMisc::PlatformTearDown();
+	FGenericPlatformMisc::RequestExit(false);
+}
+
diff --git a/Engine/Source/Developer/LowLevelTestsRunner/Public/TestRunner.h b/Engine/Source/Developer/LowLevelTestsRunner/Public/TestRunner.h
--- a/Engine/Source/Developer/LowLevelTestsRunner/Public/TestRunner.h
+++ b/Engine/Source/Developer/LowLevelTestsRunner/Public/TestRunner.h
@@ -24,17 +24,7 @@ static bool bCatchIsRunning = false;
 static bool bGAllowLogging = true;
 static bool bMultithreaded = true;
 
-/**
- * Setup phase before all tests.
- */
-void Setup();
-
-/**
- * Teardown phase after all tests.
- */
-void Teardown();
-
-#if PLATFORM_SWITCH
+#if __cplusplus
 	extern "C" const char* GetProcessExecutablePath();
 #else
 	const char* GetProcessExecutablePath();
diff --git a/Engine/Source/Programs/LowLevelTests/Tests/main.cpp b/Engine/Source/Programs/LowLevelTests/Tests/main.cpp
deleted file mode 100644
diff --git a/Engine/Source/Programs/UnrealBuildTool/Configuration/TestModuleRules.cs b/Engine/Source/Programs/UnrealBuildTool/Configuration/TestModuleRules.cs
--- a/Engine/Source/Programs/UnrealBuildTool/Configuration/TestModuleRules.cs
+++ b/Engine/Source/Programs/UnrealBuildTool/Configuration/TestModuleRules.cs
@@ -53,14 +53,9 @@ namespace UnrealBuildTool
 
 			PrivateDependencyModuleNames.AddRange(
 				new string[] {
-					"Core",
-					"Projects",
 					"LowLevelTestsRunner"
 				});
 
-			PrivateIncludePaths.Add("Runtime/Launch/Private");
-			PrivateIncludePathModuleNames.Add("Launch");
-
 			// Tests can refer to tested module's Public and Private paths
 			string ModulePublicDir = Path.Combine(TestedModule.ModuleDirectory, "Public");
 			if (System.IO.Directory.Exists(ModulePublicDir))
diff --git a/Engine/Source/Programs/UnrealBuildTool/Configuration/TestTargetRules.cs b/Engine/Source/Programs/UnrealBuildTool/Configuration/TestTargetRules.cs
--- a/Engine/Source/Programs/UnrealBuildTool/Configuration/TestTargetRules.cs
+++ b/Engine/Source/Programs/UnrealBuildTool/Configuration/TestTargetRules.cs
@@ -74,8 +74,6 @@ namespace UnrealBuildTool
 				bDebugBuildsActuallyUseDebugCRT = true;
 			}
 
-			GlobalDefinitions.AddRange(TestedTarget.GlobalDefinitions);
-
 			GlobalDefinitions.Add("STATS=0");
 
 			// Platform specific setup
diff --git a/Engine/Source/Programs/UnrealBuildTool/Configuration/UEBuildTarget.cs b/Engine/Source/Programs/UnrealBuildTool/Configuration/UEBuildTarget.cs
--- a/Engine/Source/Programs/UnrealBuildTool/Configuration/UEBuildTarget.cs
+++ b/Engine/Source/Programs/UnrealBuildTool/Configuration/UEBuildTarget.cs
@@ -124,7 +124,7 @@ namespace UnrealBuildTool
 		/// <returns></returns>
 		public override bool Equals(object? B)
 		{
-			if(Object.ReferenceEquals(B, null))
+			if (Object.ReferenceEquals(B, null))
 			{
 				return false;
 			}
@@ -193,7 +193,7 @@ namespace UnrealBuildTool
 				return new UnrealTargetPlatform(Id);
 			}
 
-			throw new BuildException(string.Format("The platform name {0} is not a valid platform name. Valid names are ({1})", Name, 
+			throw new BuildException(string.Format("The platform name {0} is not a valid platform name. Valid names are ({1})", Name,
 				string.Join(",", GetUniqueStringRegistry().GetStringNames())));
 		}
 
@@ -309,7 +309,7 @@ namespace UnrealBuildTool
 		{
 			if (value.GetType() == typeof(string))
 			{
-				return UnrealTargetPlatform.Parse((string) value);
+				return UnrealTargetPlatform.Parse((string)value);
 			}
 			return base.ConvertFrom(context, culture, value);
 		}
@@ -317,8 +317,8 @@ namespace UnrealBuildTool
 		public override object ConvertTo(ITypeDescriptorContext context, CultureInfo culture, object value, Type destinationType)
 		{
 			if (destinationType == typeof(string))
-			{ 
-				UnrealTargetPlatform Platform = (UnrealTargetPlatform) value;
+			{
+				UnrealTargetPlatform Platform = (UnrealTargetPlatform)value;
 				return Platform.ToString();
 			}
 			return base.ConvertTo(context, culture, value, destinationType);
@@ -440,7 +440,7 @@ namespace UnrealBuildTool
 		/// <returns></returns>
 		public override bool Equals(object? B)
 		{
-			if(Object.ReferenceEquals(B, null))
+			if (Object.ReferenceEquals(B, null))
 			{
 				return false;
 			}
@@ -710,7 +710,8 @@ namespace UnrealBuildTool
 			TargetRules RulesObject;
 			using (GlobalTracer.Instance.BuildSpan("RulesAssembly.CreateTargetRules()").StartActive())
 			{
-				if (!Descriptor.IsTestsTarget) {
+				if (!Descriptor.IsTestsTarget)
+				{
 					RulesObject = RulesAssembly.CreateTargetRules(Descriptor.Name, Descriptor.Platform, Descriptor.Configuration, Descriptor.Architecture, Descriptor.ProjectFile, Descriptor.AdditionalArguments);
 				}
 				else
@@ -726,7 +727,7 @@ namespace UnrealBuildTool
 			}
 
 			// Make sure this configuration is supports
-			if(!RulesObject.GetSupportedConfigurations().Contains(Descriptor.Configuration))
+			if (!RulesObject.GetSupportedConfigurations().Contains(Descriptor.Configuration))
 			{
 				throw new BuildException("{0} does not support the {1} configuration", Descriptor.Name, Descriptor.Configuration);
 			}
@@ -738,7 +739,7 @@ namespace UnrealBuildTool
 				{
 					throw new BuildException("Download support for building {0} {1} targets from the launcher.", Descriptor.Platform, RulesObject.Type);
 				}
-				else if(!InstalledPlatformInfo.IsValid(RulesObject.Type, null, null, EProjectType.Code, InstalledPlatformState.Supported))
+				else if (!InstalledPlatformInfo.IsValid(RulesObject.Type, null, null, EProjectType.Code, InstalledPlatformState.Supported))
 				{
 					throw new BuildException("{0} targets are not currently supported from this engine distribution.", RulesObject.Type);
 				}
@@ -762,7 +763,7 @@ namespace UnrealBuildTool
 			if (RulesObject.bPrecompile)
 			{
 				DirectoryReference DependencyListDir;
-				if(RulesObject.ProjectFile == null)
+				if (RulesObject.ProjectFile == null)
 				{
 					DependencyListDir = DirectoryReference.Combine(Unreal.EngineDirectory, "Intermediate", "DependencyLists", RulesObject.Name, RulesObject.Configuration.ToString(), RulesObject.Platform.ToString());
 				}
@@ -771,13 +772,13 @@ namespace UnrealBuildTool
 					DependencyListDir = DirectoryReference.Combine(RulesObject.ProjectFile.Directory, "Intermediate", "DependencyLists", RulesObject.Name, RulesObject.Configuration.ToString(), RulesObject.Platform.ToString());
 				}
 
-				if(!String.IsNullOrEmpty(Descriptor.Architecture))
+				if (!String.IsNullOrEmpty(Descriptor.Architecture))
 				{
 					DependencyListDir = new DirectoryReference(DependencyListDir.FullName + Descriptor.Architecture);
 				}
 
 				FileReference DependencyListFile;
-				if(RulesObject.bBuildAllModules)
+				if (RulesObject.bBuildAllModules)
 				{
 					DependencyListFile = FileReference.Combine(DependencyListDir, "DependencyList-AllModules.txt");
 				}
@@ -790,7 +791,7 @@ namespace UnrealBuildTool
 			}
 
 			// If we're compiling only specific files, we need to prevent unity builds from running
-			if(Descriptor.SpecificFilesToCompile.Count > 0)
+			if (Descriptor.SpecificFilesToCompile.Count > 0)
 			{
 				RulesObject.bUseUnityBuild = false;
 				RulesObject.bForceUnityBuild = false;
@@ -876,14 +877,14 @@ namespace UnrealBuildTool
 		static void ValidateSharedEnvironment(RulesAssembly RulesAssembly, string ThisTargetName, CommandLineArguments Arguments, TargetRules ThisRules)
 		{
 			// Allow disabling these checks
-			if(ThisRules.bOverrideBuildEnvironment)
+			if (ThisRules.bOverrideBuildEnvironment)
 			{
 				return;
 			}
 
 			// Get the name of the target with default settings
 			string BaseTargetName;
-			switch(ThisRules.Type)
+			switch (ThisRules.Type)
 			{
 				case TargetType.Game:
 					BaseTargetName = "UnrealGame";
@@ -907,27 +908,27 @@ namespace UnrealBuildTool
 			// Get all the configurable objects
 			object[] BaseObjects = BaseRules.GetConfigurableObjects().ToArray();
 			object[] ThisObjects = ThisRules.GetConfigurableObjects().ToArray();
-			if(BaseObjects.Length != ThisObjects.Length)
+			if (BaseObjects.Length != ThisObjects.Length)
 			{
 				throw new BuildException("Expected same number of configurable objects from base rules object.");
 			}
 
 			// Iterate through all fields with the [SharedBuildEnvironment] attribute
-			for(int Idx = 0; Idx < BaseObjects.Length; Idx++)
+			for (int Idx = 0; Idx < BaseObjects.Length; Idx++)
 			{
 				Type ObjectType = BaseObjects[Idx].GetType();
-				foreach(FieldInfo Field in ObjectType.GetFields())
+				foreach (FieldInfo Field in ObjectType.GetFields())
 				{
-					if(Field.GetCustomAttribute<RequiresUniqueBuildEnvironmentAttribute>() != null)
+					if (Field.GetCustomAttribute<RequiresUniqueBuildEnvironmentAttribute>() != null)
 					{
 						object? ThisValue = Field.GetValue(ThisObjects[Idx]);
 						object? BaseValue = Field.GetValue(BaseObjects[Idx]);
 						CheckValuesMatch(ThisRules.GetType(), ThisTargetName, BaseTargetName, Field.Name, Field.FieldType, ThisValue, BaseValue);
 					}
 				}
-				foreach(PropertyInfo Property in ObjectType.GetProperties())
+				foreach (PropertyInfo Property in ObjectType.GetProperties())
 				{
-					if(Property.GetCustomAttribute<RequiresUniqueBuildEnvironmentAttribute>() != null)
+					if (Property.GetCustomAttribute<RequiresUniqueBuildEnvironmentAttribute>() != null)
 					{
 						object? ThisValue = Property.GetValue(ThisObjects[Idx]);
 						object? BaseValue = Property.GetValue(BaseObjects[Idx]);
@@ -951,11 +952,11 @@ namespace UnrealBuildTool
 		{
 			// Check if the fields match, treating lists of strings (eg. definitions) differently to value types.
 			bool bFieldsMatch;
-			if(ThisValue == null || BaseValue == null)
+			if (ThisValue == null || BaseValue == null)
 			{
 				bFieldsMatch = (ThisValue == BaseValue);
 			}
-			else if(typeof(IEnumerable<string>).IsAssignableFrom(ValueType))
+			else if (typeof(IEnumerable<string>).IsAssignableFrom(ValueType))
 			{
 				bFieldsMatch = Enumerable.SequenceEqual((IEnumerable<string>)ThisValue, (IEnumerable<string>)BaseValue);
 			}
@@ -965,7 +966,7 @@ namespace UnrealBuildTool
 			}
 
 			// Throw an exception if they don't match
-			if(!bFieldsMatch)
+			if (!bFieldsMatch)
 			{
 				throw new BuildException("{0} modifies the value of {1}. This is not allowed, as {0} has build products in common with {2}.\nRemove the modified setting, change {0} to use a unique build environment by setting 'BuildEnvironment = TargetBuildEnvironment.Unique;' in the {3} constructor, or set bOverrideBuildEnvironment = true to force this setting on.", ThisTargetName, FieldName, BaseTargetName, RulesType.Name);
 			}
@@ -1086,6 +1087,26 @@ namespace UnrealBuildTool
 		/// </summary>
 		protected bool bCompileMonolithic = false;
 
+		/// <summary>
+		/// Keeps track if low level tests executable must build with the Editor.
+		/// </summary>
+		private bool bTestsRequireEditor = false;
+
+		/// <summary>
+		/// Keeps track if low level tests executable must build with the Engine.
+		/// </summary>
+		private bool bTestsRequireEngine = false;
+
+		/// <summary>
+		/// Keeps track if low level tests executable must build with the ApplicationCore.
+		/// </summary>
+		private bool bTestsRequireApplicationCore = false;
+
+		/// <summary>
+		/// Keeps track if low level tests executable must build with the CoreUObject.
+		/// </summary>
+		private bool bTestsRequireCoreUObject = false;
+
 		/// <summary>
 		/// Used to keep track of all modules by name.
 		/// </summary>
@@ -1116,7 +1137,7 @@ namespace UnrealBuildTool
 		/// <returns>true if it should, false if it shouldn't</returns>
 		public bool ShouldCompileMonolithic()
 		{
-			return bCompileMonolithic;	// @todo ubtmake: We need to make sure this function and similar things aren't called in assembler mode
+			return bCompileMonolithic;  // @todo ubtmake: We need to make sure this function and similar things aren't called in assembler mode
 		}
 
 		/// <summary>
@@ -1151,7 +1172,7 @@ namespace UnrealBuildTool
 			bUseSharedBuildEnvironment = (Rules.BuildEnvironment == TargetBuildEnvironment.Shared);
 			if (bUseSharedBuildEnvironment)
 			{
-				if(Rules.Type == TargetType.Program)
+				if (Rules.Type == TargetType.Program)
 				{
 					AppName = TargetName;
 				}
@@ -1222,7 +1243,7 @@ namespace UnrealBuildTool
 		/// <returns>The app name for this target type</returns>
 		public static string GetAppNameForTargetType(TargetType Type)
 		{
-			switch(Type)
+			switch (Type)
 			{
 				case TargetType.Game:
 					return "UnrealGame";
@@ -1248,14 +1269,14 @@ namespace UnrealBuildTool
 			HashSet<FileReference> Files = new HashSet<FileReference>();
 
 			// Find all the runtime dependency files in their original location
-			foreach(RuntimeDependency RuntimeDependency in RuntimeDependencies)
+			foreach (RuntimeDependency RuntimeDependency in RuntimeDependencies)
 			{
 				FileReference? SourceFile;
-				if(!RuntimeDependencyTargetFileToSourceFile.TryGetValue(RuntimeDependency.Path, out SourceFile))
+				if (!RuntimeDependencyTargetFileToSourceFile.TryGetValue(RuntimeDependency.Path, out SourceFile))
 				{
 					SourceFile = RuntimeDependency.Path;
 				}
-				if(RuntimeDependency.Type != StagedFileType.DebugNonUFS || FileReference.Exists(SourceFile))
+				if (RuntimeDependency.Type != StagedFileType.DebugNonUFS || FileReference.Exists(SourceFile))
 				{
 					Files.Add(SourceFile);
 				}
@@ -1265,7 +1286,7 @@ namespace UnrealBuildTool
 			HashSet<UEBuildModule> Modules = new HashSet<UEBuildModule>();
 			foreach (UEBuildBinary Binary in Binaries)
 			{
-				foreach(UEBuildModule Module in Binary.Modules)
+				foreach (UEBuildModule Module in Binary.Modules)
 				{
 					Modules.Add(Module);
 					Modules.UnionWith(Module.GetDependencies(true, true));
@@ -1277,7 +1298,7 @@ namespace UnrealBuildTool
 			foreach (UEBuildModule Module in Modules)
 			{
 				// Skip artificial modules
-				if(Module.RulesFile == null)
+				if (Module.RulesFile == null)
 				{
 					continue;
 				}
@@ -1341,12 +1362,12 @@ namespace UnrealBuildTool
 						{
 							ResolveLibraryName(LibraryPath, LibraryName, LibraryExtension, Files);
 						}
-							}
+					}
 
 					foreach (string LibraryName in Rules.PublicSystemLibraryPaths)
-							{
+					{
 						foreach (string LibraryPath in SystemLibraryPaths)
-								{
+						{
 							ResolveLibraryName(LibraryPath, LibraryName, LibraryExtension, Files);
 						}
 					}
@@ -1411,7 +1432,7 @@ namespace UnrealBuildTool
 			}
 
 			// Add all the dependency lists
-			foreach(FileReference DependencyListFileName in Rules.DependencyListFileNames)
+			foreach (FileReference DependencyListFileName in Rules.DependencyListFileNames)
 			{
 				Manifest.BuildProducts.Add(DependencyListFileName.FullName);
 			}
@@ -1427,7 +1448,7 @@ namespace UnrealBuildTool
 			}
 
 			// Remove anything that's not part of the plugin
-			if(ForeignPlugin != null)
+			if (ForeignPlugin != null)
 			{
 				DirectoryReference ForeignPluginDir = ForeignPlugin.Directory;
 				Manifest.BuildProducts.RemoveAll(x => !new FileReference(x).IsUnderDirectory(ForeignPluginDir));
@@ -1452,7 +1473,7 @@ namespace UnrealBuildTool
 				// Create the receipts for each folder
 				foreach (UEBuildBinary Binary in Binaries)
 				{
-					if(Binary.Type == UEBuildBinaryType.DynamicLinkLibrary)
+					if (Binary.Type == UEBuildBinaryType.DynamicLinkLibrary)
 					{
 						DirectoryReference DirectoryName = Binary.OutputFilePath.Directory;
 						bool bIsGameBinary = Binary.PrimaryModule.Rules.Context.bCanBuildDebugGame;
@@ -1493,9 +1514,9 @@ namespace UnrealBuildTool
 			// Create a unique identifier for this build which can be used to identify modules which are compatible. It's fine to share this between runs with the same makefile.
 			// By default we leave it blank when compiling a subset of modules (for hot reload, etc...), otherwise it won't match anything else. When writing to a directory
 			// that already contains a manifest, we'll reuse the build id that's already in there (see below).
-			if(String.IsNullOrEmpty(Version.BuildId))
+			if (String.IsNullOrEmpty(Version.BuildId))
 			{
-				if(Rules.bFormalBuild)
+				if (Rules.bFormalBuild)
 				{
 					// If this is a formal build, we can just the compatible changelist as the unique id.
 					Version.BuildId = String.Format("{0}", Version.EffectiveCompatibleChangelist);
@@ -1527,7 +1548,7 @@ namespace UnrealBuildTool
 							Receipt.Launch = Pair.Key;
 						}
 
-						if ( Receipt.Launch != null && Receipt.LaunchCmd != null )
+						if (Receipt.Launch != null && Receipt.LaunchCmd != null)
 						{
 							break;
 						}
@@ -1542,14 +1563,14 @@ namespace UnrealBuildTool
 			// Find all the build products and modules from this binary
 			foreach (KeyValuePair<FileReference, BuildProductType> BuildProductPair in BuildProducts)
 			{
-				if(BuildProductPair.Value != BuildProductType.BuildResource)
+				if (BuildProductPair.Value != BuildProductType.BuildResource)
 				{
 					Receipt.AddBuildProduct(BuildProductPair.Key, BuildProductPair.Value);
 				}
 			}
 
 			// Add the project file
-			if(ProjectFile != null)
+			if (ProjectFile != null)
 			{
 				Receipt.RuntimeDependencies.Add(ProjectFile, StagedFileType.UFS);
 			}
@@ -1558,7 +1579,7 @@ namespace UnrealBuildTool
 			List<UEBuildPlugin>? RuntimeDependencyPluginsList = Rules.bRuntimeDependenciesComeFromBuildPlugins ? BuildPlugins : EnabledPlugins;
 			foreach (UEBuildPlugin EnabledPlugin in RuntimeDependencyPluginsList!)
 			{
-				if(EnabledPlugin.bDescriptorNeededAtRuntime || EnabledPlugin.bDescriptorReferencedExplicitly)
+				if (EnabledPlugin.bDescriptorNeededAtRuntime || EnabledPlugin.bDescriptorReferencedExplicitly)
 				{
 					Receipt.RuntimeDependencies.Add(EnabledPlugin.File, StagedFileType.UFS);
 
@@ -1576,22 +1597,22 @@ namespace UnrealBuildTool
 							}
 						}
 
-					foreach (FileReference ChildFile in EnabledPlugin.ChildFiles)
-					{
-						if (ValidFileNames.Contains(ChildFile.GetFileNameWithoutExtension(), StringComparer.InvariantCultureIgnoreCase))
+						foreach (FileReference ChildFile in EnabledPlugin.ChildFiles)
 						{
-							Receipt.RuntimeDependencies.Add(ChildFile, StagedFileType.UFS);
+							if (ValidFileNames.Contains(ChildFile.GetFileNameWithoutExtension(), StringComparer.InvariantCultureIgnoreCase))
+							{
+								Receipt.RuntimeDependencies.Add(ChildFile, StagedFileType.UFS);
+							}
 						}
 					}
 				}
 			}
-			}
 
 			// Add all the other runtime dependencies
 			HashSet<FileReference> UniqueRuntimeDependencyFiles = new HashSet<FileReference>();
-			foreach(RuntimeDependency RuntimeDependency in RuntimeDependencies)
+			foreach (RuntimeDependency RuntimeDependency in RuntimeDependencies)
 			{
-				if(UniqueRuntimeDependencyFiles.Add(RuntimeDependency.Path))
+				if (UniqueRuntimeDependencyFiles.Add(RuntimeDependency.Path))
 				{
 					Receipt.RuntimeDependencies.Add(RuntimeDependency);
 				}
@@ -1704,7 +1725,7 @@ namespace UnrealBuildTool
 				List<UEBuildBinary> FilteredBinaries = new List<UEBuildBinary>();
 				foreach (UEBuildBinary DLLBinary in Binaries)
 				{
-					if(ModDirectories.Any(x => DLLBinary.OutputFilePath.IsUnderDirectory(x)))
+					if (ModDirectories.Any(x => DLLBinary.OutputFilePath.IsUnderDirectory(x)))
 					{
 						FilteredBinaries.Add(DLLBinary);
 					}
@@ -1719,9 +1740,9 @@ namespace UnrealBuildTool
 
 			// Build a mapping from module to its plugin
 			Dictionary<UEBuildModule, UEBuildPlugin> ModuleToPlugin = new Dictionary<UEBuildModule, UEBuildPlugin>();
-			foreach(UEBuildPlugin Plugin in BuildPlugins!)
+			foreach (UEBuildPlugin Plugin in BuildPlugins!)
 			{
-				foreach(UEBuildModule Module in Plugin.Modules)
+				foreach (UEBuildModule Module in Plugin.Modules)
 				{
 					if (!ModuleToPlugin.ContainsKey(Module))
 					{
@@ -1731,14 +1752,14 @@ namespace UnrealBuildTool
 			}
 
 			// Check there aren't any engine binaries with dependencies on game modules. This can happen when game-specific plugins override engine plugins.
-			foreach(UEBuildModule Module in Modules.Values)
+			foreach (UEBuildModule Module in Modules.Values)
 			{
-				if(Module.Binary != null)
+				if (Module.Binary != null)
 				{
 					HashSet<UEBuildModule> ReferencedModules = Module.GetDependencies(bWithIncludePathModules: true, bWithDynamicallyLoadedModules: true);
-					foreach(UEBuildModule ReferencedModule in ReferencedModules)
+					foreach (UEBuildModule ReferencedModule in ReferencedModules)
 					{
-						if(!Module.Rules.Context.Scope.Contains(ReferencedModule.Rules.Context.Scope) && !IsEngineToPluginReferenceAllowed(Module.Name, ReferencedModule.Name))
+						if (!Module.Rules.Context.Scope.Contains(ReferencedModule.Rules.Context.Scope) && !IsEngineToPluginReferenceAllowed(Module.Name, ReferencedModule.Name))
 						{
 							throw new BuildException("Module '{0}' ({1}) should not reference module '{2}' ({3}). Hierarchy is {4}.", Module.Name, Module.Rules.Context.Scope.Name, ReferencedModule.Name, ReferencedModule.Rules.Context.Scope.Name, ReferencedModule.Rules.Context.Scope.FormatHierarchy());
 						}
@@ -1747,15 +1768,15 @@ namespace UnrealBuildTool
 			}
 
 			// Check that each plugin declares its dependencies explicitly
-			foreach(UEBuildPlugin Plugin in BuildPlugins)
+			foreach (UEBuildPlugin Plugin in BuildPlugins)
 			{
-				foreach(UEBuildModule Module in Plugin.Modules)
+				foreach (UEBuildModule Module in Plugin.Modules)
 				{
 					HashSet<UEBuildModule> DependencyModules = Module.GetDependencies(bWithIncludePathModules: true, bWithDynamicallyLoadedModules: true);
-					foreach(UEBuildModule DependencyModule in DependencyModules)
+					foreach (UEBuildModule DependencyModule in DependencyModules)
 					{
 						UEBuildPlugin? DependencyPlugin;
-						if(ModuleToPlugin.TryGetValue(DependencyModule, out DependencyPlugin) && DependencyPlugin != Plugin
+						if (ModuleToPlugin.TryGetValue(DependencyModule, out DependencyPlugin) && DependencyPlugin != Plugin
 							&& !Rules.InternalPluginDependencies.Contains(DependencyPlugin.Name)
 							&& !Plugin.Dependencies!.Contains(DependencyPlugin))
 						{
@@ -1855,7 +1876,7 @@ namespace UnrealBuildTool
 			{
 				if (Platform == UnrealTargetPlatform.Win64)
 				{
-					if(!Rules.bFormalBuild)
+					if (!Rules.bFormalBuild)
 					{
 						FileReference DefaultResourceLocation = FileReference.Combine(Unreal.EngineDirectory, "Build", "Windows", "Resources", "Default.rc2");
 						if (!UnrealBuildTool.IsFileInstalled(DefaultResourceLocation))
@@ -1902,32 +1923,32 @@ namespace UnrealBuildTool
 			// Prepare all the runtime dependencies, copying them from their source folders if necessary
 			List<RuntimeDependency> RuntimeDependencies = new List<RuntimeDependency>();
 			Dictionary<FileReference, FileReference> RuntimeDependencyTargetFileToSourceFile = new Dictionary<FileReference, FileReference>();
-			foreach(UEBuildBinary Binary in Binaries)
+			foreach (UEBuildBinary Binary in Binaries)
 			{
 				Binary.PrepareRuntimeDependencies(RuntimeDependencies, RuntimeDependencyTargetFileToSourceFile, ExeDir);
 			}
 			TargetToolChain.PrepareRuntimeDependencies(RuntimeDependencies, RuntimeDependencyTargetFileToSourceFile, ExeDir);
 
-			foreach(KeyValuePair<FileReference, FileReference> Pair in RuntimeDependencyTargetFileToSourceFile)
+			foreach (KeyValuePair<FileReference, FileReference> Pair in RuntimeDependencyTargetFileToSourceFile)
 			{
-				if(!UnrealBuildTool.IsFileInstalled(Pair.Key))
+				if (!UnrealBuildTool.IsFileInstalled(Pair.Key))
 				{
 					Makefile.OutputItems.Add(Makefile.CreateCopyAction(Pair.Value, Pair.Key));
 				}
 			}
 
 			// If we're just precompiling a plugin, only include output items which are part of it
-			if(ForeignPlugin != null)
+			if (ForeignPlugin != null)
 			{
 				HashSet<FileItem> RetainOutputItems = new HashSet<FileItem>();
-				foreach(UEBuildPlugin Plugin in BuildPlugins)
+				foreach (UEBuildPlugin Plugin in BuildPlugins)
 				{
-					if(Plugin.File == ForeignPlugin)
+					if (Plugin.File == ForeignPlugin)
 					{
 						foreach (UEBuildModule Module in Plugin.Modules)
 						{
 							FileItem[]? ModuleOutputItems;
-							if(Makefile.ModuleNameToOutputItems.TryGetValue(Module.Name, out ModuleOutputItems))
+							if (Makefile.ModuleNameToOutputItems.TryGetValue(Module.Name, out ModuleOutputItems))
 							{
 								RetainOutputItems.UnionWith(ModuleOutputItems);
 							}
@@ -1951,7 +1972,7 @@ namespace UnrealBuildTool
 			BuildProducts.AddRange(RuntimeDependencyTargetFileToSourceFile.Select(x => new KeyValuePair<FileReference, BuildProductType>(x.Key, BuildProductType.RequiredResource)));
 
 			// Remove any installed build products that don't exist. They may be part of an optional install.
-			if(Unreal.IsEngineInstalled())
+			if (Unreal.IsEngineInstalled())
 			{
 				BuildProducts.RemoveAll(x => UnrealBuildTool.IsFileInstalled(x.Key) && !FileReference.Exists(x.Key));
 			}
@@ -2055,13 +2076,13 @@ namespace UnrealBuildTool
 
 				// Create actions to run the post build steps
 				FileReference[] PostBuildScripts = CreatePostBuildScripts();
-				foreach(FileReference PostBuildScript in PostBuildScripts)
+				foreach (FileReference PostBuildScript in PostBuildScripts)
 				{
 					FileReference OutputFile = new FileReference(PostBuildScript.FullName + ".ran");
 
 					Action PostBuildStepAction = Makefile.CreateAction(ActionType.PostBuildStep);
 					PostBuildStepAction.CommandPath = BuildHostPlatform.Current.Shell;
-					if(BuildHostPlatform.Current.ShellType == ShellType.Cmd)
+					if (BuildHostPlatform.Current.ShellType == ShellType.Cmd)
 					{
 						PostBuildStepAction.CommandArguments = String.Format("/C \"call \"{0}\" && type NUL >\"{1}\"\"", PostBuildScript, OutputFile);
 					}
@@ -2080,19 +2101,19 @@ namespace UnrealBuildTool
 			}
 
 			// Build a list of all the files required to build
-			foreach(FileReference DependencyListFileName in Rules.DependencyListFileNames)
+			foreach (FileReference DependencyListFileName in Rules.DependencyListFileNames)
 			{
 				WriteDependencyList(DependencyListFileName, RuntimeDependencies, RuntimeDependencyTargetFileToSourceFile);
 			}
 
 			// If we're only generating the manifest, return now
-			foreach(FileReference ManifestFileName in Rules.ManifestFileNames)
+			foreach (FileReference ManifestFileName in Rules.ManifestFileNames)
 			{
 				GenerateManifest(ManifestFileName, BuildProducts);
 			}
 
 			// Check there are no EULA or restricted folder violations
-			if(!Rules.bDisableLinking)
+			if (!Rules.bDisableLinking)
 			{
 				// Check the distribution level of all binaries based on the dependencies they have
 				if (ProjectFile == null && !Rules.bLegalToDistributeBinary)
@@ -2116,12 +2137,12 @@ namespace UnrealBuildTool
 					{
 						bResult &= Binary.CheckRestrictedFolders(RootDirectories, ModuleRestrictedFolderCache);
 					}
-					foreach(KeyValuePair<FileReference, FileReference> Pair in RuntimeDependencyTargetFileToSourceFile)
+					foreach (KeyValuePair<FileReference, FileReference> Pair in RuntimeDependencyTargetFileToSourceFile)
 					{
 						bResult &= CheckRestrictedFolders(Pair.Key, Pair.Value);
 					}
 
-					if(!bResult)
+					if (!bResult)
 					{
 						throw new BuildException("Unable to create binaries in less restricted locations than their input files.");
 					}
@@ -2132,7 +2153,7 @@ namespace UnrealBuildTool
 			}
 
 			// Add all the plugins to be tracked
-			foreach(FileReference PluginFile in PluginsBase.EnumeratePlugins(ProjectFile))
+			foreach (FileReference PluginFile in PluginsBase.EnumeratePlugins(ProjectFile))
 			{
 				FileItem PluginFileItem = FileItem.GetItemByFileReference(PluginFile);
 				Makefile.PluginFiles.Add(PluginFileItem);
@@ -2141,10 +2162,10 @@ namespace UnrealBuildTool
 			// Add all the input files to the predicate store
 			Makefile.ExternalDependencies.Add(FileItem.GetItemByFileReference(TargetRulesFile));
 			Makefile.ExternalDependencies.Add(FileItem.GetItemByFileReference(Rules.TargetSourceFile));
-			foreach(UEBuildModule Module in Modules.Values)
+			foreach (UEBuildModule Module in Modules.Values)
 			{
 				Makefile.ExternalDependencies.Add(FileItem.GetItemByFileReference(Module.RulesFile));
-				foreach(string ExternalDependency in Module.Rules.ExternalDependencies)
+				foreach (string ExternalDependency in Module.Rules.ExternalDependencies)
 				{
 					FileReference Location = FileReference.Combine(Module.RulesFile.Directory, ExternalDependency);
 					Makefile.ExternalDependencies.Add(FileItem.GetItemByFileReference(Location));
@@ -2261,9 +2282,9 @@ namespace UnrealBuildTool
 		{
 			List<RestrictedFolder> TargetRestrictedFolders = GetRestrictedFolders(TargetFile);
 			List<RestrictedFolder> SourceRestrictedFolders = GetRestrictedFolders(SourceFile);
-			foreach(RestrictedFolder SourceRestrictedFolder in SourceRestrictedFolders)
+			foreach (RestrictedFolder SourceRestrictedFolder in SourceRestrictedFolders)
 			{
-				if(!TargetRestrictedFolders.Contains(SourceRestrictedFolder))
+				if (!TargetRestrictedFolders.Contains(SourceRestrictedFolder))
 				{
 					Log.TraceError("Runtime dependency '{0}' is copied to '{1}', which does not contain a '{2}' folder.", SourceFile, TargetFile, SourceRestrictedFolder);
 					return false;
@@ -2281,11 +2302,11 @@ namespace UnrealBuildTool
 		{
 			// Find the base directory for this binary
 			DirectoryReference BaseDir;
-			if(File.IsUnderDirectory(Unreal.RootDirectory))
+			if (File.IsUnderDirectory(Unreal.RootDirectory))
 			{
 				BaseDir = Unreal.RootDirectory;
 			}
-			else if(ProjectDirectory != null && File.IsUnderDirectory(ProjectDirectory))
+			else if (ProjectDirectory != null && File.IsUnderDirectory(ProjectDirectory))
 			{
 				BaseDir = ProjectDirectory;
 			}
@@ -2329,20 +2350,20 @@ namespace UnrealBuildTool
 		{
 			// Find all the output files
 			HashSet<FileReference> OutputFiles = new HashSet<FileReference>();
-			foreach(UEBuildBinary Binary in Binaries)
+			foreach (UEBuildBinary Binary in Binaries)
 			{
 				OutputFiles.UnionWith(Binary.OutputFilePaths);
 			}
 
 			// Build a map of base filenames to their full path
 			Dictionary<string, FileReference> OutputNameToLocation = new Dictionary<string, FileReference>(StringComparer.InvariantCultureIgnoreCase);
-			foreach(FileReference OutputFile in OutputFiles)
+			foreach (FileReference OutputFile in OutputFiles)
 			{
 				OutputNameToLocation[OutputFile.GetFileName()] = OutputFile;
 			}
 
 			// Search all the output directories for files with a name matching one of our output files
-			foreach(DirectoryReference OutputDirectory in OutputFiles.Select(x => x.Directory).Distinct())
+			foreach (DirectoryReference OutputDirectory in OutputFiles.Select(x => x.Directory).Distinct())
 			{
 				if (DirectoryReference.Exists(OutputDirectory))
 				{
@@ -2374,11 +2395,11 @@ namespace UnrealBuildTool
 		/// <returns>True if the reference is allowed.</returns>
 		static bool IsEngineToPluginReferenceAllowed(string EngineModuleName, string PluginModuleName)
 		{
-			if(EngineModuleName == "AndroidDeviceDetection" && PluginModuleName == "TcpMessaging")
+			if (EngineModuleName == "AndroidDeviceDetection" && PluginModuleName == "TcpMessaging")
 			{
 				return true;
 			}
-			if(EngineModuleName == "Voice" && PluginModuleName == "AndroidPermission")
+			if (EngineModuleName == "Voice" && PluginModuleName == "AndroidPermission")
 			{
 				return true;
 			}
@@ -2438,20 +2459,20 @@ namespace UnrealBuildTool
 
 			// Find all the public definitions from each module. We pass null as the source binary to AddModuleToCompileEnvironment, which forces all _API macros to 'export' mode
 			List<string> Definitions = new List<string>(GlobalCompileEnvironment.Definitions);
-			foreach(UEBuildModule Module in Binary.Modules)
+			foreach (UEBuildModule Module in Binary.Modules)
 			{
 				Module.AddModuleToCompileEnvironment(null, null, new HashSet<DirectoryReference>(), new HashSet<DirectoryReference>(), new HashSet<DirectoryReference>(), Definitions, new List<UEBuildFramework>(), new List<FileItem>(), false);
 			}
 
 			// Write the header
-			using(StreamWriter Writer = new StreamWriter(HeaderFile.FullName))
+			using (StreamWriter Writer = new StreamWriter(HeaderFile.FullName))
 			{
 				Writer.WriteLine("#pragma once");
 				Writer.WriteLine();
-				foreach(string Definition in Definitions)
+				foreach (string Definition in Definitions)
 				{
 					int EqualsIdx = Definition.IndexOf('=');
-					if(EqualsIdx == -1)
+					if (EqualsIdx == -1)
 					{
 						Writer.WriteLine(String.Format("#define {0} 1", Definition));
 					}
@@ -2490,7 +2511,7 @@ namespace UnrealBuildTool
 						{
 							Message = string.Format("{0}\nDependant modules '{1}'", Message, string.Join("', '", NonRedistDeps));
 						}
-						if(Rules.bBreakBuildOnLicenseViolation)
+						if (Rules.bBreakBuildOnLicenseViolation)
 						{
 							Log.TraceError("ERROR: {0}", Message);
 						}
@@ -2513,23 +2534,23 @@ namespace UnrealBuildTool
 		/// </summary>
 		public static bool IsRedistributable(UEBuildModule Module)
 		{
-			if(Module.Rules != null && Module.Rules.IsRedistributableOverride.HasValue)
+			if (Module.Rules != null && Module.Rules.IsRedistributableOverride.HasValue)
 			{
 				return Module.Rules.IsRedistributableOverride.Value;
 			}
 
-			if(Module.RulesFile != null)
+			if (Module.RulesFile != null)
 			{
 				foreach (DirectoryReference ExtensionDir in Unreal.GetExtensionDirs(Unreal.EngineDirectory))
 				{
 					DirectoryReference SourceDeveloperDir = DirectoryReference.Combine(ExtensionDir, "Source/Developer");
-					if(Module.RulesFile.IsUnderDirectory(SourceDeveloperDir))
+					if (Module.RulesFile.IsUnderDirectory(SourceDeveloperDir))
 					{
 						return false;
 					}
 
 					DirectoryReference SourceEditorDir = DirectoryReference.Combine(ExtensionDir, "Source/Editor");
-					if(Module.RulesFile.IsUnderDirectory(SourceEditorDir))
+					if (Module.RulesFile.IsUnderDirectory(SourceEditorDir))
 					{
 						return false;
 					}
@@ -2556,7 +2577,7 @@ namespace UnrealBuildTool
 			// Add the plugin binaries to the build
 			foreach (UEBuildPlugin Plugin in BuildPlugins!)
 			{
-				foreach(UEBuildModuleCPP Module in Plugin.Modules)
+				foreach (UEBuildModuleCPP Module in Plugin.Modules)
 				{
 					AddModuleToBinary(Module);
 				}
@@ -2568,7 +2589,7 @@ namespace UnrealBuildTool
 			AddExtraModules();
 
 			// Create all the modules referenced by the existing binaries
-			foreach(UEBuildBinary Binary in Binaries)
+			foreach (UEBuildBinary Binary in Binaries)
 			{
 				Binary.CreateAllDependentModules(FindOrCreateModuleByName);
 			}
@@ -2579,7 +2600,7 @@ namespace UnrealBuildTool
 				List<UEBuildModule> DependencyModules = Binaries[Idx].GetAllDependencyModules(true, true);
 				foreach (UEBuildModuleCPP DependencyModule in DependencyModules.OfType<UEBuildModuleCPP>())
 				{
-					if(DependencyModule.Binary == null)
+					if (DependencyModule.Binary == null)
 					{
 						AddModuleToBinary(DependencyModule);
 					}
@@ -2587,7 +2608,7 @@ namespace UnrealBuildTool
 			}
 
 			// Add all the modules to the target if necessary.
-			if(Rules.bBuildAllModules)
+			if (Rules.bBuildAllModules)
 			{
 				AddAllValidModulesToTarget();
 			}
@@ -2595,7 +2616,7 @@ namespace UnrealBuildTool
 			// Add the external and non-C++ referenced modules to the binaries that reference them.
 			foreach (UEBuildModuleCPP Module in Modules.Values.OfType<UEBuildModuleCPP>())
 			{
-				if(Module.Binary != null)
+				if (Module.Binary != null)
 				{
 					foreach (UEBuildModule ReferencedModule in Module.GetUnboundReferences())
 					{
@@ -2639,15 +2660,15 @@ namespace UnrealBuildTool
 		{
 			// Find all the pre-build steps
 			List<Tuple<string[], UEBuildPlugin?>> PreBuildCommandBatches = new List<Tuple<string[], UEBuildPlugin?>>();
-			if(ProjectDescriptor != null && ProjectDescriptor.PreBuildSteps != null)
+			if (ProjectDescriptor != null && ProjectDescriptor.PreBuildSteps != null)
 			{
 				AddCustomBuildSteps(ProjectDescriptor.PreBuildSteps, null, PreBuildCommandBatches);
 			}
-			if(Rules.PreBuildSteps.Count > 0)
+			if (Rules.PreBuildSteps.Count > 0)
 			{
 				PreBuildCommandBatches.Add(new Tuple<string[], UEBuildPlugin?>(Rules.PreBuildSteps.ToArray(), null));
 			}
-			foreach(UEBuildPlugin BuildPlugin in BuildPlugins.Where(x => x.Descriptor.PreBuildSteps != null))
+			foreach (UEBuildPlugin BuildPlugin in BuildPlugins.Where(x => x.Descriptor.PreBuildSteps != null))
 			{
 				AddCustomBuildSteps(BuildPlugin.Descriptor.PreBuildSteps!, BuildPlugin, PreBuildCommandBatches);
 			}
@@ -2662,17 +2683,17 @@ namespace UnrealBuildTool
 		{
 			// Find all the post-build steps
 			List<Tuple<string[], UEBuildPlugin?>> PostBuildCommandBatches = new List<Tuple<string[], UEBuildPlugin?>>();
-			if(!Rules.bDisableLinking)
+			if (!Rules.bDisableLinking)
 			{
-				if(ProjectDescriptor != null && ProjectDescriptor.PostBuildSteps != null)
+				if (ProjectDescriptor != null && ProjectDescriptor.PostBuildSteps != null)
 				{
 					AddCustomBuildSteps(ProjectDescriptor.PostBuildSteps, null, PostBuildCommandBatches);
 				}
-				if(Rules.PostBuildSteps.Count > 0)
+				if (Rules.PostBuildSteps.Count > 0)
 				{
 					PostBuildCommandBatches.Add(new Tuple<string[], UEBuildPlugin?>(Rules.PostBuildSteps.ToArray(), null));
 				}
-				foreach(UEBuildPlugin BuildPlugin in BuildPlugins.Where(x => x.Descriptor.PostBuildSteps != null))
+				foreach (UEBuildPlugin BuildPlugin in BuildPlugins.Where(x => x.Descriptor.PostBuildSteps != null))
 				{
 					AddCustomBuildSteps(BuildPlugin.Descriptor.PostBuildSteps!, BuildPlugin, PostBuildCommandBatches);
 				}
@@ -2689,7 +2710,7 @@ namespace UnrealBuildTool
 		private void AddCustomBuildSteps(CustomBuildSteps BuildSteps, UEBuildPlugin? Plugin, List<Tuple<string[], UEBuildPlugin?>> CommandBatches)
 		{
 			string[]? Commands;
-			if(BuildSteps.TryGetCommands(BuildHostPlatform.Current.Platform, out Commands))
+			if (BuildSteps.TryGetCommands(BuildHostPlatform.Current.Platform, out Commands))
 			{
 				CommandBatches.Add(Tuple.Create(Commands, Plugin));
 			}
@@ -2702,7 +2723,7 @@ namespace UnrealBuildTool
 		/// <returns>Map of variable names to values</returns>
 		private Dictionary<string, string> GetTargetVariables(UEBuildPlugin? Plugin)
 		{
-			Dictionary<string, string> Variables = new Dictionary<string,string>();
+			Dictionary<string, string> Variables = new Dictionary<string, string>();
 			Variables.Add("RootDir", Unreal.RootDirectory.FullName);
 			Variables.Add("EngineDir", Unreal.EngineDirectory.FullName);
 			Variables.Add("ProjectDir", ProjectDirectory.FullName);
@@ -2710,11 +2731,11 @@ namespace UnrealBuildTool
 			Variables.Add("TargetPlatform", Platform.ToString());
 			Variables.Add("TargetConfiguration", Configuration.ToString());
 			Variables.Add("TargetType", TargetType.ToString());
-			if(ProjectFile != null)
+			if (ProjectFile != null)
 			{
 				Variables.Add("ProjectFile", ProjectFile.FullName);
 			}
-			if(Plugin != null)
+			if (Plugin != null)
 			{
 				Variables.Add("PluginDir", Plugin.Directory.FullName);
 			}
@@ -2732,26 +2753,26 @@ namespace UnrealBuildTool
 		private FileReference[] WriteCustomBuildStepScripts(UnrealTargetPlatform HostPlatform, DirectoryReference Directory, string FilePrefix, List<Tuple<string[], UEBuildPlugin?>> CommandBatches)
 		{
 			List<FileReference> ScriptFiles = new List<FileReference>();
-			foreach(Tuple<string[], UEBuildPlugin?> CommandBatch in CommandBatches)
+			foreach (Tuple<string[], UEBuildPlugin?> CommandBatch in CommandBatches)
 			{
 				// Find all the standard variables
 				Dictionary<string, string> Variables = GetTargetVariables(CommandBatch.Item2);
 
 				// Get the output path to the script
-				string ScriptExtension = (HostPlatform == UnrealTargetPlatform.Win64)? ".bat" : ".sh";
+				string ScriptExtension = (HostPlatform == UnrealTargetPlatform.Win64) ? ".bat" : ".sh";
 				FileReference ScriptFile = FileReference.Combine(Directory, String.Format("{0}-{1}{2}", FilePrefix, ScriptFiles.Count + 1, ScriptExtension));
 
 				// Write it to disk
 				List<string> Contents = new List<string>();
-				if(HostPlatform == UnrealTargetPlatform.Win64)
+				if (HostPlatform == UnrealTargetPlatform.Win64)
 				{
 					Contents.Insert(0, "@echo off");
 				}
-				foreach(string Command in CommandBatch.Item1)
+				foreach (string Command in CommandBatch.Item1)
 				{
 					Contents.Add(Utils.ExpandVariables(Command, Variables));
 				}
-				if(!DirectoryReference.Exists(ScriptFile.Directory))
+				if (!DirectoryReference.Exists(ScriptFile.Directory))
 				{
 					DirectoryReference.CreateDirectory(ScriptFile.Directory);
 				}
@@ -2807,11 +2828,11 @@ namespace UnrealBuildTool
 		{
 			// Find how many other shared PCH modules each module depends on, and use that to sort the shared PCHs by reverse order of size.
 			HashSet<UEBuildModuleCPP> SharedPCHModules = new HashSet<UEBuildModuleCPP>();
-			foreach(UEBuildBinary Binary in OriginalBinaries)
+			foreach (UEBuildBinary Binary in OriginalBinaries)
 			{
-				foreach(UEBuildModuleCPP Module in Binary.Modules.OfType<UEBuildModuleCPP>())
+				foreach (UEBuildModuleCPP Module in Binary.Modules.OfType<UEBuildModuleCPP>())
 				{
-					if(Module.Rules.SharedPCHHeaderFile != null)
+					if (Module.Rules.SharedPCHHeaderFile != null)
 					{
 						SharedPCHModules.Add(Module);
 					}
@@ -2820,7 +2841,7 @@ namespace UnrealBuildTool
 
 			// Find a priority for each shared PCH, determined as the number of other shared PCHs it includes.
 			Dictionary<UEBuildModuleCPP, int> SharedPCHModuleToPriority = new Dictionary<UEBuildModuleCPP, int>();
-			foreach(UEBuildModuleCPP SharedPCHModule in SharedPCHModules)
+			foreach (UEBuildModuleCPP SharedPCHModule in SharedPCHModules)
 			{
 				List<UEBuildModule> Dependencies = new List<UEBuildModule>();
 				SharedPCHModule.GetAllDependencyModules(Dependencies, new HashSet<UEBuildModule>(), false, false, false);
@@ -2829,13 +2850,13 @@ namespace UnrealBuildTool
 
 			// Create the shared PCH modules, in order
 			List<PrecompiledHeaderTemplate> OrderedSharedPCHModules = GlobalCompileEnvironment.SharedPCHs;
-			foreach(UEBuildModuleCPP Module in SharedPCHModuleToPriority.OrderByDescending(x => x.Value).Select(x => x.Key))
+			foreach (UEBuildModuleCPP Module in SharedPCHModuleToPriority.OrderByDescending(x => x.Value).Select(x => x.Key))
 			{
 				OrderedSharedPCHModules.Add(Module.CreateSharedPCHTemplate(this, GlobalCompileEnvironment));
 			}
 
 			// Print the ordered list of shared PCHs
-			if(OrderedSharedPCHModules.Count > 0)
+			if (OrderedSharedPCHModules.Count > 0)
 			{
 				Log.TraceVerbose("Found {0} shared PCH headers (listed in order of preference):", SharedPCHModules.Count);
 				foreach (PrecompiledHeaderTemplate SharedPCHModule in OrderedSharedPCHModules)
@@ -2876,7 +2897,7 @@ namespace UnrealBuildTool
 			HashSet<string> ValidModuleNames = new HashSet<string>();
 			foreach (UEBuildModuleCPP Module in Modules.Values.OfType<UEBuildModuleCPP>())
 			{
-				if(Module.Binary != null)
+				if (Module.Binary != null)
 				{
 					ValidModuleNames.Add(Module.Name);
 				}
@@ -2889,7 +2910,7 @@ namespace UnrealBuildTool
 			HashSet<string> FilteredModuleNames = new HashSet<string>();
 
 			// Only add engine modules for non-program targets. Programs only compile allowed modules through plugins.
-			if(TargetType != TargetType.Program)
+			if (TargetType != TargetType.Program)
 			{
 				// Find all the known module names in this assembly
 				List<string> ModuleNames = new List<string>();
@@ -2905,7 +2926,7 @@ namespace UnrealBuildTool
 
 				// Also allow anything in the developer directory in non-shipping configurations (though we deny by default
 				// unless the PrecompileForTargets setting indicates that it's actually useful at runtime).
-				if(Rules.bBuildDeveloperTools)
+				if (Rules.bBuildDeveloperTools)
 				{
 					Directories.AddRange(Unreal.GetExtensionDirs(Unreal.EngineDirectory, "Source/Developer"));
 				}
@@ -2923,9 +2944,9 @@ namespace UnrealBuildTool
 						if (!UEBuildPlatform.GetPlatformFolderNames().Any(Name => RulesType.Name.EndsWith("_" + Name)))
 						{
 							SupportedPlatformsAttribute? SupportedPlatforms = RulesType.GetCustomAttribute<SupportedPlatformsAttribute>();
-							if(SupportedPlatforms != null)
+							if (SupportedPlatforms != null)
 							{
-								if(SupportedPlatforms.Platforms.Contains(Platform))
+								if (SupportedPlatforms.Platforms.Contains(Platform))
 								{
 									FilteredModuleNames.Add(ModuleName);
 								}
@@ -2944,7 +2965,7 @@ namespace UnrealBuildTool
 
 			// Add all the plugin modules that need to be compiled
 			List<PluginInfo> Plugins = RulesAssembly.EnumeratePlugins().ToList();
-			foreach(PluginInfo Plugin in Plugins)
+			foreach (PluginInfo Plugin in Plugins)
 			{
 				// Ignore plugins which are specifically disabled by this target
 				if (Rules.DisablePlugins.Contains(Plugin.Name))
@@ -2965,7 +2986,7 @@ namespace UnrealBuildTool
 				}
 
 				// Disable any plugin that requires the build platform
-				if(Plugin.Descriptor.bRequiresBuildPlatform && ShouldExcludePlugin(Plugin, ExcludeFolders))
+				if (Plugin.Descriptor.bRequiresBuildPlatform && ShouldExcludePlugin(Plugin, ExcludeFolders))
 				{
 					continue;
 				}
@@ -2982,11 +3003,11 @@ namespace UnrealBuildTool
 					if (ModuleDescriptor.IsCompiledInConfiguration(Platform, Configuration, TargetName, TargetType, Rules.bBuildDeveloperTools, Rules.bBuildRequiresCookedData))
 					{
 						FileReference? ModuleFileName = RulesAssembly.GetModuleFileName(ModuleDescriptor.Name);
-						if(ModuleFileName == null)
+						if (ModuleFileName == null)
 						{
 							throw new BuildException("Unable to find module '{0}' referenced by {1}", ModuleDescriptor.Name, Plugin.File);
 						}
-						if(!ModuleFileName.ContainsAnyNames(ExcludeFolders, Plugin.Directory))
+						if (!ModuleFileName.ContainsAnyNames(ExcludeFolders, Plugin.Directory))
 						{
 							FilteredModuleNames.Add(ModuleDescriptor.Name);
 						}
@@ -2995,7 +3016,7 @@ namespace UnrealBuildTool
 			}
 
 			// Create rules for each remaining module, and check that it's set to be compiled
-			foreach(string FilteredModuleName in FilteredModuleNames)
+			foreach (string FilteredModuleName in FilteredModuleNames)
 			{
 				// Try to create the rules object, but catch any exceptions if it fails. Some modules (eg. SQLite) may determine that they are unavailable in the constructor.
 				ModuleRules? ModuleRules;
@@ -3017,7 +3038,7 @@ namespace UnrealBuildTool
 
 			// Now create all the precompiled modules, making sure they don't reference anything that's not in the precompiled set
 			HashSet<UEBuildModuleCPP> ValidModules = new HashSet<UEBuildModuleCPP>();
-			foreach(string ModuleName in ValidModuleNames)
+			foreach (string ModuleName in ValidModuleNames)
 			{
 				const string PrecompileReferenceChain = "allmodules option";
 				UEBuildModuleCPP Module = (UEBuildModuleCPP)FindOrCreateModuleByName(ModuleName, PrecompileReferenceChain);
@@ -3029,11 +3050,11 @@ namespace UnrealBuildTool
 			}
 
 			// Make sure precompiled modules don't reference any non-precompiled modules
-			foreach(UEBuildModuleCPP ValidModule in ValidModules)
+			foreach (UEBuildModuleCPP ValidModule in ValidModules)
 			{
-				foreach(UEBuildModuleCPP ReferencedModule in ValidModule.GetDependencies(false, true).OfType<UEBuildModuleCPP>())
+				foreach (UEBuildModuleCPP ReferencedModule in ValidModule.GetDependencies(false, true).OfType<UEBuildModuleCPP>())
 				{
-					if(!ValidModules.Contains(ReferencedModule))
+					if (!ValidModules.Contains(ReferencedModule))
 					{
 						Log.TraceWarning("Module '{0}' is not usable without module '{1}', which is not valid for this target.", ValidModule.Name, ReferencedModule.Name);
 					}
@@ -3041,9 +3062,9 @@ namespace UnrealBuildTool
 			}
 
 			// Make sure every module is built
-			foreach(UEBuildModuleCPP Module in ValidModules)
+			foreach (UEBuildModuleCPP Module in ValidModules)
 			{
-				if(Module.Binary == null)
+				if (Module.Binary == null)
 				{
 					AddModuleToBinary(Module);
 				}
@@ -3087,7 +3108,7 @@ namespace UnrealBuildTool
 			// Get the root output directory and base name (target name/app name) for this binary
 			DirectoryReference BaseOutputDirectory;
 
-			if(bUseSharedBuildEnvironment && (ModuleRules.Plugin == null || (ModuleRules.Plugin.Type != PluginType.External || ModuleRules.Plugin.bExplicitPluginTarget)))
+			if (bUseSharedBuildEnvironment && (ModuleRules.Plugin == null || (ModuleRules.Plugin.Type != PluginType.External || ModuleRules.Plugin.bExplicitPluginTarget)))
 			{
 				BaseOutputDirectory = ModuleRules.Context.DefaultOutputBaseDir;
 			}
@@ -3196,7 +3217,7 @@ namespace UnrealBuildTool
 			}
 
 			UEBuildPlatform BuildPlatform = UEBuildPlatform.GetBuildPlatform(Platform);
-			if(BuildPlatform.RequiresArchitectureSuffix())
+			if (BuildPlatform.RequiresArchitectureSuffix())
 			{
 				Result.Append(Architecture);
 			}
@@ -3266,16 +3287,16 @@ namespace UnrealBuildTool
 			Dictionary<string, UEBuildPlugin> NameToInstance = new Dictionary<string, UEBuildPlugin>(StringComparer.InvariantCultureIgnoreCase);
 
 			// Set up the foreign plugin
-			if(ForeignPlugin != null)
+			if (ForeignPlugin != null)
 			{
 				PluginReferenceDescriptor PluginReference = new PluginReferenceDescriptor(ForeignPlugin.GetFileNameWithoutExtension(), null, true);
 				AddPlugin(PluginReference, "command line", ExcludeFolders, NameToInstance, NameToInfo);
 			}
 
 			// Configure plugins explicitly enabled via target settings
-			foreach(string PluginName in Rules.EnablePlugins)
+			foreach (string PluginName in Rules.EnablePlugins)
 			{
-				if(ReferencedNames.Add(PluginName))
+				if (ReferencedNames.Add(PluginName))
 				{
 					PluginReferenceDescriptor PluginReference = new PluginReferenceDescriptor(PluginName, null, true);
 					AddPlugin(PluginReference, "target settings", ExcludeFolders, NameToInstance, NameToInfo);
@@ -3283,9 +3304,9 @@ namespace UnrealBuildTool
 			}
 
 			// Configure plugins explicitly disabled via target settings
-			foreach(string PluginName in Rules.DisablePlugins)
+			foreach (string PluginName in Rules.DisablePlugins)
 			{
-				if(ReferencedNames.Add(PluginName))
+				if (ReferencedNames.Add(PluginName))
 				{
 					PluginReferenceDescriptor PluginReference = new PluginReferenceDescriptor(PluginName, null, false);
 					AddPlugin(PluginReference, "target settings", ExcludeFolders, NameToInstance, NameToInfo);
@@ -3322,9 +3343,9 @@ namespace UnrealBuildTool
 			// Also synthesize references for plugins which are enabled by default
 			if (Rules.bCompileAgainstEngine || Rules.bCompileWithPluginSupport)
 			{
-				foreach(PluginInfo Plugin in NameToInfo.Values)
+				foreach (PluginInfo Plugin in NameToInfo.Values)
 				{
-					if(Plugin.IsEnabledByDefault(bAllowEnginePluginsEnabledByDefault) && !ReferencedNames.Contains(Plugin.Name))
+					if (Plugin.IsEnabledByDefault(bAllowEnginePluginsEnabledByDefault) && !ReferencedNames.Contains(Plugin.Name))
 					{
 						ReferencedNames.Add(Plugin.Name);
 
@@ -3344,16 +3365,16 @@ namespace UnrealBuildTool
 			}
 
 			// If this is a program, synthesize references for plugins which are enabled via the config file
-			if(TargetType == TargetType.Program)
+			if (TargetType == TargetType.Program)
 			{
 				ConfigHierarchy EngineConfig = ConfigCache.ReadHierarchy(ConfigHierarchyType.Engine, DirectoryReference.Combine(Unreal.EngineDirectory, "Programs", TargetName), Platform);
 
 				List<string>? PluginNames;
-				if(EngineConfig.GetArray("Plugins", "ProgramEnabledPlugins", out PluginNames))
+				if (EngineConfig.GetArray("Plugins", "ProgramEnabledPlugins", out PluginNames))
 				{
-					foreach(string PluginName in PluginNames)
+					foreach (string PluginName in PluginNames)
 					{
-						if(ReferencedNames.Add(PluginName))
+						if (ReferencedNames.Add(PluginName))
 						{
 							PluginReferenceDescriptor PluginReference = new PluginReferenceDescriptor(PluginName, null, true);
 							AddPlugin(PluginReference, "DefaultEngine.ini", ExcludeFolders, NameToInstance, NameToInfo);
@@ -3391,22 +3412,22 @@ namespace UnrealBuildTool
 		private UEBuildPlugin? AddPlugin(PluginReferenceDescriptor Reference, string ReferenceChain, string[] ExcludeFolders, Dictionary<string, UEBuildPlugin> NameToInstance, Dictionary<string, PluginInfo> NameToInfo)
 		{
 			// Ignore disabled references
-			if(!Reference.bEnabled)
+			if (!Reference.bEnabled)
 			{
 				return null;
 			}
 
 			// Try to get an existing reference to this plugin
 			UEBuildPlugin? Instance;
-			if(NameToInstance.TryGetValue(Reference.Name, out Instance))
+			if (NameToInstance.TryGetValue(Reference.Name, out Instance))
 			{
 				// If this is a non-optional reference, make sure that and every referenced dependency is staged
-				if(!Reference.bOptional && !Instance.bDescriptorReferencedExplicitly)
+				if (!Reference.bOptional && !Instance.bDescriptorReferencedExplicitly)
 				{
 					Instance.bDescriptorReferencedExplicitly = true;
-					if(Instance.Descriptor.Plugins != null)
+					if (Instance.Descriptor.Plugins != null)
 					{
-						foreach(PluginReferenceDescriptor NextReference in Instance.Descriptor.Plugins)
+						foreach (PluginReferenceDescriptor NextReference in Instance.Descriptor.Plugins)
 						{
 							string NextReferenceChain = String.Format("{0} -> {1}", ReferenceChain, Instance.File.GetFileName());
 							AddPlugin(NextReference, NextReferenceChain, ExcludeFolders, NameToInstance, NameToInfo);
@@ -3417,7 +3438,7 @@ namespace UnrealBuildTool
 			else
 			{
 				// Check if the plugin is required for this platform
-				if(!Reference.IsEnabledForPlatform(Platform))
+				if (!Reference.IsEnabledForPlatform(Platform))
 				{
 					Log.TraceLog("Ignoring plugin '{0}' (referenced via {1}), not enabled for platform {2}", Reference.Name, ReferenceChain, Platform);
 					return null;
@@ -3444,7 +3465,7 @@ namespace UnrealBuildTool
 
 				// Find the plugin being enabled
 				PluginInfo? Info;
-				if(!NameToInfo.TryGetValue(Reference.Name, out Info))
+				if (!NameToInfo.TryGetValue(Reference.Name, out Info))
 				{
 					if (Reference.bOptional)
 					{
@@ -3463,7 +3484,7 @@ namespace UnrealBuildTool
 				}
 
 				// Disable any plugin that requires the build platform
-				if(Info.Descriptor.bRequiresBuildPlatform && ShouldExcludePlugin(Info, ExcludeFolders))
+				if (Info.Descriptor.bRequiresBuildPlatform && ShouldExcludePlugin(Info, ExcludeFolders))
 				{
 					Log.TraceLog("Ignoring plugin '{0}' (referenced via {1}) due to missing build platform", Reference.Name, ReferenceChain);
 					return null;
@@ -3500,7 +3521,7 @@ namespace UnrealBuildTool
 						if (ModuleInfo.IsCompiledInConfiguration(Platform, Configuration, TargetName, TargetType, Rules.bBuildDeveloperTools, Rules.bBuildRequiresCookedData))
 						{
 							UEBuildModuleCPP Module = FindOrCreateCppModuleByName(ModuleInfo.Name, PluginReferenceChain);
-							if(!Instance.Modules.Contains(Module))
+							if (!Instance.Modules.Contains(Module))
 							{
 								// This could be in a child plugin so scan thorugh those as well
 								if (!Module.RulesFile.IsUnderDirectory(Info.Directory) && !Info.ChildFiles.Any(ChildFile => Module.RulesFile.IsUnderDirectory(ChildFile.Directory)))
@@ -3516,15 +3537,15 @@ namespace UnrealBuildTool
 
 				// Create the dependencies set
 				HashSet<UEBuildPlugin> Dependencies = new HashSet<UEBuildPlugin>();
-				if(Info.Descriptor.Plugins != null)
+				if (Info.Descriptor.Plugins != null)
 				{
-					foreach(PluginReferenceDescriptor NextReference in Info.Descriptor.Plugins)
+					foreach (PluginReferenceDescriptor NextReference in Info.Descriptor.Plugins)
 					{
 						UEBuildPlugin? NextInstance = AddPlugin(NextReference, PluginReferenceChain, ExcludeFolders, NameToInstance, NameToInfo);
-						if(NextInstance != null)
+						if (NextInstance != null)
 						{
 							Dependencies.Add(NextInstance);
-							if(NextInstance.Dependencies == null)
+							if (NextInstance.Dependencies == null)
 							{
 								throw new BuildException("Found circular dependency from plugin '{0}' onto itself.", NextReference.Name);
 							}
@@ -3552,7 +3573,7 @@ namespace UnrealBuildTool
 			{
 				return Plugin.File.ContainsAnyNames(ExcludeFolders, Unreal.EngineDirectory);
 			}
-			else if(ProjectFile != null)
+			else if (ProjectFile != null)
 			{
 				return Plugin.File.ContainsAnyNames(ExcludeFolders, ProjectFile.Directory);
 			}
@@ -3568,7 +3589,7 @@ namespace UnrealBuildTool
 		protected void SetupBinaries()
 		{
 			// If we're using the new method for specifying binaries, fill in the binary configurations now
-			if(Rules.LaunchModuleName == null)
+			if (Rules.LaunchModuleName == null)
 			{
 				throw new BuildException("LaunchModuleName must be set for all targets.");
 			}
@@ -3578,7 +3599,7 @@ namespace UnrealBuildTool
 
 			// Get the intermediate directory for the launch module directory. This can differ from the standard engine intermediate directory because it is always configuration-specific.
 			DirectoryReference IntermediateDirectory;
-			if(LaunchModule.RulesFile.IsUnderDirectory(Unreal.EngineDirectory) && !ShouldCompileMonolithic())
+			if (LaunchModule.RulesFile.IsUnderDirectory(Unreal.EngineDirectory) && !ShouldCompileMonolithic())
 			{
 				IntermediateDirectory = DirectoryReference.Combine(Unreal.EngineDirectory, PlatformIntermediateFolder, AppName, Configuration.ToString());
 			}
@@ -3603,7 +3624,7 @@ namespace UnrealBuildTool
 
 			// Create the binary
 			UEBuildBinary Binary = new UEBuildBinary(
-				Type: Rules.bShouldCompileAsDLL? UEBuildBinaryType.DynamicLinkLibrary : UEBuildBinaryType.Executable,
+				Type: Rules.bShouldCompileAsDLL ? UEBuildBinaryType.DynamicLinkLibrary : UEBuildBinaryType.Executable,
 				OutputFilePaths: OutputPaths,
 				IntermediateDirectory: IntermediateDirectory,
 				bAllowExports: Rules.bHasExports,
@@ -3709,7 +3730,7 @@ namespace UnrealBuildTool
 
 			// Add the 'Engine/Source' path as a global include path for all modules
 			GlobalCompileEnvironment.UserIncludePaths.Add(UnrealBuildTool.EngineSourceDirectory);
-						
+
 			[REDACTED_EMAIL]: Do any platform specific tool chain initialization here if required
 
 			UnrealTargetConfiguration EngineTargetConfiguration = Configuration == UnrealTargetConfiguration.DebugGame ? UnrealTargetConfiguration.Development : Configuration;
@@ -3772,12 +3793,12 @@ namespace UnrealBuildTool
 
 			GlobalCompileEnvironment.Definitions.Add(String.Format("IS_MONOLITHIC={0}", ShouldCompileMonolithic() ? "1" : "0"));
 
-			GlobalCompileEnvironment.Definitions.Add(String.Format("WITH_ENGINE={0}", Rules.bCompileAgainstEngine ? "1" : "0"));
+			GlobalCompileEnvironment.Definitions.Add(String.Format("WITH_ENGINE={0}", Rules.bCompileAgainstEngine || bTestsRequireEngine ? "1" : "0"));
 			GlobalCompileEnvironment.Definitions.Add(String.Format("WITH_UNREAL_DEVELOPER_TOOLS={0}", Rules.bBuildDeveloperTools ? "1" : "0"));
 			GlobalCompileEnvironment.Definitions.Add(String.Format("WITH_UNREAL_TARGET_DEVELOPER_TOOLS={0}", Rules.bBuildTargetDeveloperTools ? "1" : "0"));
 
 			// Set a macro to control whether to initialize ApplicationCore. Command line utilities should not generally need this.
-			if (Rules.bCompileAgainstApplicationCore)
+			if (Rules.bCompileAgainstApplicationCore || bTestsRequireApplicationCore)
 			{
 				GlobalCompileEnvironment.Definitions.Add("WITH_APPLICATION_CORE=1");
 			}
@@ -3786,7 +3807,7 @@ namespace UnrealBuildTool
 				GlobalCompileEnvironment.Definitions.Add("WITH_APPLICATION_CORE=0");
 			}
 
-			if (Rules.bCompileAgainstCoreUObject)
+			if (Rules.bCompileAgainstCoreUObject || bTestsRequireCoreUObject)
 			{
 				GlobalCompileEnvironment.Definitions.Add("WITH_COREUOBJECT=1");
 			}
@@ -3885,7 +3906,7 @@ namespace UnrealBuildTool
 				GlobalCompileEnvironment.Definitions.Add("USE_ESTIMATED_UTCNOW=0");
 			}
 
-			if (Rules.bCompileAgainstEditor && (Rules.Type == TargetType.Editor || Rules.Type == TargetType.Program))
+			if ((Rules.bCompileAgainstEditor && (Rules.Type == TargetType.Editor || Rules.Type == TargetType.Program)) || bTestsRequireEditor)
 			{
 				GlobalCompileEnvironment.Definitions.Add("WITH_EDITOR=1");
 				GlobalCompileEnvironment.Definitions.Add("WITH_IOSTORE_IN_EDITOR=1");
@@ -3931,13 +3952,13 @@ namespace UnrealBuildTool
 			}
 
 			// Set the define for enabling live coding
-			if(Rules.bWithLiveCoding)
+			if (Rules.bWithLiveCoding)
 			{
 				GlobalCompileEnvironment.Definitions.Add("WITH_LIVE_CODING=1");
 				if (Rules.LinkType == TargetLinkType.Monolithic)
 				{
 					GlobalCompileEnvironment.Definitions.Add(String.Format("UE_LIVE_CODING_ENGINE_DIR=\"{0}\"", Unreal.EngineDirectory.FullName.Replace("\\", "\\\\")));
-					if(ProjectFile != null)
+					if (ProjectFile != null)
 					{
 						GlobalCompileEnvironment.Definitions.Add(String.Format("UE_LIVE_CODING_PROJECT=\"{0}\"", ProjectFile.FullName.Replace("\\", "\\\\")));
 					}
@@ -4060,7 +4081,7 @@ namespace UnrealBuildTool
 			}
 
 			// Set whether the module requires an IMPLEMENT_MODULE macro
-			if(!RulesObject.bRequiresImplementModule.HasValue)
+			if (!RulesObject.bRequiresImplementModule.HasValue)
 			{
 				RulesObject.bRequiresImplementModule = (RulesObject.Type == ModuleRules.ModuleType.CPlusPlus && RulesObject.Name != Rules.LaunchModuleName);
 			}
@@ -4092,26 +4113,26 @@ namespace UnrealBuildTool
 				}
 
 				// Make sure that engine modules use shared PCHs or have an explicit private PCH
-				if(RulesObject.PCHUsage == ModuleRules.PCHUsageMode.NoSharedPCHs && RulesObject.PrivatePCHHeaderFile == null)
+				if (RulesObject.PCHUsage == ModuleRules.PCHUsageMode.NoSharedPCHs && RulesObject.PrivatePCHHeaderFile == null)
 				{
-					if(ProjectFile == null || !RulesObject.File.IsUnderDirectory(ProjectFile.Directory))
+					if (ProjectFile == null || !RulesObject.File.IsUnderDirectory(ProjectFile.Directory))
 					{
 						Log.TraceWarning("{0} module has shared PCHs disabled, but does not have a private PCH set", ModuleName);
 					}
 				}
 
 				// If we can't use a shared PCH, check there's a private PCH set
-				if(RulesObject.PCHUsage != ModuleRules.PCHUsageMode.NoPCHs && RulesObject.PCHUsage != ModuleRules.PCHUsageMode.UseExplicitOrSharedPCHs && RulesObject.PrivatePCHHeaderFile == null)
+				if (RulesObject.PCHUsage != ModuleRules.PCHUsageMode.NoPCHs && RulesObject.PCHUsage != ModuleRules.PCHUsageMode.UseExplicitOrSharedPCHs && RulesObject.PrivatePCHHeaderFile == null)
 				{
 					// Try to figure out the legacy PCH file
 					FileReference CppFile = DirectoryReference.EnumerateFiles(RulesObject.Directory, "*.cpp", SearchOption.AllDirectories).FirstOrDefault();
-					if(CppFile != null)
+					if (CppFile != null)
 					{
 						string? IncludeFile = MetadataCache.GetFirstInclude(FileItem.GetItemByFileReference(CppFile));
-						if(IncludeFile != null)
+						if (IncludeFile != null)
 						{
 							FileReference PchIncludeFile = DirectoryReference.EnumerateFiles(RulesObject.Directory, Path.GetFileName(IncludeFile), SearchOption.AllDirectories).FirstOrDefault();
-							if(PchIncludeFile != null)
+							if (PchIncludeFile != null)
 							{
 								RulesObject.PrivatePCHHeaderFile = PchIncludeFile.MakeRelativeTo(RulesObject.Directory).Replace(Path.DirectorySeparatorChar, '/');
 							}
@@ -4119,7 +4140,7 @@ namespace UnrealBuildTool
 					}
 
 					// Print a suggestion for which file to include
-					if(RulesObject.PrivatePCHHeaderFile == null)
+					if (RulesObject.PrivatePCHHeaderFile == null)
 					{
 						Log.TraceWarningOnce(RulesObject.File, "Modules must specify an explicit precompiled header (eg. PrivatePCHHeaderFile = \"Private/{0}PrivatePCH.h\") from UE 4.21 onwards.", ModuleName);
 					}
@@ -4142,7 +4163,7 @@ namespace UnrealBuildTool
 		/// <param name="Paths">List of paths to process</param>
 		private static void RemoveTrailingSlashes(List<string> Paths)
 		{
-			for(int Idx = 0; Idx < Paths.Count; Idx++)
+			for (int Idx = 0; Idx < Paths.Count; Idx++)
 			{
 				Paths[Idx] = Paths[Idx].TrimEnd('\\');
 			}
@@ -4165,13 +4186,7 @@ namespace UnrealBuildTool
 				DirectoryReference ModuleDirectory = RulesObject.File.Directory;
 
 				// If we're building a test executable, include test harness dependency
-				if (!IsTestModule && ReferenceChain.Contains(typeof(TestModuleRules).Name) && Directory.Exists(RulesObject.TestsDirectory) || RulesObject.Target.bIncludeAllTests)
-				{
-					if (RulesObject.Name != "LowLevelTestsRunner")
-					{
-						RulesObject.PrivateDependencyModuleNames.Add("LowLevelTestsRunner");
-					}
-				}
+				PrepareModuleForTests(IsTestModule, RulesObject);
 
 				// Clear the bUsePrecompiled flag if we're compiling a foreign plugin; since it's treated like an engine module, it will default to true in an installed build.
 				if (RulesObject.Plugin != null && RulesObject.Plugin.File == ForeignPlugin)
@@ -4284,6 +4299,32 @@ namespace UnrealBuildTool
 			return Module;
 		}
 
+		/// <summary>
+		/// Prepares a module for building a low level tests executable.
+		/// If the target includes all tests OR we're building a module as part of a test module chain, then we include the LowLevelTestsRunner dependency.
+		/// We also keep track of any Editor and Engine dependencies and add appropriate global definitions at compilation step.
+		/// </summary>
+		/// <param name="IsTestModule">Whether we're building a tests module or not.</param>
+		/// <param name="RulesObject">The module rules object.</param>
+		private void PrepareModuleForTests(bool IsTestModule, ModuleRules RulesObject)
+		{
+			if (!IsTestModule
+				&& Rules.IsTestTarget()
+				&& Directory.Exists(RulesObject.TestsDirectory)
+				|| Rules.bIncludeAllTests)
+			{
+				if (RulesObject.Name != "LowLevelTestsRunner")
+				{
+					RulesObject.PrivateIncludePathModuleNames.Add("LowLevelTestsRunner");
+				}
+
+				bTestsRequireEditor = (RulesObject.Name == "Editor");
+				bTestsRequireEngine = (RulesObject.Name == "Engine");
+				bTestsRequireApplicationCore = (RulesObject.Name == "ApplicationCore");
+				bTestsRequireCoreUObject = (RulesObject.Name == "CoreUObject");
+			}
+		}
+
 		/// <summary>
 		/// Constructs a new C++ module
 		/// </summary>
@@ -4294,7 +4335,7 @@ namespace UnrealBuildTool
 		public UEBuildModuleCPP FindOrCreateCppModuleByName(string ModuleName, string ReferenceChain, bool IsTestModule = false)
 		{
 			UEBuildModuleCPP? CppModule = FindOrCreateModuleByName(ModuleName, ReferenceChain, IsTestModule) as UEBuildModuleCPP;
-			if(CppModule == null)
+			if (CppModule == null)
 			{
 				throw new BuildException("'{0}' is not a C++ module (referenced via {1})", ModuleName, ReferenceChain);
 			}
@@ -4358,7 +4399,7 @@ namespace UnrealBuildTool
 			List<string> NewPathList = new List<string>();
 			foreach (string Path in PathList)
 			{
-				if(Path.StartsWith("$(", StringComparison.Ordinal))
+				if (Path.StartsWith("$(", StringComparison.Ordinal))
 				{
 					NewPathList.Add(Path);
 				}
