LLT: More initialization boilerplate
#rnx

diff --git a/Engine/Source/Programs/UnrealBuildTool/Configuration/UEBuildTarget.cs b/Engine/Source/Programs/UnrealBuildTool/Configuration/UEBuildTarget.cs
--- a/Engine/Source/Programs/UnrealBuildTool/Configuration/UEBuildTarget.cs
+++ b/Engine/Source/Programs/UnrealBuildTool/Configuration/UEBuildTarget.cs
@@ -4310,18 +4310,30 @@ namespace UnrealBuildTool
 		{
 			if (!IsTestModule
 				&& Rules.IsTestTarget()
-				&& Directory.Exists(RulesObject.TestsDirectory)
 				|| Rules.bIncludeAllTests)
 			{
-				if (RulesObject.Name != "LowLevelTestsRunner")
+				if (RulesObject.Name != "LowLevelTestsRunner" && Directory.Exists(RulesObject.TestsDirectory))
 				{
 					RulesObject.PrivateIncludePathModuleNames.Add("LowLevelTestsRunner");
 				}
 
-				bTestsRequireEditor = (RulesObject.Name == "Editor");
-				bTestsRequireEngine = (RulesObject.Name == "Engine");
-				bTestsRequireApplicationCore = (RulesObject.Name == "ApplicationCore");
-				bTestsRequireCoreUObject = (RulesObject.Name == "CoreUObject");
+				// If one of these modules is in the dependency graph, we must enable their appropriate global definitions
+				if (RulesObject.Name == "Editor")
+				{
+					bTestsRequireEditor = true;
+				}
+				if (RulesObject.Name == "Engine")
+				{
+					bTestsRequireEngine = true;
+				}
+				if (RulesObject.Name == "ApplicationCore")
+				{
+					bTestsRequireApplicationCore = true;
+				}
+				if (RulesObject.Name == "CoreUObject")
+				{
+					bTestsRequireCoreUObject = true;
+				}
 			}
 		}
 
