Remove iOS 13 from plist to enforce minimum supported version of iOS 14

diff --git a/Engine/Source/Programs/UnrealBuildTool/Platform/IOS/UEBuildIOS.cs b/Engine/Source/Programs/UnrealBuildTool/Platform/IOS/UEBuildIOS.cs
--- a/Engine/Source/Programs/UnrealBuildTool/Platform/IOS/UEBuildIOS.cs
+++ b/Engine/Source/Programs/UnrealBuildTool/Platform/IOS/UEBuildIOS.cs
@@ -374,14 +374,12 @@ namespace UnrealBuildTool
 			{
 				switch (MinimumIOSVersion)
 				{
-					case "IOS_13":
-						return "13.0";
 					case "IOS_14":
 						return "14.0";
 					case "IOS_15":
 						return "15.0";
 					default:
-						return "13.0";
+						return "14.0";
 				}
 			}
 		}
diff --git a/Engine/Source/Programs/UnrealBuildTool/Platform/IOS/UEDeployIOS.cs b/Engine/Source/Programs/UnrealBuildTool/Platform/IOS/UEDeployIOS.cs
--- a/Engine/Source/Programs/UnrealBuildTool/Platform/IOS/UEDeployIOS.cs
+++ b/Engine/Source/Programs/UnrealBuildTool/Platform/IOS/UEDeployIOS.cs
@@ -201,9 +201,6 @@ namespace UnrealBuildTool
 			string MinVersionToReturn = "";
 			switch (MinVersion)
 			{
-				case "IOS_13":
-					MinVersionToReturn = "13.0";
-					break;
 				case "IOS_14":
 					MinVersionToReturn = "14.0";
 					break;
@@ -211,8 +208,8 @@ namespace UnrealBuildTool
 					MinVersionToReturn = "15.0";
 					break;
 				default:
-					Log.TraceWarning("MinimumiOSVersion {0} specified in ini file is no longer supported, defaulting to 13.0", MinVersion);
-					MinVersionToReturn = "13.0";
+					Log.TraceWarning("MinimumiOSVersion {0} specified in ini file is no longer supported, defaulting to 14.0", MinVersion);
+					MinVersionToReturn = "14.0";
 					break;
 			}
 			return MinVersionToReturn;
diff --git a/Templates/TP_AEC_HandheldARBP/Config/DefaultEngine.ini b/Templates/TP_AEC_HandheldARBP/Config/DefaultEngine.ini
--- a/Templates/TP_AEC_HandheldARBP/Config/DefaultEngine.ini
+++ b/Templates/TP_AEC_HandheldARBP/Config/DefaultEngine.ini
@@ -109,7 +109,7 @@ MobileProvision=iOS11.mobileprovision
 bSupportsPortraitOrientation=True
 bSupportsUpsideDownOrientation=False
 FrameRateLock=PUFRL_60
-MinimumiOSVersion=IOS_13
+MinimumiOSVersion=IOS_14
 
 [/Script/Engine.Engine]
 bUseFixedFrameRate=True
diff --git a/Templates/TP_FirstPersonBP/Config/DefaultEngine.ini b/Templates/TP_FirstPersonBP/Config/DefaultEngine.ini
--- a/Templates/TP_FirstPersonBP/Config/DefaultEngine.ini
+++ b/Templates/TP_FirstPersonBP/Config/DefaultEngine.ini
@@ -17,7 +17,7 @@ GlobalDefaultGameMode=/Game/FirstPerson/Blueprints/BP_FirstPersonGameMode.BP_Fir
 GlobalDefaultServerGameMode=None
 
 [/Script/IOSRuntimeSettings.IOSRuntimeSettings]
-MinimumiOSVersion=IOS_13
+MinimumiOSVersion=IOS_14
 
 [/Script/Engine.RendererSettings]
 r.ReflectionMethod=1
diff --git a/Templates/TP_HandheldARBP/Config/DefaultEngine.ini b/Templates/TP_HandheldARBP/Config/DefaultEngine.ini
--- a/Templates/TP_HandheldARBP/Config/DefaultEngine.ini
+++ b/Templates/TP_HandheldARBP/Config/DefaultEngine.ini
@@ -107,7 +107,7 @@ MobileProvision=iOS11.mobileprovision
 bSupportsPortraitOrientation=True
 bSupportsUpsideDownOrientation=False
 FrameRateLock=PUFRL_60
-MinimumiOSVersion=IOS_13
+MinimumiOSVersion=IOS_14
 
 [/Script/Engine.Engine]
 bUseFixedFrameRate=True
diff --git a/Templates/TP_MFG_HandheldARBP/Config/DefaultEngine.ini b/Templates/TP_MFG_HandheldARBP/Config/DefaultEngine.ini
--- a/Templates/TP_MFG_HandheldARBP/Config/DefaultEngine.ini
+++ b/Templates/TP_MFG_HandheldARBP/Config/DefaultEngine.ini
@@ -109,7 +109,7 @@ MobileProvision=iOS11.mobileprovision
 bSupportsPortraitOrientation=True
 bSupportsUpsideDownOrientation=False
 FrameRateLock=PUFRL_60
-MinimumiOSVersion=IOS_13
+MinimumiOSVersion=IOS_14
 
 [/Script/Engine.Engine]
 bUseFixedFrameRate=True
diff --git a/Templates/TP_PuzzleBP/Config/DefaultEngine.ini b/Templates/TP_PuzzleBP/Config/DefaultEngine.ini
--- a/Templates/TP_PuzzleBP/Config/DefaultEngine.ini
+++ b/Templates/TP_PuzzleBP/Config/DefaultEngine.ini
@@ -12,4 +12,4 @@ GlobalDefaultGameMode=/Game/PuzzleBP/Blueprints/PuzzleGameMode.PuzzleGameMode_C
 GlobalDefaultServerGameMode=None
 
 [/Script/IOSRuntimeSettings.IOSRuntimeSettings]
-MinimumiOSVersion=IOS_13
+MinimumiOSVersion=IOS_14
diff --git a/Templates/TP_ThirdPersonBP/Config/DefaultEngine.ini b/Templates/TP_ThirdPersonBP/Config/DefaultEngine.ini
--- a/Templates/TP_ThirdPersonBP/Config/DefaultEngine.ini
+++ b/Templates/TP_ThirdPersonBP/Config/DefaultEngine.ini
@@ -12,7 +12,7 @@ GlobalDefaultGameMode=/Game/ThirdPerson/Blueprints/BP_ThirdPersonGameMode.BP_Thi
 GlobalDefaultServerGameMode=None
 
 [/Script/IOSRuntimeSettings.IOSRuntimeSettings]
-MinimumiOSVersion=IOS_13
+MinimumiOSVersion=IOS_14
 
 [/Script/Engine.RendererSettings]
 r.ReflectionMethod=1
diff --git a/Templates/TP_TopDownBP/Config/DefaultEngine.ini b/Templates/TP_TopDownBP/Config/DefaultEngine.ini
--- a/Templates/TP_TopDownBP/Config/DefaultEngine.ini
+++ b/Templates/TP_TopDownBP/Config/DefaultEngine.ini
@@ -50,7 +50,7 @@ VerticalDeviationFromGroundCompensation=0.000000
 bForceRebuildOnLoad=True
 
 [/Script/IOSRuntimeSettings.IOSRuntimeSettings]
-MinimumiOSVersion=IOS_13
+MinimumiOSVersion=IOS_14
 
 [/Script/Engine.RendererSettings]
 r.ReflectionMethod=1
diff --git a/Templates/TP_VehicleAdvBP/Config/DefaultEngine.ini b/Templates/TP_VehicleAdvBP/Config/DefaultEngine.ini
--- a/Templates/TP_VehicleAdvBP/Config/DefaultEngine.ini
+++ b/Templates/TP_VehicleAdvBP/Config/DefaultEngine.ini
@@ -18,5 +18,5 @@ GlobalDefaultGameMode=/Game/VehicleAdvBP/Blueprints/VehicleAdvGameMode.VehicleAd
 GlobalDefaultServerGameMode=None
 
 [/Script/IOSRuntimeSettings.IOSRuntimeSettings]
-MinimumiOSVersion=IOS_13
+MinimumiOSVersion=IOS_14
 
