Add support for bRetainFramePointers to LinuxToolChain.
Add command line argument for bRetainFramePointers.

diff --git a/Engine/Source/Programs/UnrealBuildTool/Configuration/TargetRules.cs b/Engine/Source/Programs/UnrealBuildTool/Configuration/TargetRules.cs
--- a/Engine/Source/Programs/UnrealBuildTool/Configuration/TargetRules.cs
+++ b/Engine/Source/Programs/UnrealBuildTool/Configuration/TargetRules.cs
@@ -1122,8 +1122,23 @@ namespace UnrealBuildTool
 		/// <summary>
 		/// Forces frame pointers to be retained this is usually required when you want reliable callstacks e.g. mallocframeprofiler
 		/// </summary>
-		[XmlConfigFile(Category = "BuildConfiguration")]
-		public bool bRetainFramePointers = true;
+		public bool bRetainFramePointers
+		{
+			get 
+			{
+				// Default to disabled on Linux to maintain legacy behavior
+				return bRetainFramePointersOverride ?? Platform.IsInGroup(UnrealPlatformGroup.Linux) == false;
+			}
+			set { bRetainFramePointersOverride = value; }
+		}
+
+		/// <summary>
+		/// Forces frame pointers to be retained this is usually required when you want reliable callstacks e.g. mallocframeprofiler
+		/// </summary>
+		[XmlConfigFile(Category = "BuildConfiguration", Name="bRetainFramePointers")]
+		[CommandLine("-RetainFramePointers", Value="true")]
+		[CommandLine("-NoRetainFramePointers", Value="false")]
+		public bool? bRetainFramePointersOverride = null;
 
 		/// <summary>
 		/// New Monolithic Graphics drivers have optional "fast calls" replacing various D3d functions
diff --git a/Engine/Source/Programs/UnrealBuildTool/Platform/Linux/LinuxToolChain.cs b/Engine/Source/Programs/UnrealBuildTool/Platform/Linux/LinuxToolChain.cs
--- a/Engine/Source/Programs/UnrealBuildTool/Platform/Linux/LinuxToolChain.cs
+++ b/Engine/Source/Programs/UnrealBuildTool/Platform/Linux/LinuxToolChain.cs
@@ -673,21 +673,32 @@ namespace UnrealBuildTool
 
 			//Result += " -DOPERATOR_NEW_INLINE=FORCENOINLINE";
 
+			bool bRetainFramePointers = CompileEnvironment.bRetainFramePointers 
+				|| Options.HasFlag(LinuxToolChainOptions.EnableAddressSanitizer) || Options.HasFlag(LinuxToolChainOptions.EnableMemorySanitizer)
+				|| CompileEnvironment.Configuration == CppConfiguration.Debug;
+
 			// shipping builds will cause this warning with "ensure", so disable only in those case
 			if (CompileEnvironment.Configuration == CppConfiguration.Shipping)
 			{
 				Result += " -Wno-unused-value";
-				Result += " -fomit-frame-pointer";
+				if (!bRetainFramePointers)
+                {
+					Result += " -fomit-frame-pointer";
+				}
 			}
 			// switches to help debugging
 			else if (CompileEnvironment.Configuration == CppConfiguration.Debug)
 			{
 				Result += " -fno-inline";                   // disable inlining for better debuggability (e.g. callstacks, "skip file" in gdb)
-				Result += " -fno-omit-frame-pointer";       // force not omitting fp
 				Result += " -fstack-protector";             // detect stack smashing
 				//Result += " -fsanitize=address";            // detect address based errors (support properly and link to libasan)
 			}
 
+			if (bRetainFramePointers)
+			{
+				Result += " -fno-optimize-sibling-calls -fno-omit-frame-pointer";
+			}
+
 			// debug info
 			// bCreateDebugInfo is normally set for all configurations, including Shipping - this is needed to enable callstack in Shipping builds (proper resolution: UEPLAT-205, separate files with debug info)
 			if (CompileEnvironment.bCreateDebugInfo)
@@ -717,7 +728,7 @@ namespace UnrealBuildTool
 				// Don't over optimise if using Address/MemorySanitizer or you'll get false positive errors due to erroneous optimisation of necessary Address/MemorySanitizer instrumentation.
 				if (Options.HasFlag(LinuxToolChainOptions.EnableAddressSanitizer) || Options.HasFlag(LinuxToolChainOptions.EnableMemorySanitizer))
 				{
-					Result += " -O1 -g -fno-optimize-sibling-calls -fno-omit-frame-pointer";
+					Result += " -O1 -g";
 				}
 				else if (Options.HasFlag(LinuxToolChainOptions.EnableThreadSanitizer))
 				{
