Exposed UBT relative base dir behavior to public api so it can be consumed by external codegen tools

diff --git a/Engine/Source/Programs/UnrealBuildTool/Configuration/UEBuildTarget.cs b/Engine/Source/Programs/UnrealBuildTool/Configuration/UEBuildTarget.cs
--- a/Engine/Source/Programs/UnrealBuildTool/Configuration/UEBuildTarget.cs
+++ b/Engine/Source/Programs/UnrealBuildTool/Configuration/UEBuildTarget.cs
@@ -3758,15 +3758,10 @@ namespace UnrealBuildTool
 			// By default, shadow source files for this target in the root OutputDirectory
 			GlobalLinkEnvironment.LocalShadowDirectory = GlobalLinkEnvironment.OutputDirectory;
 
-			DirectoryReference OutputDir = Binaries[0].OutputDir;
-			if (OutputDir.IsUnderDirectory(Unreal.EngineDirectory))
+			string? RelativeBaseDir = GetRelativeBaseDir(Binaries[0].OutputDir, Platform);
+			if (RelativeBaseDir != null)
 			{
-				DirectoryReference BaseDir = DirectoryReference.Combine(Unreal.EngineDirectory, "Binaries", Platform.ToString());
-				if (BaseDir != OutputDir)
-				{
-					string RelativeBaseDir = BaseDir.MakeRelativeTo(OutputDir).Replace(Path.DirectorySeparatorChar, '/');
-					GlobalCompileEnvironment.Definitions.Add(String.Format("UE_RELATIVE_BASE_DIR=\"{0}/\"", RelativeBaseDir));
-				}
+				GlobalCompileEnvironment.Definitions.Add(String.Format("UE_RELATIVE_BASE_DIR=\"{0}/\"", RelativeBaseDir));
 			}
 
 			bool bCompileDevTests = (Configuration != UnrealTargetConfiguration.Test && Configuration != UnrealTargetConfiguration.Shipping);
@@ -4169,6 +4164,26 @@ namespace UnrealBuildTool
 			}
 		}
 
+		/// <summary>
+		/// Determines the relative base directory path from the specified Binary directory to the Engine base directory if possible.
+		/// </summary>
+		/// <param name="BinaryDirectory">Directory that contains the platform subdirectoies for each UBT built binary</param>
+		/// <param name="TargetPlatform">Platform that the binary is for</param>
+		/// <returns>A string specifying the relative path from the binary directory to the engine base dir, or null if this isn't contained within the engine hierarhcy</returns>
+		public static string? GetRelativeBaseDir(DirectoryReference BinaryDirectory, UnrealTargetPlatform TargetPlatform)
+		{
+			if (BinaryDirectory.IsUnderDirectory(Unreal.EngineDirectory))
+			{
+				DirectoryReference BaseDir = DirectoryReference.Combine(Unreal.EngineDirectory, "Binaries", TargetPlatform.ToString());
+				if (BaseDir != BinaryDirectory)
+				{
+					return BaseDir.MakeRelativeTo(BinaryDirectory).Replace(Path.DirectorySeparatorChar, '/');
+				}
+			}
+
+			return null;
+		}
+
 		/// <summary>
 		/// Finds a module given its name.  Throws an exception if the module couldn't be found.
 		/// </summary>
