Windows-based platform extensions use the same icon resource as normal Win64.

diff --git a/Engine/Source/Programs/UnrealBuildTool/Configuration/UEBuildBinary.cs b/Engine/Source/Programs/UnrealBuildTool/Configuration/UEBuildBinary.cs
--- a/Engine/Source/Programs/UnrealBuildTool/Configuration/UEBuildBinary.cs
+++ b/Engine/Source/Programs/UnrealBuildTool/Configuration/UEBuildBinary.cs
@@ -776,7 +776,7 @@ namespace UnrealBuildTool
 			BinaryLinkEnvironment.bIsBuildingLibrary = IsBuildingLibrary(Type);
 
 			// If we don't have any resource file, use the default or compile a custom one for this module
-			if(BinaryLinkEnvironment.Platform == UnrealTargetPlatform.Win64)
+			if(BinaryLinkEnvironment.Platform.IsInGroup(UnrealPlatformGroup.Windows))
 			{
 				// Figure out if this binary has any custom resource files. Hacky check to ignore the resource file in the Launch module, since it contains dialogs that the engine needs and always needs to be included.
 				FileItem[] CustomResourceFiles = BinaryLinkEnvironment.InputFiles.Where(x => x.Location.HasExtension(".res") && !x.Location.FullName.EndsWith("\\Launch\\PCLaunch.rc.res", StringComparison.OrdinalIgnoreCase)).ToArray();
