UnrealBuildTool: Log PVS-Studio failure as an error

#rnx

diff --git a/Engine/Source/Programs/UnrealBuildTool/Executors/ParallelExecutor.cs b/Engine/Source/Programs/UnrealBuildTool/Executors/ParallelExecutor.cs
--- a/Engine/Source/Programs/UnrealBuildTool/Executors/ParallelExecutor.cs
+++ b/Engine/Source/Programs/UnrealBuildTool/Executors/ParallelExecutor.cs
@@ -6,6 +6,7 @@ using System.Collections.Generic;
 using System.Diagnostics;
 using System.IO;
 using System.Linq;
+using System.Text.Json;
 using System.Threading;
 using System.Threading.Tasks;
 
@@ -159,7 +160,7 @@ namespace UnrealBuildTool
 			using SemaphoreSlim MaxProcessSemaphore = new SemaphoreSlim(ActualNumParallelProcesses, ActualNumParallelProcesses);
 			using ProgressWriter ProgressWriter = new ProgressWriter("Compiling C++ source code...", false);
 
-			Log.TraceInformation("Building {0} {1} with {2} {3}...", TotalActions, (TotalActions == 1) ? "action" : "actions", ActualNumParallelProcesses, (ActualNumParallelProcesses == 1) ? "process" : "processes");
+			Log.TraceInformation($"Building {TotalActions} {(TotalActions == 1 ? "action" : "actions")} with {ActualNumParallelProcesses} {(ActualNumParallelProcesses == 1 ? "process" : "processes")}...");
 
 			Dictionary<LinkedAction, Task<ExecuteResults>> Tasks = new Dictionary<LinkedAction, Task<ExecuteResults>>();
 			List<Task> AllTasks = new List<Task>();
@@ -188,24 +189,24 @@ namespace UnrealBuildTool
 				Log.TraceInformation("");
 				if (ProcessGroup.TotalProcessorTime.Ticks > 0)
 				{
-					Log.TraceInformation("Total CPU Time: {0} s", ProcessGroup.TotalProcessorTime.TotalSeconds);
+					Log.TraceInformation($"Total CPU Time: {ProcessGroup.TotalProcessorTime.TotalSeconds} s");
 					Log.TraceInformation("");
 				}
 
 				if (Tasks.Count > 0)
 				{
-					Log.TraceInformation("Compilation Time Top {0}", Math.Min(20, Tasks.Count));
+					Log.TraceInformation($"Compilation Time Top {Math.Min(20, Tasks.Count)}");
 					Log.TraceInformation("");
 					foreach (var Pair in Tasks.OrderByDescending(x => x.Value.Result.ExecutionTime).Take(20))
 					{
 						string Description = $"{(Pair.Key.Inner.CommandDescription ?? Pair.Key.Inner.CommandPath.GetFileNameWithoutExtension())} {Pair.Key.Inner.StatusDescription}".Trim();
 						if (Pair.Value.Result.ProcessorTime.Ticks > 0)
 						{
-							Log.TraceInformation("{0} [ Wall Time {1:0.00} s / CPU Time {2:0.00} s ]", Description, Pair.Value.Result.ExecutionTime.TotalSeconds, Pair.Value.Result.ProcessorTime.TotalSeconds);
+							Log.TraceInformation($"{Description} [ Wall Time {Pair.Value.Result.ExecutionTime.TotalSeconds:0.00} s / CPU Time {Pair.Value.Result.ProcessorTime.TotalSeconds:0.00} s ]");
 						}
 						else
 						{
-							Log.TraceInformation("{0} [ Time {1:0.00} s ]", Description, Pair.Value.Result.ExecutionTime.TotalSeconds);
+							Log.TraceInformation($"{Description} [ Time {Pair.Value.Result.ExecutionTime.TotalSeconds:0.00} s ]");
 						}
 
 					}
@@ -232,8 +233,8 @@ namespace UnrealBuildTool
 				// Wait for all PrerequisiteActions to complete
 				ExecuteResults[] Results = await Task.WhenAll(Action.PrerequisiteActions.Select(x => Tasks[x]).ToArray());
 
-				// Cancel this task if any PrerequisiteActions fail (or were cancelled)
-				if (Results.Any(x => x.ExitCode != 0))
+				// Cancel tasks if any PrerequisiteActions fail, unless a PostBuildStep
+				if (Action.ActionType != ActionType.PostBuildStep && Results.Any(x => x.ExitCode != 0))
 				{
 					throw new OperationCanceledException();
 				}
@@ -316,7 +317,7 @@ namespace UnrealBuildTool
 				// Cancelled
 				if (ExitCode == int.MaxValue)
 				{
-					Log.TraceInformation("[{0}/{1}] {2} cancelled", CompletedActions, TotalActions, Description);
+					Log.TraceInformation($"[{CompletedActions}/{TotalActions}] {Description} cancelled");
 					return;
 				}
 
@@ -336,10 +337,17 @@ namespace UnrealBuildTool
 
 				if (bShowPerActionCompilationTimes)
 				{
-					CompilationTimes = $" (Wall: {ExecutionTime.TotalSeconds:0.00}s CPU: {ProcessorTime.TotalSeconds:0.00}s)";
+					if (ProcessorTime.Ticks > 0)
+					{
+						CompilationTimes = $" (Wall: {ExecutionTime.TotalSeconds:0.00}s CPU: {ProcessorTime.TotalSeconds:0.00}s)";
+					}
+					else
+					{
+						CompilationTimes = $" (Wall: {ExecutionTime.TotalSeconds:0.00}s)";
+					}
 				}
 
-				Log.TraceInformation("[{0}/{1}]{2}{3} {4}", CompletedActions, TotalActions, TargetDetails, CompilationTimes, Description);
+				Log.TraceInformation($"[{CompletedActions}/{TotalActions}]{TargetDetails}{CompilationTimes} {Description}");
 				foreach (string Line in LogLines.Skip(Action.bShouldOutputStatusDescription ? 0 : 1))
 				{
 					Log.TraceInformation(Line);
@@ -350,14 +358,14 @@ namespace UnrealBuildTool
 					// BEGIN TEMPORARY TO CATCH PVS-STUDIO ISSUES
 					if (LogLines.Count == 0)
 					{
-						Log.TraceInformation("[{0}/{1}]{2} {3} - Error but no output", NumCompletedActions, TotalActions, TargetDetails, Description);
-						Log.TraceInformation("[{0}/{1}]{2} {3} - {4} {5} {6} {7}", NumCompletedActions, TotalActions, TargetDetails, Description, ExitCode,
-							Action.WorkingDirectory, Action.CommandPath, Action.CommandArguments);
+						Log.TraceError($"{TargetDetails} {Description}: Exited with error code {ExitCode}");
+						Log.TraceInformation($"{TargetDetails} {Description}: WorkingDirectory {Action.WorkingDirectory}");
+						Log.TraceInformation($"{TargetDetails} {Description}: {Action.CommandPath} {Action.CommandArguments}");
 					}
 					// END TEMPORARY
 
-					// Cancel all other pending tasks
-					if (bStopCompilationAfterErrors)
+						// Cancel all other pending tasks
+						if (bStopCompilationAfterErrors)
 					{
 						CancellationTokenSource.Cancel();
 					}
diff --git a/Engine/Source/Programs/UnrealBuildTool/Platform/Windows/PVSToolChain.cs b/Engine/Source/Programs/UnrealBuildTool/Platform/Windows/PVSToolChain.cs
--- a/Engine/Source/Programs/UnrealBuildTool/Platform/Windows/PVSToolChain.cs
+++ b/Engine/Source/Programs/UnrealBuildTool/Platform/Windows/PVSToolChain.cs
@@ -585,16 +585,15 @@ namespace UnrealBuildTool
 				AnalyzeAction.CommandDescription = "Analyzing";
 				AnalyzeAction.StatusDescription = BaseFileName;
 				AnalyzeAction.WorkingDirectory = UnrealBuildTool.EngineSourceDirectory;
-				AnalyzeAction.CommandPath = BuildHostPlatform.Current.Shell;
+				AnalyzeAction.CommandPath = AnalyzerFile;
 
 				StringBuilder Arguments = new StringBuilder();
-				Arguments.Append($"/C \"\"{AnalyzerFile}\" --source-file \"{SourceFileItem.AbsolutePath}\" --output-file \"{OutputFileLocation}\" --cfg \"{ConfigFileItem.AbsolutePath}\" --i-file=\"{PreprocessedFileItem.AbsolutePath}\" --analysis-mode {(uint)Settings.ModeFlags}");
+				Arguments.Append($"--source-file \"{SourceFileItem.AbsolutePath}\" --output-file \"{OutputFileLocation}\" --cfg \"{ConfigFileItem.AbsolutePath}\" --i-file=\"{PreprocessedFileItem.AbsolutePath}\" --analysis-mode {(uint)Settings.ModeFlags}");
 				if (LicenseFile != null)
 				{
 					Arguments.Append($" --lic-file \"{LicenseFile}\"");
 					AnalyzeAction.PrerequisiteItems.Add(FileItem.GetItemByFileReference(LicenseFile));
 				}
-				Arguments.Append($" && echo. >>\"{OutputFileLocation}\"\"");
 				AnalyzeAction.CommandArguments = Arguments.ToString();
 
 				AnalyzeAction.PrerequisiteItems.Add(ConfigFileItem);
@@ -634,6 +633,8 @@ namespace UnrealBuildTool
 			FileItem InputFileListItem = Makefile.CreateIntermediateTextFile(OutputFile.ChangeExtension(".input"), InputFiles.Select(x => x.FullName));
 
 			Action AnalyzeAction = Makefile.CreateAction(ActionType.Compile);
+			AnalyzeAction.ActionType = ActionType.PostBuildStep;
+			AnalyzeAction.CommandDescription = "Process PVS-Studio Results";
 			AnalyzeAction.CommandPath = UnrealBuildTool.GetUBTPath();
 			AnalyzeAction.CommandArguments = String.Format("-Mode=PVSGather -Input=\"{0}\" -Output=\"{1}\" ", InputFileListItem.Location, OutputFile);
 			AnalyzeAction.WorkingDirectory = UnrealBuildTool.EngineSourceDirectory;
