UnrealBuildTool: Add ban list of known bad MSVC toolchains

#rnx

diff --git a/Engine/Source/Programs/UnrealBuildTool/Platform/Windows/UEBuildWindows.cs b/Engine/Source/Programs/UnrealBuildTool/Platform/Windows/UEBuildWindows.cs
--- a/Engine/Source/Programs/UnrealBuildTool/Platform/Windows/UEBuildWindows.cs
+++ b/Engine/Source/Programs/UnrealBuildTool/Platform/Windows/UEBuildWindows.cs
@@ -847,6 +847,15 @@ namespace UnrealBuildTool
 			VersionNumberRange.Parse("14.29.30133", "14.29.30136"), // VS2019 16.11.5
 		};
 
+		/// <summary>
+		/// Tested compiler toolchanges that should not be allowed.
+		/// </summary>
+		static readonly VersionNumberRange[] BannedVisualCppVersions = new VersionNumberRange[]
+		{
+			VersionNumberRange.Parse("14.31.31103", "14.31.31104"), // VS2022 17.1.0
+			VersionNumberRange.Parse("14.32.31113", "14.32.31114"), // VS2022 17.2.0
+		};
+
 		static readonly VersionNumber MinimumVisualCppVersion = new VersionNumber(14, 29, 30133);
 
 		/// <summary>
@@ -1561,6 +1570,12 @@ namespace UnrealBuildTool
 				Error = $"UnrealBuildTool requires at minimum the MSVC {MinimumVisualCppVersion} toolchain. Please install a later toolchain from the Visual Studio installer.";
 			}
 
+			VersionNumberRange? Banned = BannedVisualCppVersions.FirstOrDefault(x => x.Contains(Version));
+			if (Banned != null)
+			{
+				Error = $"UnrealBuildTool has banned the MSVC {Banned} toolchains due to compiler issues. Please install a different toolchain from the Visual Studio installer.";
+			}
+
 			Log.TraceLog("Found Visual Studio toolchain: {0} (Family={1}, FamilyRank={2}, Version={3}, Is64Bit={4}, Preview={5}, Architecture={6}, Error={7}, Redist={8})", ToolChainDir, Family, FamilyRank, Version, Is64Bit, bPreview, WindowsArchitecture.x64.ToString(), Error != null, RedistDir);
 			ToolChains.Add(new ToolChainInstallation(Family, FamilyRank, Version, Is64Bit, bPreview, WindowsArchitecture.x64, Error, ToolChainDir, RedistDir));
 
