UnrealBuildTool: Remove VS2022 17.1 toolchain ban as it is no longer needed

#rnx

diff --git a/Engine/Source/Programs/UnrealBuildTool/Platform/Windows/UEBuildWindows.cs b/Engine/Source/Programs/UnrealBuildTool/Platform/Windows/UEBuildWindows.cs
--- a/Engine/Source/Programs/UnrealBuildTool/Platform/Windows/UEBuildWindows.cs
+++ b/Engine/Source/Programs/UnrealBuildTool/Platform/Windows/UEBuildWindows.cs
@@ -852,8 +852,6 @@ namespace UnrealBuildTool
 		/// </summary>
 		static readonly VersionNumberRange[] BannedVisualCppVersions = new VersionNumberRange[]
 		{
-			VersionNumberRange.Parse("14.31.31103", "14.31.31104"), // VS2022 17.1.0
-			VersionNumberRange.Parse("14.32.31113", "14.32.31114"), // VS2022 17.2.0
 		};
 
 		static readonly VersionNumber MinimumVisualCppVersion = new VersionNumber(14, 29, 30133);
