UnrealBuildTool: Fixes for quoting of defines for ISPC command line

diff --git a/Engine/Source/Programs/UnrealBuildTool/ToolChain/UEToolChain.cs b/Engine/Source/Programs/UnrealBuildTool/ToolChain/UEToolChain.cs
--- a/Engine/Source/Programs/UnrealBuildTool/ToolChain/UEToolChain.cs
+++ b/Engine/Source/Programs/UnrealBuildTool/ToolChain/UEToolChain.cs
@@ -4,6 +4,7 @@ using System;
 using System.Collections.Generic;
 using System.IO;
 using System.Linq;
+using System.Text;
 using System.Text.RegularExpressions;
 using EpicGames.Core;
 using UnrealBuildBase;
@@ -407,6 +408,78 @@ namespace UnrealBuildTool
 			return Suffix;
 		}
 
+		private string EscapeDefinitionForISPC(string Definition)
+		{
+			// See: [REDACTED_PROFILE_URL]/ispc/blob/main/src/main.cpp#L260 for ispc's argument parsing code, which does the following (and does not support escaping):
+			// Argument      Parses as 
+			// "abc""def"    One agrument:  abcdef
+			// "'abc'"       One argument:  'abc'
+			// -D"X="Y Z""   Two arguments: -DX=Y and Z
+			// -D'X="Y Z"'   One argument:  -DX="Y Z"  (i.e. with quotes in value)
+			// -DX="Y Z"     One argument:  -DX=Y Z    (this is what we want on the command line)
+
+			// Assumes that quotes at the start and end of the value string mean that everything between them should be passed on unchanged.
+
+			int DoubleQuoteCount = Definition.Count(c => c == '"');
+			bool bHasSingleQuote = Definition.Contains('\'');
+			bool bHasSpace = Definition.Contains(' ');
+
+			string Escaped = Definition;
+
+			if (DoubleQuoteCount > 0 || bHasSingleQuote || bHasSpace)
+			{
+				int EqualsIndex = Definition.IndexOf('=');
+				string Name = Definition[0..EqualsIndex];
+				string Value = Definition[(EqualsIndex + 1)..];
+
+				string UnquotedValue = Value;
+
+				// remove one layer of quoting, if present
+				if (Value.StartsWith('"') && Value.EndsWith('"') && Value.Length != 1)
+				{
+					UnquotedValue = Value[1..^1];
+					DoubleQuoteCount -= 2;
+				}
+
+				if (DoubleQuoteCount == 0 && (bHasSingleQuote || bHasSpace))
+				{
+					Escaped = $"{Name}=\"{UnquotedValue}\"";
+				}
+				else if (!bHasSingleQuote && (bHasSpace || DoubleQuoteCount > 0))
+				{
+					// If there are no single quotes, we can use them to quote the value string
+					Escaped = $"{Name}='{UnquotedValue}'";
+				}
+				else
+				{
+					// Treat all special chars in the value string as needing explicit extra quoting. Thoroughly clumsy.
+					StringBuilder Requoted = new StringBuilder();
+					foreach (char c in UnquotedValue)
+					{
+						if (c == '"')
+						{
+							Requoted.Append("'\"'");
+						}
+						else if (c == '\'')
+						{
+							Requoted.Append("\"'\"");
+						}
+						else if (c == ' ')
+						{
+							Requoted.Append("\" \"");
+						}
+						else
+						{
+							Requoted.Append(c);
+						}
+					}
+					Escaped = $"{Name}={Requoted}";
+				}
+			}
+
+			return Escaped;
+		}
+
 		public override CPPOutput GenerateISPCHeaders(CppCompileEnvironment CompileEnvironment, List<FileItem> InputFiles, DirectoryReference OutputDir, IActionGraphBuilder Graph)
 		{
 			CPPOutput Result = new CPPOutput();
@@ -498,7 +571,7 @@ namespace UnrealBuildTool
 					// TODO: Causes ISPC compiler to generate a spurious warning about the universal character set
 					if (!Definition.Contains("\\\\U") && !Definition.Contains("\\\\u"))
 					{
-						Arguments.Add(String.Format("-D\"{0}\"", Definition));
+						Arguments.Add($"-D{EscapeDefinitionForISPC(Definition)}");
 					}
 				}
 
@@ -702,7 +775,7 @@ namespace UnrealBuildTool
 					// TODO: Causes ISPC compiler to generate a spurious warning about the universal character set
 					if (!Definition.Contains("\\\\U") && !Definition.Contains("\\\\u"))
 					{
-						Arguments.Add(String.Format("-D\"{0}\"", Definition));
+						Arguments.Add($"-D{EscapeDefinitionForISPC(Definition)}");
 					}
 				}
 
