Use a more appropriate & enduring link to upstream source

#trivial
#rnx

diff --git a/Engine/Source/Programs/UnrealBuildTool/ToolChain/UEToolChain.cs b/Engine/Source/Programs/UnrealBuildTool/ToolChain/UEToolChain.cs
--- a/Engine/Source/Programs/UnrealBuildTool/ToolChain/UEToolChain.cs
+++ b/Engine/Source/Programs/UnrealBuildTool/ToolChain/UEToolChain.cs
@@ -410,7 +410,7 @@ namespace UnrealBuildTool
 
 		private string EscapeDefinitionForISPC(string Definition)
 		{
-			// See: [REDACTED_PROFILE_URL]/ispc/blob/main/src/main.cpp#L260 for ispc's argument parsing code, which does the following (and does not support escaping):
+			// See: [REDACTED_PROFILE_URL]/ispc/blob/4ee767560cd752eaf464c124eb7ef1b0fd37f1df/src/main.cpp#L264 for ispc's argument parsing code, which does the following (and does not support escaping):
 			// Argument      Parses as 
 			// "abc""def"    One agrument:  abcdef
 			// "'abc'"       One argument:  'abc'
