Upgrade to nuget package Microsoft.Win32.Registry 5.0.0

#trivial
#rnx

diff --git a/Engine/Source/Programs/Shared/EpicGames.Build/EpicGames.Build.csproj b/Engine/Source/Programs/Shared/EpicGames.Build/EpicGames.Build.csproj
--- a/Engine/Source/Programs/Shared/EpicGames.Build/EpicGames.Build.csproj
+++ b/Engine/Source/Programs/Shared/EpicGames.Build/EpicGames.Build.csproj
@@ -32,7 +32,7 @@
   </ItemGroup>
   <ItemGroup>
     <PackageReference Include="Microsoft.CSharp" Version="4.7.0" />
-    <PackageReference Include="Microsoft.Win32.Registry" Version="4.7.0" />
+    <PackageReference Include="Microsoft.Win32.Registry" Version="5.0.0" />
     <PackageReference Include="System.Data.DataSetExtensions" Version="4.5.0" />
   </ItemGroup>
-</Project>
\ No newline at end of file
+</Project>
diff --git a/Engine/Source/Programs/Shared/EpicGames.Build/Platform/Mac/ApplePlatformSDK.cs b/Engine/Source/Programs/Shared/EpicGames.Build/Platform/Mac/ApplePlatformSDK.cs
--- a/Engine/Source/Programs/Shared/EpicGames.Build/Platform/Mac/ApplePlatformSDK.cs
+++ b/Engine/Source/Programs/Shared/EpicGames.Build/Platform/Mac/ApplePlatformSDK.cs
@@ -69,7 +69,7 @@ namespace UnrealBuildBase
 
 			string PackagesKeyName = "Software\\Classes\\Local Settings\\Software\\Microsoft\\Windows\\CurrentVersion\\AppModel\\PackageRepository\\Packages";
 
-			RegistryKey PackagesKey = Registry.LocalMachine.OpenSubKey(PackagesKeyName);
+			RegistryKey? PackagesKey = Registry.LocalMachine.OpenSubKey(PackagesKeyName);
 			if (PackagesKey != null)
 			{
 				string[] PackageSubKeyNames = PackagesKey.GetSubKeyNames();
@@ -80,10 +80,14 @@ namespace UnrealBuildBase
 					{
 						string FullPackageSubKeyName = PackagesKeyName + "\\" + PackageSubKeyName;
 
-						RegistryKey iTunesKey = Registry.LocalMachine.OpenSubKey(FullPackageSubKeyName);
+						RegistryKey? iTunesKey = Registry.LocalMachine.OpenSubKey(FullPackageSubKeyName);
 						if (iTunesKey != null)
 						{
-							InstallPath = (string)iTunesKey.GetValue("Path") + "\\AMDS32\\MobileDevice.dll";
+							object? Value = iTunesKey.GetValue("Path");
+							if (Value != null)
+							{
+								InstallPath = (string)Value + "\\AMDS32\\MobileDevice.dll";
+							}
 							break;
 						}
 					}
diff --git a/Engine/Source/Programs/Shared/EpicGames.Perforce/EpicGames.Perforce.csproj b/Engine/Source/Programs/Shared/EpicGames.Perforce/EpicGames.Perforce.csproj
--- a/Engine/Source/Programs/Shared/EpicGames.Perforce/EpicGames.Perforce.csproj
+++ b/Engine/Source/Programs/Shared/EpicGames.Perforce/EpicGames.Perforce.csproj
@@ -59,7 +59,7 @@
   </ItemGroup>
 
   <ItemGroup>
-    <PackageReference Include="Microsoft.Win32.Registry" Version="4.7.0" />
+    <PackageReference Include="Microsoft.Win32.Registry" Version="5.0.0" />
     <PackageReference Include="System.Linq.Async" Version="5.0.0" />
   </ItemGroup>
 
diff --git a/Engine/Source/Programs/UnrealBuildTool/Executors/Experimental/FASTBuild.cs b/Engine/Source/Programs/UnrealBuildTool/Executors/Experimental/FASTBuild.cs
--- a/Engine/Source/Programs/UnrealBuildTool/Executors/Experimental/FASTBuild.cs
+++ b/Engine/Source/Programs/UnrealBuildTool/Executors/Experimental/FASTBuild.cs
@@ -797,7 +797,7 @@ namespace UnrealBuildTool
 
 		public string GetRegistryValue(string keyName, string valueName, object defaultValue)
 		{
-			object returnValue = (string)Microsoft.Win32.Registry.GetValue("HKEY_LOCAL_MACHINE\\SOFTWARE\\" + keyName, valueName, defaultValue);
+			object? returnValue = Microsoft.Win32.Registry.GetValue("HKEY_LOCAL_MACHINE\\SOFTWARE\\" + keyName, valueName, defaultValue);
 			if (returnValue != null)
 				return returnValue.ToString()!;
 
@@ -805,7 +805,7 @@ namespace UnrealBuildTool
 			if (returnValue != null)
 				return returnValue.ToString()!;
 
-			returnValue = (string)Microsoft.Win32.Registry.GetValue("HKEY_LOCAL_MACHINE\\SOFTWARE\\Wow6432Node\\" + keyName, valueName, defaultValue);
+			returnValue = Microsoft.Win32.Registry.GetValue("HKEY_LOCAL_MACHINE\\SOFTWARE\\Wow6432Node\\" + keyName, valueName, defaultValue);
 			if (returnValue != null)
 				return returnValue.ToString()!;
 
diff --git a/Engine/Source/Programs/UnrealBuildTool/Executors/XGE.cs b/Engine/Source/Programs/UnrealBuildTool/Executors/XGE.cs
--- a/Engine/Source/Programs/UnrealBuildTool/Executors/XGE.cs
+++ b/Engine/Source/Programs/UnrealBuildTool/Executors/XGE.cs
@@ -137,7 +137,7 @@ namespace UnrealBuildTool
 			{
 				using(RegistryKey BaseKey = RegistryKey.OpenBaseKey(RegistryHive.LocalMachine, View))
 				{
-					using (RegistryKey Key = BaseKey.OpenSubKey("SOFTWARE\\Xoreax\\IncrediBuild\\Builder", false))
+					using (RegistryKey? Key = BaseKey.OpenSubKey("SOFTWARE\\Xoreax\\IncrediBuild\\Builder", false))
 					{
 						if(Key != null)
 						{
@@ -168,7 +168,7 @@ namespace UnrealBuildTool
 		{
 			using (RegistryKey BaseKey = RegistryKey.OpenBaseKey(Hive, View))
 			{
-				using (RegistryKey SubKey = BaseKey.OpenSubKey(KeyName))
+				using (RegistryKey? SubKey = BaseKey.OpenSubKey(KeyName))
 				{
 					if (SubKey != null)
 					{
@@ -586,7 +586,7 @@ namespace UnrealBuildTool
 			// @todo: There is a KB coming that will fix this. Once that KB is available, test if it is present. Stalls will not be a problem if it is.
 			//
 			// Stalls are possible. However there is a workaround in XGE build 1659 and newer that can avoid the issue.
-			string? XGEVersion = (BuildHostPlatform.Current.Platform == UnrealTargetPlatform.Win64) ? (string)Registry.GetValue(@"HKEY_LOCAL_MACHINE\SOFTWARE\Wow6432Node\Xoreax\IncrediBuild\Builder", "Version", null) : null;
+			string? XGEVersion = (BuildHostPlatform.Current.Platform == UnrealTargetPlatform.Win64) ? (string?)Registry.GetValue(@"HKEY_LOCAL_MACHINE\SOFTWARE\Wow6432Node\Xoreax\IncrediBuild\Builder", "Version", null) : null;
 			if (XGEVersion != null)
 			{
 				int XGEBuildNumber;
diff --git a/Engine/Source/Programs/UnrealBuildTool/Platform/Windows/MicrosoftPlatformSDK.cs b/Engine/Source/Programs/UnrealBuildTool/Platform/Windows/MicrosoftPlatformSDK.cs
--- a/Engine/Source/Programs/UnrealBuildTool/Platform/Windows/MicrosoftPlatformSDK.cs
+++ b/Engine/Source/Programs/UnrealBuildTool/Platform/Windows/MicrosoftPlatformSDK.cs
@@ -314,7 +314,7 @@ namespace UnrealBuildTool
 		{
 			foreach (KeyValuePair<RegistryKey, string> InstallRoot in InstallDirRoots)
 			{
-				using (RegistryKey Key = InstallRoot.Key.OpenSubKey(InstallRoot.Value + KeySuffix))
+				using (RegistryKey? Key = InstallRoot.Key.OpenSubKey(InstallRoot.Value + KeySuffix))
 				{
 					if (Key != null && TryReadDirRegistryKey(Key.Name, ValueName, out InstallDir))
 					{
@@ -337,7 +337,7 @@ namespace UnrealBuildTool
 			HashSet<string> AllSubKeys = new HashSet<string>(StringComparer.Ordinal);
 			foreach (KeyValuePair<RegistryKey, string> Root in InstallDirRoots)
 			{
-				using (RegistryKey Key = Root.Key.OpenSubKey(Root.Value + KeyName))
+				using (RegistryKey? Key = Root.Key.OpenSubKey(Root.Value + KeyName))
 				{
 					if (Key == null)
 					{
diff --git a/Engine/Source/Programs/UnrealBuildTool/ProjectFiles/ProjectFileGenerator.cs b/Engine/Source/Programs/UnrealBuildTool/ProjectFiles/ProjectFileGenerator.cs
--- a/Engine/Source/Programs/UnrealBuildTool/ProjectFiles/ProjectFileGenerator.cs
+++ b/Engine/Source/Programs/UnrealBuildTool/ProjectFiles/ProjectFileGenerator.cs
@@ -1827,7 +1827,7 @@ namespace UnrealBuildTool
 		private bool CheckRegistryKey(RegistryKey RootKey, string VisualStudioVersion)
 		{
 			bool bInstalled = false;
-			RegistryKey VSSubKey = RootKey.OpenSubKey("SOFTWARE\\Microsoft\\VisualStudio\\" + VisualStudioVersion + "\\Projects\\{E53F8FEA-EAE0-44A6-8774-FFD645390401}");
+			RegistryKey? VSSubKey = RootKey.OpenSubKey("SOFTWARE\\Microsoft\\VisualStudio\\" + VisualStudioVersion + "\\Projects\\{E53F8FEA-EAE0-44A6-8774-FFD645390401}");
 			if (VSSubKey != null)
 			{
 				bInstalled = true;
diff --git a/Engine/Source/Programs/UnrealBuildTool/UnrealBuildTool.csproj b/Engine/Source/Programs/UnrealBuildTool/UnrealBuildTool.csproj
--- a/Engine/Source/Programs/UnrealBuildTool/UnrealBuildTool.csproj
+++ b/Engine/Source/Programs/UnrealBuildTool/UnrealBuildTool.csproj
@@ -109,7 +109,7 @@
   </ItemGroup>
   <ItemGroup>
     <PackageReference Include="Microsoft.CodeAnalysis.CSharp" Version="3.7.0" />
-    <PackageReference Include="Microsoft.Win32.Registry" Version="4.7.0" />
+    <PackageReference Include="Microsoft.Win32.Registry" Version="5.0.0" />
     <PackageReference Include="OpenTracing" Version="0.12.1" />
     <PackageReference Include="System.Management" Version="4.7.0" />
     <PackageReference Include="System.Security.Cryptography.Csp" Version="4.3.0" />
