Module-level setup/teardown with the help of source file build order overrides.
- optional alter of source file build order: towards beginning or end
- test build: correctly handle inclusion of special dependencies such as Engine, Editor, CoreUObject and ApplicationCore
- existing tests: add setup and teardown tests that initialize test modules. This causes setup and teardown to be local to module instead of global, which was not the correct way to run tests across multiple modules.

diff --git a/Engine/Source/Developer/LowLevelTestsRunner/Private/TestMain.cpp b/Engine/Source/Developer/LowLevelTestsRunner/Private/TestMain.cpp
--- a/Engine/Source/Developer/LowLevelTestsRunner/Private/TestMain.cpp
+++ b/Engine/Source/Developer/LowLevelTestsRunner/Private/TestMain.cpp
@@ -5,3 +5,4 @@
 TCHAR GInternalProjectName[64] = TEXT("LowLevelTests");
 const TCHAR* GForeignEngineDir = nullptr;
 FEngineLoop GEngineLoop;
+bool GIsConsoleExecutable = true;
diff --git a/Engine/Source/Developer/LowLevelTestsRunner/Private/TestRunner.cpp b/Engine/Source/Developer/LowLevelTestsRunner/Private/TestRunner.cpp
--- a/Engine/Source/Developer/LowLevelTestsRunner/Private/TestRunner.cpp
+++ b/Engine/Source/Developer/LowLevelTestsRunner/Private/TestRunner.cpp
@@ -4,7 +4,6 @@
 
 #include "TestRunner.h"
 #include "TestHarness.h"
-#include "Modules/ModuleManager.h"
 #include "HAL/PlatformTLS.h"
 
 #include "CommonEngineInit.inl"
@@ -14,6 +13,19 @@
 #define USE_GLOBAL_ENGINE_SETUP
 #endif
 
+// Test run interceptor
+struct TestRunListener : public Catch::TestEventListenerBase {
+	using TestEventListenerBase::TestEventListenerBase; // inherit constructor
+public:
+	void testCaseStarting(Catch::TestCaseInfo  const& TestInfo) override {
+		if (bDebug)
+		{
+			std::cout << TestInfo.lineInfo.file << ":" << TestInfo.lineInfo.line << " with tags " << TestInfo.tagsAsString() << std::endl;
+		}
+	}
+};
+
+CATCH_REGISTER_LISTENER(TestRunListener);
 
 void GlobalSetup()
 {
@@ -58,18 +70,6 @@ int RunTests(int argc, const char* argv[])
 	std::this_thread::sleep_for(std::chrono::milliseconds(5000));
 #endif
 
-// Global initialization
-#ifdef USE_GLOBAL_ENGINE_SETUP
-	GlobalSetup();
-#endif
-
-	TArray<FName> ModuleNames;
-	FModuleManager::Get().FindModules(TEXT("*LowLevelTests"), ModuleNames);
-	for (FName ModuleName : ModuleNames)
-	{
-		FModuleManager::Get().LoadModule(ModuleName);
-	}
-
 	//Read command-line from file (if any). Some platforms do this earlier.
 #ifndef PLATFORM_SKIP_ADDITIONAL_ARGS
 	int ArgsOverrideNum = 0;
@@ -115,8 +115,17 @@ int RunTests(int argc, const char* argv[])
 		{
 			bMultithreaded = false;
 		}
+		if (std::strcmp(argv[i], "--debug") == 0)
+		{
+			bDebug = true;
+		}
 	}
 
+// Global initialization
+#ifdef USE_GLOBAL_ENGINE_SETUP
+	GlobalSetup();
+#endif
+
 	int SessionResult = 0;
 	{
 		TGuardValue<bool> CatchRunning(bCatchIsRunning, true);
@@ -124,11 +133,6 @@ int RunTests(int argc, const char* argv[])
 		CatchArgv.Reset();
 	}
 
-	for (FName ModuleName : ModuleNames)
-	{
-		FModuleManager::Get().UnloadModule(ModuleName);
-	}
-
 // Global cleanup
 #if defined(USE_GLOBAL_ENGINE_SETUP)
 	GlobalTeardown();
diff --git a/Engine/Source/Developer/LowLevelTestsRunner/Public/TestRunner.h b/Engine/Source/Developer/LowLevelTestsRunner/Public/TestRunner.h
--- a/Engine/Source/Developer/LowLevelTestsRunner/Public/TestRunner.h
+++ b/Engine/Source/Developer/LowLevelTestsRunner/Public/TestRunner.h
@@ -23,6 +23,7 @@
 static bool bCatchIsRunning = false;
 static bool bGAllowLogging = true;
 static bool bMultithreaded = true;
+static bool bDebug = false;
 
 #if __cplusplus
 	extern "C" const char* GetProcessExecutablePath();
diff --git a/Engine/Source/Programs/LowLevelTests/LowLevelTests.Build.cs b/Engine/Source/Programs/LowLevelTests/LowLevelTests.Build.cs
--- a/Engine/Source/Programs/LowLevelTests/LowLevelTests.Build.cs
+++ b/Engine/Source/Programs/LowLevelTests/LowLevelTests.Build.cs
@@ -42,11 +42,12 @@ public class LowLevelTests : ModuleRules
 			// Disable exception handling so that tests can assert for exceptions
 			bEnableObjCExceptions = false;
 			bEnableExceptions = false;
-
+		
 			PrivateDependencyModuleNames.AddRange(
 				new string[] {
 				"Core",
 				"Cbor",
+				"CoreUObject",
 				"Projects",
 				"LowLevelTestsRunner"
 				});
@@ -63,6 +64,17 @@ public class LowLevelTests : ModuleRules
 
 		if (this.GetType() != typeof(LowLevelTests))
 		{
+			string SetupFile = Path.Combine("Private", "setup.cpp");
+			string TeardownFile = Path.Combine("Private", "teardown.cpp");
+			if (File.Exists(Path.Combine(ModuleDirectory.ToString(), SetupFile)))
+			{
+				BuildOrderSettings.AddBuildOrderOverride(SetupFile, SourceFileBuildOrder.First);
+			}
+			if (File.Exists(Path.Combine(ModuleDirectory.ToString(), TeardownFile)))
+			{
+				BuildOrderSettings.AddBuildOrderOverride(TeardownFile, SourceFileBuildOrder.Last);
+			}
+
 			UpdateBuildGraphPropertiesFile();
 		}
 	}
diff --git a/Engine/Source/Programs/UnrealBuildTool/Configuration/ModuleRules.cs b/Engine/Source/Programs/UnrealBuildTool/Configuration/ModuleRules.cs
--- a/Engine/Source/Programs/UnrealBuildTool/Configuration/ModuleRules.cs
+++ b/Engine/Source/Programs/UnrealBuildTool/Configuration/ModuleRules.cs
@@ -206,6 +206,22 @@ namespace UnrealBuildTool
 			VisibileForDll,
 		}
 
+		/// <summary>
+		/// Alter build order of source files for specific cases where necessary
+		/// An example is test files that must be executed first for module level setup or last for module level teardown
+		/// </summary>
+		public enum SourceFileBuildOrder
+		{
+			/// <summary>
+			/// Moves the order of the module source file at the beginning of compilation
+			/// </summary>
+			First,
+			/// <summary>
+			/// Moves the order of the module source file at the end of compilation
+			/// </summary>
+			Last,
+		}
+
 		/// <summary>
 		/// Information about a file which is required by the target at runtime, and must be moved around with it.
 		/// </summary>
@@ -447,6 +463,51 @@ namespace UnrealBuildTool
 			}
 		}
 
+		/// <summary>
+		/// Specifies build order overrides for source files in this module
+		/// A source file can be moved towards the beginning or end of compilation
+		/// Example use case: test setup and test teardown files that need to compile first and last respectively in a module
+		/// </summary>
+		public class SourceFilesBuildOrderSettings
+		{
+			private DirectoryReference ModuleDirectory;
+			private Dictionary<FileItem, SourceFileBuildOrder> BuildOrderOverridesPrivate;
+
+			/// <summary>
+			/// Constructs <see cref="SourceFilesBuildOrderSettings"/> given module directory.
+			/// </summary>
+			/// <param name="InModuleDirectory">Module source directory</param>
+			public SourceFilesBuildOrderSettings(DirectoryReference InModuleDirectory)
+			{
+				ModuleDirectory = InModuleDirectory;
+				BuildOrderOverridesPrivate = new Dictionary<FileItem, SourceFileBuildOrder>();
+			}
+
+			/// <summary>
+			/// Slightly alter the order of build of a module source file by placing it either at the beginning or end of compilation
+			/// </summary>
+			/// <param name="InRelativeSourceFile">Relative path of source file to module's directory</param>
+			/// <param name="InBuildOrderOverride">A <see cref="SourceFileBuildOrder"/> specifying order placement: 
+			/// <see cref="SourceFileBuildOrder.First"/> for beginning of compilation, <see cref="SourceFileBuildOrder.Last"/> for end</param>
+			public void AddBuildOrderOverride(string InRelativeSourceFile, SourceFileBuildOrder InBuildOrderOverride)
+			{
+				FileItem File = FileItem.GetItemByPath(Path.Combine(ModuleDirectory.FullName, InRelativeSourceFile));
+				if (File.Exists)
+				{
+					BuildOrderOverridesPrivate.Add(File, InBuildOrderOverride);
+				}
+				else
+				{
+					Log.TraceWarning("Cannot apply build order override, file doesn't exist: {0}", File.AbsolutePath);
+				}
+			}
+
+			/// <summary>
+			/// Get build order overrides map of module source file to <see cref="SourceFileBuildOrder"/>.
+			/// </summary>
+			public IReadOnlyDictionary<FileItem, SourceFileBuildOrder> Overrides => BuildOrderOverridesPrivate;
+		}
+
 		/// <summary>
 		/// Name of this module
 		/// </summary>
@@ -1066,6 +1127,22 @@ namespace UnrealBuildTool
 			}
 		}
 
+		/// <summary>
+		/// Optional compilation order override rules for module's source files.
+		/// </summary>
+		public SourceFilesBuildOrderSettings BuildOrderSettings
+		{
+			get
+			{
+				if (BuildOrderOverridesPrivate == null)
+				{
+					BuildOrderOverridesPrivate = new SourceFilesBuildOrderSettings(Directory);
+				}
+				return BuildOrderOverridesPrivate;
+			}
+		}
+		private SourceFilesBuildOrderSettings BuildOrderOverridesPrivate;
+
 #nullable disable
 		/// <summary>
 		/// Constructor. For backwards compatibility while the parameterless constructor is being phased out, initialization which would happen here is done by 
diff --git a/Engine/Source/Programs/UnrealBuildTool/Configuration/TestModuleRules.cs b/Engine/Source/Programs/UnrealBuildTool/Configuration/TestModuleRules.cs
--- a/Engine/Source/Programs/UnrealBuildTool/Configuration/TestModuleRules.cs
+++ b/Engine/Source/Programs/UnrealBuildTool/Configuration/TestModuleRules.cs
@@ -51,10 +51,23 @@ namespace UnrealBuildTool
 			bEnableObjCExceptions = false;
 			bEnableExceptions = false;
 
-			PrivateDependencyModuleNames.AddRange(
-				new string[] {
-					"LowLevelTestsRunner"
-				});
+			string SetupFile = Path.Combine("Private", "setup.cpp");
+			string TeardownFile = Path.Combine("Private", "teardown.cpp");
+			if (System.IO.File.Exists(Path.Combine(Directory.ToString(), SetupFile)))
+			{
+				BuildOrderSettings.AddBuildOrderOverride(SetupFile, SourceFileBuildOrder.First);
+			}
+			if (System.IO.File.Exists(Path.Combine(Directory.ToString(), TeardownFile)))
+			{
+				BuildOrderSettings.AddBuildOrderOverride(TeardownFile, SourceFileBuildOrder.Last);
+			}
+
+			PrivateDependencyModuleNames.AddRange(TestedModule.PrivateDependencyModuleNames);
+
+			if (!PrivateDependencyModuleNames.Contains("LowLevelTestsRunner"))
+			{
+				PrivateDependencyModuleNames.Add("LowLevelTestsRunner");
+			}
 
 			// Tests can refer to tested module's Public and Private paths
 			string ModulePublicDir = Path.Combine(TestedModule.ModuleDirectory, "Public");
diff --git a/Engine/Source/Programs/UnrealBuildTool/Configuration/UEBuildModuleCPP.cs b/Engine/Source/Programs/UnrealBuildTool/Configuration/UEBuildModuleCPP.cs
--- a/Engine/Source/Programs/UnrealBuildTool/Configuration/UEBuildModuleCPP.cs
+++ b/Engine/Source/Programs/UnrealBuildTool/Configuration/UEBuildModuleCPP.cs
@@ -2,16 +2,16 @@
 
 using System;
 using System.Collections.Generic;
-using System.Diagnostics;
 using System.IO;
 using System.Linq;
 using System.Text;
-using System.Threading.Tasks;
 using EpicGames.Core;
 using UnrealBuildBase;
 
 namespace UnrealBuildTool
 {
+	using SourceOrderOverrides = IReadOnlyDictionary<FileItem, ModuleRules.SourceFileBuildOrder>;
+
 	/// <summary>
 	/// A module that is compiled from C++ code.
 	/// </summary>
@@ -182,7 +182,6 @@ namespace UnrealBuildTool
             new KeyValuePair<string, string>("UnrealEd", "MeshPaint"),
         };
 
-
 		public UEBuildModuleCPP(ModuleRules Rules, DirectoryReference IntermediateDirectory, DirectoryReference? GeneratedCodeDirectory)
 			: base(Rules, IntermediateDirectory)
 		{
@@ -964,11 +963,16 @@ namespace UnrealBuildTool
 
 			if (NormalFiles.Count > 0)
 			{
+				// Optionally alter order of compilation
+				TryAlterCompilationOrder(Rules.BuildOrderSettings.Overrides, NormalFiles);
 				OutputFiles = ToolChain.CompileCPPFiles(CompileEnvironment, NormalFiles, IntermediateDirectory, Name, Graph);
 			}
 
 			if (AdaptiveFiles.Count > 0)
 			{
+				// Optionally alter order of compilation
+				TryAlterCompilationOrder(Rules.BuildOrderSettings.Overrides, AdaptiveFiles);
+
 				// Create the new compile environment. Always turn off PCH due to different compiler settings.
 				CppCompileEnvironment AdaptiveUnityEnvironment = new CppCompileEnvironment(ModuleCompileEnvironment);
 				if(Target.bAdaptiveUnityDisablesOptimizations)
@@ -1533,6 +1537,38 @@ namespace UnrealBuildTool
 			return InputFiles;
 		}
 
+		private void TryAlterCompilationOrder(SourceOrderOverrides OrderOverrides, List<FileItem> Files)
+		{
+			FileItem FileToAlter;
+			ModuleRules.SourceFileBuildOrder Order;
+			int FoundIndex;
+			if (OrderOverrides != null && OrderOverrides.Count > 0)
+			{
+				foreach (var FileOrder in OrderOverrides)
+				{
+					FileToAlter = FileOrder.Key;
+					Order = FileOrder.Value;
+					FoundIndex = Files.FindIndex((File) => File.Location == FileToAlter.Location);
+					if (Order == ModuleRules.SourceFileBuildOrder.First)
+					{
+						if (FoundIndex > 0)
+						{
+							Files.Insert(0, Files[FoundIndex]);
+							Files.RemoveAt(FoundIndex + 1);
+						}
+					}
+					else if (Order == ModuleRules.SourceFileBuildOrder.Last)
+					{
+						if (FoundIndex >= 0 && FoundIndex < Files.Count - 1)
+						{
+							Files.Add(Files[FoundIndex]);
+							Files.RemoveAt(FoundIndex);
+						}
+					}
+				}
+			}
+		}
+
 		/// <summary>
 		/// Finds all the source files that should be built for this module
 		/// </summary>
diff --git a/Engine/Source/Programs/UnrealBuildTool/Configuration/UEBuildTarget.cs b/Engine/Source/Programs/UnrealBuildTool/Configuration/UEBuildTarget.cs
--- a/Engine/Source/Programs/UnrealBuildTool/Configuration/UEBuildTarget.cs
+++ b/Engine/Source/Programs/UnrealBuildTool/Configuration/UEBuildTarget.cs
@@ -1107,6 +1107,11 @@ namespace UnrealBuildTool
 		/// </summary>
 		private bool bTestsRequireCoreUObject = false;
 
+		/// <summary>
+		/// Track low level tests runner module for configuration.
+		/// </summary>
+		private ModuleRules? LowLevelTestsRunnerModule;
+
 		/// <summary>
 		/// Used to keep track of all modules by name.
 		/// </summary>
@@ -4333,6 +4338,10 @@ namespace UnrealBuildTool
 				{
 					RulesObject.PrivateIncludePathModuleNames.Add("LowLevelTestsRunner");
 				}
+				else if (RulesObject.Name == "LowLevelTestsRunner")
+				{
+					LowLevelTestsRunnerModule = RulesObject;
+				}
 
 				// If one of these modules is in the dependency graph, we must enable their appropriate global definitions
 				if (RulesObject.Name == "Editor")
@@ -4351,6 +4360,26 @@ namespace UnrealBuildTool
 				{
 					bTestsRequireCoreUObject = true;
 				}
+
+				if (LowLevelTestsRunnerModule != null)
+				{
+					if (bTestsRequireEditor && !LowLevelTestsRunnerModule.PrivateDependencyModuleNames.Contains("Editor"))
+					{
+						LowLevelTestsRunnerModule.PrivateDependencyModuleNames.Add("Editor");
+					}
+					if (bTestsRequireEngine && !LowLevelTestsRunnerModule.PrivateDependencyModuleNames.Contains("Engine"))
+					{
+						LowLevelTestsRunnerModule.PrivateDependencyModuleNames.Add("Engine");
+					}
+					if (bTestsRequireApplicationCore && !LowLevelTestsRunnerModule.PrivateDependencyModuleNames.Contains("ApplicationCore"))
+					{
+						LowLevelTestsRunnerModule.PrivateDependencyModuleNames.Add("ApplicationCore");
+					}
+					if (bTestsRequireCoreUObject && !LowLevelTestsRunnerModule.PrivateDependencyModuleNames.Contains("CoreUObject"))
+					{
+						LowLevelTestsRunnerModule.PrivateDependencyModuleNames.Add("CoreUObject");
+					}
+				}
 			}
 		}
 
