UnrealBuildTool Preprocessor: Support I32 integer suffix

#rnx

diff --git a/Engine/Source/Programs/UnrealBuildTool/Preprocessor/PreprocessorExpression.cs b/Engine/Source/Programs/UnrealBuildTool/Preprocessor/PreprocessorExpression.cs
--- a/Engine/Source/Programs/UnrealBuildTool/Preprocessor/PreprocessorExpression.cs
+++ b/Engine/Source/Programs/UnrealBuildTool/Preprocessor/PreprocessorExpression.cs
@@ -418,6 +418,12 @@ namespace UnrealBuildTool
 					Index++;
 				}
 			}
+			else if (Index + 2 < Literal.Length && (Literal[Index] == 'i' || Literal[Index] == 'I') && Literal[Index + 1] == '3' && Literal[Index + 2] == '2')
+			{
+				// 'I32' (Microsoft extension)
+				Index += 3;
+				bAllowTrailingUnsigned = false;
+			}
 			else if(Index + 2 < Literal.Length && (Literal[Index] == 'i' || Literal[Index] == 'I') && Literal[Index + 1] == '6' && Literal[Index + 2] == '4')
 			{
 				// 'I64' (Microsoft extension)
