UnrealBuildTool Preprocessor: Support headers in Xcode framework paths

#rnx

diff --git a/Engine/Source/Programs/UnrealBuildTool/Preprocessor/Preprocessor.cs b/Engine/Source/Programs/UnrealBuildTool/Preprocessor/Preprocessor.cs
--- a/Engine/Source/Programs/UnrealBuildTool/Preprocessor/Preprocessor.cs
+++ b/Engine/Source/Programs/UnrealBuildTool/Preprocessor/Preprocessor.cs
@@ -64,6 +64,11 @@ namespace UnrealBuildTool
 		/// </summary>
 		List<DirectoryItem> IncludeDirectories = new List<DirectoryItem>();
 
+		/// <summary>
+        /// Framework paths to look in
+        /// </summary>
+		List<DirectoryItem> FrameworkDirectories = new List<DirectoryItem>();
+
 		/// <summary>
 		/// Set of all included files with the #pragma once directive
 		/// </summary>
@@ -209,6 +214,41 @@ namespace UnrealBuildTool
 			AddIncludePath(new DirectoryReference(DirectoryName));
 		}
 
+		/// <summary>
+		/// Adds a framework path to the preprocessor
+		/// </summary>
+		/// <param name="Directory">The framework path</param>
+		public void AddFrameworkPath(DirectoryItem Directory)
+		{
+			if (!FrameworkDirectories.Contains(Directory))
+			{
+				FrameworkDirectories.Add(Directory);
+			}
+		}
+
+		/// <summary>
+		/// Adds a framework path to the preprocessor
+		/// </summary>
+		/// <param name="Location">The framework path</param>
+		public void AddFrameworkPath(DirectoryReference Location)
+		{
+			DirectoryItem Directory = DirectoryItem.GetItemByDirectoryReference(Location);
+			if (!Directory.Exists)
+			{
+				throw new FileNotFoundException("Unable to find " + Location.FullName);
+			}
+			AddFrameworkPath(Directory);
+		}
+
+		/// <summary>
+		/// Adds a framework path to the preprocessor
+		/// </summary>
+		/// <param name="DirectoryName">The framework path</param>
+		public void AddFrameworkPath(string DirectoryName)
+		{
+			AddFrameworkPath(new DirectoryReference(DirectoryName));
+		}
+
 		/// <summary>
 		/// Try to resolve an quoted include against the list of include directories. Uses search order described by https://msdn.microsoft.com/en-us/library/36k2cdd4.aspx.
 		/// </summary>
@@ -282,6 +322,21 @@ namespace UnrealBuildTool
 				}
 			}
 
+			// Try to match the include path against any of the MacOS framework Header paths
+			if (Fragments.Length > 1)
+			{
+				foreach (DirectoryItem BaseDirectory in FrameworkDirectories)
+				{
+					if (BaseDirectory.TryGetDirectory($"{Fragments[0]}.framework", out DirectoryItem? FrameworkBaseDirectory) &&
+						FrameworkBaseDirectory.TryGetDirectory("Headers", out DirectoryItem? HeaderDirectory) &&
+						TryResolveRelativeIncludePath(HeaderDirectory, Fragments.Skip(1).ToArray(), out FileItem? ResolvedFile))
+					{
+						File = ResolvedFile;
+						return true;
+					}
+				}
+			}
+
 			// Failed to find the file
 			File = null;
 			return false;
diff --git a/Engine/Source/Programs/UnrealBuildTool/Preprocessor/PreprocessorContext.cs b/Engine/Source/Programs/UnrealBuildTool/Preprocessor/PreprocessorContext.cs
--- a/Engine/Source/Programs/UnrealBuildTool/Preprocessor/PreprocessorContext.cs
+++ b/Engine/Source/Programs/UnrealBuildTool/Preprocessor/PreprocessorContext.cs
@@ -82,7 +82,7 @@ namespace UnrealBuildTool
 		/// </summary>
 		/// <param name="SourceFile">The source file being parsed</param>
 		/// <param name="Outer">The outer context</param>
-		public PreprocessorFileContext(SourceFile SourceFile, PreprocessorContext Outer)
+		public PreprocessorFileContext(SourceFile SourceFile, PreprocessorContext? Outer)
 			: base(Outer)
 		{
 			this.SourceFile = SourceFile;
