UnrealBuildTool Preprocessor: Handle additional builtin identifiers

#rnx

diff --git a/Engine/Source/Programs/UnrealBuildTool/Preprocessor/Identifier.cs b/Engine/Source/Programs/UnrealBuildTool/Preprocessor/Identifier.cs
--- a/Engine/Source/Programs/UnrealBuildTool/Preprocessor/Identifier.cs
+++ b/Engine/Source/Programs/UnrealBuildTool/Preprocessor/Identifier.cs
@@ -104,8 +104,17 @@ namespace UnrealBuildTool
 		public static readonly Identifier Alignof = Identifier.FindOrAdd("alignof");
 		public static readonly Identifier __has_builtin = Identifier.FindOrAdd("__has_builtin");
 		public static readonly Identifier __has_feature = Identifier.FindOrAdd("__has_feature");
+		public static readonly Identifier __has_warning = Identifier.FindOrAdd("__has_warning");
 		public static readonly Identifier __building_module = Identifier.FindOrAdd("__building_module");
 		public static readonly Identifier __pragma = Identifier.FindOrAdd("__pragma");
+		public static readonly Identifier __builtin_return_address = Identifier.FindOrAdd("__builtin_return_address");
+		public static readonly Identifier __builtin_frame_address = Identifier.FindOrAdd("__builtin_frame_address");
+		public static readonly Identifier __has_cpp_attribute = Identifier.FindOrAdd("__has_cpp_attribute");
+		public static readonly Identifier __has_keyword = Identifier.FindOrAdd("__has_keyword");
+		public static readonly Identifier __has_extension = Identifier.FindOrAdd("__has_extension");
+		public static readonly Identifier __has_include = Identifier.FindOrAdd("__has_include");
+		public static readonly Identifier __has_include_next = Identifier.FindOrAdd("__has_include_next");
+		public static readonly Identifier __is_target_arch = Identifier.FindOrAdd("__is_target_arch");
 	}
 
 	/// <summary>
diff --git a/Engine/Source/Programs/UnrealBuildTool/Preprocessor/PreprocessorExpression.cs b/Engine/Source/Programs/UnrealBuildTool/Preprocessor/PreprocessorExpression.cs
--- a/Engine/Source/Programs/UnrealBuildTool/Preprocessor/PreprocessorExpression.cs
+++ b/Engine/Source/Programs/UnrealBuildTool/Preprocessor/PreprocessorExpression.cs
@@ -56,7 +56,7 @@ namespace UnrealBuildTool
 			long Result = EvaluateTernary(Context, Tokens, ref Idx);
 			if (Idx != Tokens.Count)
 			{
-				throw new PreprocessorException(Context, "Garbage after end of expression");
+				throw new PreprocessorException(Context, $"Garbage after end of expression: {Token.Format(Tokens)}");
 			}
 			return Result;
 		}
@@ -84,7 +84,7 @@ namespace UnrealBuildTool
 				// Check for the colon in the middle
 				if(Tokens[Idx].Type != TokenType.Colon)
 				{
-					throw new PreprocessorException(Context, "Expected colon for conditional operator");
+					throw new PreprocessorException(Context, $"Expected colon for conditional operator, found {Tokens[Idx].Text}");
 				}
 
 				// Read the right expression
@@ -193,7 +193,7 @@ namespace UnrealBuildTool
 						Lhs %= Rhs;
 						break;
 					default:
-						throw new NotImplementedException(String.Format("Binary operator '{0}' has not been implemented", Operator.Type));
+						throw new NotImplementedException($"Binary operator '{Operator.Type}' has not been implemented");
 				}
 			}
 
@@ -232,17 +232,17 @@ namespace UnrealBuildTool
 					Identifier Text = Tokens[Idx].Identifier!;
 					if(Text == Identifiers.Sizeof)
 					{
-						throw new NotImplementedException();
+						throw new NotImplementedException(Text.ToString());
 					}
 					else if(Text == Identifiers.Alignof)
 					{
-						throw new NotImplementedException();
+						throw new NotImplementedException(Text.ToString());
 					}
 					else if (Text == Identifiers.__has_builtin)
 					{
 						if (Tokens[Idx + 1].Type != TokenType.LeftParen || Tokens[Idx + 3].Type != TokenType.RightParen)
 						{
-							throw new NotImplementedException();
+							throw new NotImplementedException(Text.ToString());
 						}
 						Idx += 4;
 						return 0;
@@ -251,7 +251,16 @@ namespace UnrealBuildTool
 					{
 						if(Tokens[Idx + 1].Type != TokenType.LeftParen || Tokens[Idx + 3].Type != TokenType.RightParen)
 						{
-							throw new NotImplementedException();
+							throw new NotImplementedException(Text.ToString());
+						}
+						Idx += 4;
+						return 0;
+					}
+					else if (Text == Identifiers.__has_warning)
+					{
+						if (Tokens[Idx + 1].Type != TokenType.LeftParen || Tokens[Idx + 3].Type != TokenType.RightParen)
+						{
+							throw new NotImplementedException(Text.ToString());
 						}
 						Idx += 4;
 						return 0;
@@ -260,11 +269,71 @@ namespace UnrealBuildTool
 					{
 						if(Tokens[Idx + 1].Type != TokenType.LeftParen || Tokens[Idx + 3].Type != TokenType.RightParen)
 						{
-							throw new NotImplementedException();
+							throw new NotImplementedException(Text.ToString());
+						}
+						Idx += 4;
+						return 0;
+					}
+					else if (Text == Identifiers.__pragma)
+					{
+						if (Tokens[Idx + 1].Type != TokenType.LeftParen || Tokens[Idx + 3].Type != TokenType.RightParen)
+						{
+							throw new NotImplementedException(Text.ToString());
 						}
 						Idx += 4;
 						return 0;
 					}
+					else if (Text == Identifiers.__builtin_return_address)
+					{
+						if (Tokens[Idx + 1].Type != TokenType.LeftParen || Tokens[Idx + 3].Type != TokenType.RightParen)
+						{
+							throw new NotImplementedException(Text.ToString());
+						}
+						Idx += 4;
+						return 0;
+					}
+					else if (Text == Identifiers.__builtin_frame_address)
+					{
+						if (Tokens[Idx + 1].Type != TokenType.LeftParen || Tokens[Idx + 3].Type != TokenType.RightParen)
+						{
+							throw new NotImplementedException(Text.ToString());
+						}
+						Idx += 4;
+						return 0;
+					}
+					else if (Text == Identifiers.__has_keyword)
+					{
+						if (Tokens[Idx + 1].Type != TokenType.LeftParen || Tokens[Idx + 3].Type != TokenType.RightParen)
+						{
+							throw new NotImplementedException(Text.ToString());
+						}
+						Idx += 4;
+						return 0;
+					}
+					else if (Text == Identifiers.__has_extension)
+					{
+						if (Tokens[Idx + 1].Type != TokenType.LeftParen || Tokens[Idx + 3].Type != TokenType.RightParen)
+						{
+							throw new NotImplementedException(Text.ToString());
+						}
+						Idx += 4;
+						return 0;
+					}
+					else if (Text == Identifiers.__is_target_arch)
+					{
+						if (Tokens[Idx + 1].Type != TokenType.LeftParen || Tokens[Idx + 3].Type != TokenType.RightParen)
+						{
+							throw new NotImplementedException(Text.ToString());
+						}
+						Idx += 4;
+						return 0;
+					}
+					else if (Text == Identifiers.__has_cpp_attribute || Text == Identifiers.__has_include || Text == Identifiers.__has_include_next)
+					{
+						Idx += Tokens.Count - Idx;
+						return 0;
+					}
+
 					else
 					{
 						return EvaluatePrimary(Context, Tokens, ref Idx);
