UnrealBuildTool Preprocessor: Support gathering list of includes when parsing files

#rnx

diff --git a/Engine/Source/Programs/UnrealBuildTool/Preprocessor/Preprocessor.cs b/Engine/Source/Programs/UnrealBuildTool/Preprocessor/Preprocessor.cs
--- a/Engine/Source/Programs/UnrealBuildTool/Preprocessor/Preprocessor.cs
+++ b/Engine/Source/Programs/UnrealBuildTool/Preprocessor/Preprocessor.cs
@@ -74,6 +74,11 @@ namespace UnrealBuildTool
 		/// </summary>
 		HashSet<FileItem> PragmaOnceFiles = new HashSet<FileItem>();
 
+		/// <summary>
+		/// Set of any files that has been processed
+		/// </summary>
+		HashSet<FileItem> ProcessedFiles = new HashSet<FileItem>();
+
 		/// <summary>
 		/// The current state of the preprocessor
 		/// </summary>
@@ -94,6 +99,15 @@ namespace UnrealBuildTool
 		/// </summary>
 		int Counter;
 
+		/// <summary>
+		/// List of files included by the preprocessor
+		/// </summary>
+		/// <returns>Enumerable of processed files</returns>
+		public IEnumerable<FileItem> GetProcessedFiles()
+		{
+			return ProcessedFiles.AsEnumerable();
+		}
+
 		/// <summary>
 		/// Default constructor
 		/// </summary>
@@ -371,7 +385,8 @@ namespace UnrealBuildTool
 		/// <param name="OuterContext">Outer context information, for error messages</param>
 		/// <param name="SourceFileCache">Cache for source files</param>
 		/// <param name="bShowIncludes">Show all the included files, in order</param>
-		public void ParseFile(FileItem File, List<SourceFileFragment> Fragments, PreprocessorContext OuterContext, SourceFileMetadataCache SourceFileCache, bool bShowIncludes)
+		/// <param name="bIgnoreMissingIncludes">Suppress exceptions if an include path can not be resolved</param>
+		public void ParseFile(FileItem File, List<SourceFileFragment> Fragments, PreprocessorContext? OuterContext, SourceFileMetadataCache SourceFileCache, bool bShowIncludes = false, bool bIgnoreMissingIncludes = false)
 		{
 			// If the file has already been included and had a #pragma once directive, don't include it again
 			if(PragmaOnceFiles.Contains(File))
@@ -379,6 +394,11 @@ namespace UnrealBuildTool
 				return;
 			}
 
+			if (!ProcessedFiles.Contains(File))
+			{
+				ProcessedFiles.Add(File);
+			}
+
 			// Output a trace of the included files
 			if(bShowIncludes)
 			{
@@ -404,10 +424,13 @@ namespace UnrealBuildTool
 					if(State.IsCurrentBranchActive())
 					{
 						// Parse the directive
-						FileItem IncludedFile = ParseIncludeDirective(Markup, Context);
+						FileItem? IncludedFile = ParseIncludeDirective(Markup, Context, bIgnoreMissingIncludes);
 
 						// Parse the included file
-						ParseFile(IncludedFile, Fragments, Context, SourceFileCache, bShowIncludes);
+						if (IncludedFile != null)
+						{
+							ParseFile(IncludedFile, Fragments, Context, SourceFileCache, bShowIncludes, bIgnoreMissingIncludes);
+						}
 					}
 					Context.MarkupIdx++;
 				}
@@ -432,8 +455,9 @@ namespace UnrealBuildTool
 		/// </summary>
 		/// <param name="Markup">Markup for the include directive</param>
 		/// <param name="Context">Current preprocessor context</param>
+		/// <param name="bIgnoreMissingIncludes">Suppress exceptions if an include path can not be resolved</param>
 		/// <returns>Included file</returns>
-		FileItem ParseIncludeDirective(SourceFileMarkup Markup, PreprocessorContext Context)
+		FileItem? ParseIncludeDirective(SourceFileMarkup Markup, PreprocessorFileContext Context, bool bIgnoreMissingIncludes = false)
 		{
 			// Expand macros in the given tokens
 			List<Token> ExpandedTokens = new List<Token>();
@@ -464,7 +488,14 @@ namespace UnrealBuildTool
 			FileItem? IncludedFile;
 			if(!TryResolveIncludePath(Context, IncludePath, Type, out IncludedFile))
 			{
-				throw new PreprocessorException(Context, "Couldn't resolve include '{0}'", IncludePath);
+				if (bIgnoreMissingIncludes)
+				{
+					Log.TraceWarningOnce("Couldn't resolve include '{0}' ({1})", IncludePath, Context.SourceFile.Location);
+				}
+				else
+				{
+					throw new PreprocessorException(Context, "Couldn't resolve include '{0}' ({1})", IncludePath, Context.SourceFile.Location);
+				}
 			}
 			return IncludedFile;
 		}
