[Backout] - CL18683346
Original CL Desc
-----------------------------------------------------------------
UnrealBuildTool: Ensure name suffixes are added if the TargetName is replaced in UEBuildTarget (manually merge due to conflict)

#rnx

diff --git a/Engine/Source/Programs/UnrealBuildTool/Configuration/TargetRules.cs b/Engine/Source/Programs/UnrealBuildTool/Configuration/TargetRules.cs
--- a/Engine/Source/Programs/UnrealBuildTool/Configuration/TargetRules.cs
+++ b/Engine/Source/Programs/UnrealBuildTool/Configuration/TargetRules.cs
@@ -189,20 +189,11 @@ namespace UnrealBuildTool
 		{
 			get
 			{
-				if (!string.IsNullOrEmpty(NameOverride))
+				if (!String.IsNullOrEmpty(NameOverride))
 				{
-					if (NameSuffixes.Count > 0)
-					{
-						return $"{NameOverride}-{NameSuffix}";
-					}
 					return NameOverride;
 				}
 
-				if (NameSuffixes.Count > 0)
-				{
-					return $"{DefaultName}-{NameSuffix}";
-				}
-
 				return DefaultName;
 			}
 			set
@@ -214,36 +205,14 @@ namespace UnrealBuildTool
 		/// <summary>
 		/// If the Name of this target has been overriden
 		/// </summary>
-		public bool IsNameOverriden() { return !string.IsNullOrEmpty(NameOverride) || NameSuffixes.Count > 0; }
-
-		/// <summary>
-		/// Add an optional suffix to append to Name
-		/// </summary>
-		/// <param name="Suffix">The string to append</param>
-		public void AddNameSuffix(string Suffix) { NameSuffixes.Add(Suffix); }
-
-		/// <summary>
-		/// Optional suffix to append to Name
-		/// </summary>
-		public string? NameSuffix
-		{
-			get
-			{
-				return NameSuffixes.Count > 0 ? string.Join("-", NameSuffixes) : null;
-			}
-		}
-
-		private HashSet<string> NameSuffixes = new HashSet<string>();
+		public bool IsNameOverriden() { return !String.IsNullOrEmpty(NameOverride); }
 
 		/// <summary>
 		/// Override the name used for this target
 		/// </summary>
 		[CommandLine("-TargetNameOverride=")]
-		private string? NameOverride { get; set; }
+		private string? NameOverride;
 
-		/// <summary>
-		/// The default name for this target, 
-		/// </summary>
 		private readonly string DefaultName;
 
 		private TestTargetRules? TestTargetRules;
@@ -379,18 +348,6 @@ namespace UnrealBuildTool
 		[CommandLine("-AllModules")]
 		public bool bBuildAllModules = false;
 
-		/// <summary>
-		/// Decides whether to compile with all the "Tests" folders from all dependent modules.
-		/// </summary>
-		public bool bIncludeAllTests
-		{
-			get { return bIncludeAllTestsOverride ?? false; }
-		}
-		/// <summary>
-		/// Set this override to true in derived target classes to compile with all the "Tests" folders from all dependent modules.
-		/// </summary>
-		protected bool? bIncludeAllTestsOverride;
-
 		/// <summary>
 		/// Additional plugins that are built for this target type but not enabled.
 		/// </summary>
@@ -632,20 +589,6 @@ namespace UnrealBuildTool
 		/// </summary>
 		public bool bForceBuildShaderFormats = false;
 
-		/// <summary>
-		/// Override for including extra shader formats
-		/// </summary>
-		public bool? bNeedsExtraShaderFormatsOverride;
-
-		/// <summary>
-		/// Whether we should include any extra shader formats. By default this is only enabled for Program and Editor targets.
-		/// </summary>
-		public bool bNeedsExtraShaderFormats
-		{
-			set { bNeedsExtraShaderFormatsOverride = value; }
-			get { return bNeedsExtraShaderFormatsOverride ?? (bForceBuildShaderFormats || bBuildDeveloperTools) && (Type == TargetType.Editor || Type == TargetType.Program); }
-		}
-
 		/// <summary>
 		/// Whether we should compile SQLite using the custom "Unreal" platform (true), or using the native platform (false).
 		/// </summary>
@@ -766,19 +709,6 @@ namespace UnrealBuildTool
 		}
 		private bool? bWithServerCodeOverride;
 
-		/// <summary>
-		/// Compile with FName storing the number part in the name table. 
-		/// Saves memory when most names are not numbered and those that are are referenced multiple times.
-		/// The game and engine must ensure they reuse numbered names similarly to name strings to avoid leaking memory.
-		/// </summary>
-		[RequiresUniqueBuildEnvironment]
-		public bool bFNameOutlineNumber
-		{
-			get { return bFNameOutlineNumberOverride ?? false;  }
-			set { bFNameOutlineNumberOverride = value;  }
-		}
-		private bool? bFNameOutlineNumberOverride;
-
 		/// <summary>
 		/// When enabled, Push Model Networking support will be compiled in.
 		/// This can help reduce CPU overhead of networking, at the cost of more memory.
@@ -1135,23 +1065,8 @@ namespace UnrealBuildTool
 		/// <summary>
 		/// Forces frame pointers to be retained this is usually required when you want reliable callstacks e.g. mallocframeprofiler
 		/// </summary>
-		public bool bRetainFramePointers
-		{
-			get 
-			{
-				// Default to disabled on Linux to maintain legacy behavior
-				return bRetainFramePointersOverride ?? Platform.IsInGroup(UnrealPlatformGroup.Linux) == false;
-			}
-			set { bRetainFramePointersOverride = value; }
-		}
-
-		/// <summary>
-		/// Forces frame pointers to be retained this is usually required when you want reliable callstacks e.g. mallocframeprofiler
-		/// </summary>
-		[XmlConfigFile(Category = "BuildConfiguration", Name="bRetainFramePointers")]
-		[CommandLine("-RetainFramePointers", Value="true")]
-		[CommandLine("-NoRetainFramePointers", Value="false")]
-		public bool? bRetainFramePointersOverride = null;
+		[XmlConfigFile(Category = "BuildConfiguration")]
+		public bool bRetainFramePointers = true;
 
 		/// <summary>
 		/// New Monolithic Graphics drivers have optional "fast calls" replacing various D3d functions
@@ -2076,11 +1991,6 @@ namespace UnrealBuildTool
 			get { return Inner.Name; }
 		}
 
-		public string? NameSuffix
-		{
-			get { return Inner.NameSuffix; }
-		}
-
 		public TestTargetRules TestTarget { get { return Inner.CreateOrGetTestTarget(); } }
 
 		public TargetRules TestedTarget {
@@ -2189,11 +2099,6 @@ namespace UnrealBuildTool
 			get { return Inner.bBuildAllModules; }
 		}
 
-		public bool bIncludeAllTests
-		{
-			get { return Inner.bIncludeAllTests; }
-		}
-		
 		public IEnumerable<string> AdditionalPlugins
 		{
 			get { return Inner.AdditionalPlugins; }
@@ -2314,7 +2219,12 @@ namespace UnrealBuildTool
 
 		public bool bCompileISPC
 		{
-			get { return Inner.bCompileISPC; }
+			get { return Inner.bCompileISPC && GlobalDefinitions.Contains("UE_LARGE_WORLD_COORDINATES_DISABLED=1"); }   // LWC_TODO: Temporarily disable ISPC when LWC is turned on. To be removed when double support is added to ISPC.
+		}
+
+		public bool bLWCDisabled 
+		{
+			get { return GlobalDefinitions.Contains("UE_LARGE_WORLD_COORDINATES_DISABLED=1"); }
 		}
 		
 		public bool bCompilePython
@@ -2357,11 +2267,6 @@ namespace UnrealBuildTool
 			get { return Inner.bForceBuildShaderFormats; }
 		}
 
-		public bool bNeedsExtraShaderFormats
-		{
-			get { return Inner.bNeedsExtraShaderFormats; }
-		}
-
 		public bool bCompileCustomSQLitePlatform
 		{
 			get { return Inner.bCompileCustomSQLitePlatform; }
@@ -2437,11 +2342,6 @@ namespace UnrealBuildTool
 			get { return Inner.bWithServerCode; }
 		}
 
-		public bool bFNameOutlineNumber
-		{
-			get { return Inner.bFNameOutlineNumber;  }
-		}
-
 		public bool bWithPushModel
 		{
 			get { return Inner.bWithPushModel; }
diff --git a/Engine/Source/Programs/UnrealBuildTool/Configuration/UEBuildTarget.cs b/Engine/Source/Programs/UnrealBuildTool/Configuration/UEBuildTarget.cs
--- a/Engine/Source/Programs/UnrealBuildTool/Configuration/UEBuildTarget.cs
+++ b/Engine/Source/Programs/UnrealBuildTool/Configuration/UEBuildTarget.cs
@@ -10,7 +10,6 @@ using System.Globalization;
 using System.IO;
 using System.Xml;
 using System.Runtime.Serialization;
-using System.Text.Json.Serialization;
 using System.Text.RegularExpressions;
 using EpicGames.Core;
 using System.Reflection;
@@ -124,7 +123,7 @@ namespace UnrealBuildTool
 		/// <returns></returns>
 		public override bool Equals(object? B)
 		{
-			if (Object.ReferenceEquals(B, null))
+			if(Object.ReferenceEquals(B, null))
 			{
 				return false;
 			}
@@ -193,7 +192,7 @@ namespace UnrealBuildTool
 				return new UnrealTargetPlatform(Id);
 			}
 
-			throw new BuildException(string.Format("The platform name {0} is not a valid platform name. Valid names are ({1})", Name,
+			throw new BuildException(string.Format("The platform name {0} is not a valid platform name. Valid names are ({1})", Name, 
 				string.Join(",", GetUniqueStringRegistry().GetStringNames())));
 		}
 
@@ -309,7 +308,7 @@ namespace UnrealBuildTool
 		{
 			if (value.GetType() == typeof(string))
 			{
-				return UnrealTargetPlatform.Parse((string)value);
+				return UnrealTargetPlatform.Parse((string) value);
 			}
 			return base.ConvertFrom(context, culture, value);
 		}
@@ -317,8 +316,8 @@ namespace UnrealBuildTool
 		public override object ConvertTo(ITypeDescriptorContext context, CultureInfo culture, object value, Type destinationType)
 		{
 			if (destinationType == typeof(string))
-			{
-				UnrealTargetPlatform Platform = (UnrealTargetPlatform)value;
+			{ 
+				UnrealTargetPlatform Platform = (UnrealTargetPlatform) value;
 				return Platform.ToString();
 			}
 			return base.ConvertTo(context, culture, value, destinationType);
@@ -440,7 +439,7 @@ namespace UnrealBuildTool
 		/// <returns></returns>
 		public override bool Equals(object? B)
 		{
-			if (Object.ReferenceEquals(B, null))
+			if(Object.ReferenceEquals(B, null))
 			{
 				return false;
 			}
@@ -596,7 +595,6 @@ namespace UnrealBuildTool
 	/// <summary>
 	/// The type of configuration a target can be built for.  Roughly order by optimization level. 
 	/// </summary>
-	[JsonConverter(typeof(JsonStringEnumConverter))]
 	public enum UnrealTargetConfiguration
 	{
 		/// <summary>
@@ -710,8 +708,7 @@ namespace UnrealBuildTool
 			TargetRules RulesObject;
 			using (GlobalTracer.Instance.BuildSpan("RulesAssembly.CreateTargetRules()").StartActive())
 			{
-				if (!Descriptor.IsTestsTarget)
-				{
+				if (!Descriptor.IsTestsTarget) {
 					RulesObject = RulesAssembly.CreateTargetRules(Descriptor.Name, Descriptor.Platform, Descriptor.Configuration, Descriptor.Architecture, Descriptor.ProjectFile, Descriptor.AdditionalArguments);
 				}
 				else
@@ -727,7 +724,7 @@ namespace UnrealBuildTool
 			}
 
 			// Make sure this configuration is supports
-			if (!RulesObject.GetSupportedConfigurations().Contains(Descriptor.Configuration))
+			if(!RulesObject.GetSupportedConfigurations().Contains(Descriptor.Configuration))
 			{
 				throw new BuildException("{0} does not support the {1} configuration", Descriptor.Name, Descriptor.Configuration);
 			}
@@ -739,7 +736,7 @@ namespace UnrealBuildTool
 				{
 					throw new BuildException("Download support for building {0} {1} targets from the launcher.", Descriptor.Platform, RulesObject.Type);
 				}
-				else if (!InstalledPlatformInfo.IsValid(RulesObject.Type, null, null, EProjectType.Code, InstalledPlatformState.Supported))
+				else if(!InstalledPlatformInfo.IsValid(RulesObject.Type, null, null, EProjectType.Code, InstalledPlatformState.Supported))
 				{
 					throw new BuildException("{0} targets are not currently supported from this engine distribution.", RulesObject.Type);
 				}
@@ -763,7 +760,7 @@ namespace UnrealBuildTool
 			if (RulesObject.bPrecompile)
 			{
 				DirectoryReference DependencyListDir;
-				if (RulesObject.ProjectFile == null)
+				if(RulesObject.ProjectFile == null)
 				{
 					DependencyListDir = DirectoryReference.Combine(Unreal.EngineDirectory, "Intermediate", "DependencyLists", RulesObject.Name, RulesObject.Configuration.ToString(), RulesObject.Platform.ToString());
 				}
@@ -772,13 +769,13 @@ namespace UnrealBuildTool
 					DependencyListDir = DirectoryReference.Combine(RulesObject.ProjectFile.Directory, "Intermediate", "DependencyLists", RulesObject.Name, RulesObject.Configuration.ToString(), RulesObject.Platform.ToString());
 				}
 
-				if (!String.IsNullOrEmpty(Descriptor.Architecture))
+				if(!String.IsNullOrEmpty(Descriptor.Architecture))
 				{
 					DependencyListDir = new DirectoryReference(DependencyListDir.FullName + Descriptor.Architecture);
 				}
 
 				FileReference DependencyListFile;
-				if (RulesObject.bBuildAllModules)
+				if(RulesObject.bBuildAllModules)
 				{
 					DependencyListFile = FileReference.Combine(DependencyListDir, "DependencyList-AllModules.txt");
 				}
@@ -791,7 +788,7 @@ namespace UnrealBuildTool
 			}
 
 			// If we're compiling only specific files, we need to prevent unity builds from running
-			if (Descriptor.SpecificFilesToCompile.Count > 0)
+			if(Descriptor.SpecificFilesToCompile.Count > 0)
 			{
 				RulesObject.bUseUnityBuild = false;
 				RulesObject.bForceUnityBuild = false;
@@ -877,14 +874,14 @@ namespace UnrealBuildTool
 		static void ValidateSharedEnvironment(RulesAssembly RulesAssembly, string ThisTargetName, CommandLineArguments Arguments, TargetRules ThisRules)
 		{
 			// Allow disabling these checks
-			if (ThisRules.bOverrideBuildEnvironment)
+			if(ThisRules.bOverrideBuildEnvironment)
 			{
 				return;
 			}
 
 			// Get the name of the target with default settings
 			string BaseTargetName;
-			switch (ThisRules.Type)
+			switch(ThisRules.Type)
 			{
 				case TargetType.Game:
 					BaseTargetName = "UnrealGame";
@@ -908,27 +905,27 @@ namespace UnrealBuildTool
 			// Get all the configurable objects
 			object[] BaseObjects = BaseRules.GetConfigurableObjects().ToArray();
 			object[] ThisObjects = ThisRules.GetConfigurableObjects().ToArray();
-			if (BaseObjects.Length != ThisObjects.Length)
+			if(BaseObjects.Length != ThisObjects.Length)
 			{
 				throw new BuildException("Expected same number of configurable objects from base rules object.");
 			}
 
 			// Iterate through all fields with the [SharedBuildEnvironment] attribute
-			for (int Idx = 0; Idx < BaseObjects.Length; Idx++)
+			for(int Idx = 0; Idx < BaseObjects.Length; Idx++)
 			{
 				Type ObjectType = BaseObjects[Idx].GetType();
-				foreach (FieldInfo Field in ObjectType.GetFields())
+				foreach(FieldInfo Field in ObjectType.GetFields())
 				{
-					if (Field.GetCustomAttribute<RequiresUniqueBuildEnvironmentAttribute>() != null)
+					if(Field.GetCustomAttribute<RequiresUniqueBuildEnvironmentAttribute>() != null)
 					{
 						object? ThisValue = Field.GetValue(ThisObjects[Idx]);
 						object? BaseValue = Field.GetValue(BaseObjects[Idx]);
 						CheckValuesMatch(ThisRules.GetType(), ThisTargetName, BaseTargetName, Field.Name, Field.FieldType, ThisValue, BaseValue);
 					}
 				}
-				foreach (PropertyInfo Property in ObjectType.GetProperties())
+				foreach(PropertyInfo Property in ObjectType.GetProperties())
 				{
-					if (Property.GetCustomAttribute<RequiresUniqueBuildEnvironmentAttribute>() != null)
+					if(Property.GetCustomAttribute<RequiresUniqueBuildEnvironmentAttribute>() != null)
 					{
 						object? ThisValue = Property.GetValue(ThisObjects[Idx]);
 						object? BaseValue = Property.GetValue(BaseObjects[Idx]);
@@ -952,11 +949,11 @@ namespace UnrealBuildTool
 		{
 			// Check if the fields match, treating lists of strings (eg. definitions) differently to value types.
 			bool bFieldsMatch;
-			if (ThisValue == null || BaseValue == null)
+			if(ThisValue == null || BaseValue == null)
 			{
 				bFieldsMatch = (ThisValue == BaseValue);
 			}
-			else if (typeof(IEnumerable<string>).IsAssignableFrom(ValueType))
+			else if(typeof(IEnumerable<string>).IsAssignableFrom(ValueType))
 			{
 				bFieldsMatch = Enumerable.SequenceEqual((IEnumerable<string>)ThisValue, (IEnumerable<string>)BaseValue);
 			}
@@ -966,7 +963,7 @@ namespace UnrealBuildTool
 			}
 
 			// Throw an exception if they don't match
-			if (!bFieldsMatch)
+			if(!bFieldsMatch)
 			{
 				throw new BuildException("{0} modifies the value of {1}. This is not allowed, as {0} has build products in common with {2}.\nRemove the modified setting, change {0} to use a unique build environment by setting 'BuildEnvironment = TargetBuildEnvironment.Unique;' in the {3} constructor, or set bOverrideBuildEnvironment = true to force this setting on.", ThisTargetName, FieldName, BaseTargetName, RulesType.Name);
 			}
@@ -1087,31 +1084,6 @@ namespace UnrealBuildTool
 		/// </summary>
 		protected bool bCompileMonolithic = false;
 
-		/// <summary>
-		/// Keeps track if low level tests executable must build with the Editor.
-		/// </summary>
-		private bool bTestsRequireEditor = false;
-
-		/// <summary>
-		/// Keeps track if low level tests executable must build with the Engine.
-		/// </summary>
-		private bool bTestsRequireEngine = false;
-
-		/// <summary>
-		/// Keeps track if low level tests executable must build with the ApplicationCore.
-		/// </summary>
-		private bool bTestsRequireApplicationCore = false;
-
-		/// <summary>
-		/// Keeps track if low level tests executable must build with the CoreUObject.
-		/// </summary>
-		private bool bTestsRequireCoreUObject = false;
-
-		/// <summary>
-		/// Track low level tests runner module for configuration.
-		/// </summary>
-		private ModuleRules? LowLevelTestsRunnerModule;
-
 		/// <summary>
 		/// Used to keep track of all modules by name.
 		/// </summary>
@@ -1142,7 +1114,7 @@ namespace UnrealBuildTool
 		/// <returns>true if it should, false if it shouldn't</returns>
 		public bool ShouldCompileMonolithic()
 		{
-			return bCompileMonolithic;  // @todo ubtmake: We need to make sure this function and similar things aren't called in assembler mode
+			return bCompileMonolithic;	// @todo ubtmake: We need to make sure this function and similar things aren't called in assembler mode
 		}
 
 		/// <summary>
@@ -1177,17 +1149,13 @@ namespace UnrealBuildTool
 			bUseSharedBuildEnvironment = (Rules.BuildEnvironment == TargetBuildEnvironment.Shared);
 			if (bUseSharedBuildEnvironment)
 			{
-				if (Rules.Type == TargetType.Program)
+				if(Rules.Type == TargetType.Program)
 				{
 					AppName = TargetName;
 				}
 				else
 				{
 					AppName = GetAppNameForTargetType(Rules.Type);
-					if (!string.IsNullOrEmpty(Rules.NameSuffix))
-					{
-						AppName = $"{AppName}-{Rules.NameSuffix}";
-					}
 				}
 			}
 
@@ -1248,7 +1216,7 @@ namespace UnrealBuildTool
 		/// <returns>The app name for this target type</returns>
 		public static string GetAppNameForTargetType(TargetType Type)
 		{
-			switch (Type)
+			switch(Type)
 			{
 				case TargetType.Game:
 					return "UnrealGame";
@@ -1274,14 +1242,14 @@ namespace UnrealBuildTool
 			HashSet<FileReference> Files = new HashSet<FileReference>();
 
 			// Find all the runtime dependency files in their original location
-			foreach (RuntimeDependency RuntimeDependency in RuntimeDependencies)
+			foreach(RuntimeDependency RuntimeDependency in RuntimeDependencies)
 			{
 				FileReference? SourceFile;
-				if (!RuntimeDependencyTargetFileToSourceFile.TryGetValue(RuntimeDependency.Path, out SourceFile))
+				if(!RuntimeDependencyTargetFileToSourceFile.TryGetValue(RuntimeDependency.Path, out SourceFile))
 				{
 					SourceFile = RuntimeDependency.Path;
 				}
-				if (RuntimeDependency.Type != StagedFileType.DebugNonUFS || FileReference.Exists(SourceFile))
+				if(RuntimeDependency.Type != StagedFileType.DebugNonUFS || FileReference.Exists(SourceFile))
 				{
 					Files.Add(SourceFile);
 				}
@@ -1291,7 +1259,7 @@ namespace UnrealBuildTool
 			HashSet<UEBuildModule> Modules = new HashSet<UEBuildModule>();
 			foreach (UEBuildBinary Binary in Binaries)
 			{
-				foreach (UEBuildModule Module in Binary.Modules)
+				foreach(UEBuildModule Module in Binary.Modules)
 				{
 					Modules.Add(Module);
 					Modules.UnionWith(Module.GetDependencies(true, true));
@@ -1303,7 +1271,7 @@ namespace UnrealBuildTool
 			foreach (UEBuildModule Module in Modules)
 			{
 				// Skip artificial modules
-				if (Module.RulesFile == null)
+				if(Module.RulesFile == null)
 				{
 					continue;
 				}
@@ -1367,12 +1335,12 @@ namespace UnrealBuildTool
 						{
 							ResolveLibraryName(LibraryPath, LibraryName, LibraryExtension, Files);
 						}
-					}
+							}
 
 					foreach (string LibraryName in Rules.PublicSystemLibraryPaths)
-					{
+							{
 						foreach (string LibraryPath in SystemLibraryPaths)
-						{
+								{
 							ResolveLibraryName(LibraryPath, LibraryName, LibraryExtension, Files);
 						}
 					}
@@ -1437,7 +1405,7 @@ namespace UnrealBuildTool
 			}
 
 			// Add all the dependency lists
-			foreach (FileReference DependencyListFileName in Rules.DependencyListFileNames)
+			foreach(FileReference DependencyListFileName in Rules.DependencyListFileNames)
 			{
 				Manifest.BuildProducts.Add(DependencyListFileName.FullName);
 			}
@@ -1453,7 +1421,7 @@ namespace UnrealBuildTool
 			}
 
 			// Remove anything that's not part of the plugin
-			if (ForeignPlugin != null)
+			if(ForeignPlugin != null)
 			{
 				DirectoryReference ForeignPluginDir = ForeignPlugin.Directory;
 				Manifest.BuildProducts.RemoveAll(x => !new FileReference(x).IsUnderDirectory(ForeignPluginDir));
@@ -1478,7 +1446,7 @@ namespace UnrealBuildTool
 				// Create the receipts for each folder
 				foreach (UEBuildBinary Binary in Binaries)
 				{
-					if (Binary.Type == UEBuildBinaryType.DynamicLinkLibrary)
+					if(Binary.Type == UEBuildBinaryType.DynamicLinkLibrary)
 					{
 						DirectoryReference DirectoryName = Binary.OutputFilePath.Directory;
 						bool bIsGameBinary = Binary.PrimaryModule.Rules.Context.bCanBuildDebugGame;
@@ -1519,9 +1487,9 @@ namespace UnrealBuildTool
 			// Create a unique identifier for this build which can be used to identify modules which are compatible. It's fine to share this between runs with the same makefile.
 			// By default we leave it blank when compiling a subset of modules (for hot reload, etc...), otherwise it won't match anything else. When writing to a directory
 			// that already contains a manifest, we'll reuse the build id that's already in there (see below).
-			if (String.IsNullOrEmpty(Version.BuildId))
+			if(String.IsNullOrEmpty(Version.BuildId))
 			{
-				if (Rules.bFormalBuild)
+				if(Rules.bFormalBuild)
 				{
 					// If this is a formal build, we can just the compatible changelist as the unique id.
 					Version.BuildId = String.Format("{0}", Version.EffectiveCompatibleChangelist);
@@ -1553,7 +1521,7 @@ namespace UnrealBuildTool
 							Receipt.Launch = Pair.Key;
 						}
 
-						if (Receipt.Launch != null && Receipt.LaunchCmd != null)
+						if ( Receipt.Launch != null && Receipt.LaunchCmd != null )
 						{
 							break;
 						}
@@ -1568,14 +1536,14 @@ namespace UnrealBuildTool
 			// Find all the build products and modules from this binary
 			foreach (KeyValuePair<FileReference, BuildProductType> BuildProductPair in BuildProducts)
 			{
-				if (BuildProductPair.Value != BuildProductType.BuildResource)
+				if(BuildProductPair.Value != BuildProductType.BuildResource)
 				{
 					Receipt.AddBuildProduct(BuildProductPair.Key, BuildProductPair.Value);
 				}
 			}
 
 			// Add the project file
-			if (ProjectFile != null)
+			if(ProjectFile != null)
 			{
 				Receipt.RuntimeDependencies.Add(ProjectFile, StagedFileType.UFS);
 			}
@@ -1584,7 +1552,7 @@ namespace UnrealBuildTool
 			List<UEBuildPlugin>? RuntimeDependencyPluginsList = Rules.bRuntimeDependenciesComeFromBuildPlugins ? BuildPlugins : EnabledPlugins;
 			foreach (UEBuildPlugin EnabledPlugin in RuntimeDependencyPluginsList!)
 			{
-				if (EnabledPlugin.bDescriptorNeededAtRuntime || EnabledPlugin.bDescriptorReferencedExplicitly)
+				if(EnabledPlugin.bDescriptorNeededAtRuntime || EnabledPlugin.bDescriptorReferencedExplicitly)
 				{
 					Receipt.RuntimeDependencies.Add(EnabledPlugin.File, StagedFileType.UFS);
 
@@ -1602,22 +1570,22 @@ namespace UnrealBuildTool
 							}
 						}
 
-						foreach (FileReference ChildFile in EnabledPlugin.ChildFiles)
+					foreach (FileReference ChildFile in EnabledPlugin.ChildFiles)
+					{
+						if (ValidFileNames.Contains(ChildFile.GetFileNameWithoutExtension(), StringComparer.InvariantCultureIgnoreCase))
 						{
-							if (ValidFileNames.Contains(ChildFile.GetFileNameWithoutExtension(), StringComparer.InvariantCultureIgnoreCase))
-							{
-								Receipt.RuntimeDependencies.Add(ChildFile, StagedFileType.UFS);
-							}
+							Receipt.RuntimeDependencies.Add(ChildFile, StagedFileType.UFS);
 						}
 					}
 				}
 			}
+			}
 
 			// Add all the other runtime dependencies
 			HashSet<FileReference> UniqueRuntimeDependencyFiles = new HashSet<FileReference>();
-			foreach (RuntimeDependency RuntimeDependency in RuntimeDependencies)
+			foreach(RuntimeDependency RuntimeDependency in RuntimeDependencies)
 			{
-				if (UniqueRuntimeDependencyFiles.Add(RuntimeDependency.Path))
+				if(UniqueRuntimeDependencyFiles.Add(RuntimeDependency.Path))
 				{
 					Receipt.RuntimeDependencies.Add(RuntimeDependency);
 				}
@@ -1730,7 +1698,7 @@ namespace UnrealBuildTool
 				List<UEBuildBinary> FilteredBinaries = new List<UEBuildBinary>();
 				foreach (UEBuildBinary DLLBinary in Binaries)
 				{
-					if (ModDirectories.Any(x => DLLBinary.OutputFilePath.IsUnderDirectory(x)))
+					if(ModDirectories.Any(x => DLLBinary.OutputFilePath.IsUnderDirectory(x)))
 					{
 						FilteredBinaries.Add(DLLBinary);
 					}
@@ -1745,9 +1713,9 @@ namespace UnrealBuildTool
 
 			// Build a mapping from module to its plugin
 			Dictionary<UEBuildModule, UEBuildPlugin> ModuleToPlugin = new Dictionary<UEBuildModule, UEBuildPlugin>();
-			foreach (UEBuildPlugin Plugin in BuildPlugins!)
+			foreach(UEBuildPlugin Plugin in BuildPlugins!)
 			{
-				foreach (UEBuildModule Module in Plugin.Modules)
+				foreach(UEBuildModule Module in Plugin.Modules)
 				{
 					if (!ModuleToPlugin.ContainsKey(Module))
 					{
@@ -1757,14 +1725,14 @@ namespace UnrealBuildTool
 			}
 
 			// Check there aren't any engine binaries with dependencies on game modules. This can happen when game-specific plugins override engine plugins.
-			foreach (UEBuildModule Module in Modules.Values)
+			foreach(UEBuildModule Module in Modules.Values)
 			{
-				if (Module.Binary != null)
+				if(Module.Binary != null)
 				{
 					HashSet<UEBuildModule> ReferencedModules = Module.GetDependencies(bWithIncludePathModules: true, bWithDynamicallyLoadedModules: true);
-					foreach (UEBuildModule ReferencedModule in ReferencedModules)
+					foreach(UEBuildModule ReferencedModule in ReferencedModules)
 					{
-						if (!Module.Rules.Context.Scope.Contains(ReferencedModule.Rules.Context.Scope) && !IsEngineToPluginReferenceAllowed(Module.Name, ReferencedModule.Name))
+						if(!Module.Rules.Context.Scope.Contains(ReferencedModule.Rules.Context.Scope) && !IsEngineToPluginReferenceAllowed(Module.Name, ReferencedModule.Name))
 						{
 							throw new BuildException("Module '{0}' ({1}) should not reference module '{2}' ({3}). Hierarchy is {4}.", Module.Name, Module.Rules.Context.Scope.Name, ReferencedModule.Name, ReferencedModule.Rules.Context.Scope.Name, ReferencedModule.Rules.Context.Scope.FormatHierarchy());
 						}
@@ -1773,15 +1741,15 @@ namespace UnrealBuildTool
 			}
 
 			// Check that each plugin declares its dependencies explicitly
-			foreach (UEBuildPlugin Plugin in BuildPlugins)
+			foreach(UEBuildPlugin Plugin in BuildPlugins)
 			{
-				foreach (UEBuildModule Module in Plugin.Modules)
+				foreach(UEBuildModule Module in Plugin.Modules)
 				{
 					HashSet<UEBuildModule> DependencyModules = Module.GetDependencies(bWithIncludePathModules: true, bWithDynamicallyLoadedModules: true);
-					foreach (UEBuildModule DependencyModule in DependencyModules)
+					foreach(UEBuildModule DependencyModule in DependencyModules)
 					{
 						UEBuildPlugin? DependencyPlugin;
-						if (ModuleToPlugin.TryGetValue(DependencyModule, out DependencyPlugin) && DependencyPlugin != Plugin
+						if(ModuleToPlugin.TryGetValue(DependencyModule, out DependencyPlugin) && DependencyPlugin != Plugin
 							&& !Rules.InternalPluginDependencies.Contains(DependencyPlugin.Name)
 							&& !Plugin.Dependencies!.Contains(DependencyPlugin))
 						{
@@ -1881,7 +1849,7 @@ namespace UnrealBuildTool
 			{
 				if (Platform == UnrealTargetPlatform.Win64)
 				{
-					if (!Rules.bFormalBuild)
+					if(!Rules.bFormalBuild)
 					{
 						FileReference DefaultResourceLocation = FileReference.Combine(Unreal.EngineDirectory, "Build", "Windows", "Resources", "Default.rc2");
 						if (!UnrealBuildTool.IsFileInstalled(DefaultResourceLocation))
@@ -1928,32 +1896,32 @@ namespace UnrealBuildTool
 			// Prepare all the runtime dependencies, copying them from their source folders if necessary
 			List<RuntimeDependency> RuntimeDependencies = new List<RuntimeDependency>();
 			Dictionary<FileReference, FileReference> RuntimeDependencyTargetFileToSourceFile = new Dictionary<FileReference, FileReference>();
-			foreach (UEBuildBinary Binary in Binaries)
+			foreach(UEBuildBinary Binary in Binaries)
 			{
 				Binary.PrepareRuntimeDependencies(RuntimeDependencies, RuntimeDependencyTargetFileToSourceFile, ExeDir);
 			}
 			TargetToolChain.PrepareRuntimeDependencies(RuntimeDependencies, RuntimeDependencyTargetFileToSourceFile, ExeDir);
 
-			foreach (KeyValuePair<FileReference, FileReference> Pair in RuntimeDependencyTargetFileToSourceFile)
+			foreach(KeyValuePair<FileReference, FileReference> Pair in RuntimeDependencyTargetFileToSourceFile)
 			{
-				if (!UnrealBuildTool.IsFileInstalled(Pair.Key))
+				if(!UnrealBuildTool.IsFileInstalled(Pair.Key))
 				{
 					Makefile.OutputItems.Add(Makefile.CreateCopyAction(Pair.Value, Pair.Key));
 				}
 			}
 
 			// If we're just precompiling a plugin, only include output items which are part of it
-			if (ForeignPlugin != null)
+			if(ForeignPlugin != null)
 			{
 				HashSet<FileItem> RetainOutputItems = new HashSet<FileItem>();
-				foreach (UEBuildPlugin Plugin in BuildPlugins)
+				foreach(UEBuildPlugin Plugin in BuildPlugins)
 				{
-					if (Plugin.File == ForeignPlugin)
+					if(Plugin.File == ForeignPlugin)
 					{
 						foreach (UEBuildModule Module in Plugin.Modules)
 						{
 							FileItem[]? ModuleOutputItems;
-							if (Makefile.ModuleNameToOutputItems.TryGetValue(Module.Name, out ModuleOutputItems))
+							if(Makefile.ModuleNameToOutputItems.TryGetValue(Module.Name, out ModuleOutputItems))
 							{
 								RetainOutputItems.UnionWith(ModuleOutputItems);
 							}
@@ -1977,7 +1945,7 @@ namespace UnrealBuildTool
 			BuildProducts.AddRange(RuntimeDependencyTargetFileToSourceFile.Select(x => new KeyValuePair<FileReference, BuildProductType>(x.Key, BuildProductType.RequiredResource)));
 
 			// Remove any installed build products that don't exist. They may be part of an optional install.
-			if (Unreal.IsEngineInstalled())
+			if(Unreal.IsEngineInstalled())
 			{
 				BuildProducts.RemoveAll(x => UnrealBuildTool.IsFileInstalled(x.Key) && !FileReference.Exists(x.Key));
 			}
@@ -2081,13 +2049,13 @@ namespace UnrealBuildTool
 
 				// Create actions to run the post build steps
 				FileReference[] PostBuildScripts = CreatePostBuildScripts();
-				foreach (FileReference PostBuildScript in PostBuildScripts)
+				foreach(FileReference PostBuildScript in PostBuildScripts)
 				{
 					FileReference OutputFile = new FileReference(PostBuildScript.FullName + ".ran");
 
 					Action PostBuildStepAction = Makefile.CreateAction(ActionType.PostBuildStep);
 					PostBuildStepAction.CommandPath = BuildHostPlatform.Current.Shell;
-					if (BuildHostPlatform.Current.ShellType == ShellType.Cmd)
+					if(BuildHostPlatform.Current.ShellType == ShellType.Cmd)
 					{
 						PostBuildStepAction.CommandArguments = String.Format("/C \"call \"{0}\" && type NUL >\"{1}\"\"", PostBuildScript, OutputFile);
 					}
@@ -2106,19 +2074,19 @@ namespace UnrealBuildTool
 			}
 
 			// Build a list of all the files required to build
-			foreach (FileReference DependencyListFileName in Rules.DependencyListFileNames)
+			foreach(FileReference DependencyListFileName in Rules.DependencyListFileNames)
 			{
 				WriteDependencyList(DependencyListFileName, RuntimeDependencies, RuntimeDependencyTargetFileToSourceFile);
 			}
 
 			// If we're only generating the manifest, return now
-			foreach (FileReference ManifestFileName in Rules.ManifestFileNames)
+			foreach(FileReference ManifestFileName in Rules.ManifestFileNames)
 			{
 				GenerateManifest(ManifestFileName, BuildProducts);
 			}
 
 			// Check there are no EULA or restricted folder violations
-			if (!Rules.bDisableLinking)
+			if(!Rules.bDisableLinking)
 			{
 				// Check the distribution level of all binaries based on the dependencies they have
 				if (ProjectFile == null && !Rules.bLegalToDistributeBinary)
@@ -2142,12 +2110,12 @@ namespace UnrealBuildTool
 					{
 						bResult &= Binary.CheckRestrictedFolders(RootDirectories, ModuleRestrictedFolderCache);
 					}
-					foreach (KeyValuePair<FileReference, FileReference> Pair in RuntimeDependencyTargetFileToSourceFile)
+					foreach(KeyValuePair<FileReference, FileReference> Pair in RuntimeDependencyTargetFileToSourceFile)
 					{
 						bResult &= CheckRestrictedFolders(Pair.Key, Pair.Value);
 					}
 
-					if (!bResult)
+					if(!bResult)
 					{
 						throw new BuildException("Unable to create binaries in less restricted locations than their input files.");
 					}
@@ -2158,7 +2126,7 @@ namespace UnrealBuildTool
 			}
 
 			// Add all the plugins to be tracked
-			foreach (FileReference PluginFile in PluginsBase.EnumeratePlugins(ProjectFile))
+			foreach(FileReference PluginFile in PluginsBase.EnumeratePlugins(ProjectFile))
 			{
 				FileItem PluginFileItem = FileItem.GetItemByFileReference(PluginFile);
 				Makefile.PluginFiles.Add(PluginFileItem);
@@ -2167,10 +2135,10 @@ namespace UnrealBuildTool
 			// Add all the input files to the predicate store
 			Makefile.ExternalDependencies.Add(FileItem.GetItemByFileReference(TargetRulesFile));
 			Makefile.ExternalDependencies.Add(FileItem.GetItemByFileReference(Rules.TargetSourceFile));
-			foreach (UEBuildModule Module in Modules.Values)
+			foreach(UEBuildModule Module in Modules.Values)
 			{
 				Makefile.ExternalDependencies.Add(FileItem.GetItemByFileReference(Module.RulesFile));
-				foreach (string ExternalDependency in Module.Rules.ExternalDependencies)
+				foreach(string ExternalDependency in Module.Rules.ExternalDependencies)
 				{
 					FileReference Location = FileReference.Combine(Module.RulesFile.Directory, ExternalDependency);
 					Makefile.ExternalDependencies.Add(FileItem.GetItemByFileReference(Location));
@@ -2287,9 +2255,9 @@ namespace UnrealBuildTool
 		{
 			List<RestrictedFolder> TargetRestrictedFolders = GetRestrictedFolders(TargetFile);
 			List<RestrictedFolder> SourceRestrictedFolders = GetRestrictedFolders(SourceFile);
-			foreach (RestrictedFolder SourceRestrictedFolder in SourceRestrictedFolders)
+			foreach(RestrictedFolder SourceRestrictedFolder in SourceRestrictedFolders)
 			{
-				if (!TargetRestrictedFolders.Contains(SourceRestrictedFolder))
+				if(!TargetRestrictedFolders.Contains(SourceRestrictedFolder))
 				{
 					Log.TraceError("Runtime dependency '{0}' is copied to '{1}', which does not contain a '{2}' folder.", SourceFile, TargetFile, SourceRestrictedFolder);
 					return false;
@@ -2307,11 +2275,11 @@ namespace UnrealBuildTool
 		{
 			// Find the base directory for this binary
 			DirectoryReference BaseDir;
-			if (File.IsUnderDirectory(Unreal.RootDirectory))
+			if(File.IsUnderDirectory(Unreal.RootDirectory))
 			{
 				BaseDir = Unreal.RootDirectory;
 			}
-			else if (ProjectDirectory != null && File.IsUnderDirectory(ProjectDirectory))
+			else if(ProjectDirectory != null && File.IsUnderDirectory(ProjectDirectory))
 			{
 				BaseDir = ProjectDirectory;
 			}
@@ -2355,20 +2323,20 @@ namespace UnrealBuildTool
 		{
 			// Find all the output files
 			HashSet<FileReference> OutputFiles = new HashSet<FileReference>();
-			foreach (UEBuildBinary Binary in Binaries)
+			foreach(UEBuildBinary Binary in Binaries)
 			{
 				OutputFiles.UnionWith(Binary.OutputFilePaths);
 			}
 
 			// Build a map of base filenames to their full path
 			Dictionary<string, FileReference> OutputNameToLocation = new Dictionary<string, FileReference>(StringComparer.InvariantCultureIgnoreCase);
-			foreach (FileReference OutputFile in OutputFiles)
+			foreach(FileReference OutputFile in OutputFiles)
 			{
 				OutputNameToLocation[OutputFile.GetFileName()] = OutputFile;
 			}
 
 			// Search all the output directories for files with a name matching one of our output files
-			foreach (DirectoryReference OutputDirectory in OutputFiles.Select(x => x.Directory).Distinct())
+			foreach(DirectoryReference OutputDirectory in OutputFiles.Select(x => x.Directory).Distinct())
 			{
 				if (DirectoryReference.Exists(OutputDirectory))
 				{
@@ -2400,11 +2368,11 @@ namespace UnrealBuildTool
 		/// <returns>True if the reference is allowed.</returns>
 		static bool IsEngineToPluginReferenceAllowed(string EngineModuleName, string PluginModuleName)
 		{
-			if (EngineModuleName == "AndroidDeviceDetection" && PluginModuleName == "TcpMessaging")
+			if(EngineModuleName == "AndroidDeviceDetection" && PluginModuleName == "TcpMessaging")
 			{
 				return true;
 			}
-			if (EngineModuleName == "Voice" && PluginModuleName == "AndroidPermission")
+			if(EngineModuleName == "Voice" && PluginModuleName == "AndroidPermission")
 			{
 				return true;
 			}
@@ -2464,20 +2432,20 @@ namespace UnrealBuildTool
 
 			// Find all the public definitions from each module. We pass null as the source binary to AddModuleToCompileEnvironment, which forces all _API macros to 'export' mode
 			List<string> Definitions = new List<string>(GlobalCompileEnvironment.Definitions);
-			foreach (UEBuildModule Module in Binary.Modules)
+			foreach(UEBuildModule Module in Binary.Modules)
 			{
 				Module.AddModuleToCompileEnvironment(null, null, new HashSet<DirectoryReference>(), new HashSet<DirectoryReference>(), new HashSet<DirectoryReference>(), Definitions, new List<UEBuildFramework>(), new List<FileItem>(), false);
 			}
 
 			// Write the header
-			using (StreamWriter Writer = new StreamWriter(HeaderFile.FullName))
+			using(StreamWriter Writer = new StreamWriter(HeaderFile.FullName))
 			{
 				Writer.WriteLine("#pragma once");
 				Writer.WriteLine();
-				foreach (string Definition in Definitions)
+				foreach(string Definition in Definitions)
 				{
 					int EqualsIdx = Definition.IndexOf('=');
-					if (EqualsIdx == -1)
+					if(EqualsIdx == -1)
 					{
 						Writer.WriteLine(String.Format("#define {0} 1", Definition));
 					}
@@ -2516,7 +2484,7 @@ namespace UnrealBuildTool
 						{
 							Message = string.Format("{0}\nDependant modules '{1}'", Message, string.Join("', '", NonRedistDeps));
 						}
-						if (Rules.bBreakBuildOnLicenseViolation)
+						if(Rules.bBreakBuildOnLicenseViolation)
 						{
 							Log.TraceError("ERROR: {0}", Message);
 						}
@@ -2539,23 +2507,23 @@ namespace UnrealBuildTool
 		/// </summary>
 		public static bool IsRedistributable(UEBuildModule Module)
 		{
-			if (Module.Rules != null && Module.Rules.IsRedistributableOverride.HasValue)
+			if(Module.Rules != null && Module.Rules.IsRedistributableOverride.HasValue)
 			{
 				return Module.Rules.IsRedistributableOverride.Value;
 			}
 
-			if (Module.RulesFile != null)
+			if(Module.RulesFile != null)
 			{
 				foreach (DirectoryReference ExtensionDir in Unreal.GetExtensionDirs(Unreal.EngineDirectory))
 				{
 					DirectoryReference SourceDeveloperDir = DirectoryReference.Combine(ExtensionDir, "Source/Developer");
-					if (Module.RulesFile.IsUnderDirectory(SourceDeveloperDir))
+					if(Module.RulesFile.IsUnderDirectory(SourceDeveloperDir))
 					{
 						return false;
 					}
 
 					DirectoryReference SourceEditorDir = DirectoryReference.Combine(ExtensionDir, "Source/Editor");
-					if (Module.RulesFile.IsUnderDirectory(SourceEditorDir))
+					if(Module.RulesFile.IsUnderDirectory(SourceEditorDir))
 					{
 						return false;
 					}
@@ -2582,7 +2550,7 @@ namespace UnrealBuildTool
 			// Add the plugin binaries to the build
 			foreach (UEBuildPlugin Plugin in BuildPlugins!)
 			{
-				foreach (UEBuildModuleCPP Module in Plugin.Modules)
+				foreach(UEBuildModuleCPP Module in Plugin.Modules)
 				{
 					AddModuleToBinary(Module);
 				}
@@ -2594,7 +2562,7 @@ namespace UnrealBuildTool
 			AddExtraModules();
 
 			// Create all the modules referenced by the existing binaries
-			foreach (UEBuildBinary Binary in Binaries)
+			foreach(UEBuildBinary Binary in Binaries)
 			{
 				Binary.CreateAllDependentModules(FindOrCreateModuleByName);
 			}
@@ -2605,7 +2573,7 @@ namespace UnrealBuildTool
 				List<UEBuildModule> DependencyModules = Binaries[Idx].GetAllDependencyModules(true, true);
 				foreach (UEBuildModuleCPP DependencyModule in DependencyModules.OfType<UEBuildModuleCPP>())
 				{
-					if (DependencyModule.Binary == null)
+					if(DependencyModule.Binary == null)
 					{
 						AddModuleToBinary(DependencyModule);
 					}
@@ -2613,7 +2581,7 @@ namespace UnrealBuildTool
 			}
 
 			// Add all the modules to the target if necessary.
-			if (Rules.bBuildAllModules)
+			if(Rules.bBuildAllModules)
 			{
 				AddAllValidModulesToTarget();
 			}
@@ -2621,7 +2589,7 @@ namespace UnrealBuildTool
 			// Add the external and non-C++ referenced modules to the binaries that reference them.
 			foreach (UEBuildModuleCPP Module in Modules.Values.OfType<UEBuildModuleCPP>())
 			{
-				if (Module.Binary != null)
+				if(Module.Binary != null)
 				{
 					foreach (UEBuildModule ReferencedModule in Module.GetUnboundReferences())
 					{
@@ -2665,15 +2633,15 @@ namespace UnrealBuildTool
 		{
 			// Find all the pre-build steps
 			List<Tuple<string[], UEBuildPlugin?>> PreBuildCommandBatches = new List<Tuple<string[], UEBuildPlugin?>>();
-			if (ProjectDescriptor != null && ProjectDescriptor.PreBuildSteps != null)
+			if(ProjectDescriptor != null && ProjectDescriptor.PreBuildSteps != null)
 			{
 				AddCustomBuildSteps(ProjectDescriptor.PreBuildSteps, null, PreBuildCommandBatches);
 			}
-			if (Rules.PreBuildSteps.Count > 0)
+			if(Rules.PreBuildSteps.Count > 0)
 			{
 				PreBuildCommandBatches.Add(new Tuple<string[], UEBuildPlugin?>(Rules.PreBuildSteps.ToArray(), null));
 			}
-			foreach (UEBuildPlugin BuildPlugin in BuildPlugins.Where(x => x.Descriptor.PreBuildSteps != null))
+			foreach(UEBuildPlugin BuildPlugin in BuildPlugins.Where(x => x.Descriptor.PreBuildSteps != null))
 			{
 				AddCustomBuildSteps(BuildPlugin.Descriptor.PreBuildSteps!, BuildPlugin, PreBuildCommandBatches);
 			}
@@ -2688,17 +2656,17 @@ namespace UnrealBuildTool
 		{
 			// Find all the post-build steps
 			List<Tuple<string[], UEBuildPlugin?>> PostBuildCommandBatches = new List<Tuple<string[], UEBuildPlugin?>>();
-			if (!Rules.bDisableLinking)
+			if(!Rules.bDisableLinking)
 			{
-				if (ProjectDescriptor != null && ProjectDescriptor.PostBuildSteps != null)
+				if(ProjectDescriptor != null && ProjectDescriptor.PostBuildSteps != null)
 				{
 					AddCustomBuildSteps(ProjectDescriptor.PostBuildSteps, null, PostBuildCommandBatches);
 				}
-				if (Rules.PostBuildSteps.Count > 0)
+				if(Rules.PostBuildSteps.Count > 0)
 				{
 					PostBuildCommandBatches.Add(new Tuple<string[], UEBuildPlugin?>(Rules.PostBuildSteps.ToArray(), null));
 				}
-				foreach (UEBuildPlugin BuildPlugin in BuildPlugins.Where(x => x.Descriptor.PostBuildSteps != null))
+				foreach(UEBuildPlugin BuildPlugin in BuildPlugins.Where(x => x.Descriptor.PostBuildSteps != null))
 				{
 					AddCustomBuildSteps(BuildPlugin.Descriptor.PostBuildSteps!, BuildPlugin, PostBuildCommandBatches);
 				}
@@ -2715,7 +2683,7 @@ namespace UnrealBuildTool
 		private void AddCustomBuildSteps(CustomBuildSteps BuildSteps, UEBuildPlugin? Plugin, List<Tuple<string[], UEBuildPlugin?>> CommandBatches)
 		{
 			string[]? Commands;
-			if (BuildSteps.TryGetCommands(BuildHostPlatform.Current.Platform, out Commands))
+			if(BuildSteps.TryGetCommands(BuildHostPlatform.Current.Platform, out Commands))
 			{
 				CommandBatches.Add(Tuple.Create(Commands, Plugin));
 			}
@@ -2728,7 +2696,7 @@ namespace UnrealBuildTool
 		/// <returns>Map of variable names to values</returns>
 		private Dictionary<string, string> GetTargetVariables(UEBuildPlugin? Plugin)
 		{
-			Dictionary<string, string> Variables = new Dictionary<string, string>();
+			Dictionary<string, string> Variables = new Dictionary<string,string>();
 			Variables.Add("RootDir", Unreal.RootDirectory.FullName);
 			Variables.Add("EngineDir", Unreal.EngineDirectory.FullName);
 			Variables.Add("ProjectDir", ProjectDirectory.FullName);
@@ -2736,11 +2704,11 @@ namespace UnrealBuildTool
 			Variables.Add("TargetPlatform", Platform.ToString());
 			Variables.Add("TargetConfiguration", Configuration.ToString());
 			Variables.Add("TargetType", TargetType.ToString());
-			if (ProjectFile != null)
+			if(ProjectFile != null)
 			{
 				Variables.Add("ProjectFile", ProjectFile.FullName);
 			}
-			if (Plugin != null)
+			if(Plugin != null)
 			{
 				Variables.Add("PluginDir", Plugin.Directory.FullName);
 			}
@@ -2758,26 +2726,26 @@ namespace UnrealBuildTool
 		private FileReference[] WriteCustomBuildStepScripts(UnrealTargetPlatform HostPlatform, DirectoryReference Directory, string FilePrefix, List<Tuple<string[], UEBuildPlugin?>> CommandBatches)
 		{
 			List<FileReference> ScriptFiles = new List<FileReference>();
-			foreach (Tuple<string[], UEBuildPlugin?> CommandBatch in CommandBatches)
+			foreach(Tuple<string[], UEBuildPlugin?> CommandBatch in CommandBatches)
 			{
 				// Find all the standard variables
 				Dictionary<string, string> Variables = GetTargetVariables(CommandBatch.Item2);
 
 				// Get the output path to the script
-				string ScriptExtension = (HostPlatform == UnrealTargetPlatform.Win64) ? ".bat" : ".sh";
+				string ScriptExtension = (HostPlatform == UnrealTargetPlatform.Win64)? ".bat" : ".sh";
 				FileReference ScriptFile = FileReference.Combine(Directory, String.Format("{0}-{1}{2}", FilePrefix, ScriptFiles.Count + 1, ScriptExtension));
 
 				// Write it to disk
 				List<string> Contents = new List<string>();
-				if (HostPlatform == UnrealTargetPlatform.Win64)
+				if(HostPlatform == UnrealTargetPlatform.Win64)
 				{
 					Contents.Insert(0, "@echo off");
 				}
-				foreach (string Command in CommandBatch.Item1)
+				foreach(string Command in CommandBatch.Item1)
 				{
 					Contents.Add(Utils.ExpandVariables(Command, Variables));
 				}
-				if (!DirectoryReference.Exists(ScriptFile.Directory))
+				if(!DirectoryReference.Exists(ScriptFile.Directory))
 				{
 					DirectoryReference.CreateDirectory(ScriptFile.Directory);
 				}
@@ -2833,11 +2801,11 @@ namespace UnrealBuildTool
 		{
 			// Find how many other shared PCH modules each module depends on, and use that to sort the shared PCHs by reverse order of size.
 			HashSet<UEBuildModuleCPP> SharedPCHModules = new HashSet<UEBuildModuleCPP>();
-			foreach (UEBuildBinary Binary in OriginalBinaries)
+			foreach(UEBuildBinary Binary in OriginalBinaries)
 			{
-				foreach (UEBuildModuleCPP Module in Binary.Modules.OfType<UEBuildModuleCPP>())
+				foreach(UEBuildModuleCPP Module in Binary.Modules.OfType<UEBuildModuleCPP>())
 				{
-					if (Module.Rules.SharedPCHHeaderFile != null)
+					if(Module.Rules.SharedPCHHeaderFile != null)
 					{
 						SharedPCHModules.Add(Module);
 					}
@@ -2846,7 +2814,7 @@ namespace UnrealBuildTool
 
 			// Find a priority for each shared PCH, determined as the number of other shared PCHs it includes.
 			Dictionary<UEBuildModuleCPP, int> SharedPCHModuleToPriority = new Dictionary<UEBuildModuleCPP, int>();
-			foreach (UEBuildModuleCPP SharedPCHModule in SharedPCHModules)
+			foreach(UEBuildModuleCPP SharedPCHModule in SharedPCHModules)
 			{
 				List<UEBuildModule> Dependencies = new List<UEBuildModule>();
 				SharedPCHModule.GetAllDependencyModules(Dependencies, new HashSet<UEBuildModule>(), false, false, false);
@@ -2855,13 +2823,13 @@ namespace UnrealBuildTool
 
 			// Create the shared PCH modules, in order
 			List<PrecompiledHeaderTemplate> OrderedSharedPCHModules = GlobalCompileEnvironment.SharedPCHs;
-			foreach (UEBuildModuleCPP Module in SharedPCHModuleToPriority.OrderByDescending(x => x.Value).Select(x => x.Key))
+			foreach(UEBuildModuleCPP Module in SharedPCHModuleToPriority.OrderByDescending(x => x.Value).Select(x => x.Key))
 			{
 				OrderedSharedPCHModules.Add(Module.CreateSharedPCHTemplate(this, GlobalCompileEnvironment));
 			}
 
 			// Print the ordered list of shared PCHs
-			if (OrderedSharedPCHModules.Count > 0)
+			if(OrderedSharedPCHModules.Count > 0)
 			{
 				Log.TraceVerbose("Found {0} shared PCH headers (listed in order of preference):", SharedPCHModules.Count);
 				foreach (PrecompiledHeaderTemplate SharedPCHModule in OrderedSharedPCHModules)
@@ -2902,7 +2870,7 @@ namespace UnrealBuildTool
 			HashSet<string> ValidModuleNames = new HashSet<string>();
 			foreach (UEBuildModuleCPP Module in Modules.Values.OfType<UEBuildModuleCPP>())
 			{
-				if (Module.Binary != null)
+				if(Module.Binary != null)
 				{
 					ValidModuleNames.Add(Module.Name);
 				}
@@ -2915,7 +2883,7 @@ namespace UnrealBuildTool
 			HashSet<string> FilteredModuleNames = new HashSet<string>();
 
 			// Only add engine modules for non-program targets. Programs only compile allowed modules through plugins.
-			if (TargetType != TargetType.Program)
+			if(TargetType != TargetType.Program)
 			{
 				// Find all the known module names in this assembly
 				List<string> ModuleNames = new List<string>();
@@ -2931,7 +2899,7 @@ namespace UnrealBuildTool
 
 				// Also allow anything in the developer directory in non-shipping configurations (though we deny by default
 				// unless the PrecompileForTargets setting indicates that it's actually useful at runtime).
-				if (Rules.bBuildDeveloperTools)
+				if(Rules.bBuildDeveloperTools)
 				{
 					Directories.AddRange(Unreal.GetExtensionDirs(Unreal.EngineDirectory, "Source/Developer"));
 				}
@@ -2949,9 +2917,9 @@ namespace UnrealBuildTool
 						if (!UEBuildPlatform.GetPlatformFolderNames().Any(Name => RulesType.Name.EndsWith("_" + Name)))
 						{
 							SupportedPlatformsAttribute? SupportedPlatforms = RulesType.GetCustomAttribute<SupportedPlatformsAttribute>();
-							if (SupportedPlatforms != null)
+							if(SupportedPlatforms != null)
 							{
-								if (SupportedPlatforms.Platforms.Contains(Platform))
+								if(SupportedPlatforms.Platforms.Contains(Platform))
 								{
 									FilteredModuleNames.Add(ModuleName);
 								}
@@ -2970,7 +2938,7 @@ namespace UnrealBuildTool
 
 			// Add all the plugin modules that need to be compiled
 			List<PluginInfo> Plugins = RulesAssembly.EnumeratePlugins().ToList();
-			foreach (PluginInfo Plugin in Plugins)
+			foreach(PluginInfo Plugin in Plugins)
 			{
 				// Ignore plugins which are specifically disabled by this target
 				if (Rules.DisablePlugins.Contains(Plugin.Name))
@@ -2991,7 +2959,7 @@ namespace UnrealBuildTool
 				}
 
 				// Disable any plugin that requires the build platform
-				if (Plugin.Descriptor.bRequiresBuildPlatform && ShouldExcludePlugin(Plugin, ExcludeFolders))
+				if(Plugin.Descriptor.bRequiresBuildPlatform && ShouldExcludePlugin(Plugin, ExcludeFolders))
 				{
 					continue;
 				}
@@ -3008,11 +2976,11 @@ namespace UnrealBuildTool
 					if (ModuleDescriptor.IsCompiledInConfiguration(Platform, Configuration, TargetName, TargetType, Rules.bBuildDeveloperTools, Rules.bBuildRequiresCookedData))
 					{
 						FileReference? ModuleFileName = RulesAssembly.GetModuleFileName(ModuleDescriptor.Name);
-						if (ModuleFileName == null)
+						if(ModuleFileName == null)
 						{
 							throw new BuildException("Unable to find module '{0}' referenced by {1}", ModuleDescriptor.Name, Plugin.File);
 						}
-						if (!ModuleFileName.ContainsAnyNames(ExcludeFolders, Plugin.Directory))
+						if(!ModuleFileName.ContainsAnyNames(ExcludeFolders, Plugin.Directory))
 						{
 							FilteredModuleNames.Add(ModuleDescriptor.Name);
 						}
@@ -3021,7 +2989,7 @@ namespace UnrealBuildTool
 			}
 
 			// Create rules for each remaining module, and check that it's set to be compiled
-			foreach (string FilteredModuleName in FilteredModuleNames)
+			foreach(string FilteredModuleName in FilteredModuleNames)
 			{
 				// Try to create the rules object, but catch any exceptions if it fails. Some modules (eg. SQLite) may determine that they are unavailable in the constructor.
 				ModuleRules? ModuleRules;
@@ -3043,7 +3011,7 @@ namespace UnrealBuildTool
 
 			// Now create all the precompiled modules, making sure they don't reference anything that's not in the precompiled set
 			HashSet<UEBuildModuleCPP> ValidModules = new HashSet<UEBuildModuleCPP>();
-			foreach (string ModuleName in ValidModuleNames)
+			foreach(string ModuleName in ValidModuleNames)
 			{
 				const string PrecompileReferenceChain = "allmodules option";
 				UEBuildModuleCPP Module = (UEBuildModuleCPP)FindOrCreateModuleByName(ModuleName, PrecompileReferenceChain);
@@ -3055,11 +3023,11 @@ namespace UnrealBuildTool
 			}
 
 			// Make sure precompiled modules don't reference any non-precompiled modules
-			foreach (UEBuildModuleCPP ValidModule in ValidModules)
+			foreach(UEBuildModuleCPP ValidModule in ValidModules)
 			{
-				foreach (UEBuildModuleCPP ReferencedModule in ValidModule.GetDependencies(false, true).OfType<UEBuildModuleCPP>())
+				foreach(UEBuildModuleCPP ReferencedModule in ValidModule.GetDependencies(false, true).OfType<UEBuildModuleCPP>())
 				{
-					if (!ValidModules.Contains(ReferencedModule))
+					if(!ValidModules.Contains(ReferencedModule))
 					{
 						Log.TraceWarning("Module '{0}' is not usable without module '{1}', which is not valid for this target.", ValidModule.Name, ReferencedModule.Name);
 					}
@@ -3067,9 +3035,9 @@ namespace UnrealBuildTool
 			}
 
 			// Make sure every module is built
-			foreach (UEBuildModuleCPP Module in ValidModules)
+			foreach(UEBuildModuleCPP Module in ValidModules)
 			{
-				if (Module.Binary == null)
+				if(Module.Binary == null)
 				{
 					AddModuleToBinary(Module);
 				}
@@ -3112,8 +3080,7 @@ namespace UnrealBuildTool
 		{
 			// Get the root output directory and base name (target name/app name) for this binary
 			DirectoryReference BaseOutputDirectory;
-
-			if (bUseSharedBuildEnvironment && (ModuleRules.Plugin == null || (ModuleRules.Plugin.Type != PluginType.External || ModuleRules.Plugin.bExplicitPluginTarget)))
+			if(bUseSharedBuildEnvironment && ModuleRules.Plugin?.Type != PluginType.External)
 			{
 				BaseOutputDirectory = ModuleRules.Context.DefaultOutputBaseDir;
 			}
@@ -3141,10 +3108,8 @@ namespace UnrealBuildTool
 				ModuleConfiguration = UnrealTargetConfiguration.Development;
 			}
 
-			bool bUseExternalFolder = ModuleRules.Plugin != null && (ModuleRules.Plugin.Type == PluginType.External && !ModuleRules.Plugin.bExplicitPluginTarget);
-
 			// Get the output and intermediate directories for this module
-			DirectoryReference IntermediateDirectory = DirectoryReference.Combine(BaseOutputDirectory, GetPlatformIntermediateFolder(Platform, Architecture, bUseExternalFolder), AppName, ModuleConfiguration.ToString());
+			DirectoryReference IntermediateDirectory = DirectoryReference.Combine(BaseOutputDirectory, GetPlatformIntermediateFolder(Platform, Architecture, ModuleRules.Plugin?.Type == PluginType.External), AppName, ModuleConfiguration.ToString());
 
 			// Append a subdirectory if the module rules specifies one
 			if (!String.IsNullOrEmpty(ModuleRules.BinariesSubFolder))
@@ -3222,7 +3187,7 @@ namespace UnrealBuildTool
 			}
 
 			UEBuildPlatform BuildPlatform = UEBuildPlatform.GetBuildPlatform(Platform);
-			if (BuildPlatform.RequiresArchitectureSuffix())
+			if(BuildPlatform.RequiresArchitectureSuffix())
 			{
 				Result.Append(Architecture);
 			}
@@ -3292,16 +3257,16 @@ namespace UnrealBuildTool
 			Dictionary<string, UEBuildPlugin> NameToInstance = new Dictionary<string, UEBuildPlugin>(StringComparer.InvariantCultureIgnoreCase);
 
 			// Set up the foreign plugin
-			if (ForeignPlugin != null)
+			if(ForeignPlugin != null)
 			{
 				PluginReferenceDescriptor PluginReference = new PluginReferenceDescriptor(ForeignPlugin.GetFileNameWithoutExtension(), null, true);
 				AddPlugin(PluginReference, "command line", ExcludeFolders, NameToInstance, NameToInfo);
 			}
 
 			// Configure plugins explicitly enabled via target settings
-			foreach (string PluginName in Rules.EnablePlugins)
+			foreach(string PluginName in Rules.EnablePlugins)
 			{
-				if (ReferencedNames.Add(PluginName))
+				if(ReferencedNames.Add(PluginName))
 				{
 					PluginReferenceDescriptor PluginReference = new PluginReferenceDescriptor(PluginName, null, true);
 					AddPlugin(PluginReference, "target settings", ExcludeFolders, NameToInstance, NameToInfo);
@@ -3309,9 +3274,9 @@ namespace UnrealBuildTool
 			}
 
 			// Configure plugins explicitly disabled via target settings
-			foreach (string PluginName in Rules.DisablePlugins)
+			foreach(string PluginName in Rules.DisablePlugins)
 			{
-				if (ReferencedNames.Add(PluginName))
+				if(ReferencedNames.Add(PluginName))
 				{
 					PluginReferenceDescriptor PluginReference = new PluginReferenceDescriptor(PluginName, null, false);
 					AddPlugin(PluginReference, "target settings", ExcludeFolders, NameToInstance, NameToInfo);
@@ -3348,9 +3313,9 @@ namespace UnrealBuildTool
 			// Also synthesize references for plugins which are enabled by default
 			if (Rules.bCompileAgainstEngine || Rules.bCompileWithPluginSupport)
 			{
-				foreach (PluginInfo Plugin in NameToInfo.Values)
+				foreach(PluginInfo Plugin in NameToInfo.Values)
 				{
-					if (Plugin.IsEnabledByDefault(bAllowEnginePluginsEnabledByDefault) && !ReferencedNames.Contains(Plugin.Name))
+					if(Plugin.IsEnabledByDefault(bAllowEnginePluginsEnabledByDefault) && !ReferencedNames.Contains(Plugin.Name))
 					{
 						ReferencedNames.Add(Plugin.Name);
 
@@ -3370,16 +3335,16 @@ namespace UnrealBuildTool
 			}
 
 			// If this is a program, synthesize references for plugins which are enabled via the config file
-			if (TargetType == TargetType.Program)
+			if(TargetType == TargetType.Program)
 			{
 				ConfigHierarchy EngineConfig = ConfigCache.ReadHierarchy(ConfigHierarchyType.Engine, DirectoryReference.Combine(Unreal.EngineDirectory, "Programs", TargetName), Platform);
 
 				List<string>? PluginNames;
-				if (EngineConfig.GetArray("Plugins", "ProgramEnabledPlugins", out PluginNames))
+				if(EngineConfig.GetArray("Plugins", "ProgramEnabledPlugins", out PluginNames))
 				{
-					foreach (string PluginName in PluginNames)
+					foreach(string PluginName in PluginNames)
 					{
-						if (ReferencedNames.Add(PluginName))
+						if(ReferencedNames.Add(PluginName))
 						{
 							PluginReferenceDescriptor PluginReference = new PluginReferenceDescriptor(PluginName, null, true);
 							AddPlugin(PluginReference, "DefaultEngine.ini", ExcludeFolders, NameToInstance, NameToInfo);
@@ -3417,22 +3382,22 @@ namespace UnrealBuildTool
 		private UEBuildPlugin? AddPlugin(PluginReferenceDescriptor Reference, string ReferenceChain, string[] ExcludeFolders, Dictionary<string, UEBuildPlugin> NameToInstance, Dictionary<string, PluginInfo> NameToInfo)
 		{
 			// Ignore disabled references
-			if (!Reference.bEnabled)
+			if(!Reference.bEnabled)
 			{
 				return null;
 			}
 
 			// Try to get an existing reference to this plugin
 			UEBuildPlugin? Instance;
-			if (NameToInstance.TryGetValue(Reference.Name, out Instance))
+			if(NameToInstance.TryGetValue(Reference.Name, out Instance))
 			{
 				// If this is a non-optional reference, make sure that and every referenced dependency is staged
-				if (!Reference.bOptional && !Instance.bDescriptorReferencedExplicitly)
+				if(!Reference.bOptional && !Instance.bDescriptorReferencedExplicitly)
 				{
 					Instance.bDescriptorReferencedExplicitly = true;
-					if (Instance.Descriptor.Plugins != null)
+					if(Instance.Descriptor.Plugins != null)
 					{
-						foreach (PluginReferenceDescriptor NextReference in Instance.Descriptor.Plugins)
+						foreach(PluginReferenceDescriptor NextReference in Instance.Descriptor.Plugins)
 						{
 							string NextReferenceChain = String.Format("{0} -> {1}", ReferenceChain, Instance.File.GetFileName());
 							AddPlugin(NextReference, NextReferenceChain, ExcludeFolders, NameToInstance, NameToInfo);
@@ -3443,7 +3408,7 @@ namespace UnrealBuildTool
 			else
 			{
 				// Check if the plugin is required for this platform
-				if (!Reference.IsEnabledForPlatform(Platform))
+				if(!Reference.IsEnabledForPlatform(Platform))
 				{
 					Log.TraceLog("Ignoring plugin '{0}' (referenced via {1}), not enabled for platform {2}", Reference.Name, ReferenceChain, Platform);
 					return null;
@@ -3470,7 +3435,7 @@ namespace UnrealBuildTool
 
 				// Find the plugin being enabled
 				PluginInfo? Info;
-				if (!NameToInfo.TryGetValue(Reference.Name, out Info))
+				if(!NameToInfo.TryGetValue(Reference.Name, out Info))
 				{
 					if (Reference.bOptional)
 					{
@@ -3489,7 +3454,7 @@ namespace UnrealBuildTool
 				}
 
 				// Disable any plugin that requires the build platform
-				if (Info.Descriptor.bRequiresBuildPlatform && ShouldExcludePlugin(Info, ExcludeFolders))
+				if(Info.Descriptor.bRequiresBuildPlatform && ShouldExcludePlugin(Info, ExcludeFolders))
 				{
 					Log.TraceLog("Ignoring plugin '{0}' (referenced via {1}) due to missing build platform", Reference.Name, ReferenceChain);
 					return null;
@@ -3526,7 +3491,7 @@ namespace UnrealBuildTool
 						if (ModuleInfo.IsCompiledInConfiguration(Platform, Configuration, TargetName, TargetType, Rules.bBuildDeveloperTools, Rules.bBuildRequiresCookedData))
 						{
 							UEBuildModuleCPP Module = FindOrCreateCppModuleByName(ModuleInfo.Name, PluginReferenceChain);
-							if (!Instance.Modules.Contains(Module))
+							if(!Instance.Modules.Contains(Module))
 							{
 								// This could be in a child plugin so scan thorugh those as well
 								if (!Module.RulesFile.IsUnderDirectory(Info.Directory) && !Info.ChildFiles.Any(ChildFile => Module.RulesFile.IsUnderDirectory(ChildFile.Directory)))
@@ -3542,15 +3507,15 @@ namespace UnrealBuildTool
 
 				// Create the dependencies set
 				HashSet<UEBuildPlugin> Dependencies = new HashSet<UEBuildPlugin>();
-				if (Info.Descriptor.Plugins != null)
+				if(Info.Descriptor.Plugins != null)
 				{
-					foreach (PluginReferenceDescriptor NextReference in Info.Descriptor.Plugins)
+					foreach(PluginReferenceDescriptor NextReference in Info.Descriptor.Plugins)
 					{
 						UEBuildPlugin? NextInstance = AddPlugin(NextReference, PluginReferenceChain, ExcludeFolders, NameToInstance, NameToInfo);
-						if (NextInstance != null)
+						if(NextInstance != null)
 						{
 							Dependencies.Add(NextInstance);
-							if (NextInstance.Dependencies == null)
+							if(NextInstance.Dependencies == null)
 							{
 								throw new BuildException("Found circular dependency from plugin '{0}' onto itself.", NextReference.Name);
 							}
@@ -3578,7 +3543,7 @@ namespace UnrealBuildTool
 			{
 				return Plugin.File.ContainsAnyNames(ExcludeFolders, Unreal.EngineDirectory);
 			}
-			else if (ProjectFile != null)
+			else if(ProjectFile != null)
 			{
 				return Plugin.File.ContainsAnyNames(ExcludeFolders, ProjectFile.Directory);
 			}
@@ -3594,7 +3559,7 @@ namespace UnrealBuildTool
 		protected void SetupBinaries()
 		{
 			// If we're using the new method for specifying binaries, fill in the binary configurations now
-			if (Rules.LaunchModuleName == null)
+			if(Rules.LaunchModuleName == null)
 			{
 				throw new BuildException("LaunchModuleName must be set for all targets.");
 			}
@@ -3604,7 +3569,7 @@ namespace UnrealBuildTool
 
 			// Get the intermediate directory for the launch module directory. This can differ from the standard engine intermediate directory because it is always configuration-specific.
 			DirectoryReference IntermediateDirectory;
-			if (LaunchModule.RulesFile.IsUnderDirectory(Unreal.EngineDirectory) && !ShouldCompileMonolithic())
+			if(LaunchModule.RulesFile.IsUnderDirectory(Unreal.EngineDirectory) && !ShouldCompileMonolithic())
 			{
 				IntermediateDirectory = DirectoryReference.Combine(Unreal.EngineDirectory, PlatformIntermediateFolder, AppName, Configuration.ToString());
 			}
@@ -3629,7 +3594,7 @@ namespace UnrealBuildTool
 
 			// Create the binary
 			UEBuildBinary Binary = new UEBuildBinary(
-				Type: Rules.bShouldCompileAsDLL ? UEBuildBinaryType.DynamicLinkLibrary : UEBuildBinaryType.Executable,
+				Type: Rules.bShouldCompileAsDLL? UEBuildBinaryType.DynamicLinkLibrary : UEBuildBinaryType.Executable,
 				OutputFilePaths: OutputPaths,
 				IntermediateDirectory: IntermediateDirectory,
 				bAllowExports: Rules.bHasExports,
@@ -3763,10 +3728,15 @@ namespace UnrealBuildTool
 			// By default, shadow source files for this target in the root OutputDirectory
 			GlobalLinkEnvironment.LocalShadowDirectory = GlobalLinkEnvironment.OutputDirectory;
 
-			string? RelativeBaseDir = GetRelativeBaseDir(Binaries[0].OutputDir, Platform);
-			if (RelativeBaseDir != null)
+			DirectoryReference OutputDir = Binaries[0].OutputDir;
+			if (OutputDir.IsUnderDirectory(Unreal.EngineDirectory))
 			{
-				GlobalCompileEnvironment.Definitions.Add(String.Format("UE_RELATIVE_BASE_DIR=\"{0}/\"", RelativeBaseDir));
+				DirectoryReference BaseDir = DirectoryReference.Combine(Unreal.EngineDirectory, "Binaries", Platform.ToString());
+				if (BaseDir != OutputDir)
+				{
+					string RelativeBaseDir = BaseDir.MakeRelativeTo(OutputDir).Replace(Path.DirectorySeparatorChar, '/');
+					GlobalCompileEnvironment.Definitions.Add(String.Format("UE_RELATIVE_BASE_DIR=\"{0}/\"", RelativeBaseDir));
+				}
 			}
 
 			bool bCompileDevTests = (Configuration != UnrealTargetConfiguration.Test && Configuration != UnrealTargetConfiguration.Shipping);
@@ -3793,12 +3763,12 @@ namespace UnrealBuildTool
 
 			GlobalCompileEnvironment.Definitions.Add(String.Format("IS_MONOLITHIC={0}", ShouldCompileMonolithic() ? "1" : "0"));
 
-			GlobalCompileEnvironment.Definitions.Add(String.Format("WITH_ENGINE={0}", Rules.bCompileAgainstEngine || bTestsRequireEngine ? "1" : "0"));
+			GlobalCompileEnvironment.Definitions.Add(String.Format("WITH_ENGINE={0}", Rules.bCompileAgainstEngine ? "1" : "0"));
 			GlobalCompileEnvironment.Definitions.Add(String.Format("WITH_UNREAL_DEVELOPER_TOOLS={0}", Rules.bBuildDeveloperTools ? "1" : "0"));
 			GlobalCompileEnvironment.Definitions.Add(String.Format("WITH_UNREAL_TARGET_DEVELOPER_TOOLS={0}", Rules.bBuildTargetDeveloperTools ? "1" : "0"));
 
 			// Set a macro to control whether to initialize ApplicationCore. Command line utilities should not generally need this.
-			if (Rules.bCompileAgainstApplicationCore || bTestsRequireApplicationCore)
+			if (Rules.bCompileAgainstApplicationCore)
 			{
 				GlobalCompileEnvironment.Definitions.Add("WITH_APPLICATION_CORE=1");
 			}
@@ -3807,7 +3777,7 @@ namespace UnrealBuildTool
 				GlobalCompileEnvironment.Definitions.Add("WITH_APPLICATION_CORE=0");
 			}
 
-			if (Rules.bCompileAgainstCoreUObject || bTestsRequireCoreUObject)
+			if (Rules.bCompileAgainstCoreUObject)
 			{
 				GlobalCompileEnvironment.Definitions.Add("WITH_COREUOBJECT=1");
 			}
@@ -3906,7 +3876,7 @@ namespace UnrealBuildTool
 				GlobalCompileEnvironment.Definitions.Add("USE_ESTIMATED_UTCNOW=0");
 			}
 
-			if ((Rules.bCompileAgainstEditor && (Rules.Type == TargetType.Editor || Rules.Type == TargetType.Program)) || bTestsRequireEditor)
+			if (Rules.bCompileAgainstEditor && (Rules.Type == TargetType.Editor || Rules.Type == TargetType.Program))
 			{
 				GlobalCompileEnvironment.Definitions.Add("WITH_EDITOR=1");
 				GlobalCompileEnvironment.Definitions.Add("WITH_IOSTORE_IN_EDITOR=1");
@@ -3931,8 +3901,6 @@ namespace UnrealBuildTool
 				GlobalCompileEnvironment.Definitions.Add("WITH_SERVER_CODE=0");
 			}
 
-			GlobalCompileEnvironment.Definitions.Add(String.Format("UE_FNAME_OUTLINE_NUMBER={0}", Rules.bFNameOutlineNumber? 1 : 0));
-
 			// Set the defines for Push Model
 			if (Rules.bWithPushModel)
 			{
@@ -3954,13 +3922,13 @@ namespace UnrealBuildTool
 			}
 
 			// Set the define for enabling live coding
-			if (Rules.bWithLiveCoding)
+			if(Rules.bWithLiveCoding)
 			{
 				GlobalCompileEnvironment.Definitions.Add("WITH_LIVE_CODING=1");
 				if (Rules.LinkType == TargetLinkType.Monolithic)
 				{
 					GlobalCompileEnvironment.Definitions.Add(String.Format("UE_LIVE_CODING_ENGINE_DIR=\"{0}\"", Unreal.EngineDirectory.FullName.Replace("\\", "\\\\")));
-					if (ProjectFile != null)
+					if(ProjectFile != null)
 					{
 						GlobalCompileEnvironment.Definitions.Add(String.Format("UE_LIVE_CODING_PROJECT=\"{0}\"", ProjectFile.FullName.Replace("\\", "\\\\")));
 					}
@@ -3984,7 +3952,6 @@ namespace UnrealBuildTool
 			// Whether C++20 coroutines are enabled
 			if (Rules.bEnableCppCoroutinesForEvaluation)
 			{
-				Log.TraceInformationOnce($"NOTE: C++ coroutine support is considered experimental and should be used for evaluation purposes only ({nameof(TargetRules.bEnableCppCoroutinesForEvaluation)})");
 				GlobalCompileEnvironment.Definitions.Add("WITH_CPP_COROUTINES=1");
 				GlobalCompileEnvironment.bEnableCoroutines = true;
 			}
@@ -4083,7 +4050,7 @@ namespace UnrealBuildTool
 			}
 
 			// Set whether the module requires an IMPLEMENT_MODULE macro
-			if (!RulesObject.bRequiresImplementModule.HasValue)
+			if(!RulesObject.bRequiresImplementModule.HasValue)
 			{
 				RulesObject.bRequiresImplementModule = (RulesObject.Type == ModuleRules.ModuleType.CPlusPlus && RulesObject.Name != Rules.LaunchModuleName);
 			}
@@ -4115,26 +4082,26 @@ namespace UnrealBuildTool
 				}
 
 				// Make sure that engine modules use shared PCHs or have an explicit private PCH
-				if (RulesObject.PCHUsage == ModuleRules.PCHUsageMode.NoSharedPCHs && RulesObject.PrivatePCHHeaderFile == null)
+				if(RulesObject.PCHUsage == ModuleRules.PCHUsageMode.NoSharedPCHs && RulesObject.PrivatePCHHeaderFile == null)
 				{
-					if (ProjectFile == null || !RulesObject.File.IsUnderDirectory(ProjectFile.Directory))
+					if(ProjectFile == null || !RulesObject.File.IsUnderDirectory(ProjectFile.Directory))
 					{
 						Log.TraceWarning("{0} module has shared PCHs disabled, but does not have a private PCH set", ModuleName);
 					}
 				}
 
 				// If we can't use a shared PCH, check there's a private PCH set
-				if (RulesObject.PCHUsage != ModuleRules.PCHUsageMode.NoPCHs && RulesObject.PCHUsage != ModuleRules.PCHUsageMode.UseExplicitOrSharedPCHs && RulesObject.PrivatePCHHeaderFile == null)
+				if(RulesObject.PCHUsage != ModuleRules.PCHUsageMode.NoPCHs && RulesObject.PCHUsage != ModuleRules.PCHUsageMode.UseExplicitOrSharedPCHs && RulesObject.PrivatePCHHeaderFile == null)
 				{
 					// Try to figure out the legacy PCH file
 					FileReference CppFile = DirectoryReference.EnumerateFiles(RulesObject.Directory, "*.cpp", SearchOption.AllDirectories).FirstOrDefault();
-					if (CppFile != null)
+					if(CppFile != null)
 					{
 						string? IncludeFile = MetadataCache.GetFirstInclude(FileItem.GetItemByFileReference(CppFile));
-						if (IncludeFile != null)
+						if(IncludeFile != null)
 						{
 							FileReference PchIncludeFile = DirectoryReference.EnumerateFiles(RulesObject.Directory, Path.GetFileName(IncludeFile), SearchOption.AllDirectories).FirstOrDefault();
-							if (PchIncludeFile != null)
+							if(PchIncludeFile != null)
 							{
 								RulesObject.PrivatePCHHeaderFile = PchIncludeFile.MakeRelativeTo(RulesObject.Directory).Replace(Path.DirectorySeparatorChar, '/');
 							}
@@ -4142,7 +4109,7 @@ namespace UnrealBuildTool
 					}
 
 					// Print a suggestion for which file to include
-					if (RulesObject.PrivatePCHHeaderFile == null)
+					if(RulesObject.PrivatePCHHeaderFile == null)
 					{
 						Log.TraceWarningOnce(RulesObject.File, "Modules must specify an explicit precompiled header (eg. PrivatePCHHeaderFile = \"Private/{0}PrivatePCH.h\") from UE 4.21 onwards.", ModuleName);
 					}
@@ -4165,32 +4132,12 @@ namespace UnrealBuildTool
 		/// <param name="Paths">List of paths to process</param>
 		private static void RemoveTrailingSlashes(List<string> Paths)
 		{
-			for (int Idx = 0; Idx < Paths.Count; Idx++)
+			for(int Idx = 0; Idx < Paths.Count; Idx++)
 			{
 				Paths[Idx] = Paths[Idx].TrimEnd('\\');
 			}
 		}
 
-		/// <summary>
-		/// Determines the relative base directory path from the specified Binary directory to the Engine base directory if possible.
-		/// </summary>
-		/// <param name="BinaryDirectory">Directory that contains the platform subdirectoies for each UBT built binary</param>
-		/// <param name="TargetPlatform">Platform that the binary is for</param>
-		/// <returns>A string specifying the relative path from the binary directory to the engine base dir, or null if this isn't contained within the engine hierarhcy</returns>
-		public static string? GetRelativeBaseDir(DirectoryReference BinaryDirectory, UnrealTargetPlatform TargetPlatform)
-		{
-			if (BinaryDirectory.IsUnderDirectory(Unreal.EngineDirectory))
-			{
-				DirectoryReference BaseDir = DirectoryReference.Combine(Unreal.EngineDirectory, "Binaries", TargetPlatform.ToString());
-				if (BaseDir != BinaryDirectory)
-				{
-					return BaseDir.MakeRelativeTo(BinaryDirectory).Replace(Path.DirectorySeparatorChar, '/');
-				}
-			}
-
-			return null;
-		}
-
 		/// <summary>
 		/// Finds a module given its name.  Throws an exception if the module couldn't be found.
 		/// </summary>
@@ -4208,7 +4155,11 @@ namespace UnrealBuildTool
 				DirectoryReference ModuleDirectory = RulesObject.File.Directory;
 
 				// If we're building a test executable, include test harness dependency
-				PrepareModuleForTests(IsTestModule, RulesObject);
+				if (!IsTestModule && ReferenceChain.Contains(typeof(TestModuleRules).Name) && Directory.Exists(RulesObject.TestsDirectory))
+				{
+					// TODO: If explicit Tests module rules exists, use that one instead
+					RulesObject.PrivateDependencyModuleNames.Add("LowLevelTestsRunner");
+				}
 
 				// Clear the bUsePrecompiled flag if we're compiling a foreign plugin; since it's treated like an engine module, it will default to true in an installed build.
 				if (RulesObject.Plugin != null && RulesObject.Plugin.File == ForeignPlugin)
@@ -4321,68 +4272,6 @@ namespace UnrealBuildTool
 			return Module;
 		}
 
-		/// <summary>
-		/// Prepares a module for building a low level tests executable.
-		/// If the target includes all tests OR we're building a module as part of a test module chain, then we include the LowLevelTestsRunner dependency.
-		/// We also keep track of any Editor and Engine dependencies and add appropriate global definitions at compilation step.
-		/// </summary>
-		/// <param name="IsTestModule">Whether we're building a tests module or not.</param>
-		/// <param name="RulesObject">The module rules object.</param>
-		private void PrepareModuleForTests(bool IsTestModule, ModuleRules RulesObject)
-		{
-			if (!IsTestModule
-				&& Rules.IsTestTarget()
-				|| Rules.bIncludeAllTests)
-			{
-				if (RulesObject.Name != "LowLevelTestsRunner" && Directory.Exists(RulesObject.TestsDirectory))
-				{
-					RulesObject.PrivateIncludePathModuleNames.Add("LowLevelTestsRunner");
-				}
-				else if (RulesObject.Name == "LowLevelTestsRunner")
-				{
-					LowLevelTestsRunnerModule = RulesObject;
-				}
-
-				// If one of these modules is in the dependency graph, we must enable their appropriate global definitions
-				if (RulesObject.Name == "Editor")
-				{
-					bTestsRequireEditor = true;
-				}
-				if (RulesObject.Name == "Engine")
-				{
-					bTestsRequireEngine = true;
-				}
-				if (RulesObject.Name == "ApplicationCore")
-				{
-					bTestsRequireApplicationCore = true;
-				}
-				if (RulesObject.Name == "CoreUObject")
-				{
-					bTestsRequireCoreUObject = true;
-				}
-
-				if (LowLevelTestsRunnerModule != null)
-				{
-					if (bTestsRequireEditor && !LowLevelTestsRunnerModule.PrivateDependencyModuleNames.Contains("Editor"))
-					{
-						LowLevelTestsRunnerModule.PrivateDependencyModuleNames.Add("Editor");
-					}
-					if (bTestsRequireEngine && !LowLevelTestsRunnerModule.PrivateDependencyModuleNames.Contains("Engine"))
-					{
-						LowLevelTestsRunnerModule.PrivateDependencyModuleNames.Add("Engine");
-					}
-					if (bTestsRequireApplicationCore && !LowLevelTestsRunnerModule.PrivateDependencyModuleNames.Contains("ApplicationCore"))
-					{
-						LowLevelTestsRunnerModule.PrivateDependencyModuleNames.Add("ApplicationCore");
-					}
-					if (bTestsRequireCoreUObject && !LowLevelTestsRunnerModule.PrivateDependencyModuleNames.Contains("CoreUObject"))
-					{
-						LowLevelTestsRunnerModule.PrivateDependencyModuleNames.Add("CoreUObject");
-					}
-				}
-			}
-		}
-
 		/// <summary>
 		/// Constructs a new C++ module
 		/// </summary>
@@ -4393,7 +4282,7 @@ namespace UnrealBuildTool
 		public UEBuildModuleCPP FindOrCreateCppModuleByName(string ModuleName, string ReferenceChain, bool IsTestModule = false)
 		{
 			UEBuildModuleCPP? CppModule = FindOrCreateModuleByName(ModuleName, ReferenceChain, IsTestModule) as UEBuildModuleCPP;
-			if (CppModule == null)
+			if(CppModule == null)
 			{
 				throw new BuildException("'{0}' is not a C++ module (referenced via {1})", ModuleName, ReferenceChain);
 			}
@@ -4457,7 +4346,7 @@ namespace UnrealBuildTool
 			List<string> NewPathList = new List<string>();
 			foreach (string Path in PathList)
 			{
-				if (Path.StartsWith("$(", StringComparison.Ordinal))
+				if(Path.StartsWith("$(", StringComparison.Ordinal))
 				{
 					NewPathList.Add(Path);
 				}
diff --git a/Engine/Source/Programs/UnrealBuildTool/Platform/Linux/UEBuildLinux.cs b/Engine/Source/Programs/UnrealBuildTool/Platform/Linux/UEBuildLinux.cs
--- a/Engine/Source/Programs/UnrealBuildTool/Platform/Linux/UEBuildLinux.cs
+++ b/Engine/Source/Programs/UnrealBuildTool/Platform/Linux/UEBuildLinux.cs
@@ -200,27 +200,31 @@ namespace UnrealBuildTool
 
 		public override void ValidateTarget(TargetRules Target)
 		{
-			HashSet<string> Sanitizers = new HashSet<string>();
-			if (Target.LinuxPlatform.bEnableAddressSanitizer)
-			{
-				Sanitizers.Add("A");
-			}
-			if (Target.LinuxPlatform.bEnableThreadSanitizer)
-			{
-				Sanitizers.Add("T");
-			}
-			if (Target.LinuxPlatform.bEnableUndefinedBehaviorSanitizer)
-			{
-				Sanitizers.Add("UB");
-			}
-			if (Target.LinuxPlatform.bEnableMemorySanitizer)
+			if (!Target.IsNameOverriden())
 			{
-				Sanitizers.Add("M");
-			}
+				string? SanitizerSuffix = null;
 
-			if (Sanitizers.Count > 0)
-			{
-				Target.AddNameSuffix($"{string.Join(string.Empty, Sanitizers)}San");
+				if (Target.LinuxPlatform.bEnableAddressSanitizer)
+				{
+					SanitizerSuffix = "ASan";
+				}
+				else if (Target.LinuxPlatform.bEnableThreadSanitizer)
+				{
+					SanitizerSuffix = "TSan";
+				}
+				else if (Target.LinuxPlatform.bEnableUndefinedBehaviorSanitizer)
+				{
+					SanitizerSuffix = "UBSan";
+				}
+				else if (Target.LinuxPlatform.bEnableMemorySanitizer)
+				{
+					SanitizerSuffix = "MSan";
+				}
+
+				if (!String.IsNullOrEmpty(SanitizerSuffix))
+				{
+					Target.Name = Target.Name + "-" + SanitizerSuffix;
+				}
 			}
 
 			if (Target.bAllowLTCG && Target.LinkType != TargetLinkType.Monolithic)
