UnrealBuildTool: Add command line flag -PreprocessDepends which will run the compiler preprocesor to generate the dependency list for a compile in a separate action, so the include list can be consumed before executing the compile action. Currently Mac compiles only.

#rnx

diff --git a/Engine/Source/Programs/UnrealBuildTool/Platform/Mac/MacToolChain.cs b/Engine/Source/Programs/UnrealBuildTool/Platform/Mac/MacToolChain.cs
--- a/Engine/Source/Programs/UnrealBuildTool/Platform/Mac/MacToolChain.cs
+++ b/Engine/Source/Programs/UnrealBuildTool/Platform/Mac/MacToolChain.cs
@@ -149,6 +149,8 @@ namespace UnrealBuildTool
 
 		private static List<FileItem> BundleDependencies = new List<FileItem>();
 
+		private bool bPreprocessDepends = false;
+
 		private static void SetupXcodePaths(bool bVerbose)
 		{
 		}
@@ -157,6 +159,8 @@ namespace UnrealBuildTool
 		{
 			base.SetUpGlobalEnvironment(Target);
 
+			bPreprocessDepends = Target.MacPlatform.bPreprocessDepends;
+
 			// validation, because sometimes this is called from a shell script and quoting messes up		
 			if (!Target.Architecture.All(C => char.IsLetterOrDigit(C) || C == '_' || C == '+'))
 			{
@@ -564,7 +568,7 @@ namespace UnrealBuildTool
 				}
 
 				// Generate the included header dependency list
-				if(CompileEnvironment.bGenerateDependenciesFile)
+				if(!bPreprocessDepends && CompileEnvironment.bGenerateDependenciesFile)
 				{
 					FileItem DependencyListFile = FileItem.GetItemByFileReference(FileReference.Combine(OutputDir, Path.GetFileName(SourceFile.AbsolutePath) + ".d"));
 					FileArguments += string.Format(" -MD -MF\"{0}\"", DependencyListFile.AbsolutePath.Replace('\\', '/'));
@@ -609,7 +613,6 @@ namespace UnrealBuildTool
 					string ArchPath = "/usr/bin/arch";
 					CompileAction.CommandPath = new FileReference(ArchPath);
 					CompileAction.CommandArguments = string.Format("-{0} {1} {2}", MacExports.HostArchitecture, CompilerPath, AllArgs);
-
 				}
 				else
 				{
@@ -626,6 +629,53 @@ namespace UnrealBuildTool
 				CompileAction.bCanExecuteRemotely = Extension != ".C";
 				CompileAction.bShouldOutputStatusDescription = true;
 				CompileAction.CommandVersion = GetFullClangVersion();
+
+				if (bPreprocessDepends && CompileEnvironment.bGenerateDependenciesFile)
+				{
+					Action PrepassAction = Graph.CreateAction(ActionType.Compile);
+					PrepassAction.PrerequisiteItems.AddRange(CompileAction.PrerequisiteItems);
+					PrepassAction.CommandDescription = "Preprocess Depends";
+					PrepassAction.StatusDescription = CompileAction.StatusDescription;
+					PrepassAction.bIsGCCCompiler = true;
+					PrepassAction.bCanExecuteRemotely = false;
+					PrepassAction.bShouldOutputStatusDescription = true;
+					PrepassAction.CommandVersion = CompileAction.CommandVersion;
+					PrepassAction.WorkingDirectory = CompileAction.WorkingDirectory;
+
+					string PreprocessArguments = Arguments.ToString();
+					string PreprocessFileArguments = FileArguments;
+					PreprocessArguments = PreprocessArguments.Replace(" -c ", " ");
+
+					FileItem DependencyListFile = FileItem.GetItemByFileReference(FileReference.Combine(OutputDir, Path.GetFileName(SourceFile.AbsolutePath) + ".d"));
+					PreprocessFileArguments += string.Format(" -M -MF\"{0}\"", DependencyListFile.AbsolutePath.Replace('\\', '/'));
+					PrepassAction.DependencyListFile = DependencyListFile;
+					PrepassAction.ProducedItems.Add(DependencyListFile);
+
+					PreprocessFileArguments = PreprocessFileArguments.Replace(" -ftime-trace", string.Empty);
+					PreprocessFileArguments = PreprocessFileArguments.Replace(string.Format(" -o \"{0}\"", CompileAction.ProducedItems.First().AbsolutePath), String.Empty);
+
+					PrepassAction.DeleteItems.AddRange(PrepassAction.ProducedItems);
+
+					string AllPreprocessArgs = PreprocessArguments + PreprocessFileArguments + EscapedAdditionalArgs;
+
+					PrepassAction.WorkingDirectory = GetMacDevSrcRoot();
+
+					if (MacExports.IsRunningUnderRosetta)
+					{
+						string ArchPath = "/usr/bin/arch";
+						PrepassAction.CommandPath = new FileReference(ArchPath);
+						PrepassAction.CommandArguments = string.Format("-{0} {1} {2}", MacExports.HostArchitecture, CompilerPath, AllPreprocessArgs);
+					}
+					else
+					{
+						PrepassAction.CommandPath = new FileReference(CompilerPath);
+						PrepassAction.CommandArguments = AllPreprocessArgs;
+					}
+
+					CompileAction.DependencyListFile = DependencyListFile;
+					CompileAction.PrerequisiteItems.Add(DependencyListFile);
+				}
+
 			}
 			return Result;
 		}
diff --git a/Engine/Source/Programs/UnrealBuildTool/Platform/Mac/UEBuildMac.cs b/Engine/Source/Programs/UnrealBuildTool/Platform/Mac/UEBuildMac.cs
--- a/Engine/Source/Programs/UnrealBuildTool/Platform/Mac/UEBuildMac.cs
+++ b/Engine/Source/Programs/UnrealBuildTool/Platform/Mac/UEBuildMac.cs
@@ -40,6 +40,13 @@ namespace UnrealBuildTool
 		[CommandLine("-NoDSYM", Value = "false")]
 		[XmlConfigFile(Category = "BuildConfiguration", Name = "bUseDSYMFiles")]
 		public bool bUseDSYMFiles = true;
+
+		/// <summary>
+		/// Generate dependency files by preprocessing. This is only recommended when distributing builds as it adds additional overhead.
+		/// </summary>
+		[CommandLine("-PreprocessDepends")]
+		[XmlConfigFile(Category = "BuildConfiguration", Name = "bPreprocessDepends")]
+		public bool bPreprocessDepends = false;
 	}
 
 	/// <summary>
@@ -80,9 +87,14 @@ namespace UnrealBuildTool
 		public bool bEnableUndefinedBehaviorSanitizer
 		{
 			get { return Inner.bEnableUndefinedBehaviorSanitizer; }
-		}		
+		}
+
+		public bool bPreprocessDepends
+		{
+			get { return Inner.bPreprocessDepends; }
+		}
 
-		#pragma warning restore CS1591
+#pragma warning restore CS1591
 		#endregion
 	}
 
