Build - Adding additional logging to explain what the enabled and disabled plugins are so there's something to look for.

diff --git a/Engine/Source/Programs/UnrealBuildTool/Configuration/UEBuildTarget.cs b/Engine/Source/Programs/UnrealBuildTool/Configuration/UEBuildTarget.cs
--- a/Engine/Source/Programs/UnrealBuildTool/Configuration/UEBuildTarget.cs
+++ b/Engine/Source/Programs/UnrealBuildTool/Configuration/UEBuildTarget.cs
@@ -941,7 +941,10 @@ namespace UnrealBuildTool
 			bool bUsesTargetReceiptToEnablePlugins = (ThisRules.Type == TargetType.Editor && ThisRules.LinkType != TargetLinkType.Monolithic);
 			if (!bUsesTargetReceiptToEnablePlugins && (ThisRules.EnablePlugins.Count > 0 || ThisRules.DisablePlugins.Count > 0))
 			{
-				throw new BuildException("Explicitly enabling and disabling plugins for a target is only supported when using a unique build environment (eg. for monolithic game targets).");
+				throw new BuildException(string.Format("Explicitly enabling and disabling plugins for a target is only supported when using a unique build environment (eg. for monolithic game targets). EnabledPlugins={0}, DisabledPlugins={1}",
+					string.Join(", ", ThisRules.EnablePlugins),
+					string.Join(", ", ThisRules.DisablePlugins)
+				));
 			}
 		}
 
