UnrealBuildTool: When using bBuildAllModules, ensure dependencies are correctly added for plugins

#rnx

diff --git a/Engine/Source/Programs/UnrealBuildTool/Configuration/UEBuildTarget.cs b/Engine/Source/Programs/UnrealBuildTool/Configuration/UEBuildTarget.cs
--- a/Engine/Source/Programs/UnrealBuildTool/Configuration/UEBuildTarget.cs
+++ b/Engine/Source/Programs/UnrealBuildTool/Configuration/UEBuildTarget.cs
@@ -3040,19 +3040,32 @@ namespace UnrealBuildTool
 				}
 			}
 
+			// Used to capture modules that were referenced while creating all the precompiled modules
+
 			// Now create all the precompiled modules, making sure they don't reference anything that's not in the precompiled set
-			HashSet<UEBuildModuleCPP> ValidModules = new HashSet<UEBuildModuleCPP>();
+
+			// Gather the set of all modules traversed while processing ValidModuleNames, to ensure we have any extra modules found during RecursivelyCreateModules()
+			HashSet<UEBuildModule> AllModules = new HashSet<UEBuildModule>();
+
+			const string PrecompileReferenceChain = "allmodules option";
 			foreach (string ModuleName in ValidModuleNames)
 			{
-				const string PrecompileReferenceChain = "allmodules option";
-				UEBuildModuleCPP Module = (UEBuildModuleCPP)FindOrCreateModuleByName(ModuleName, PrecompileReferenceChain);
-				Module.RecursivelyCreateModules(FindOrCreateModuleByName, PrecompileReferenceChain);
-				if (!Module.bDependsOnVerse || Rules.bUseVerse)
-				{
-					ValidModules.Add(Module);
-				}
+				UEBuildModule Module = FindOrCreateModuleByName(ModuleName, PrecompileReferenceChain);
+				AllModules.Add(Module);
+				Module.RecursivelyCreateModules(
+					(string ModuleName, string ReferenceChain) => 
+					{ 
+						UEBuildModule FoundModule = FindOrCreateModuleByName(ModuleName, ReferenceChain);
+						AllModules.Add(FoundModule);
+						return FoundModule; 
+					}, 
+					PrecompileReferenceChain);
 			}
 
+			// Exclude additional modules that were added only for include-path-only purposes, and those that will have their Verse dependency satisfied
+			HashSet<UEBuildModuleCPP> ValidModules = new HashSet<UEBuildModuleCPP>(
+				AllModules.OfType<UEBuildModuleCPP>() .Where(x => x.PrivateIncludePathModules != null && (!x.bDependsOnVerse || Rules.bUseVerse)));
+
 			// Make sure precompiled modules don't reference any non-precompiled modules
 			foreach (UEBuildModuleCPP ValidModule in ValidModules)
 			{
@@ -3060,17 +3073,13 @@ namespace UnrealBuildTool
 				{
 					if (!ValidModules.Contains(ReferencedModule))
 					{
-						Log.TraceWarning("Module '{0}' is not usable without module '{1}', which is not valid for this target.", ValidModule.Name, ReferencedModule.Name);
+						Log.TraceError("Module '{0}' is not usable without module '{1}', which is not valid for this target.", ValidModule.Name, ReferencedModule.Name);
 					}
 				}
-			}
 
-			// Make sure every module is built
-			foreach (UEBuildModuleCPP Module in ValidModules)
-			{
-				if (Module.Binary == null)
+				if (ValidModule.Binary == null)
 				{
-					AddModuleToBinary(Module);
+					AddModuleToBinary(ValidModule);
 				}
 			}
 		}
