UnrealBuildTool: Fix breakage due to changed function signature.

#trivial
#rnx

diff --git a/Engine/Source/Programs/UnrealBuildTool/Configuration/UEBuildTarget.cs b/Engine/Source/Programs/UnrealBuildTool/Configuration/UEBuildTarget.cs
--- a/Engine/Source/Programs/UnrealBuildTool/Configuration/UEBuildTarget.cs
+++ b/Engine/Source/Programs/UnrealBuildTool/Configuration/UEBuildTarget.cs
@@ -3053,9 +3053,9 @@ namespace UnrealBuildTool
 				UEBuildModule Module = FindOrCreateModuleByName(ModuleName, PrecompileReferenceChain);
 				AllModules.Add(Module);
 				Module.RecursivelyCreateModules(
-					(string ModuleName, string ReferenceChain) => 
+					(string ModuleName, string ReferenceChain, bool IsTestModule) => 
 					{ 
-						UEBuildModule FoundModule = FindOrCreateModuleByName(ModuleName, ReferenceChain);
+						UEBuildModule FoundModule = FindOrCreateModuleByName(ModuleName, ReferenceChain, IsTestModule);
 						AllModules.Add(FoundModule);
 						return FoundModule; 
 					}, 
