fix exceptions when running with bPGOProfile=true on Windows targets with no project.

#rnx

diff --git a/Engine/Source/Programs/UnrealBuildTool/Platform/Windows/UEBuildWindows.cs b/Engine/Source/Programs/UnrealBuildTool/Platform/Windows/UEBuildWindows.cs
--- a/Engine/Source/Programs/UnrealBuildTool/Platform/Windows/UEBuildWindows.cs
+++ b/Engine/Source/Programs/UnrealBuildTool/Platform/Windows/UEBuildWindows.cs
@@ -2244,13 +2244,13 @@ namespace UnrealBuildTool
 
 			LinkEnvironment.ModuleDefinitionFile = Target.WindowsPlatform.ModuleDefinitionFile;
 
-			if (Target.bPGOOptimize || Target.bPGOProfile)
+			if ((Target.bPGOOptimize || Target.bPGOProfile) && Target.ProjectFile != null)
 			{
 				// LTCG is required for PGO
 				//CompileEnvironment.bAllowLTCG = true;
 				//LinkEnvironment.bAllowLTCG = true;
 
-				CompileEnvironment.PGODirectory = Path.Combine(DirectoryReference.FromFile(Target.ProjectFile!).FullName, "Platforms", "Windows", "Build", "PGO");
+				CompileEnvironment.PGODirectory = Path.Combine(DirectoryReference.FromFile(Target.ProjectFile).FullName, "Platforms", "Windows", "Build", "PGO");
 				CompileEnvironment.PGOFilenamePrefix = string.Format("{0}-{1}-{2}", Target.Name, Target.Platform, Target.Configuration);
 
 				LinkEnvironment.PGODirectory = CompileEnvironment.PGODirectory;
