Dont try to override the Editor with ASan until things are done at a much higher level. This causes issues with EngineTest + ASan

diff --git a/Engine/Source/Programs/UnrealBuildTool/Platform/Linux/UEBuildLinux.cs b/Engine/Source/Programs/UnrealBuildTool/Platform/Linux/UEBuildLinux.cs
--- a/Engine/Source/Programs/UnrealBuildTool/Platform/Linux/UEBuildLinux.cs
+++ b/Engine/Source/Programs/UnrealBuildTool/Platform/Linux/UEBuildLinux.cs
@@ -200,7 +200,10 @@ namespace UnrealBuildTool
 
 		public override void ValidateTarget(TargetRules Target)
 		{
-			if (!Target.IsNameOverriden())
+			// Editor target types get overwritten in UEBuildTarget.cs so lets avoid adding this here
+			// as it will create an UnrealEditor-ASan.target while the binary is UnrealEditor. These need to be promoted to
+			// higher level concepts vs this hacky solution
+			if (!Target.IsNameOverriden() && Target.Type != TargetType.Editor)
 			{
 				string? SanitizerSuffix = null;
 
