Include .ipp files when adding files from include paths

#rnx

diff --git a/Engine/Source/Programs/UnrealBuildTool/Configuration/UEBuildTarget.cs b/Engine/Source/Programs/UnrealBuildTool/Configuration/UEBuildTarget.cs
--- a/Engine/Source/Programs/UnrealBuildTool/Configuration/UEBuildTarget.cs
+++ b/Engine/Source/Programs/UnrealBuildTool/Configuration/UEBuildTarget.cs
@@ -1390,7 +1390,7 @@ namespace UnrealBuildTool
 						foreach (string IncludeFileName in Directory.EnumerateFiles(IncludePath, "*", SearchOption.AllDirectories))
 						{
 							string Extension = Path.GetExtension(IncludeFileName).ToLower();
-							if (Extension == ".h" || Extension == ".inl" || Extension == ".hpp")
+							if (Extension == ".h" || Extension == ".inl" || Extension == ".hpp" || Extension == ".ipp")
 							{
 								Files.Add(new FileReference(IncludeFileName));
 							}
