New bTCHARIsUTF8 variable added to .Target.cs file to make a target use UTF8CHARs as its TCHARs.

diff --git a/Engine/Source/Programs/UnrealBuildTool/Configuration/TargetRules.cs b/Engine/Source/Programs/UnrealBuildTool/Configuration/TargetRules.cs
--- a/Engine/Source/Programs/UnrealBuildTool/Configuration/TargetRules.cs
+++ b/Engine/Source/Programs/UnrealBuildTool/Configuration/TargetRules.cs
@@ -858,6 +858,12 @@ namespace UnrealBuildTool
 		[RequiresUniqueBuildEnvironment]
 		public bool bUseChecksInShipping = false;
 
+		/// <summary>
+		/// Whether to turn on UTF-8 mode, mapping TCHAR to UTF8CHAR.
+		/// </summary>
+		[RequiresUniqueBuildEnvironment]
+		public bool bTCHARIsUTF8 = false;
+
 		/// <summary>
 		/// Whether to use the EstimatedUtcNow or PlatformUtcNow.  EstimatedUtcNow is appropriate in
 		/// cases where PlatformUtcNow can be slow.
@@ -2471,6 +2477,11 @@ namespace UnrealBuildTool
 			get { return Inner.bUseChecksInShipping; }
 		}
 
+		public bool bTCHARIsUTF8
+		{
+			get { return Inner.bTCHARIsUTF8; }
+		}
+
 		public bool bUseEstimatedUtcNow
 		{
 			get { return Inner.bUseEstimatedUtcNow; }
diff --git a/Engine/Source/Programs/UnrealBuildTool/Configuration/UEBuildTarget.cs b/Engine/Source/Programs/UnrealBuildTool/Configuration/UEBuildTarget.cs
--- a/Engine/Source/Programs/UnrealBuildTool/Configuration/UEBuildTarget.cs
+++ b/Engine/Source/Programs/UnrealBuildTool/Configuration/UEBuildTarget.cs
@@ -3905,6 +3905,15 @@ namespace UnrealBuildTool
 				GlobalCompileEnvironment.Definitions.Add("USE_CHECKS_IN_SHIPPING=0");
 			}
 
+			if (Rules.bTCHARIsUTF8)
+			{
+				GlobalCompileEnvironment.Definitions.Add("USE_UTF8_TCHARS=1");
+			}
+			else
+			{
+				GlobalCompileEnvironment.Definitions.Add("USE_UTF8_TCHARS=0");
+			}
+
 			if (Rules.bUseEstimatedUtcNow)
 			{
 				GlobalCompileEnvironment.Definitions.Add("USE_ESTIMATED_UTCNOW=1");
diff --git a/Engine/Source/Runtime/Core/Public/HAL/Platform.h b/Engine/Source/Runtime/Core/Public/HAL/Platform.h
--- a/Engine/Source/Runtime/Core/Public/HAL/Platform.h
+++ b/Engine/Source/Runtime/Core/Public/HAL/Platform.h
@@ -257,7 +257,7 @@
 #endif
 #define PLATFORM_WIDECHAR_IS_CHAR16 PLATFORM_TCHAR_IS_CHAR16
 #ifndef PLATFORM_TCHAR_IS_UTF8CHAR
-	#define PLATFORM_TCHAR_IS_UTF8CHAR			0
+	#define PLATFORM_TCHAR_IS_UTF8CHAR			USE_UTF8_TCHARS
 #endif
 #ifndef PLATFORM_UCS2CHAR_IS_UTF16CHAR
 	// Currently true, but we don't want it to be true
diff --git a/Engine/Source/Runtime/Core/Public/Windows/WindowsPlatform.h b/Engine/Source/Runtime/Core/Public/Windows/WindowsPlatform.h
--- a/Engine/Source/Runtime/Core/Public/Windows/WindowsPlatform.h
+++ b/Engine/Source/Runtime/Core/Public/Windows/WindowsPlatform.h
@@ -14,8 +14,6 @@
 	#error "Windows Vista and earlier are no longer supported"
 #endif
 
-#define PLATFORM_TCHAR_IS_UTF8CHAR				0
-
 /**
 * Windows specific types
 **/
@@ -29,7 +27,7 @@ struct FWindowsPlatformTypes : public FGenericPlatformTypes
 	typedef long				SSIZE_T;
 #endif
 
-#if PLATFORM_TCHAR_IS_UTF8CHAR
+#if USE_UTF8_TCHARS
 	typedef UTF8CHAR TCHAR;
 #endif
 };
