When packaging an iOS app on BigSur/Catalina, it is unable to be installed on an iOS15 device. This fix now sets the "--generate-entitlement-der" flag for 'codesign'.

diff --git a/Engine/Source/Programs/UnrealBuildTool/ProjectFiles/Xcode/XcodeProject.cs b/Engine/Source/Programs/UnrealBuildTool/ProjectFiles/Xcode/XcodeProject.cs
--- a/Engine/Source/Programs/UnrealBuildTool/ProjectFiles/Xcode/XcodeProject.cs
+++ b/Engine/Source/Programs/UnrealBuildTool/ProjectFiles/Xcode/XcodeProject.cs
@@ -796,7 +796,7 @@ namespace UnrealBuildTool
 				"\\tfor FRAMEWORK in ${FRAMEWORK_DIR}/*.framework; do\\n" +
 					"\\t\\t[ -d \\\"${FRAMEWORK}\\\" ] || continue\\n" +
 					"\\t\\techo Codesigning ${FRAMEWORK}\\n" +
-					"\\t\\tcodesign --force --sign ${EXPANDED_CODE_SIGN_IDENTITY} --verbose --preserve-metadata=identifier,entitlements,flags --timestamp=none \\\"${FRAMEWORK}\\\"\\n" +
+					"\\t\\tcodesign --force --sign ${EXPANDED_CODE_SIGN_IDENTITY} --generate-entitlement-der --verbose --preserve-metadata=identifier,entitlements,flags --timestamp=none \\\"${FRAMEWORK}\\\"\\n" +
 				"\\tdone\\n" +
 				"fi\\n";
 
@@ -1394,6 +1394,10 @@ namespace UnrealBuildTool
 					}
 				}
 			}
+
+			// #jira UE-143619: Pre Monterey macOS requires this option for a packaged app to run on iOS15 due to new code signature format. Could be removed once Monterey is minimum OS.
+			Content.Append("\t\t\t\tOTHER_CODE_SIGN_FLAGS = \"--generate-entitlement-der\";" + ProjectFileGenerator.NewLine);
+
 			Content.Append("\t\t\t\tMACOSX_DEPLOYMENT_TARGET = " + MacToolChain.Settings.MacOSVersion + ";" + ProjectFileGenerator.NewLine);
 			Content.Append("\t\t\t\tINFOPLIST_OUTPUT_FORMAT = xml;" + ProjectFileGenerator.NewLine);
 			Content.Append("\t\t\t\tCOMBINE_HIDPI_IMAGES = YES;" + ProjectFileGenerator.NewLine);
