Fix being created UnrealEditor-ASan.target when building Editor with ASan enabled.

This is breaking EngineTest automated testing due how we looked up exe path based on the *.target file name.

*Note* This is pretty hacky, but since non-monolithic/editor builds already collide with the non-asan build dir this is ending up no worse.
Checked with Game/Monolithics are still working and correctly being renamed UnrealGame-ASan + UnrealGame-ASan.target

diff --git a/Engine/Source/Programs/UnrealBuildTool/Platform/Linux/UEBuildLinux.cs b/Engine/Source/Programs/UnrealBuildTool/Platform/Linux/UEBuildLinux.cs
--- a/Engine/Source/Programs/UnrealBuildTool/Platform/Linux/UEBuildLinux.cs
+++ b/Engine/Source/Programs/UnrealBuildTool/Platform/Linux/UEBuildLinux.cs
@@ -200,10 +200,12 @@ namespace UnrealBuildTool
 
 		public override void ValidateTarget(TargetRules Target)
 		{
-			// Editor target types get overwritten in UEBuildTarget.cs so lets avoid adding this here
-			// as it will create an UnrealEditor-ASan.target while the binary is UnrealEditor. These need to be promoted to
-			// higher level concepts vs this hacky solution
-			if (!Target.IsNameOverriden() && Target.Type != TargetType.Editor)
+			// Editor target types get overwritten in UEBuildTarget.cs so lets avoid adding this here. ResetTarget is called with
+			// default settings for TargetRules meanings Type == Game once then Type == Editor a 2nd time when building the Editor.
+			// BuildVersion string is not set at this point so we can avoid setting a Sanitizer suffix if this is the first ResetTarget
+			// with an unset TargetType. This avoids creating UnrealEditor-ASan.target while the binary is UnrealEditor.
+			// These need to be promoted to higher level concepts vs this hacky solution
+			if (!Target.IsNameOverriden() && !String.IsNullOrEmpty(Target.BuildVersion) && Target.Type != TargetType.Editor)
 			{
 				string? SanitizerSuffix = null;
 
