- Added PDBPAGESIZE option for monolithic editor configuration, and check the toolchain version to catch unsupported versions

diff --git a/Engine/Source/Programs/UnrealBuildTool/Configuration/TargetRules.cs b/Engine/Source/Programs/UnrealBuildTool/Configuration/TargetRules.cs
--- a/Engine/Source/Programs/UnrealBuildTool/Configuration/TargetRules.cs
+++ b/Engine/Source/Programs/UnrealBuildTool/Configuration/TargetRules.cs
@@ -1801,9 +1801,6 @@ namespace UnrealBuildTool
 		protected void SetDefaultsForCookedEditor(bool bIsCookedCooker, bool bIsForExternalUse)
 		{
 			LinkType = TargetLinkType.Monolithic;
-			// enable > 4gb pdb file support, which needs an up-to-date compiler and can make pdbs a little bigger, so we 
-			// only enable it when needed
-			WindowsPlatform.AdditionalLinkerOptions = "/PDBPAGESIZE:8192";
 
 			if (!bIsCookedCooker)
 			{
diff --git a/Engine/Source/Programs/UnrealBuildTool/Platform/Windows/VCToolChain.cs b/Engine/Source/Programs/UnrealBuildTool/Platform/Windows/VCToolChain.cs
--- a/Engine/Source/Programs/UnrealBuildTool/Platform/Windows/VCToolChain.cs
+++ b/Engine/Source/Programs/UnrealBuildTool/Platform/Windows/VCToolChain.cs
@@ -1024,6 +1024,18 @@ namespace UnrealBuildTool
 				}
 			}
 
+			// for monolithic editor builds, add the PDBPAGESIZE option, (VS 16.11, VC toolchain 14.29.30133), but the pdb will be too large without this
+			// some monolithic game builds could be too large as well, but they can be added in a .Target.cs with:
+			//   			WindowsPlatform.AdditionalLinkerOptions = "/PDBPAGESIZE:8192";
+			if (Target.LinkType == TargetLinkType.Monolithic && Target.Type == TargetType.Editor)
+			{
+				if (EnvVars.CompilerVersion < VersionNumber.Parse("14.29.30133"))
+				{
+					throw new BuildException($"Monolithic editors now require VC Toolchain 14.29.30133 (the toolchain for Visual Studio 16.11) ({EnvVars.CompilerVersion} < {VersionNumber.Parse("14.29.30133")})");
+				}
+				Arguments.Add("/PDBPAGESIZE:8192");
+			}
+
 			//
 			//	Shipping & LTCG
 			//
