UE-130910 UE5 Binary fails to package BP projects for HoloLens due to Missing Game Binary
-For hololens we generate a new target file <ProjectName>-HoloLens-<Configuration>Multi.target in the UnrealBuildTool(UBT) HoloLensDeploy.cs PrepTargetForDeployment function when packaging (accumulating information from *x64.target and/or *arm64.target).  However with a blueprint only project that does not enable or disable any plugins (vs default) that contain runtime code UBT is not run at all, so the Multi.target file is not created and packaging fails.
-Fortunately there is already a platform hook to require require that UBT run, so we can simply have HoloLensPlatform::RequiresBuild return true.

diff --git a/Engine/Source/Programs/UnrealBuildTool/Platform/HoloLens/UEBuildHoloLens.cs b/Engine/Source/Programs/UnrealBuildTool/Platform/HoloLens/UEBuildHoloLens.cs
--- a/Engine/Source/Programs/UnrealBuildTool/Platform/HoloLens/UEBuildHoloLens.cs
+++ b/Engine/Source/Programs/UnrealBuildTool/Platform/HoloLens/UEBuildHoloLens.cs
@@ -720,6 +720,16 @@ namespace UnrealBuildTool
 		{
 			return new HoloLensToolChain(Target);
 		}
+
+		/// <summary>
+		/// Check for the build requirement due to platform requirements
+		/// return true if the project requires a build
+		/// </summary>
+		public override bool RequiresBuild(UnrealTargetPlatform Platform, DirectoryReference ProjectDirectoryName)
+		{
+			// Hololens requires UBT run before packaging so that HoloLensDeploy PrepTargetForDeployment will create the Multi.target file which is used in hololens packaging.
+			return true;
+		}
 	}
 
 
