GenerateClangDatabase: Pass C++ Standard Flag correctly to clang-cl

diff --git a/Engine/Source/Programs/UnrealBuildTool/Modes/GenerateClangDatabase.cs b/Engine/Source/Programs/UnrealBuildTool/Modes/GenerateClangDatabase.cs
--- a/Engine/Source/Programs/UnrealBuildTool/Modes/GenerateClangDatabase.cs
+++ b/Engine/Source/Programs/UnrealBuildTool/Modes/GenerateClangDatabase.cs
@@ -66,6 +66,8 @@ namespace UnrealBuildTool
 					// Find the location of the compiler
 					FileReference ClangPath = FindClangCompiler(Target);
 
+					bool IsWindowsClang = ClangPath.GetFileName().Equals("clang-cl.exe", StringComparison.OrdinalIgnoreCase);
+
 					// Convince each module to output its generated code include path
 					foreach (UEBuildBinary Binary in Target.Binaries)
 					{
@@ -98,14 +100,14 @@ namespace UnrealBuildTool
 								switch (ModuleCompileEnvironment.CppStandard)
 								{
 									case CppStandardVersion.Cpp14:
-										CommandBuilder.AppendFormat(" -std=c++14");
+										CommandBuilder.AppendFormat(IsWindowsClang ? " /std:c++14" : " -std=c++14");
 										break;
-									case CppStandardVersion.Latest:
 									case CppStandardVersion.Cpp17:
-										CommandBuilder.AppendFormat(" -std=c++17");
+										CommandBuilder.AppendFormat(IsWindowsClang ? " /std:c++17" : " -std=c++17");
 										break;
 									case CppStandardVersion.Cpp20:
-										CommandBuilder.AppendFormat(" -std=c++20");
+									case CppStandardVersion.Latest:
+										CommandBuilder.AppendFormat(IsWindowsClang ? " /std:c++latest" : " -std=c++20");
 										break;
 									default:
 										throw new BuildException($"Unsupported C++ standard type set: {ModuleCompileEnvironment.CppStandard}");
