UnrealBuildTool: Escape special characters for file path in generated vcxproj

#trivial
#rnx

diff --git a/Engine/Source/Programs/UnrealBuildTool/Platform/HoloLens/HoloLensProjectGenerator.cs b/Engine/Source/Programs/UnrealBuildTool/Platform/HoloLens/HoloLensProjectGenerator.cs
--- a/Engine/Source/Programs/UnrealBuildTool/Platform/HoloLens/HoloLensProjectGenerator.cs
+++ b/Engine/Source/Programs/UnrealBuildTool/Platform/HoloLens/HoloLensProjectGenerator.cs
@@ -5,6 +5,7 @@ using System.Collections.Generic;
 using System.Text;
 using System.Diagnostics;
 using System.IO;
+using System.Security;
 using System.Xml.Linq;
 using EpicGames.Core;
 using UnrealBuildBase;
@@ -129,8 +130,8 @@ namespace UnrealBuildTool
 			{
 				LayoutDirString += "	<PropertyGroup " + InConditionString + ">" + ProjectFileGenerator.NewLine;
 				LayoutDirString += "		<RemoveExtraDeployFiles>false</RemoveExtraDeployFiles>" + ProjectFileGenerator.NewLine;
-				LayoutDirString += "		<LayoutDir>" + DirectoryReference.Combine(NMakeOutputPath.Directory, "AppX").FullName + "</LayoutDir>" + ProjectFileGenerator.NewLine;
-				LayoutDirString += "		<AppxPackageRecipe>" + FileReference.Combine(NMakeOutputPath.Directory, ProjectFilePath.GetFileNameWithoutExtension() + ".build.appxrecipe").FullName + "</AppxPackageRecipe>" + ProjectFileGenerator.NewLine;
+				LayoutDirString += "		<LayoutDir>" + SecurityElement.Escape(DirectoryReference.Combine(NMakeOutputPath.Directory, "AppX").FullName) + "</LayoutDir>" + ProjectFileGenerator.NewLine;
+				LayoutDirString += "		<AppxPackageRecipe>" + SecurityElement.Escape(FileReference.Combine(NMakeOutputPath.Directory, ProjectFilePath.GetFileNameWithoutExtension() + ".build.appxrecipe").FullName) + "</AppxPackageRecipe>" + ProjectFileGenerator.NewLine;
 				LayoutDirString += "	</PropertyGroup>" + ProjectFileGenerator.NewLine;
 
 				// another hijack - this is a convenient point to make sure that HoloLens-appropriate debuggers are available
