UnrealBuildTool: Make System.Text.RegularExpressions available for build rules.

#trivial
#rnx

diff --git a/Engine/Source/Programs/UnrealBuildTool/System/DynamicCompilation.cs b/Engine/Source/Programs/UnrealBuildTool/System/DynamicCompilation.cs
--- a/Engine/Source/Programs/UnrealBuildTool/System/DynamicCompilation.cs
+++ b/Engine/Source/Programs/UnrealBuildTool/System/DynamicCompilation.cs
@@ -217,6 +217,7 @@ namespace UnrealBuildTool
 			MetadataReferences.Add(MetadataReference.CreateFromFile(Assembly.Load("System.Linq").Location));
 			MetadataReferences.Add(MetadataReference.CreateFromFile(Assembly.Load("System.Private.Xml").Location));
 			MetadataReferences.Add(MetadataReference.CreateFromFile(Assembly.Load("System.Private.Xml.Linq").Location));
+			MetadataReferences.Add(MetadataReference.CreateFromFile(Assembly.Load("System.Text.RegularExpressions").Location));
 			MetadataReferences.Add(MetadataReference.CreateFromFile(Assembly.Load("System.Console").Location));
 			MetadataReferences.Add(MetadataReference.CreateFromFile(Assembly.Load("System.Runtime.Extensions").Location));
 			MetadataReferences.Add(MetadataReference.CreateFromFile(Assembly.Load("netstandard").Location));
