change VC AppX manifest generator class declaration to public & add comments.

 - also remove some unused functions

#rnx

diff --git a/Engine/Source/Programs/UnrealBuildTool/Platform/Windows/ResXResource/UEResXWriter.cs b/Engine/Source/Programs/UnrealBuildTool/Platform/Windows/ResXResource/UEResXWriter.cs
--- a/Engine/Source/Programs/UnrealBuildTool/Platform/Windows/ResXResource/UEResXWriter.cs
+++ b/Engine/Source/Programs/UnrealBuildTool/Platform/Windows/ResXResource/UEResXWriter.cs
@@ -9,8 +9,14 @@ using System.Xml;
 
 namespace UnrealBuildTool
 {
-	class UEResXWriter
+	/// <summary>
+	/// Resource writer for VC app manifest generator
+	/// </summary>
+	public class UEResXWriter
 	{
+		/// <summary>
+		/// Create a new resource writer at the given path
+		/// </summary>
 		public UEResXWriter(string InPath)
 		{
 
@@ -26,11 +32,17 @@ namespace UnrealBuildTool
 			RootElement.AppendChild(CreateEntry("resheader", "writer", typeof(UEResXWriter).AssemblyQualifiedName!, null));
 		}
 
+		/// <summary>
+		/// Closes the writer
+		/// </summary>
 		public void Close()
 		{
 			Document.Save(Filename);
 		}
 
+		/// <summary>
+		/// Add the given resource reference
+		/// </summary>
 		public void AddResource(string InName, string InValue)
 		{
 			RootElement.AppendChild(CreateEntry("data", InName, InValue, "preserve"));
diff --git a/Engine/Source/Programs/UnrealBuildTool/Platform/Windows/VCManifestGenerator.cs b/Engine/Source/Programs/UnrealBuildTool/Platform/Windows/VCManifestGenerator.cs
--- a/Engine/Source/Programs/UnrealBuildTool/Platform/Windows/VCManifestGenerator.cs
+++ b/Engine/Source/Programs/UnrealBuildTool/Platform/Windows/VCManifestGenerator.cs
@@ -15,37 +15,66 @@ using System.Diagnostics.CodeAnalysis;
 
 namespace UnrealBuildTool
 {
-    abstract class VCManifestGenerator
-    {
+	/// <summary>
+	/// Base class for VC appx manifest generation
+	/// </summary>
+	abstract public class VCManifestGenerator
+	{
+		/// default schema namespace
         protected virtual string Schema2010NS { get { return "http://schemas.microsoft.com/appx/2010/manifest"; } }
-        protected virtual string Schema2013NS { get { return "http://schemas.microsoft.com/appx/2013/manifest"; } }
 
+		/// config section for platform-specific target settings
 		protected virtual string IniSection_PlatformTargetSettings { get { return string.Format( "/Script/{0}PlatformEditor.{0}TargetSettings", Platform.ToString() ); } }
+		
+		/// config section for general target settings
 		protected virtual string IniSection_GeneralProjectSettings { get { return "/Script/EngineSettings.GeneralProjectSettings"; } }
 
+		/// default subdirectory for build resources
 		protected const string BuildResourceSubPath = "Resources";
+
+		/// default subdirectory for engine resources
 		protected const string EngineResourceSubPath = "DefaultImages";
 
+		/// platform to use for reading configuration
 		protected virtual UnrealTargetPlatform ConfigPlatform { get { return Platform; } }
 
-        // Manifest compliance values
+        /// Manifest compliance values
         protected const int MaxResourceEntries = 200;
 
-        // INI configuration cache
-        protected ConfigHierarchy? EngineIni;
-        protected ConfigHierarchy? GameIni;
+		/// cached engine ini
+		protected ConfigHierarchy? EngineIni;
+
+		/// cached game ini
+		protected ConfigHierarchy? GameIni;
+
+		/// the default culture to use
+		protected string? DefaultCulture;
+
+		/// all cultures that will be staged
+		protected List<string>? CulturesToStage;
+
+		/// resource writer for the default culture
+		protected UEResXWriter? DefaultResourceWriter;
 
-        protected string? DefaultCulture;
-        protected List<string>? CulturesToStage;
+		/// resource writers for each culture
+		protected Dictionary<string, UEResXWriter>? PerCultureResourceWriters;
 
-        protected UEResXWriter? DefaultResourceWriter;
-        protected Dictionary<string, UEResXWriter>? PerCultureResourceWriters;
-        protected UnrealTargetPlatform Platform;
+		/// the platform to generate the manifest for
+		protected UnrealTargetPlatform Platform;
+
+		/// project file to use
 		protected FileReference? ProjectFile;
-        protected string? ProjectPath;
-        protected string? OutputPath;
-        protected string? IntermediatePath;
 
+		/// directory containing the project
+		protected string? ProjectPath;
+
+		/// output path - where the manifest will be created
+		protected string? OutputPath;
+
+		/// intermediate path - used for temporary files
+		protected string? IntermediatePath;
+
+		/// files that have been updated
 		protected List<string>? UpdatedFilePaths;
 
 		/// <summary>
@@ -56,6 +85,9 @@ namespace UnrealBuildTool
 			this.Platform = InPlatform;
 		}
 
+		/// <summary>
+		/// Lookup a switch in a dictionary
+		/// </summary>
         protected static bool SafeGetBool(IDictionary<string, string> InDictionary, string Key, bool DefaultValue = false)
 		{
 			if (InDictionary.ContainsKey(Key))
@@ -67,7 +99,10 @@ namespace UnrealBuildTool
 			return DefaultValue;
 		}
 
-        protected static bool CreateCheckDirectory(string TargetDirectory)
+		/// <summary>
+		/// Attempts to create the given directory
+		/// </summary>
+		protected static bool CreateCheckDirectory(string TargetDirectory)
 		{
 			if (!Directory.Exists(TargetDirectory))
 			{
@@ -89,7 +124,10 @@ namespace UnrealBuildTool
 			return true;
 		}
 
-        protected static void RecursivelyForceDeleteDirectory(string InDirectoryToDelete)
+		/// <summary>
+		/// Delete the entire directory tree
+		/// </summary>
+		protected static void RecursivelyForceDeleteDirectory(string InDirectoryToDelete)
 		{
 			if (Directory.Exists(InDirectoryToDelete))
 			{
@@ -161,7 +199,10 @@ namespace UnrealBuildTool
         }
 
 
-        protected string? ValidatePackageVersion(string InVersionNumber)
+		/// <summary>
+		/// Returns a valid version of the given package version string
+		/// </summary>
+		protected string? ValidatePackageVersion(string InVersionNumber)
 		{
 			string WorkingVersionNumber = Regex.Replace(InVersionNumber, "[^.0-9]", "");
 			string CompletedVersionString = "";
@@ -211,7 +252,10 @@ namespace UnrealBuildTool
 			return CompletedVersionString;
 		}
 
-        protected string? ValidateProjectBaseName(string InApplicationId)
+		/// <summary>
+		/// Returns a valid version of the given application id
+		/// </summary>
+		protected string? ValidateProjectBaseName(string InApplicationId)
 		{
 			string ReturnVal = Regex.Replace(InApplicationId, "[^A-Za-z0-9]", "");
 			if (ReturnVal != null)
@@ -226,6 +270,9 @@ namespace UnrealBuildTool
 			return ReturnVal;
 		}
 
+		/// <summary>
+		/// Reads an integer from the cached ini files
+		/// </summary>
 		[return: NotNullIfNotNull("DefaultValue")]
         protected string? ReadIniString(string? Key, string Section, string? DefaultValue = null)
 		{
@@ -242,6 +289,9 @@ namespace UnrealBuildTool
 			return DefaultValue;
 		}
 
+		/// <summary>
+		/// Reads a string from the cached ini files
+		/// </summary>
 		[return: NotNullIfNotNull("DefaultValue")]
         protected string? GetConfigString(string PlatformKey, string? GenericKey, string? DefaultValue = null)
 		{
@@ -249,7 +299,10 @@ namespace UnrealBuildTool
 			return ReadIniString(PlatformKey, IniSection_PlatformTargetSettings, GenericValue);
 		}
 
-        protected bool GetConfigBool(string PlatformKey, string? GenericKey, bool DefaultValue = false)
+		/// <summary>
+		/// Reads a bool from the cached ini files
+		/// </summary>
+		protected bool GetConfigBool(string PlatformKey, string? GenericKey, bool DefaultValue = false)
 		{
 			var GenericValue = ReadIniString(GenericKey, IniSection_GeneralProjectSettings, null);
 			var ResultStr = ReadIniString(PlatformKey, IniSection_PlatformTargetSettings, GenericValue);
@@ -262,7 +315,10 @@ namespace UnrealBuildTool
 			return ResultStr == "true" || ResultStr == "1" || ResultStr == "yes";
 		}
 
-        protected string GetConfigColor(string PlatformConfigKey, string DefaultValue)
+		/// <summary>
+		/// Reads a color from the cached ini files
+		/// </summary>
+		protected string GetConfigColor(string PlatformConfigKey, string DefaultValue)
 		{
 			var ConfigValue = GetConfigString(PlatformConfigKey, null, null);
 			if (ConfigValue == null)
@@ -282,6 +338,9 @@ namespace UnrealBuildTool
 			return DefaultValue;
 		}
 
+		/// <summary>
+		/// Attempts to locate the given resource binary file in several known folder locations
+		/// </summary>
 		protected virtual bool FindResourceBinaryFile( out string SourcePath, string ResourceFileName, bool AllowEngineFallback = true)
 		{
 			// look in project normal Build location
@@ -312,13 +371,19 @@ namespace UnrealBuildTool
 			return bFileExists;
 		}
 
+		/// <summary>
+		/// Determines whether the given resource binary file can be found in one of the known folder locations
+		/// </summary>
 		protected bool DoesResourceBinaryFileExist(string ResourceFileName, bool AllowEngineFallback = true)
 		{
 			string SourcePath;
 			return FindResourceBinaryFile( out SourcePath, ResourceFileName, AllowEngineFallback );
 		}
 
-        protected bool CopyAndReplaceBinaryIntermediate(string ResourceFileName, bool AllowEngineFallback = true)
+		/// <summary>
+		/// Updates the given resource binary file
+		/// </summary>
+		protected bool CopyAndReplaceBinaryIntermediate(string ResourceFileName, bool AllowEngineFallback = true)
 		{
 			string TargetPath = Path.Combine(IntermediatePath!, BuildResourceSubPath);
 			string SourcePath;
@@ -373,7 +438,10 @@ namespace UnrealBuildTool
 			return true;
 		}
 
-        protected void CompareAndReplaceModifiedTarget(string IntermediatePath, string TargetPath)
+		/// <summary>
+		/// Updates the given file if necessary
+		/// </summary>
+		protected void CompareAndReplaceModifiedTarget(string IntermediatePath, string TargetPath)
 		{
 			if (!File.Exists(IntermediatePath))
 			{
@@ -424,7 +492,10 @@ namespace UnrealBuildTool
 			}
 		}
 
-        protected void CopyResourcesToTargetDir()
+		/// <summary>
+		/// Copies all of the generated files to the output folder
+		/// </summary>
+		protected void CopyResourcesToTargetDir()
 		{
 			string TargetPath = Path.Combine(OutputPath!, BuildResourceSubPath);
 			string SourcePath = Path.Combine(IntermediatePath!, BuildResourceSubPath);
@@ -480,7 +551,10 @@ namespace UnrealBuildTool
 			}
 		}
 
-        protected string AddResourceEntry(string ResourceEntryName, string ConfigKey, string GenericINISection, string GenericINIKey, string DefaultValue, string ValueSuffix = "")
+		/// <summary>
+		/// Adds the given string to the culture string writers
+		/// </summary>
+		protected string AddResourceEntry(string ResourceEntryName, string ConfigKey, string GenericINISection, string GenericINIKey, string DefaultValue, string ValueSuffix = "")
 		{
 			string? ConfigScratchValue = null;
 
@@ -545,6 +619,9 @@ namespace UnrealBuildTool
 			return "ms-resource:" + ResourceEntryName;
 		}
 
+		/// <summary>
+		/// Adds an additional string to all culture resource writers
+		/// </summary>
 		protected string AddExternalResourceEntry(string ResourceEntryName, string DefaultValue, Dictionary<string,string> CultureIdToCultureValues)
 		{
 			DefaultResourceWriter!.AddResource(ResourceEntryName, DefaultValue);
@@ -558,6 +635,9 @@ namespace UnrealBuildTool
 			return "ms-resource:" + ResourceEntryName;
 		}
 
+		/// <summary>
+		/// Adds a debug-only string to all resource writers
+		/// </summary>
 		protected string AddDebugResourceString(string ResourceEntryName, string Value)
 		{
 			DefaultResourceWriter!.AddResource(ResourceEntryName, Value);
@@ -571,12 +651,18 @@ namespace UnrealBuildTool
 			return "ms-resource:" + ResourceEntryName;
 		}
 
+		/// <summary>
+		/// Get the XName from a given string and schema pair
+		/// </summary>
 		protected virtual XName GetName( string BaseName, string SchemaName )
 		{
 			return XName.Get(BaseName);
 		}
 
-        protected XElement GetResources()
+		/// <summary>
+		/// Get the resources element
+		/// </summary>
+		protected XElement GetResources()
 		{
 			var ResourceCulturesList = CulturesToStage.ToList();
 			// Move the default culture to the front of the list
@@ -596,6 +682,9 @@ namespace UnrealBuildTool
 			return new XElement(GetName("Resources", Schema2010NS), CultureElements);
 		}
 
+		/// <summary>
+		/// Get the package identity name string
+		/// </summary>
 		protected string GetIdentityPackageName()
 		{
             // Read the PackageName from config
@@ -618,12 +707,18 @@ namespace UnrealBuildTool
 			return PackageName;
 		}
 
+		/// <summary>
+		/// Get the publisher name string
+		/// </summary>
 		protected string GetIdentityPublisherName()
 		{
             var PublisherName = GetConfigString("PublisherName", "CompanyDistinguishedName", "CN=NoPublisher");
 			return PublisherName;
 		}
 
+		/// <summary>
+		/// Get the package version string
+		/// </summary>
 		protected string? GetIdentityVersionNumber()
 		{
             var VersionNumber = GetConfigString("PackageVersion", "ProjectVersion", "1.0.0.0");
@@ -639,7 +734,10 @@ namespace UnrealBuildTool
 			return VersionNumber;
 		}
 
-        protected XElement GetIdentity(out string IdentityName)
+		/// <summary>
+		/// Get the package identity element
+		/// </summary>
+		protected XElement GetIdentity(out string IdentityName)
         {
             var PackageName = GetIdentityPackageName();
             var PublisherName = GetIdentityPublisherName();
@@ -653,6 +751,9 @@ namespace UnrealBuildTool
                 new XAttribute("Version", VersionNumber));
         }
 
+		/// <summary>
+		/// Updates the given package version to include the engine build version, if requested
+		/// </summary>
 		protected virtual string? IncludeBuildVersionInPackageVersion(string? VersionNumber)
 		{
 			BuildVersion? BuildVersionForPackage;
@@ -670,30 +771,36 @@ namespace UnrealBuildTool
 			return VersionNumber;
 		}
 
-		protected abstract string GetSDKDirectory();
-
+		/// <summary>
+		/// Get the path to the makepri.exe tool
+		/// </summary>
 		protected abstract string GetMakePriBinaryPath();
 
+		/// <summary>
+		/// Get any additional platform-specific parameters for makepri.exe
+		/// </summary>
 		protected virtual string GetMakePriExtraCommandLine()
 		{
 			return "";
 		}
 
+		/// <summary>
+		/// Return the entire manifest element
+		/// </summary>
 		protected abstract XElement GetManifest(List<UnrealTargetConfiguration> TargetConfigs, List<string> Executables, out string IdentityName);
 
+		/// <summary>
+		/// Perform any platform-specific processing on the manifest before it is saved
+		/// </summary>
 		protected virtual void ProcessManifest(List<UnrealTargetConfiguration> TargetConfigs, List<string> Executables, string ManifestName, string ManifestTargetPath, string ManifestIntermediatePath)
         {
 		}
 
-        public List<string>? CreateManifest(string InManifestName, string InOutputPath, string InIntermediatePath, FileReference? InProjectFile, string InProjectDirectory, List<UnrealTargetConfiguration> InTargetConfigs, List<string> InExecutables)
+		/// <summary>
+		/// Create a manifest and return the list of modified files
+		/// </summary>
+		public List<string>? CreateManifest(string InManifestName, string InOutputPath, string InIntermediatePath, FileReference? InProjectFile, string InProjectDirectory, List<UnrealTargetConfiguration> InTargetConfigs, List<string> InExecutables)
 		{
-			// Verify we can find the SDK.
-			string SDKDirectory = GetSDKDirectory();
-			if (string.IsNullOrEmpty(SDKDirectory))
-			{
-				return null;
-			}
-
 			// Check parameter values are valid.
 			if (InTargetConfigs.Count != InExecutables.Count)
 			{
