UnrealBuildTool: Fix SN-DBS json encoding (from UDN)

#rnx

diff --git a/Engine/Source/Programs/UnrealBuildTool/Executors/SNDBS.cs b/Engine/Source/Programs/UnrealBuildTool/Executors/SNDBS.cs
--- a/Engine/Source/Programs/UnrealBuildTool/Executors/SNDBS.cs
+++ b/Engine/Source/Programs/UnrealBuildTool/Executors/SNDBS.cs
@@ -65,6 +65,8 @@ namespace UnrealBuildTool
 
 			// Build the json script file to describe all the actions and their dependencies
 			var ActionIds = Actions.ToDictionary(a => a, a => Guid.NewGuid().ToString());
+			JsonSerializerOptions JsonOption = new JsonSerializerOptions();
+			JsonOption.Encoder = System.Text.Encodings.Web.JavaScriptEncoder.UnsafeRelaxedJsonEscaping;
 			File.WriteAllText(ScriptFile.FullName, JsonSerializer.Serialize(new Dictionary<string, object>()
 			{
 				["jobs"] = Actions.ToDictionary(a => ActionIds[a], a =>
@@ -97,7 +99,7 @@ namespace UnrealBuildTool
 
 					return Job;
 				})
-			}));
+			}, JsonOption));
 
 			PrepareToolTemplates();
 			bool bHasRewrites = GenerateSNDBSIncludeRewriteRules();
