SN-DBS executor improvements
- Prevent the unnecessary scanning of response files
- Prevent the unnecessary scanning and distribution of header files for other platforms
- Add bAllowOverVpn and VpnSubnets XmlConfigFile settings to allow restriction of distribution over VPN
- Utilize ProgressWriter to better mark up progress text as is done with other executor implementations

diff --git a/Engine/Source/Programs/UnrealBuildTool/Executors/HybridExecutor.cs b/Engine/Source/Programs/UnrealBuildTool/Executors/HybridExecutor.cs
--- a/Engine/Source/Programs/UnrealBuildTool/Executors/HybridExecutor.cs
+++ b/Engine/Source/Programs/UnrealBuildTool/Executors/HybridExecutor.cs
@@ -30,11 +30,11 @@ namespace UnrealBuildTool
 		/// <summary>
 		/// Constructor
 		/// </summary>
-		public HybridExecutor(int InMaxLocalActions, ActionExecutor? InLocalExecutor = null, ActionExecutor? InRemoteExecutor = null)
+		public HybridExecutor(List<TargetDescriptor> TargetDescriptors, int InMaxLocalActions, ActionExecutor? InLocalExecutor = null, ActionExecutor? InRemoteExecutor = null)
 		{
 			MaxLocalActions = InMaxLocalActions;
 			LocalExecutor = InLocalExecutor ?? new ParallelExecutor(MaxLocalActions);
-			RemoteExecutor = InRemoteExecutor ?? (XGE.IsAvailable() ? (ActionExecutor)new XGE() : new SNDBS());
+			RemoteExecutor = InRemoteExecutor ?? (XGE.IsAvailable() ? (ActionExecutor)new XGE() : new SNDBS(TargetDescriptors));
 
 			XmlConfig.ApplyTo(this);
 
diff --git a/Engine/Source/Programs/UnrealBuildTool/Executors/SNDBS.cs b/Engine/Source/Programs/UnrealBuildTool/Executors/SNDBS.cs
--- a/Engine/Source/Programs/UnrealBuildTool/Executors/SNDBS.cs
+++ b/Engine/Source/Programs/UnrealBuildTool/Executors/SNDBS.cs
@@ -4,10 +4,16 @@ using System;
 using System.Collections;
 using System.Collections.Generic;
 using System.Diagnostics;
+using System.Diagnostics.CodeAnalysis;
 using System.IO;
 using System.Linq;
+using System.Net;
+using System.Net.NetworkInformation;
+using System.Runtime.InteropServices;
 using System.ServiceProcess;
+using System.Text;
 using System.Text.Json;
+using System.Text.RegularExpressions;
 using EpicGames.Core;
 using UnrealBuildBase;
 
@@ -18,7 +24,8 @@ namespace UnrealBuildTool
 		public override string Name => "SNDBS";
 
 		private static readonly string? SCERoot = Environment.GetEnvironmentVariable("SCE_ROOT_DIR");
-		private static readonly string SNDBSExecutable = Path.Combine(SCERoot ?? string.Empty, "Common", "SN-DBS", "bin", "dbsbuild.exe");
+		private static readonly string SNDBSBuildExe = Path.Combine(SCERoot ?? string.Empty, "Common", "SN-DBS", "bin", "dbsbuild.exe");
+		private static readonly string SNDBSUtilExe = Path.Combine(SCERoot ?? string.Empty, "Common", "SN-DBS", "bin", "dbsutil.exe");
 
 		private static readonly DirectoryReference IntermediateDir = DirectoryReference.Combine(Unreal.EngineDirectory, "Intermediate", "Build", "SNDBS");
 		private static readonly FileReference IncludeRewriteRulesFile = FileReference.Combine(IntermediateDir, "include-rewrite-rules.ini");
@@ -26,7 +33,32 @@ namespace UnrealBuildTool
 
 		private Dictionary<string, string> ActiveTemplates = BuiltInTemplates.ToDictionary(p => p.Key, p => p.Value);
 
-		public bool EnableEcho { get; set; } = false;
+		/// <summary>
+		/// When enabled, SN-DBS will stop compiling targets after a compile error occurs.  Recommended, as it saves computing resources for others.
+		/// </summary>
+		[XmlConfigFile(Category = "BuildConfiguration")]
+		bool bStopSNDBSCompilationAfterErrors = false;
+
+		/// <summary>
+		/// When set to false, SNDBS will not be enabled when running connected to the coordinator over VPN. Configure VPN-assigned subnets via the VpnSubnets parameter.
+		/// </summary>
+		[XmlConfigFile(Category = "SNDBS")]
+		static bool bAllowOverVpn = true;
+
+		/// <summary>
+		/// List of subnets containing IP addresses assigned by VPN
+		/// </summary>
+		[XmlConfigFile(Category = "SNDBS")]
+		static string[]? VpnSubnets = null;
+
+		private const string ProgressMarkupPrefix = "@action:";
+
+		private List<TargetDescriptor> TargetDescriptors;
+
+		public SNDBS(List<TargetDescriptor> InTargetDescriptors)
+		{
+			TargetDescriptors = InTargetDescriptors;
+		}
 
 		public SNDBS AddTemplate(string ExeName, string TemplateContents)
 		{
@@ -34,16 +66,131 @@ namespace UnrealBuildTool
 			return this;
 		}
 
+		static bool TryGetBrokerHost([NotNullWhen(true)] out string? OutBroker)
+		{
+			if (BuildHostPlatform.Current.Platform == UnrealTargetPlatform.Win64)
+			{
+				string BrokerHostName = "";
+				Regex FindHost = new Regex(@"Active broker is ""(\S +)"" \((\S+)\)");
+				var LocalProcess = new Process();
+				LocalProcess.StartInfo = new ProcessStartInfo(SNDBSUtilExe, $"-connected");
+				LocalProcess.OutputDataReceived += (Sender, Args) =>
+				{
+					Match Result = FindHost.Match(Args.Data);
+					if (Result.Success)
+					{
+						BrokerHostName = Result.Groups[1].Value;
+					}
+				};
+				if (Utils.RunLocalProcess(LocalProcess) == 1 && BrokerHostName.Length > 0)
+				{
+					OutBroker = BrokerHostName;
+					return true;
+				}
+			}
+
+			OutBroker = null;
+			return false;
+		}
+
+		[DllImport("iphlpapi")]
+		static extern int GetBestInterface(uint dwDestAddr, ref int pdwBestIfIndex);
+
+		static NetworkInterface? GetInterfaceForHost(string Host)
+		{
+			if (BuildHostPlatform.Current.Platform == UnrealTargetPlatform.Win64)
+			{
+				IPHostEntry HostEntry = Dns.GetHostEntry(Host);
+				foreach (IPAddress HostAddress in HostEntry.AddressList)
+				{
+					int InterfaceIdx = 0;
+					if (GetBestInterface(BitConverter.ToUInt32(HostAddress.GetAddressBytes(), 0), ref InterfaceIdx) == 0)
+					{
+						foreach (NetworkInterface Interface in NetworkInterface.GetAllNetworkInterfaces())
+						{
+							IPv4InterfaceProperties Properties = Interface.GetIPProperties().GetIPv4Properties();
+							if (Properties.Index == InterfaceIdx)
+							{
+								return Interface;
+							}
+						}
+					}
+				}
+			}
+			return null;
+		}
+
+		public static bool IsHostOnVpn(string HostName)
+		{
+			// If there aren't any defined subnets, just early out
+			if (VpnSubnets == null || VpnSubnets.Length == 0)
+			{
+				return false;
+			}
+
+			// Parse all the subnets from the config file
+			List<Subnet> ParsedVpnSubnets = new List<Subnet>();
+			foreach (string VpnSubnet in VpnSubnets)
+			{
+				ParsedVpnSubnets.Add(Subnet.Parse(VpnSubnet));
+			}
+
+			// Check if any network adapters have an IP within one of these subnets
+			try
+			{
+				NetworkInterface? Interface = GetInterfaceForHost(HostName);
+				if (Interface != null && Interface.OperationalStatus == OperationalStatus.Up)
+				{
+					IPInterfaceProperties Properties = Interface.GetIPProperties();
+					foreach (UnicastIPAddressInformation UnicastAddressInfo in Properties.UnicastAddresses)
+					{
+						byte[] AddressBytes = UnicastAddressInfo.Address.GetAddressBytes();
+						foreach (Subnet Subnet in ParsedVpnSubnets)
+						{
+							if (Subnet.Contains(AddressBytes))
+							{
+								if (!bAllowOverVpn)
+								{
+									Log.TraceInformationOnce("XGE coordinator {0} will be not be used over VPN (adapter '{1}' with IP {2} is in subnet {3}). Set <XGE><bAllowOverVpn>true</bAllowOverVpn></XGE> in BuildConfiguration.xml to override.", HostName, Interface.Description, UnicastAddressInfo.Address, Subnet);
+								}
+								return true;
+							}
+						}
+					}
+				}
+			}
+			catch (Exception Ex)
+			{
+				Log.TraceWarning("Unable to check whether host {0} is connected to VPN:\n{1}", HostName, ExceptionUtils.FormatExceptionDetails(Ex));
+			}
+			return false;
+		}
+
 		public static bool IsAvailable()
 		{
 			// Check the executable exists on disk
-			if (SCERoot == null || !File.Exists(SNDBSExecutable))
+			if (SCERoot == null || !File.Exists(SNDBSBuildExe))
 			{
 				return false;
 			}
 
 			// Check the service is running
-			return ServiceController.GetServices().Any(s => s.ServiceName.StartsWith("SNDBS") && s.Status == ServiceControllerStatus.Running);
+			if (!ServiceController.GetServices().Any(s => s.ServiceName.StartsWith("SNDBS") && s.Status == ServiceControllerStatus.Running))
+			{
+				return false;
+			}
+
+			// Check if we're connected over VPN
+			if (!bAllowOverVpn && VpnSubnets != null && VpnSubnets.Length > 0)
+			{
+				string? BrokerHost;
+				if (TryGetBrokerHost(out BrokerHost) && IsHostOnVpn(BrokerHost))
+				{
+					return false;
+				}
+			}
+
+			return true;
 		}
 
 		public override bool ExecuteActions(List<LinkedAction> Actions)
@@ -82,20 +229,14 @@ namespace UnrealBuildTool
 
 					if (a.PrerequisiteItems.Count() > 0)
 					{
-						Job["explicit_input_files"] = a.PrerequisiteItems.Select(i => new Dictionary<string, object>()
+						Job["explicit_input_files"] = a.PrerequisiteItems.Where(i => !i.AbsolutePath.EndsWith(".response")).Select(i => new Dictionary<string, object>()
 						{
 							["filename"] = i.AbsolutePath
 						}).ToList();
 					}
 
-					if (EnableEcho)
-					{
-						var EchoString = string.Join(" ", string.IsNullOrWhiteSpace(a.CommandDescription) ? string.Empty : $"[{a.CommandDescription}]", a.StatusDescription);
-						if (!string.IsNullOrWhiteSpace(EchoString))
-						{
-							Job["echo"] = EchoString;
-						}
-					}
+					string CommandDescription = string.IsNullOrWhiteSpace(a.CommandDescription) ? a.ActionType.ToString() : a.CommandDescription;
+					Job["echo"] = $"{ProgressMarkupPrefix}{CommandDescription}:{a.StatusDescription}";
 
 					return Job;
 				})
@@ -104,11 +245,94 @@ namespace UnrealBuildTool
 			PrepareToolTemplates();
 			bool bHasRewrites = GenerateSNDBSIncludeRewriteRules();
 
-			var LocalProcess = new Process();
-			LocalProcess.StartInfo = new ProcessStartInfo(SNDBSExecutable, $"-q -p \"Unreal Engine Tasks\" -s \"{ScriptFile}\" -templates \"{IntermediateDir}\"{(bHasRewrites ? $" --include-rewrite-rules \"{IncludeRewriteRulesFile}\"" : "")}");
-			LocalProcess.OutputDataReceived += (Sender, Args) => Log.TraceInformation("{0}", Args.Data);
-			LocalProcess.ErrorDataReceived += (Sender, Args) => Log.TraceInformation("{0}", Args.Data);
-			return Utils.RunLocalProcess(LocalProcess) == 0;
+			var ConfigList = TargetDescriptors.Select(Descriptor => $"{Descriptor.Name}|{Descriptor.Platform}|{Descriptor.Configuration}");
+			var ConfigDescription = string.Join(",", ConfigList);
+
+			var StartInfo = new ProcessStartInfo(
+				SNDBSBuildExe,
+				$"-q -p \"{ConfigDescription}\" -s \"{ScriptFile}\" -templates \"{IntermediateDir}\""
+				);
+			StartInfo.UseShellExecute = false;
+			if (bHasRewrites)
+			{
+				StartInfo.Arguments += $" --include-rewrite-rules \"{IncludeRewriteRulesFile}\"";
+			}
+			if (!bStopSNDBSCompilationAfterErrors)
+			{
+				StartInfo.Arguments += " -k";
+			}
+			return ExecuteProcessWithProgressMarkup(StartInfo, Actions.Count);
+		}
+
+		/// <summary>
+		/// Executes the process, parsing progress markup as part of the output.
+		/// </summary>
+		private bool ExecuteProcessWithProgressMarkup(ProcessStartInfo SnDbsStartInfo, int NumActions)
+		{
+			using (ProgressWriter Writer = new ProgressWriter("Compiling C++ source files...", false))
+			{
+				int NumCompletedActions = 0;
+				string CurrentStatus = "";
+
+				// Create a wrapper delegate that will parse the output actions
+				DataReceivedEventHandler EventHandlerWrapper = (Sender, Args) =>
+				{
+					if (Args.Data != null)
+					{
+						string Text = Args.Data;
+						if (Text.StartsWith(ProgressMarkupPrefix))
+						{
+							Writer.Write(++NumCompletedActions, NumActions);
+
+							Text = Args.Data.Substring(ProgressMarkupPrefix.Length).Trim();
+							var ActionInfo = Text.Split(':');
+							Log.TraceInformation($"[{NumCompletedActions}/{NumActions}] {ActionInfo[0]} {ActionInfo[1]}");
+							CurrentStatus = ActionInfo[1];
+							return;
+						}
+						// Suppress redundant tool output of status we already printed (e.g., msvc cl prints compile unit name always)
+						if (!Text.Equals(CurrentStatus))
+						{
+							Log.TraceInformation(Text);
+						}
+					}
+				};
+
+				try
+				{
+					// Start the process, redirecting stdout/stderr if requested.
+					Process LocalProcess = new Process();
+					LocalProcess.StartInfo = SnDbsStartInfo;
+					bool bShouldRedirectOuput = EventHandlerWrapper != null;
+					if (bShouldRedirectOuput)
+					{
+						SnDbsStartInfo.RedirectStandardError = true;
+						SnDbsStartInfo.RedirectStandardOutput = true;
+						LocalProcess.EnableRaisingEvents = true;
+						LocalProcess.OutputDataReceived += EventHandlerWrapper;
+						LocalProcess.ErrorDataReceived += EventHandlerWrapper;
+					}
+					LocalProcess.Start();
+					if (bShouldRedirectOuput)
+					{
+						LocalProcess.BeginOutputReadLine();
+						LocalProcess.BeginErrorReadLine();
+					}
+
+					Log.TraceInformation("Distributing {0} action{1} to SN-DBS",
+						NumActions,
+						NumActions == 1 ? "" : "s");
+
+					// Wait until the process is finished and return whether it all the tasks successfully executed.
+					LocalProcess.WaitForExit();
+					return LocalProcess.ExitCode == 0;
+				}
+				catch (Exception Ex)
+				{
+					Log.WriteException(Ex, null);
+					return false;
+				}
+			}
 		}
 
 		private void PrepareToolTemplates()
@@ -132,9 +356,9 @@ namespace UnrealBuildTool
 
 		private bool GenerateSNDBSIncludeRewriteRules()
 		{
-			// Get all registered, distinct platform names.
-			var Platforms = UEBuildPlatform.GetRegisteredPlatforms()
-				.Select(Platform => UEBuildPlatform.GetBuildPlatform(Platform).GetPlatformName())
+			// Get all distinct platform names being used in this build.
+			var Platforms = TargetDescriptors
+				.Select(TargetDescriptor => UEBuildPlatform.GetBuildPlatform(TargetDescriptor.Platform).GetPlatformName())
 				.Distinct()
 				.ToList();
 
diff --git a/Engine/Source/Programs/UnrealBuildTool/Modes/BuildMode.cs b/Engine/Source/Programs/UnrealBuildTool/Modes/BuildMode.cs
--- a/Engine/Source/Programs/UnrealBuildTool/Modes/BuildMode.cs
+++ b/Engine/Source/Programs/UnrealBuildTool/Modes/BuildMode.cs
@@ -374,7 +374,7 @@ namespace UnrealBuildTool
 					if (PreprocessActions.Count > 0)
 					{
 						Log.TraceInformation("Updating module dependencies...");
-						ActionGraph.ExecuteActions(BuildConfiguration, PreprocessActions);
+						ActionGraph.ExecuteActions(BuildConfiguration, PreprocessActions, TargetDescriptors);
 
 						foreach (FileItem ProducedItem in PreprocessActions.SelectMany(x => x.ProducedItems))
 						{
@@ -603,7 +603,7 @@ namespace UnrealBuildTool
 
 						using (GlobalTracer.Instance.BuildSpan("ActionGraph.ExecuteActions()").StartActive())
 						{
-							ActionGraph.ExecuteActions(BuildConfiguration, MergedActionsToExecute);
+							ActionGraph.ExecuteActions(BuildConfiguration, MergedActionsToExecute, TargetDescriptors);
 						}
 					}
 
diff --git a/Engine/Source/Programs/UnrealBuildTool/Modes/ExecuteMode.cs b/Engine/Source/Programs/UnrealBuildTool/Modes/ExecuteMode.cs
--- a/Engine/Source/Programs/UnrealBuildTool/Modes/ExecuteMode.cs
+++ b/Engine/Source/Programs/UnrealBuildTool/Modes/ExecuteMode.cs
@@ -58,7 +58,8 @@ namespace UnrealBuildTool
 			// Execute the actions
 			using (GlobalTracer.Instance.BuildSpan("ActionGraph.ExecuteActions()").StartActive())
 			{
-				ActionGraph.ExecuteActions(BuildConfiguration, Actions);
+				List<TargetDescriptor> TargetDescriptors = TargetDescriptor.ParseCommandLine(Arguments, false, false, false);
+				ActionGraph.ExecuteActions(BuildConfiguration, Actions, TargetDescriptors);
 			}
 
 			return 0;
diff --git a/Engine/Source/Programs/UnrealBuildTool/Modes/JsonExportMode.cs b/Engine/Source/Programs/UnrealBuildTool/Modes/JsonExportMode.cs
--- a/Engine/Source/Programs/UnrealBuildTool/Modes/JsonExportMode.cs
+++ b/Engine/Source/Programs/UnrealBuildTool/Modes/JsonExportMode.cs
@@ -67,7 +67,7 @@ namespace UnrealBuildTool
 						if (PrerequisiteActions.Count > 0)
 						{
 							Log.TraceInformation("Exeucting actions that produce source files...");
-							ActionGraph.ExecuteActions(BuildConfiguration, PrerequisiteActions);
+							ActionGraph.ExecuteActions(BuildConfiguration, PrerequisiteActions, new List<TargetDescriptor> { TargetDescriptor });
 						}
 					}
 				}
diff --git a/Engine/Source/Programs/UnrealBuildTool/System/ActionGraph.cs b/Engine/Source/Programs/UnrealBuildTool/System/ActionGraph.cs
--- a/Engine/Source/Programs/UnrealBuildTool/System/ActionGraph.cs
+++ b/Engine/Source/Programs/UnrealBuildTool/System/ActionGraph.cs
@@ -330,7 +330,7 @@ namespace UnrealBuildTool
 		/// <summary>
 		/// Executes a list of actions.
 		/// </summary>
-		public static void ExecuteActions(BuildConfiguration BuildConfiguration, List<LinkedAction> ActionsToExecute)
+		public static void ExecuteActions(BuildConfiguration BuildConfiguration, List<LinkedAction> ActionsToExecute, List<TargetDescriptor> TargetDescriptors)
 		{
 			if (ActionsToExecute.Count == 0)
 			{
@@ -342,7 +342,7 @@ namespace UnrealBuildTool
 				ActionExecutor Executor;
 				if (BuildConfiguration.bAllowHybridExecutor && HybridExecutor.IsAvailable())
 				{
-					Executor = new HybridExecutor(BuildConfiguration.MaxParallelActions);
+					Executor = new HybridExecutor(TargetDescriptors, BuildConfiguration.MaxParallelActions);
 				}
 				else if (BuildConfiguration.bAllowXGE && XGE.IsAvailable())
 				{
@@ -354,7 +354,7 @@ namespace UnrealBuildTool
 				}
 				else if (BuildConfiguration.bAllowSNDBS && SNDBS.IsAvailable())
 				{
-					Executor = new SNDBS();
+					Executor = new SNDBS(TargetDescriptors);
 				}
 				else
 				{
