Fix cross building the Linux UnrealEditor from Windows where we no longer was running the FixDeps.bat script which is required to resolve circular depends issues

diff --git a/Engine/Source/Programs/UnrealBuildTool/Platform/Linux/LinuxToolChain.cs b/Engine/Source/Programs/UnrealBuildTool/Platform/Linux/LinuxToolChain.cs
--- a/Engine/Source/Programs/UnrealBuildTool/Platform/Linux/LinuxToolChain.cs
+++ b/Engine/Source/Programs/UnrealBuildTool/Platform/Linux/LinuxToolChain.cs
@@ -1468,11 +1468,6 @@ namespace UnrealBuildTool
 		{
 			if (bUseFixdeps)
 			{
-				if (!LinkEnvironment.bIsCrossReferenced)
-				{
-					return null;
-				}
-
 				Log.TraceVerbose("Adding postlink step");
 
 				bool bUseCmdExe = BuildHostPlatform.Current.ShellType == ShellType.Cmd;
@@ -1482,6 +1477,12 @@ namespace UnrealBuildTool
 
 				FileItem FixDepsScript = FileItem.GetItemByFileReference(FileReference.Combine(LinkEnvironment.LocalShadowDirectory!, ScriptName));
 
+				// if we never generated one we did not have a circular depends that needed fixing up
+				if (!FixDepsScript.Exists)
+				{
+					return null;
+				}
+
 				Action PostLinkAction = Graph.CreateAction(ActionType.Link);
 				PostLinkAction.WorkingDirectory = UnrealBuildTool.EngineSourceDirectory;
 				PostLinkAction.CommandPath = ShellBinary;
