Introduced new property "VersePath" in plugin descriptor and module rules

#rnx

diff --git a/Engine/Source/Programs/UnrealBuildTool/Configuration/ModuleRules.cs b/Engine/Source/Programs/UnrealBuildTool/Configuration/ModuleRules.cs
--- a/Engine/Source/Programs/UnrealBuildTool/Configuration/ModuleRules.cs
+++ b/Engine/Source/Programs/UnrealBuildTool/Configuration/ModuleRules.cs
@@ -1043,6 +1043,11 @@ namespace UnrealBuildTool
 		/// </summary>
 		public bool? bRequiresImplementModule;
 
+		/// <summary>
+		/// If this module has associated Verse code, this is the Verse root path of it
+		/// </summary>
+		public string? VersePath;
+
 		/// <summary>
 		/// Whether this module qualifies included headers from other modules relative to the root of their 'Public' folder. This reduces the number
 		/// of search paths that have to be passed to the compiler, improving performance and reducing the length of the compiler command line.
diff --git a/Engine/Source/Programs/UnrealBuildTool/System/PluginDescriptor.cs b/Engine/Source/Programs/UnrealBuildTool/System/PluginDescriptor.cs
--- a/Engine/Source/Programs/UnrealBuildTool/System/PluginDescriptor.cs
+++ b/Engine/Source/Programs/UnrealBuildTool/System/PluginDescriptor.cs
@@ -139,6 +139,11 @@ namespace UnrealBuildTool
 		/// </summary>
 		public LocalizationTargetDescriptor[]? LocalizationTargets;
 
+		/// <summary>
+		/// The Verse path to the root of this plugin's content directory
+		/// </summary>
+		public string? VersePath;
+
 		/// <summary>
 		/// Whether this plugin should be enabled by default for all projects
 		/// </summary>
@@ -288,6 +293,8 @@ namespace UnrealBuildTool
 				LocalizationTargets = Array.ConvertAll(LocalizationTargetsArray, x => LocalizationTargetDescriptor.FromJsonObject(x));
 			}
 
+			RawObject.TryGetStringField("VersePath", out VersePath);
+
 			bool bEnabledByDefaultValue;
 			if(RawObject.TryGetBoolField("EnabledByDefault", out bEnabledByDefaultValue))
 			{
@@ -382,7 +389,11 @@ namespace UnrealBuildTool
 			{
 				Writer.WriteValue("EngineVersion", EngineVersion);
 			}
-			if(bEnabledByDefault.HasValue)
+			if (!String.IsNullOrEmpty(VersePath))
+			{
+				Writer.WriteValue("VersePath", VersePath);
+			}
+			if (bEnabledByDefault.HasValue)
 			{
 				Writer.WriteValue("EnabledByDefault", bEnabledByDefault.Value);
 			}
diff --git a/Engine/Source/Programs/UnrealBuildTool/System/TargetMakefile.cs b/Engine/Source/Programs/UnrealBuildTool/System/TargetMakefile.cs
--- a/Engine/Source/Programs/UnrealBuildTool/System/TargetMakefile.cs
+++ b/Engine/Source/Programs/UnrealBuildTool/System/TargetMakefile.cs
@@ -19,7 +19,7 @@ namespace UnrealBuildTool
 		/// <summary>
 		/// The version number to write
 		/// </summary>
-		public const int CurrentVersion = 27;
+		public const int CurrentVersion = 28;
 
 		/// <summary>
 		/// The time at which the makefile was created
diff --git a/Engine/Source/Runtime/Projects/Private/PluginDescriptor.cpp b/Engine/Source/Runtime/Projects/Private/PluginDescriptor.cpp
--- a/Engine/Source/Runtime/Projects/Private/PluginDescriptor.cpp
+++ b/Engine/Source/Runtime/Projects/Private/PluginDescriptor.cpp
@@ -210,6 +210,8 @@ bool FPluginDescriptor::Read(const FJsonObject& Object, FText* OutFailReason /*=
 		return false;
 	}
 
+	Object.TryGetStringField(TEXT("VersePath"), VersePath);
+
 	bool bEnabledByDefault;
 	if(Object.TryGetBoolField(TEXT("EnabledByDefault"), bEnabledByDefault))
 	{
@@ -320,6 +322,11 @@ void FPluginDescriptor::UpdateJson(FJsonObject& JsonObject) const
 		JsonObject.RemoveField(TEXT("EditorCustomVirtualPath"));
 	}
 
+	if (!VersePath.IsEmpty())
+	{
+		JsonObject.SetStringField(TEXT("VersePath"), VersePath);
+	}
+
 	if (EnabledByDefault != EPluginEnabledByDefault::Unspecified)
 	{
 		JsonObject.SetBoolField(TEXT("EnabledByDefault"), (EnabledByDefault == EPluginEnabledByDefault::Enabled));
diff --git a/Engine/Source/Runtime/Projects/Public/PluginDescriptor.h b/Engine/Source/Runtime/Projects/Public/PluginDescriptor.h
--- a/Engine/Source/Runtime/Projects/Public/PluginDescriptor.h
+++ b/Engine/Source/Runtime/Projects/Public/PluginDescriptor.h
@@ -81,6 +81,9 @@ struct PROJECTS_API FPluginDescriptor
 	/** List of all localization targets associated with this plugin */
 	TArray<FLocalizationTargetDescriptor> LocalizationTargets;
 
+	/** The Verse path to the root of this plugin's content directory */
+	FString VersePath;
+
 	/** Whether this plugin should be enabled by default for all projects */
 	EPluginEnabledByDefault EnabledByDefault;
 
