UHT in UBT initial version.  Currently disabled.

diff --git a/Engine/Source/Programs/Shared/EpicGames.UHT/Utils/UhtSymbolTable.cs b/Engine/Source/Programs/Shared/EpicGames.UHT/Utils/UhtSymbolTable.cs
new file mode 100644
--- /dev/null
+++ b/Engine/Source/Programs/Shared/EpicGames.UHT/Utils/UhtSymbolTable.cs
@@ -0,0 +1,376 @@
+// Copyright Epic Games, Inc. All Rights Reserved.
+
+using EpicGames.Core;
+using EpicGames.UHT.Types;
+using System;
+using System.Collections.Generic;
+
+namespace EpicGames.UHT.Utils
+{
+
+	/// <summary>
+	/// Symbol table
+	/// </summary>
+	internal class UhtSymbolTable
+	{
+
+		/// <summary>
+		/// Represents symbol name lookup chain start.  Symbol chains are based on 
+		/// the caseless name.
+		/// </summary>
+		struct Lookup
+		{
+
+			/// <summary>
+			/// The type index of the symbol
+			/// </summary>
+			public int SymbolIndex;
+
+			/// <summary>
+			/// When searching the caseless chain, the cased index is used to match the symbol based
+			/// on the case named.
+			/// </summary>
+			public int CasedIndex;
+		}
+
+		/// <summary>
+		/// Entry in the symbol table
+		/// </summary>
+		struct Symbol
+		{
+
+			/// <summary>
+			/// The type associated with the symbol
+			/// </summary>
+			public UhtType Type;
+
+			/// <summary>
+			/// Mask of different find options which will match this symbol
+			/// </summary>
+			public UhtFindOptions MatchOptions;
+
+			/// <summary>
+			/// The case lookup index for matching by case
+			/// </summary>
+			public int CasedIndex;
+
+			/// <summary>
+			/// The next index in the symbol change based on cassless lookup
+			/// </summary>
+			public int NextIndex;
+
+			/// <summary>
+			/// The last index in the chain.  This index is only used when the symbol entry is also acting as the list
+			/// </summary>
+			public int LastIndex;
+		}
+
+		/// <summary>
+		/// Number of unique cased symbol names.  
+		/// </summary>
+		private int CasedCount = 0;
+
+		/// <summary>
+		/// Case name lookup table that returns the symbol index and the case index
+		/// </summary>
+		private Dictionary<string, Lookup> CasedDictionary = new Dictionary<string, Lookup>(StringComparer.Ordinal);
+
+		/// <summary>
+		/// Caseless name lookup table that returns the symbol index
+		/// </summary>
+		private Dictionary<string, int> CaselessDictionary = new Dictionary<string, int>(StringComparer.OrdinalIgnoreCase);
+
+		/// <summary>
+		/// Collection of symbols in the table
+		/// </summary>
+		private Symbol[] Symbols;
+
+		/// <summary>
+		/// Constructs a new symbol table.
+		/// </summary>
+		/// <param name="TypeCount">Number of types in the table.</param>
+		public UhtSymbolTable(int TypeCount)
+		{
+			this.Symbols = new Symbol[TypeCount];
+		}
+
+		/// <summary>
+		/// Add a new type to the symbol table
+		/// </summary>
+		/// <param name="Type">The type being added</param>
+		/// <param name="Name">The name of the type which could be the source name or the engine name</param>
+		public void Add(UhtType Type, string Name)
+		{
+			Lookup Existing;
+			if (this.CasedDictionary.TryGetValue(Name, out Existing))
+			{
+				AddExisting(Type, Existing.CasedIndex, Existing.SymbolIndex);
+				return;
+			}
+
+			int CasedIndex = ++this.CasedCount;
+
+			int SymbolIndex;
+			if (this.CaselessDictionary.TryGetValue(Name, out SymbolIndex))
+			{
+				this.CasedDictionary.Add(Name, new Lookup { SymbolIndex = SymbolIndex, CasedIndex = CasedIndex });
+				AddExisting(Type, CasedIndex, SymbolIndex);
+				return;
+			}
+
+			SymbolIndex = Type.TypeIndex;
+			this.Symbols[SymbolIndex] = new Symbol { Type = Type, MatchOptions = Type.EngineType.FindOptions(), CasedIndex = CasedIndex, NextIndex = 0, LastIndex = SymbolIndex };
+			this.CaselessDictionary.Add(Name, SymbolIndex);
+			this.CasedDictionary.Add(Name, new Lookup { SymbolIndex = SymbolIndex, CasedIndex = CasedIndex });
+		}
+
+		/// <summary>
+		/// Replace an entry in the symbol table.  This is used during property resolution to replace the 
+		/// parser property (which could not resolve the property prior to the symbol table being created)
+		/// with the fully resolved property.
+		/// </summary>
+		/// <param name="OldType">The old type being replaced.</param>
+		/// <param name="NewType">The new type.</param>
+		/// <param name="Name">The name of the type.</param>
+		/// <exception cref="UhtIceException">Thrown if the symbol wasn't found.</exception>
+		public void Replace(UhtType OldType, UhtType NewType, string Name)
+		{
+			int SymbolIndex;
+			if (this.CaselessDictionary.TryGetValue(Name, out SymbolIndex))
+			{
+				for (; SymbolIndex != 0; SymbolIndex = this.Symbols[SymbolIndex].NextIndex)
+				{
+					if (this.Symbols[SymbolIndex].Type == OldType)
+					{
+						this.Symbols[SymbolIndex].Type = NewType;
+						return;
+					}
+				}
+			}
+			throw new UhtIceException("Attempt to replace a type that wasn't found");
+		}
+
+		/// <summary>
+		/// Lookup the given name using cased string compare.
+		/// </summary>
+		/// <param name="StartingType">Starting type used to limit the scope of the search.</param>
+		/// <param name="Options">Options controlling what is search and what is returned.</param>
+		/// <param name="Name">Name to locate.</param>
+		/// <returns>Found type or null if not found.</returns>
+		public UhtType? FindCasedType(UhtType? StartingType, UhtFindOptions Options, string Name)
+		{
+			Lookup Existing;
+			if (this.CasedDictionary.TryGetValue(Name, out Existing))
+			{
+				return FindType(StartingType, Options, Existing);
+			}
+			return null;
+		}
+
+		/// <summary>
+		/// Lookup the given name using caseless string compare.
+		/// </summary>
+		/// <param name="StartingType">Starting type used to limit the scope of the search.</param>
+		/// <param name="Options">Options controlling what is search and what is returned.</param>
+		/// <param name="Name">Name to locate.</param>
+		/// <returns>Found type or null if not found.</returns>
+		public UhtType? FindCaselessType(UhtType? StartingType, UhtFindOptions Options, string Name)
+		{
+			int SymbolIndex;
+			if (this.CaselessDictionary.TryGetValue(Name, out SymbolIndex))
+			{
+				return FindType(StartingType, Options, new Lookup { SymbolIndex = SymbolIndex, CasedIndex = 0 });
+			}
+			return null;
+		}
+
+		/// <summary>
+		/// Lookup the given name.
+		/// </summary>
+		/// <param name="StartingType">Starting type used to limit the scope of the search.</param>
+		/// <param name="Options">Options controlling what is search and what is returned.</param>
+		/// <param name="Lookup">Starting lookup location.</param>
+		/// <returns>Found type or null if not found.</returns>
+		private UhtType? FindType(UhtType? StartingType, UhtFindOptions Options, Lookup Lookup)
+		{
+			if (StartingType != null)
+			{
+				UhtType? Found = FindTypeSuperChain(StartingType, Options, ref Lookup);
+				if (Found != null)
+				{
+					return Found;
+				}
+
+				Found = FindTypeOuterChain(StartingType, Options, ref Lookup);
+				if (Found != null)
+				{
+					return Found;
+				}
+
+				if (!Options.HasAnyFlags(UhtFindOptions.NoIncludes))
+				{
+					UhtHeaderFile HeaderFile = StartingType.HeaderFile;
+					foreach (UhtHeaderFile IncludedFile in HeaderFile.IncludedHeaders)
+					{
+						Found = this.FindSymbolChain(IncludedFile, Options, ref Lookup);
+						if (Found != null)
+						{
+							break;
+						}
+					}
+					if (Found != null)
+					{
+						return Found;
+					}
+				}
+			}
+
+			// Global search.  Match anything that has an owner of parent
+			if (!Options.HasAnyFlags(UhtFindOptions.NoGlobal))
+			{
+				for (int Index = Lookup.SymbolIndex; Index != 0; Index = this.Symbols[Index].NextIndex)
+				{
+					if (IsMatch(Options, Index, Lookup.CasedIndex) && this.Symbols[Index].Type.Outer is UhtHeaderFile)
+					{
+						return this.Symbols[Index].Type;
+					}
+				}
+			}
+
+			// Can't find at all
+			return null;
+		}
+
+		/// <summary>
+		/// Lookup the given name using the super class/struct chain.
+		/// </summary>
+		/// <param name="StartingType">Starting type used to limit the scope of the search.</param>
+		/// <param name="Options">Options controlling what is search and what is returned.</param>
+		/// <param name="Lookup">Starting lookup location.</param>
+		/// <returns>Found type or null if not found.</returns>
+		private UhtType? FindTypeSuperChain(UhtType StartingType, UhtFindOptions Options, ref Lookup Lookup)
+		{
+			UhtType? CurrentType = StartingType;
+
+			// In a super chain search, we have to start at a UHTStruct that contributes to the symbol table
+			for (; CurrentType != null; CurrentType = CurrentType.Outer)
+			{
+				if (CurrentType is UhtStruct Struct && Struct.EngineType.AddChildrenToSymbolTable())
+				{
+					break;
+				}
+			}
+
+			// Not symbol that supports a super chain
+			if (CurrentType == null)
+			{
+				return null;
+			}
+
+			// If requested, skip self
+			if (Options.HasAnyFlags(UhtFindOptions.NoSelf) && CurrentType == StartingType)
+			{
+				CurrentType = ((UhtStruct)CurrentType).Super;
+			}
+
+			// Search the chain
+			for (; CurrentType != null; CurrentType = ((UhtStruct)CurrentType).Super)
+			{
+				UhtType? FoundType = FindSymbolChain(CurrentType, Options, ref Lookup);
+				if (FoundType != null)
+				{
+					return FoundType;
+				}
+				if (Options.HasAnyFlags(UhtFindOptions.NoParents))
+				{
+					return null;
+				}
+			}
+			return null;
+		}
+
+		/// <summary>
+		/// Lookup the given name using the outer chain
+		/// </summary>
+		/// <param name="StartingType">Starting type used to limit the scope of the search.</param>
+		/// <param name="Options">Options controlling what is search and what is returned.</param>
+		/// <param name="Lookup">Starting lookup location.</param>
+		/// <returns>Found type or null if not found.</returns>
+		private UhtType? FindTypeOuterChain(UhtType StartingType, UhtFindOptions Options, ref Lookup Lookup)
+		{
+			UhtType? CurrentType = StartingType;
+
+			// If requested, skip self
+			if (Options.HasAnyFlags(UhtFindOptions.NoSelf) && CurrentType == StartingType)
+			{
+				CurrentType = CurrentType.Outer;
+			}
+
+			// Search the chain
+			for (; CurrentType != null; CurrentType = CurrentType.Outer)
+			{
+				if (CurrentType is UhtPackage)
+				{
+					return null;
+				}
+				UhtType? FoundType = FindSymbolChain(CurrentType, Options, ref Lookup);
+				if (FoundType != null)
+				{
+					return FoundType;
+				}
+				if (Options.HasAnyFlags(UhtFindOptions.NoOuter))
+				{
+					return null;
+				}
+			}
+			return null;
+		}
+
+		/// <summary>
+		/// Lookup the given name 
+		/// </summary>
+		/// <param name="Owner">Matching owner.</param>
+		/// <param name="Options">Options controlling what is search and what is returned.</param>
+		/// <param name="Lookup">Starting lookup location.</param>
+		/// <returns>Found type or null if not found.</returns>
+		private UhtType? FindSymbolChain(UhtType Owner, UhtFindOptions Options, ref Lookup Lookup)
+		{
+			for (int Index = Lookup.SymbolIndex; Index != 0; Index = this.Symbols[Index].NextIndex)
+			{
+				if (IsMatch(Options, Index, Lookup.CasedIndex) && this.Symbols[Index].Type.Outer == Owner)
+				{
+					return this.Symbols[Index].Type;
+				}
+			}
+			return null;
+		}
+
+		/// <summary>
+		/// Add a new type to the given symbol chain
+		/// </summary>
+		/// <param name="Type">Type being added</param>
+		/// <param name="CasedIndex">Cased index</param>
+		/// <param name="SymbolIndex">Symbol index</param>
+		private void AddExisting(UhtType Type, int CasedIndex, int SymbolIndex)
+		{
+			int TypeIndex = Type.TypeIndex;
+			this.Symbols[TypeIndex] = new Symbol { Type = Type, MatchOptions = Type.EngineType.FindOptions(), CasedIndex = CasedIndex, NextIndex = 0, LastIndex = 0 };
+			this.Symbols[this.Symbols[SymbolIndex].LastIndex].NextIndex = TypeIndex;
+			this.Symbols[SymbolIndex].LastIndex = TypeIndex;
+		}
+
+		/// <summary>
+		/// Test to see if the given symbol matches the options
+		/// </summary>
+		/// <param name="Options">Options to match</param>
+		/// <param name="SymbolIndex">Symbol index</param>
+		/// <param name="CasedIndex">Case index</param>
+		/// <returns>True if the symbol is a match</returns>
+		private bool IsMatch(UhtFindOptions Options, int SymbolIndex, int CasedIndex)
+		{
+			return (CasedIndex == 0 || CasedIndex == this.Symbols[SymbolIndex].CasedIndex) && 
+				this.Symbols[SymbolIndex].Type.bVisibleType && 
+				(this.Symbols[SymbolIndex].MatchOptions & Options) != 0;
+		}
+	}
+}
diff --git a/Engine/Source/Programs/Shared/EpicGames.UHT/Utils/UhtTables.cs b/Engine/Source/Programs/Shared/EpicGames.UHT/Utils/UhtTables.cs
new file mode 100644
--- /dev/null
+++ b/Engine/Source/Programs/Shared/EpicGames.UHT/Utils/UhtTables.cs
@@ -0,0 +1,405 @@
+// Copyright Epic Games, Inc. All Rights Reserved.
+
+using EpicGames.Core;
+using EpicGames.UHT.Tables;
+using System;
+using System.Collections.Generic;
+using System.Diagnostics.CodeAnalysis;
+using System.Reflection;
+
+namespace EpicGames.UHT.Utils
+{
+
+	/// <summary>
+	/// Defines all the table names for the standard UHT types
+	/// </summary>
+	public static class UhtTableNames
+	{
+		/// <summary>
+		/// The class base table is common to UCLASS, UINTERFACE and native interfaces
+		/// </summary>
+		public const string ClassBase = "ClassBase";
+
+		/// <summary>
+		/// Table for UCLASS
+		/// </summary>
+		public const string Class = "Class";
+
+		/// <summary>
+		/// Default table applied to all other tables
+		/// </summary>
+		public const string Default = "Default";
+
+		/// <summary>
+		/// Table for UENUM
+		/// </summary>
+		public const string Enum = "Enum";
+
+		/// <summary>
+		/// Table for all types considered a UField
+		/// </summary>
+		public const string Field = "Field";
+
+		/// <summary>
+		/// Table for all functions
+		/// </summary>
+		public const string Function = "Function";
+
+		/// <summary>
+		/// Table for the global/file scope
+		/// </summary>
+		public const string Global = "Global";
+
+		/// <summary>
+		/// Table for UINTERFACE
+		/// </summary>
+		public const string Interface = "Interface";
+
+		/// <summary>
+		/// Table for interfaces
+		/// </summary>
+		public const string NativeInterface = "NativeInterface";
+
+		/// <summary>
+		/// Table for all UObject types
+		/// </summary>
+		public const string Object = "Object";
+
+		/// <summary>
+		/// Table for properties that are function arguments
+		/// </summary>
+		public const string PropertyArgument = "PropertyArgument";
+
+		/// <summary>
+		/// Table for properties that are struct/class members
+		/// </summary>
+		public const string PropertyMember = "PropertyMember";
+
+		/// <summary>
+		/// Table for USTRUCT
+		/// </summary>
+		public const string ScriptStruct = "ScriptStruct";
+
+		/// <summary>
+		/// Table for all UStruct objects (structs, classes, and functions)
+		/// </summary>
+		public const string Struct = "Struct";
+	}
+
+	/// <summary>
+	/// Base class for table lookup system.
+	/// </summary>
+	public class UhtLookupTableBase
+	{
+		/// <summary>
+		/// This table inherits entries for the given table
+		/// </summary>
+		public UhtLookupTableBase? ParentTable = null;
+
+		/// <summary>
+		/// Name of the table
+		/// </summary>
+		public string TableName = string.Empty;
+
+		/// <summary>
+		/// User facing name of the table
+		/// </summary>
+		public string UserName
+		{
+			get => string.IsNullOrEmpty(this.UserNameInternal) ? this.TableName : this.UserNameInternal;
+			set => this.UserNameInternal = value;
+		}
+
+		/// <summary>
+		/// Check to see if the table is internal
+		/// </summary>
+		public bool Internal = false;
+
+		/// <summary>
+		/// If true, this table has been implemented and not just created on demand by another table
+		/// </summary>
+		public bool Implemented = false;
+
+		/// <summary>
+		/// Internal version of the user name.  If it hasn't been set, then the table name will be used
+		/// </summary>
+		private string UserNameInternal = string.Empty;
+
+		/// <summary>
+		/// Merge the lookup table.  Duplicates will be ignored.
+		/// </summary>
+		/// <param name="BaseTable">Base table being merged</param>
+		public virtual void Merge(UhtLookupTableBase BaseTable)
+		{
+		}
+
+		/// <summary>
+		/// Given a method name, try to extract the entry name for a table
+		/// </summary>
+		/// <param name="ClassType">Class containing the method</param>
+		/// <param name="MethodInfo">Method information</param>
+		/// <param name="InName">Optional name supplied by the attributes.  If specified, this name will be returned instead of extracted from the method name</param>
+		/// <param name="Suffix">Required suffix</param>
+		/// <returns>The extracted name or the supplied name</returns>
+		public static string GetSuffixedName(Type ClassType, MethodInfo MethodInfo, string? InName, string Suffix)
+		{
+			string Name = InName ?? string.Empty;
+			if (string.IsNullOrEmpty(Name))
+			{
+				if (MethodInfo.Name.EndsWith(Suffix))
+				{
+					Name = MethodInfo.Name.Substring(0, MethodInfo.Name.Length - Suffix.Length);
+				}
+				else
+				{
+					throw new UhtIceException($"The '{Suffix}' attribute on the {ClassType.Name}.{MethodInfo.Name} method doesn't have a name specified or the method name doesn't end in '{2}'.");
+				}
+			}
+			return Name;
+		}
+	}
+
+	/// <summary>
+	/// Lookup tables provide a method of associating actions with given C++ keyword or UE specifier
+	/// </summary>
+	/// <typeparam name="TValue">Keyword or specifier information</typeparam>
+	public class UhtLookupTable<TValue> : UhtLookupTableBase
+	{
+
+		/// <summary>
+		/// Lookup dictionary for the specifiers
+		/// </summary>
+		private Dictionary<StringView, TValue> Lookup;
+
+		/// <summary>
+		/// Construct a new table
+		/// </summary>
+		public UhtLookupTable(StringViewComparer Comparer)
+		{
+			this.Lookup = new Dictionary<StringView, TValue>(Comparer);
+		}
+
+		/// <summary>
+		/// Add the given value to the lookup table.  It will throw an exception if it is a duplicate.
+		/// </summary>
+		/// <param name="Key">Key to be added</param>
+		/// <param name="Value">Value to be added</param>
+		public UhtLookupTable<TValue> Add(string Key, TValue Value)
+		{
+			this.Lookup.Add(Key, Value);
+			return this;
+		}
+
+		/// <summary>
+		/// Attempt to fetch the value associated with the key
+		/// </summary>
+		/// <param name="Key">Lookup key</param>
+		/// <param name="Value">Value associated with the key</param>
+		/// <returns>True if the key was found, false if not</returns>
+		public bool TryGetValue(StringView Key, [MaybeNullWhen(false)] out TValue Value)
+		{
+			return this.Lookup.TryGetValue(Key, out Value);
+		}
+
+		/// <summary>
+		/// Merge the lookup table.  Duplicates will be ignored.
+		/// </summary>
+		/// <param name="BaseTable">Base table being merged</param>
+		public override void Merge(UhtLookupTableBase BaseTable)
+		{
+			foreach (var KVP in ((UhtLookupTable<TValue>)BaseTable).Lookup)
+			{
+				this.Lookup.TryAdd(KVP.Key, KVP.Value);
+			}
+		}
+	}
+
+	/// <summary>
+	/// A collection of lookup tables by name.
+	/// </summary>
+	/// <typeparam name="TTable">Table type</typeparam>
+	public class UhtLookupTables<TTable> where TTable : UhtLookupTableBase, new()
+	{
+
+		/// <summary>
+		/// Collection of named tables
+		/// </summary>
+		public Dictionary<string, TTable> Tables = new Dictionary<string, TTable>();
+
+		/// <summary>
+		/// The name of the group of tables
+		/// </summary>
+		public string Name;
+
+		/// <summary>
+		/// Create a new group of tables
+		/// </summary>
+		/// <param name="Name">The name of the group</param>
+		public UhtLookupTables(string Name)
+		{
+			this.Name = Name;
+		}
+
+		/// <summary>
+		/// Given a table name, return the table.  If not found, a new one will be added with the given name.
+		/// </summary>
+		/// <param name="TableName">The name of the table to return</param>
+		/// <returns>The table associated with the name.</returns>
+		public TTable Get(string TableName)
+		{
+			TTable? Table;
+			if (!this.Tables.TryGetValue(TableName, out Table))
+			{
+				Table = new TTable();
+				Table.TableName = TableName;
+				this.Tables.Add(TableName, Table);
+			}
+			return Table;
+		}
+
+		/// <summary>
+		/// Create a table with the given information.  If the table already exists, it will be initialized with the given data.
+		/// </summary>
+		/// <param name="TableName">The name of the table</param>
+		/// <param name="UserName">The user facing name of the name</param>
+		/// <param name="ParentTableName">The parent table name used to merge table entries</param>
+		/// <param name="bInternal">If true, the table is internal and won't be visible to the user.</param>
+		/// <returns>The created table</returns>
+		public TTable Create(string TableName, string UserName, string? ParentTableName, bool bInternal = false)
+		{
+			TTable Table = Get(TableName);
+			Table.UserName = UserName;
+			Table.Internal = bInternal;
+			Table.Implemented = true;
+			if (!string.IsNullOrEmpty(ParentTableName))
+			{
+				Table.ParentTable = Get(ParentTableName);
+			}
+			return Table;
+		}
+
+		/// <summary>
+		/// Merge the contents of all parent tables into their children.  This is done so that the 
+		/// parent chain need not be searched when looking for table entries.
+		/// </summary>
+		/// <exception cref="UhtIceException">Thrown if there are problems with the tables.</exception>
+		public void Merge()
+		{
+			List<TTable> OrderedList = new List<TTable>(this.Tables.Count);
+			List<TTable> RemainingList = new List<TTable>(this.Tables.Count);
+			HashSet<UhtLookupTableBase> DoneTables = new HashSet<UhtLookupTableBase>();
+
+			// Collect all of the tables
+			foreach (KeyValuePair<string, TTable> KVP in this.Tables)
+			{
+				if (!KVP.Value.Implemented)
+				{
+					throw new UhtIceException($"{this.Name} table '{KVP.Value.TableName}' has been referenced but not implemented");
+				}
+				RemainingList.Add(KVP.Value);
+			}
+
+			// Perform a topological sort of the tables
+			while (RemainingList.Count != 0)
+			{
+				bool bAddedOne = false;
+				for (int i = 0; i < RemainingList.Count;)
+				{
+					TTable Table = RemainingList[i];
+					if (Table.ParentTable == null || DoneTables.Contains(Table.ParentTable))
+					{
+						OrderedList.Add(Table);
+						DoneTables.Add(Table);
+						RemainingList[i] = RemainingList[RemainingList.Count - 1];
+						RemainingList.RemoveAt(RemainingList.Count - 1);
+						bAddedOne = true;
+					}
+					else
+					{
+						++i;
+					}
+				}
+				if (!bAddedOne)
+				{
+					throw new UhtIceException($"Circular dependency in {this.GetType().Name}.{this.Name} tables detected");
+				}
+			}
+
+			// Merge the tables
+			foreach (TTable Table in OrderedList)
+			{
+				if (Table.ParentTable != null)
+				{
+					Table.Merge((TTable)Table.ParentTable);
+				}
+			}
+		}
+	}
+
+	/// <summary>
+	/// Bootstraps the standard UHT tables
+	/// </summary>
+	[UnrealHeaderTool(InitMethod= "InitStandardTables")]
+	public static class UhtStandardTables
+	{
+
+		/// <summary>
+		/// Enumeration that specifies if the table is a specifier and/or keyword table
+		/// </summary>
+		[Flags]
+		public enum EUhtCreateTablesFlags
+		{
+			/// <summary>
+			/// A keyword table will be created
+			/// </summary>
+			Keyword = 1 << 0,
+
+			/// <summary>
+			/// A specifier table will be created
+			/// </summary>
+			Specifiers = 1 << 1,
+		}
+
+		/// <summary>
+		/// Create all of the standard scope tables.
+		/// </summary>
+		public static void InitStandardTables()
+		{
+			CreateTables(UhtTableNames.Default, "Default", EUhtCreateTablesFlags.Specifiers | EUhtCreateTablesFlags.Keyword, null, true);
+			CreateTables(UhtTableNames.Global, "Global", EUhtCreateTablesFlags.Keyword, UhtTableNames.Default, false);
+			CreateTables(UhtTableNames.PropertyArgument, "Argument/Return", EUhtCreateTablesFlags.Specifiers, UhtTableNames.Default, false);
+			CreateTables(UhtTableNames.PropertyMember, "Member", EUhtCreateTablesFlags.Specifiers, UhtTableNames.Default, false);
+			CreateTables(UhtTableNames.Object, "Object", EUhtCreateTablesFlags.Specifiers | EUhtCreateTablesFlags.Keyword, UhtTableNames.Default, true);
+			CreateTables(UhtTableNames.Field, "Field", EUhtCreateTablesFlags.Specifiers | EUhtCreateTablesFlags.Keyword, UhtTableNames.Object, true);
+			CreateTables(UhtTableNames.Enum, "Enum", EUhtCreateTablesFlags.Specifiers | EUhtCreateTablesFlags.Keyword, UhtTableNames.Field, false);
+			CreateTables(UhtTableNames.Struct, "Struct", EUhtCreateTablesFlags.Specifiers | EUhtCreateTablesFlags.Keyword, UhtTableNames.Field, true);
+			CreateTables(UhtTableNames.ClassBase, "ClassBase", EUhtCreateTablesFlags.Specifiers | EUhtCreateTablesFlags.Keyword, UhtTableNames.Struct, true);
+			CreateTables(UhtTableNames.Class, "Class", EUhtCreateTablesFlags.Specifiers | EUhtCreateTablesFlags.Keyword, UhtTableNames.ClassBase, false);
+			CreateTables(UhtTableNames.Interface, "Interface", EUhtCreateTablesFlags.Specifiers | EUhtCreateTablesFlags.Keyword, UhtTableNames.ClassBase, false);
+			CreateTables(UhtTableNames.NativeInterface, "IInterface", EUhtCreateTablesFlags.Specifiers | EUhtCreateTablesFlags.Keyword, UhtTableNames.ClassBase, false);
+			CreateTables(UhtTableNames.Function, "Function", EUhtCreateTablesFlags.Specifiers | EUhtCreateTablesFlags.Keyword, UhtTableNames.Struct, false);
+			CreateTables(UhtTableNames.ScriptStruct, "Struct", EUhtCreateTablesFlags.Specifiers | EUhtCreateTablesFlags.Keyword, UhtTableNames.Struct, false);
+		}
+
+		/// <summary>
+		/// Creates a series of tables given the supplied setup
+		/// </summary>
+		/// <param name="TableName">Name of the table.</param>
+		/// <param name="TableUserName">Name presented to the user via error messages.</param>
+		/// <param name="CreateTables">Types of tables to be created.</param>
+		/// <param name="ParentTableName">Name of the parent table or null for none.</param>
+		/// <param name="bInternal">If true, this table will not be included in any error messages.</param>
+		public static void CreateTables(string TableName, string TableUserName, EUhtCreateTablesFlags CreateTables, string? ParentTableName, bool bInternal = false)
+		{
+			if (CreateTables.HasFlag(EUhtCreateTablesFlags.Keyword))
+			{
+				UhtKeywordTables.Instance.Create(TableName, TableUserName, ParentTableName, bInternal);
+			}
+			if (CreateTables.HasFlag(EUhtCreateTablesFlags.Specifiers))
+			{
+				UhtSpecifierTables.Instance.Create(TableName, TableUserName, ParentTableName, bInternal);
+				UhtSpecifierValidatorTables.Instance.Create(TableName, TableUserName, ParentTableName, bInternal);
+			}
+		}
+	}
+}
diff --git a/Engine/Source/Programs/Shared/EpicGames.UHT/Utils/UhtUtilities.cs b/Engine/Source/Programs/Shared/EpicGames.UHT/Utils/UhtUtilities.cs
new file mode 100644
--- /dev/null
+++ b/Engine/Source/Programs/Shared/EpicGames.UHT/Utils/UhtUtilities.cs
@@ -0,0 +1,309 @@
+// Copyright Epic Games, Inc. All Rights Reserved.
+
+using EpicGames.Core;
+using System;
+using System.Collections.Generic;
+using System.Linq;
+using System.Text;
+
+namespace EpicGames.UHT.Utils
+{
+
+	/// <summary>
+	/// Extensions to List that provides some uniqueness support to list elements 
+	/// </summary>
+	public static class UhtListExtensions
+	{
+
+		/// <summary>
+		/// Add the given value if it isn't already contained within the list
+		/// </summary>
+		/// <param name="Container">Destination container</param>
+		/// <param name="Value">Value to be added</param>
+		/// <returns>True if the value was added, false if it was already present</returns>
+		public static bool AddUnique(this List<string> Container, string Value)
+		{
+			if (Container.Contains(Value, StringComparer.OrdinalIgnoreCase))
+			{
+				return false;
+			}
+			Container.Add(Value);
+			return true;
+		}
+
+		/// <summary>
+		/// Add the given values if they aren't already contained within the list
+		/// </summary>
+		/// <param name="Container">Destination container</param>
+		/// <param name="Values">Values to be added</param>
+		public static void AddUniqueRange(this List<string> Container, IEnumerable<StringView>? Values)
+		{
+			if (Values != null)
+			{
+				foreach (StringView Value in Values)
+				{
+					AddUnique(Container, Value.ToString());
+				}
+			}
+		}
+
+		/// <summary>
+		/// Remove the given value but swap the last entry into the eliminated slot
+		/// </summary>
+		/// <param name="Container">Container being modified</param>
+		/// <param name="Value">Value to be removed</param>
+		/// <returns>True if the value was removed, false if not</returns>
+		public static bool RemoveSwap(this List<string> Container, string Value)
+		{
+			int Index = Container.FindIndex(n => Value.Equals(n, StringComparison.OrdinalIgnoreCase));
+			if (Index >= 0)
+			{
+				if (Index != Container.Count - 1)
+				{
+					Container[Index] = Container[Container.Count - 1];
+				}
+				Container.RemoveAt(Container.Count - 1);
+				return true;
+			}
+			return false;
+		}
+
+		/// <summary>
+		/// Remove a range of values from a container using swapping
+		/// </summary>
+		/// <param name="Container">Container to be modified</param>
+		/// <param name="Values">List of values to be removed</param>
+		public static void RemoveSwapRange(this List<string> Container, IEnumerable<StringView>? Values)
+		{
+			if (Values != null)
+			{
+				foreach (StringView Value in Values)
+				{
+					RemoveSwap(Container, Value.ToString());
+				}
+			}
+		}
+	}
+
+	/// <summary>
+	/// UnrealEngine names often differ from the names in the source file.  The following 
+	/// structure represents the different parts of the name
+	/// </summary>
+	public struct UhtEngineNameParts
+	{
+
+		/// <summary>
+		/// Any prefix removed from the source name to create the engine name
+		/// </summary>
+		public StringView Prefix;
+
+		/// <summary>
+		/// The name to be used by the Unreal Engine.
+		/// </summary>
+		public StringView EngineName;
+
+		/// <summary>
+		/// The name contained the "DEPRECATED" text which has been removed from the engine name
+		/// </summary>
+		public bool bIsDeprecated;
+	}
+
+	/// <summary>
+	/// Assorted utility functions
+	/// </summary>
+	public class UhtUtilities
+	{
+
+		/// <summary>
+		/// Given a collection of names, return a string containing the text of those names concatenated.
+		/// </summary>
+		/// <param name="TypeNames">Collect of names to be merged</param>
+		/// <param name="AndOr">Text used to separate the names</param>
+		/// <param name="bQuote">If true, add quotes around the names</param>
+		/// <returns>Merged names</returns>
+		public static string MergeTypeNames(IEnumerable<string> TypeNames, string AndOr, bool bQuote = false)
+		{
+			List<string> Local = new List<string>(TypeNames);
+
+			if (Local.Count == 0)
+			{
+				return "";
+			}
+
+			Local.Sort();
+
+			StringBuilder Builder = new StringBuilder();
+			for (int Index = 0; Index < Local.Count; ++Index)
+			{
+				if (Index != 0)
+				{
+					Builder.Append(", ");
+					if (Index == Local.Count - 1)
+					{
+						Builder.Append(AndOr);
+						Builder.Append(' ');
+					}
+				}
+				if (bQuote)
+				{
+					Builder.Append("'");
+					Builder.Append(Local[Index]);
+					Builder.Append("'");
+				}
+				else
+				{
+					Builder.Append(Local[Index]);
+				}
+			}
+			return Builder.ToString();
+		}
+
+		/// <summary>
+		/// Split the given source name into the engine name parts
+		/// </summary>
+		/// <param name="SourceName">Source name</param>
+		/// <returns>Resulting engine name parts</returns>
+		public static UhtEngineNameParts GetEngineNameParts(StringView SourceName)
+		{
+			if (SourceName.Span.Length == 0)
+			{
+				return new UhtEngineNameParts { Prefix = new StringView(string.Empty), EngineName = new StringView(string.Empty), bIsDeprecated = false };
+			}
+
+			switch (SourceName.Span[0])
+			{
+				case 'I':
+				case 'A':
+				case 'U':
+					// If it is a class prefix, check for deprecated class prefix also
+					if (SourceName.Span.Length > 12 && SourceName.Span.Slice(1).StartsWith("DEPRECATED_"))
+					{
+						return new UhtEngineNameParts { Prefix = new StringView(SourceName, 0, 12), EngineName = new StringView(SourceName, 12), bIsDeprecated = true };
+					}
+					else
+					{
+						return new UhtEngineNameParts { Prefix = new StringView(SourceName, 0, 1), EngineName = new StringView(SourceName, 1), bIsDeprecated = false };
+					}
+
+				case 'F':
+				case 'T':
+					// Struct prefixes are also fine.
+					return new UhtEngineNameParts { Prefix = new StringView(SourceName, 0, 1), EngineName = new StringView(SourceName, 1), bIsDeprecated = false };
+
+				default:
+					// If it's not a class or struct prefix, it's invalid
+					return new UhtEngineNameParts { Prefix = new StringView(string.Empty), EngineName = new StringView(SourceName), bIsDeprecated = false };
+			}
+		}
+	}
+
+	/// <summary>
+	/// String builder class that has support for StringView so that if a single instance of
+	/// a StringView is appended, it is returned.
+	/// </summary>
+	public class StringViewBuilder
+	{
+
+		/// <summary>
+		/// When only a string view has been appended, this references that StringView
+		/// </summary>
+		private StringView StringViewInternal = new StringView();
+
+		/// <summary>
+		/// Represents more complex data being appended
+		/// </summary>
+		private StringBuilder? StringBuilderInternal = null;
+
+		/// <summary>
+		/// Set to true when the builder has switched to a StringBuilder (NOTE: This can probably be removed)
+		/// </summary>
+		private bool bUseStringBuilder = false;
+
+		/// <summary>
+		/// The length of the appended data
+		/// </summary>
+		public int Length
+		{
+			get
+			{
+				if (this.bUseStringBuilder && this.StringBuilderInternal != null)
+				{
+					return this.StringBuilderInternal.Length;
+				}
+				else
+				{
+					return this.StringViewInternal.Span.Length;
+				}
+			}
+		}
+
+		/// <summary>
+		/// Return the appended data as a StringView
+		/// </summary>
+		/// <returns>Contents of the builder</returns>
+		public StringView ToStringView()
+		{
+			return this.bUseStringBuilder ? new StringView(this.StringBuilderInternal!.ToString()) : this.StringViewInternal;
+		}
+
+		/// <summary>
+		/// Return the appended data as a string
+		/// </summary>
+		/// <returns>Contents of the builder</returns>
+		public override string ToString()
+		{
+			return this.bUseStringBuilder ? this.StringBuilderInternal!.ToString() : this.StringViewInternal.ToString();
+		}
+
+		/// <summary>
+		/// Append a StringView
+		/// </summary>
+		/// <param name="Text">Text to be appended</param>
+		/// <returns>The string builder</returns>
+		public StringViewBuilder Append(StringView Text)
+		{
+			if (this.bUseStringBuilder)
+			{
+				this.StringBuilderInternal!.Append(Text.Span);
+			}
+			else if (this.StringViewInternal.Span.Length > 0)
+			{
+				SwitchToStringBuilder();
+				this.StringBuilderInternal!.Append(Text.Span);
+			}
+			else
+			{
+				this.StringViewInternal = Text;
+			}
+			return this;
+		}
+
+		/// <summary>
+		/// Append a character
+		/// </summary>
+		/// <param name="C">Character to be appended</param>
+		/// <returns>The string builder</returns>
+		public StringViewBuilder Append(char C)
+		{
+			SwitchToStringBuilder();
+			this.StringBuilderInternal!.Append(C);
+			return this;
+		}
+
+		/// <summary>
+		/// If not already done, switch the builder to using a StringBuilder
+		/// </summary>
+		private void SwitchToStringBuilder()
+		{
+			if (!this.bUseStringBuilder)
+			{
+				if (this.StringBuilderInternal == null)
+				{
+					this.StringBuilderInternal = new StringBuilder();
+				}
+				this.bUseStringBuilder = true;
+				this.StringBuilderInternal.Append(this.StringViewInternal.Span);
+			}
+		}
+	}
+}
diff --git a/Engine/Source/Programs/UnrealBuildTool.Tests/UnrealHeaderToolTests.cs b/Engine/Source/Programs/UnrealBuildTool.Tests/UnrealHeaderToolTests.cs
new file mode 100644
--- /dev/null
+++ b/Engine/Source/Programs/UnrealBuildTool.Tests/UnrealHeaderToolTests.cs
@@ -0,0 +1,30 @@
+// Copyright Epic Games, Inc. All Rights Reserved.
+
+using EpicGames.Core;
+using EpicGames.UHT.Utils;
+using Microsoft.VisualStudio.TestTools.UnitTesting;
+using UnrealBuildTool.Modes;
+
+namespace UnrealBuildToolTests
+{
+	[TestClass]
+	public class UnrealHeaderToolTests
+	{
+		[TestMethod]
+		public void Run()
+		{
+			string[] Arguments = new string[] { };
+			CommandLineArguments CommandLineArguments = new CommandLineArguments(Arguments);
+			UhtGlobalOptions Options = new UhtGlobalOptions(CommandLineArguments);
+
+			// Initialize the attributes
+			UhtAttributeScanner.Scan();
+
+			// Initialize the config
+			UhtConfigImpl.Read(CommandLineArguments);
+
+			// Run the tests
+			Assert.IsTrue(UhtTestHarness.RunTests(Options));
+		}
+	}
+}
diff --git a/Engine/Source/Programs/UnrealBuildTool/Configuration/BuildConfiguration.cs b/Engine/Source/Programs/UnrealBuildTool/Configuration/BuildConfiguration.cs
--- a/Engine/Source/Programs/UnrealBuildTool/Configuration/BuildConfiguration.cs
+++ b/Engine/Source/Programs/UnrealBuildTool/Configuration/BuildConfiguration.cs
@@ -130,6 +130,12 @@ namespace UnrealBuildTool
 		[XmlConfigFile(Category = "UEBuildConfiguration")]
 		public bool bForceDebugUnrealHeaderTool = false;
 
+		/// <summary>
+		/// If true, use UHT internal to UBT
+		/// </summary>
+		[XmlConfigFile(Category = "UEBuildConfiguration")]
+		public bool bUseBuiltInUnrealHeaderTool = false;
+
 		/// <summary>
 		/// Whether to skip compiling rules assemblies and just assume they are valid
 		/// </summary>
diff --git a/Engine/Source/Programs/UnrealBuildTool/Modes/UnrealHeaderToolMode.cs b/Engine/Source/Programs/UnrealBuildTool/Modes/UnrealHeaderToolMode.cs
new file mode 100644
--- /dev/null
+++ b/Engine/Source/Programs/UnrealBuildTool/Modes/UnrealHeaderToolMode.cs
@@ -0,0 +1,1098 @@
+// Copyright Epic Games, Inc. All Rights Reserved.
+
+using EpicGames.Core;
+using EpicGames.UHT.Tables;
+using EpicGames.UHT.Tokenizer;
+using EpicGames.UHT.Types;
+using EpicGames.UHT.Utils;
+using System;
+using System.Collections.Generic;
+using System.IO;
+using System.Reflection;
+using System.Text;
+using UnrealBuildBase;
+
+namespace UnrealBuildTool.Modes
+{
+
+	/// <summary>
+	/// Implement the UHT configuration interface.  Due to the configuration system being fairly embedded into
+	/// UBT, the implementation must be part of UBT.
+	/// </summary>
+	public class UhtConfigImpl : IUhtConfig
+	{
+		private readonly ConfigHierarchy Ini;
+
+		/// <summary>
+		/// Types that have been renamed, treat the old deprecated name as the new name for code generation
+		/// </summary>
+		private readonly IReadOnlyDictionary<StringView, StringView> TypeRedirectMap;
+
+		/// <summary>
+		/// Metadata that have been renamed, treat the old deprecated name as the new name for code generation
+		/// </summary>
+		private readonly IReadOnlyDictionary<string, string> MetaDataRedirectMap;
+
+		/// <summary>
+		/// Supported units in the game
+		/// </summary>
+		private readonly ReadOnlyHashSet<StringView> Units;
+
+		/// <summary>
+		/// Special parsed struct names that do not require a prefix
+		/// </summary>
+		private readonly ReadOnlyHashSet<StringView> StructsWithNoPrefix;
+
+		/// <summary>
+		/// Special parsed struct names that have a 'T' prefix
+		/// </summary>
+		private readonly ReadOnlyHashSet<StringView> StructsWithTPrefix;
+
+		/// <summary>
+		/// Mapping from 'human-readable' macro substring to # of parameters for delegate declarations
+		/// Index 0 is 1 parameter, Index 1 is 2, etc...
+		/// </summary>
+		private readonly IReadOnlyList<StringView> DelegateParameterCountStrings;
+
+		/// <summary>
+		/// Default version of generated code. Defaults to oldest possible, unless specified otherwise in config.
+		/// </summary>
+		private readonly EGeneratedCodeVersion DefaultGeneratedCodeVersionInternal = EGeneratedCodeVersion.V1;
+
+		/// <summary>
+		/// Internal version of pointer warning for native pointers in the engine
+		/// </summary>
+		private readonly UhtPointerMemberBehavior EngineNativePointerMemberBehaviorInternal = UhtPointerMemberBehavior.AllowSilently;
+
+		/// <summary>
+		/// Internal version of pointer warning for object pointers in the engine
+		/// </summary>
+		private readonly UhtPointerMemberBehavior EngineObjectPtrMemberBehaviorInternal = UhtPointerMemberBehavior.AllowSilently;
+
+		/// <summary>
+		/// Internal version of pointer warning for native pointers outside the engine
+		/// </summary>
+		private readonly UhtPointerMemberBehavior NonEngineNativePointerMemberBehaviorInternal = UhtPointerMemberBehavior.AllowSilently;
+
+		/// <summary>
+		/// Internal version of pointer warning for object pointers outside the engine
+		/// </summary>
+		private readonly UhtPointerMemberBehavior NonEngineObjectPtrMemberBehaviorInternal = UhtPointerMemberBehavior.AllowSilently;
+
+		/// <summary>
+		/// If true setters and getters will be automatically (without specifying their function names on a property) parsed and generated if a function with matching signature is found
+		/// </summary>
+		private readonly bool bAllowAutomaticSettersAndGettersInternal;
+
+		#region IUhtConfig Implementation
+		/// <inheritdoc/>
+		public EGeneratedCodeVersion DefaultGeneratedCodeVersion => this.DefaultGeneratedCodeVersionInternal;
+
+		/// <inheritdoc/>
+		public UhtPointerMemberBehavior EngineNativePointerMemberBehavior => this.EngineNativePointerMemberBehaviorInternal;
+
+		/// <inheritdoc/>
+		public UhtPointerMemberBehavior EngineObjectPtrMemberBehavior => this.EngineObjectPtrMemberBehaviorInternal;
+
+		/// <inheritdoc/>
+		public UhtPointerMemberBehavior NonEngineNativePointerMemberBehavior => this.NonEngineNativePointerMemberBehaviorInternal;
+
+		/// <inheritdoc/>
+		public UhtPointerMemberBehavior NonEngineObjectPtrMemberBehavior => this.NonEngineObjectPtrMemberBehaviorInternal;
+
+		/// <inheritdoc/>
+		public bool bAllowAutomaticSettersAndGetters => this.bAllowAutomaticSettersAndGettersInternal;
+
+		/// <inheritdoc/>
+		public void RedirectTypeIdentifier(ref UhtToken Token)
+		{
+			if (!Token.IsIdentifier())
+			{
+				throw new Exception("Attempt to redirect type identifier when the token isn't an identifier.");
+			}
+
+			StringView Redirect;
+			if (this.TypeRedirectMap.TryGetValue(Token.Value, out Redirect))
+			{
+				Token.Value = Redirect;
+			}
+		}
+
+		/// <inheritdoc/>
+		public bool RedirectMetaDataKey(string Key, out string NewKey)
+		{
+			string? Redirect;
+			if (this.MetaDataRedirectMap.TryGetValue(Key, out Redirect))
+			{
+				NewKey = Redirect;
+				return Key != NewKey;
+			}
+			else
+			{
+				NewKey = Key;
+				return false;
+			}
+		}
+
+		/// <inheritdoc/>
+		public bool IsValidUnits(StringView Units)
+		{
+			return this.Units.Contains(Units);
+		}
+
+		/// <inheritdoc/>
+		public bool IsStructWithTPrefix(StringView Name)
+		{
+			return this.StructsWithTPrefix.Contains(Name);
+		}
+
+		/// <inheritdoc/>
+		public int FindDelegateParameterCount(StringView DelegateMacro)
+		{
+			IReadOnlyList<StringView> DelegateParameterCountStrings = this.DelegateParameterCountStrings;
+			for (int Index = 0, Count = this.DelegateParameterCountStrings.Count; Index < Count; ++Index)
+			{
+				if (DelegateMacro.Span.Contains(this.DelegateParameterCountStrings[Index].Span, StringComparison.Ordinal))
+				{
+					return Index;
+				}
+			}
+			return -1;
+		}
+
+		/// <inheritdoc/>
+		public StringView GetDelegateParameterCountString(int Index)
+		{
+			return Index >= 0 ? this.DelegateParameterCountStrings[Index] : "";
+		}
+
+		/// <inheritdoc/>
+		public bool IsExporterEnabled(string Name)
+		{
+			bool Value = false;
+			this.Ini.GetBool("UnrealHeaderTool", Name, out Value);
+			return Value;
+		}
+		#endregion
+
+		/// <summary>
+		/// Read the UHT configuration
+		/// </summary>
+		public static void Read(CommandLineArguments Args)
+		{
+			UhtConfigImpl Instance = new UhtConfigImpl(Args);
+			UhtConfig.Instance = Instance;
+		}
+
+		private UhtConfigImpl(CommandLineArguments Args)
+		{
+			DirectoryReference ConfigDirectory = DirectoryReference.Combine(Unreal.EngineDirectory, "Programs/UnrealBuildTool");
+			this.Ini = ConfigCache.ReadHierarchy(ConfigHierarchyType.Engine, ConfigDirectory, BuildHostPlatform.Current.Platform, "", Args.GetRawArray());
+
+			this.bAllowAutomaticSettersAndGettersInternal = GetBoolean("UnrealHeaderTool", "AutomaticSettersAndGetters", false);
+			this.TypeRedirectMap = GetRedirectsStringView("UnrealHeaderTool", "TypeRedirects", "OldType", "NewType");
+			this.MetaDataRedirectMap = GetRedirectsString("CoreUObject.Metadata", "MetadataRedirects", "OldKey", "NewKey");
+			this.StructsWithNoPrefix = GetHashSet("UnrealHeaderTool", "StructsWithNoPrefix", StringViewComparer.Ordinal);
+			this.StructsWithTPrefix = GetHashSet("UnrealHeaderTool", "StructsWithTPrefix", StringViewComparer.Ordinal);
+			this.Units = GetHashSet("UnrealHeaderTool", "Units", StringViewComparer.OrdinalIgnoreCase);
+			this.DelegateParameterCountStrings = GetList("UnrealHeaderTool", "DelegateParameterCountStrings");
+			this.DefaultGeneratedCodeVersionInternal = GetGeneratedCodeVersion("UnrealHeaderTool", "DefaultGeneratedCodeVersion", EGeneratedCodeVersion.V1);
+			this.EngineNativePointerMemberBehaviorInternal = GetPointerMemberBehavior("UnrealHeaderTool", "EngineNativePointerMemberBehavior", UhtPointerMemberBehavior.AllowSilently);
+			this.EngineObjectPtrMemberBehaviorInternal = GetPointerMemberBehavior("UnrealHeaderTool", "EngineObjectPtrMemberBehavior", UhtPointerMemberBehavior.AllowSilently);
+			this.NonEngineNativePointerMemberBehaviorInternal = GetPointerMemberBehavior("UnrealHeaderTool", "NonEngineNativePointerMemberBehavior", UhtPointerMemberBehavior.AllowSilently);
+			this.NonEngineObjectPtrMemberBehaviorInternal = GetPointerMemberBehavior("UnrealHeaderTool", "NonEngineObjectPtrMemberBehavior", UhtPointerMemberBehavior.AllowSilently);
+		}
+
+		private bool GetBoolean(string SectionName, string KeyName, bool bDefault)
+		{
+			bool bValue;
+			if (this.Ini.TryGetValue(SectionName, KeyName, out bValue))
+			{
+				return bValue;
+			}
+			return bDefault;
+		}
+
+		private UhtPointerMemberBehavior GetPointerMemberBehavior(string SectionName, string KeyName, UhtPointerMemberBehavior Default)
+		{
+			string? BehaviorStr;
+			if (this.Ini.TryGetValue(SectionName, KeyName, out BehaviorStr))
+			{
+				UhtPointerMemberBehavior Value;
+				if (!Enum.TryParse(BehaviorStr, out Value))
+				{
+					throw new Exception(string.Format("Unrecognized native pointer member behavior '{0}'", BehaviorStr));
+				}
+				return Value;
+			}
+			return Default;
+		}
+
+		private EGeneratedCodeVersion GetGeneratedCodeVersion(string SectionName, string KeyName, EGeneratedCodeVersion Default)
+		{
+			string? BehaviorStr;
+			if (this.Ini.TryGetValue(SectionName, KeyName, out BehaviorStr))
+			{
+				EGeneratedCodeVersion Value;
+				if (!Enum.TryParse(BehaviorStr, out Value))
+				{
+					throw new Exception(string.Format("Unrecognized generated code version '{0}'", BehaviorStr));
+				}
+				return Value;
+			}
+			return Default;
+		}
+
+		private IReadOnlyDictionary<StringView, StringView> GetRedirectsStringView(string Section, string Key, string OldKeyName, string NewKeyName)
+		{
+			Dictionary<StringView, StringView> Redirects = new Dictionary<StringView, StringView>();
+
+			IReadOnlyList<string>? StringList;
+			if (this.Ini.TryGetValues(Section, Key, out StringList))
+			{
+				foreach (string Line in StringList)
+				{
+					Dictionary<string, string>? Properties;
+					if (ConfigHierarchy.TryParse(Line, out Properties))
+					{
+						string? OldKey;
+						if (!Properties.TryGetValue(OldKeyName, out OldKey))
+						{
+							throw new Exception(string.Format("Unable to get the {0} from the {1} value", OldKeyName, Key));
+						}
+						string? NewKey;
+						if (!Properties.TryGetValue(NewKeyName, out NewKey))
+						{
+							throw new Exception(string.Format("Unable to get the {0} from the {1} value", NewKeyName, Key));
+						}
+						Redirects.Add(OldKey, NewKey);
+					}
+				}
+			}
+			return Redirects;
+		}
+
+		private IReadOnlyDictionary<string, string> GetRedirectsString(string Section, string Key, string OldKeyName, string NewKeyName)
+		{
+			Dictionary<string, string> Redirects = new Dictionary<string, string>();
+
+			IReadOnlyList<string>? StringList;
+			if (this.Ini.TryGetValues(Section, Key, out StringList))
+			{
+				foreach (string Line in StringList)
+				{
+					Dictionary<string, string>? Properties;
+					if (ConfigHierarchy.TryParse(Line, out Properties))
+					{
+						string? OldKey;
+						if (!Properties.TryGetValue(OldKeyName, out OldKey))
+						{
+							throw new Exception(string.Format("Unable to get the {0} from the {1} value", OldKeyName, Key));
+						}
+						string? NewKey;
+						if (!Properties.TryGetValue(NewKeyName, out NewKey))
+						{
+							throw new Exception(string.Format("Unable to get the {0} from the {1} value", NewKeyName, Key));
+						}
+						Redirects.Add(OldKey, NewKey);
+					}
+				}
+			}
+			return Redirects;
+		}
+
+		private IReadOnlyList<StringView> GetList(string Section, string Key)
+		{
+			List<StringView> List = new List<StringView>();
+
+			IReadOnlyList<string>? StringList;
+			if (this.Ini.TryGetValues(Section, Key, out StringList))
+			{
+				foreach (string Value in StringList)
+				{
+					List.Add(new StringView(Value));
+				}
+			}
+			return List;
+		}
+
+		private ReadOnlyHashSet<StringView> GetHashSet(string Section, string Key, StringViewComparer Comparer)
+		{
+			HashSet<StringView> Set = new HashSet<StringView>(Comparer);
+
+			IReadOnlyList<string>? StringList;
+			if (this.Ini.TryGetValues(Section, Key, out StringList))
+			{
+				foreach (string Value in StringList)
+				{
+					Set.Add(new StringView(Value));
+				}
+			}
+			return Set;
+		}
+	}
+
+	/// <summary>
+	/// Global options for UBT (any modes)
+	/// </summary>
+	class UhtGlobalOptions
+	{
+		/// <summary>
+		/// User asked for help
+		/// </summary>
+		[CommandLine(Prefix = "-Help", Description = "Display this help.")]
+		[CommandLine(Prefix = "-h")]
+		[CommandLine(Prefix = "--help")]
+		public bool bGetHelp = false;
+
+		/// <summary>
+		/// The amount of detail to write to the log
+		/// </summary>
+		[CommandLine(Prefix = "-Verbose", Value = "Verbose", Description = "Increase output verbosity")]
+		[CommandLine(Prefix = "-VeryVerbose", Value = "VeryVerbose", Description = "Increase output verbosity more")]
+		public LogEventType LogOutputLevel = LogEventType.Log;
+
+		/// <summary>
+		/// Specifies the path to a log file to write. Note that the default mode (eg. building, generating project files) will create a log file by default if this not specified.
+		/// </summary>
+		[CommandLine(Prefix = "-Log", Description = "Specify a log file location instead of the default Engine/Programs/UnrealBuildTool/Log_UHT.txt")]
+		public FileReference? LogFileName = null;
+
+		/// <summary>
+		/// Whether to include timestamps in the log
+		/// </summary>
+		[CommandLine(Prefix = "-Timestamps", Description = "Include timestamps in the log")]
+		public bool bLogTimestamps = false;
+
+		/// <summary>
+		/// Whether to format messages in MsBuild format
+		/// </summary>
+		[CommandLine(Prefix = "-FromMsBuild", Description = "Format messages for msbuild")]
+		public bool bLogFromMsBuild = false;
+
+		/// <summary>
+		/// Disables all logging including the default log location
+		/// </summary>
+		[CommandLine(Prefix = "-NoLog", Description = "Disable log file creation including the default log file")]
+		public bool bNoLog = false;
+
+		[CommandLine(Prefix = "-Test", Description = "Run testing scripts")]
+		public bool bTest = false;
+
+		[CommandLine("-WarningsAsErrors", Description = "Treat warnings as errors")]
+		public bool bWarningsAsErrors = false;
+
+		[CommandLine("-NoGoWide", Description = "Disable concurrent parsing and code generation")]
+		public bool bNoGoWide = false;
+
+		[CommandLine("-WriteRef", Description = "Write all the output to a reference directory")]
+		public bool bWriteRef = false;
+
+		[CommandLine("-VerifyRef", Description = "Write all the output to a verification directory and compare to the reference output")]
+		public bool bVerifyRef = false;
+
+		[CommandLine("-FailIfGeneratedCodeChanges", Description = "Consider any changes to output files as being an error")]
+		public bool bFailIfGeneratedCodeChanges = false;
+
+		[CommandLine("-NoOutput", Description = "Do not save any output files other than reference output")]
+		public bool bNoOutput = false;
+
+		[CommandLine("-IncludeDebugOutput", Description = "Include extra content in generated output to assist with debugging")]
+		public bool bIncludeDebugOutput = false;
+
+		/// <summary>
+		/// Initialize the options with the given command line arguments
+		/// </summary>
+		/// <param name="Arguments"></param>
+		public UhtGlobalOptions(CommandLineArguments Arguments)
+		{
+			Arguments.ApplyTo(this);
+		}
+	}
+
+	/// <summary>
+	/// File manager for the test harness
+	/// </summary>
+	public class UhtTestFileManager : IUhtFileManager
+	{
+		/// <summary>
+		/// Collection of test fragments that can be read
+		/// </summary>
+		public Dictionary<string, UhtSourceFragment> SourceFragments = new Dictionary<string, UhtSourceFragment>();
+
+		/// <summary>
+		/// All output segments generated by code gen
+		/// </summary>
+		public SortedDictionary<string, string> Outputs = new SortedDictionary<string, string>();
+
+		private readonly IUhtFileManager InnerManager;
+		private readonly string? RootDirectory;
+
+		/// <summary>
+		/// Construct a new instance of the test file manager
+		/// </summary>
+		/// <param name="RootDirectory">Root directory of the UE</param>
+		public UhtTestFileManager(string RootDirectory)
+		{
+			this.RootDirectory = RootDirectory;
+			this.InnerManager = new UhtStdFileManager();
+		}
+
+		/// <inheritdoc/>
+		public string GetFullFilePath(string FilePath)
+		{
+			if (this.RootDirectory == null)
+			{
+				return FilePath;
+			}
+			else
+			{
+				return Path.Combine(this.RootDirectory, FilePath);
+			}
+		}
+
+		/// <inheritdoc/>
+		public bool ReadSource(string FilePath, out UhtSourceFragment Fragment)
+		{
+			if (this.SourceFragments.TryGetValue(FilePath, out Fragment))
+			{
+				return true;
+			}
+
+			return InnerManager.ReadSource(GetFullFilePath(FilePath), out Fragment);
+		}
+
+		/// <inheritdoc/>
+		public UhtBuffer? ReadOutput(string FilePath)
+		{
+			return null;
+		}
+
+		/// <inheritdoc/>
+		public bool WriteOutput(string FilePath, ReadOnlySpan<char> Contents)
+		{
+			lock (this.Outputs)
+			{
+				this.Outputs.Add(FilePath, Contents.ToString());
+			}
+			return true;
+		}
+
+		/// <inheritdoc/>
+		public bool RenameOutput(string OldFilePath, string NewFilePath)
+		{
+			lock (this.Outputs)
+			{
+				if (this.Outputs.TryGetValue(OldFilePath, out string? Content))
+				{
+					this.Outputs.Remove(OldFilePath);
+					this.Outputs.Add(NewFilePath, Content);
+				}
+			}
+			return true;
+		}
+
+		/// <summary>
+		/// Add a source file fragment to the session.  When requests are made to read sources, the 
+		/// fragment list will be searched first.
+		/// </summary>
+		/// <param name="SourceFile">Source file</param>
+		/// <param name="FilePath">The relative path to add</param>
+		/// <param name="LineNumber">Starting line number</param>
+		/// <param name="Data">The data associated with the path</param>
+		public void AddSourceFragment(UhtSourceFile SourceFile, string FilePath, int LineNumber, StringView Data)
+		{
+			this.SourceFragments.Add(FilePath, new UhtSourceFragment { SourceFile = SourceFile, FilePath = FilePath, LineNumber = LineNumber, Data = Data });
+		}
+	}
+
+	/// <summary>
+	/// Testing harness to run the test scripts
+	/// </summary>
+	class UhtTestHarness
+	{
+		private enum ScriptFragmentType
+		{
+			Unknown,
+			Manifest,
+			Header,
+			Console,
+			Output,
+		}
+
+		private struct ScriptFragment
+		{
+			public ScriptFragmentType Type;
+			public string Name;
+			public int LineNumber;
+			public StringView Header;
+			public StringView Body;
+			public bool External;
+		}
+
+		private static bool RunScriptTest(UhtGlobalOptions Options, string TestDirectory, string TestOutputDirectory, string Script)
+		{
+			string InPath = Path.Combine(TestDirectory, Script);
+			string OutPath = Path.Combine(TestOutputDirectory, Script);
+
+			UhtTestFileManager TestFileManager = new UhtTestFileManager(TestDirectory);
+			UhtSession Session = new UhtSession
+			{
+				FileManager = TestFileManager,
+				RootDirectory = TestDirectory,
+				bWarningsAsErrors = Options.bWarningsAsErrors,
+				bRelativePathInLog = true,
+				bGoWide = !Options.bNoGoWide,
+				bNoOutput = false,
+				bCullOutput = false,
+				bCacheMessages = true,
+				bIncludeDebugOutput = true,
+			};
+
+			// Read the testing script
+			List<ScriptFragment> ScriptFragments = new List<ScriptFragment>();
+			int ManifestIndex = -1;
+			int ConsoleIndex = -1;
+			UhtSourceFile ScriptSourceFile = new UhtSourceFile(Session, Script);
+			Dictionary<string, int> OutputFragments = new Dictionary<string, int>();
+			Session.Try(ScriptSourceFile, () =>
+			{
+				ScriptSourceFile.Read();
+				UhtTokenBufferReader Reader = new UhtTokenBufferReader(ScriptSourceFile, ScriptSourceFile.Data.Memory);
+
+				bool bDone = false;
+				while (!bDone)
+				{
+
+					// Scan for the fragment header
+					ScriptFragmentType Type = ScriptFragmentType.Unknown;
+					string Name = "";
+					int HeaderStartPos = Reader.InputPos;
+					int HeaderEndPos = HeaderStartPos;
+					int LineNumber = 1;
+					while (true)
+					{
+						using (var SaveState = new UhtTokenSaveState(Reader))
+						{
+							UhtToken Token = Reader.GetLine();
+							if (Token.TokenType == UhtTokenType.EndOfFile)
+							{
+								break;
+							}
+							if (Token.Value.Span.Length == 0 || (Token.Value.Span.Length > 0 && Token.Value.Span[0] != '!'))
+							{
+								break;
+							}
+							HeaderEndPos = Reader.InputPos;
+
+							int EndCommandPos = Token.Value.Span.IndexOf(' ');
+							if (EndCommandPos == -1)
+							{
+								EndCommandPos = Token.Value.Span.Length;
+							}
+							string ScriptFragmentTypeString = Token.Value.Span.Slice(1, EndCommandPos - 1).Trim().ToString();
+
+							if (!System.Enum.TryParse<ScriptFragmentType>(ScriptFragmentTypeString, true, out Type))
+							{
+								continue;
+							}
+							if (Type == ScriptFragmentType.Unknown)
+							{
+								continue;
+							}
+
+							Name = Token.Value.Span.Slice(EndCommandPos).Trim().ToString();
+							LineNumber = Token.InputLine;
+							SaveState.AbandonState();
+							break;
+						}
+					}
+
+					// Scan for the fragment body
+					int BodyStartPos = Reader.InputPos;
+					int BodyEndPos = BodyStartPos;
+					while (true)
+					{
+						using (var SaveState = new UhtTokenSaveState(Reader))
+						{
+							UhtToken Token = Reader.GetLine();
+							if (Token.TokenType == UhtTokenType.EndOfFile)
+							{
+								bDone = true;
+								break;
+							}
+							if (Token.Value.Span.Length > 0 && Token.Value.Span[0] == '!')
+							{
+								break;
+							}
+							BodyEndPos = Reader.InputPos;
+							SaveState.AbandonState();
+						}
+					}
+
+					ScriptFragments.Add(new ScriptFragment
+					{
+						Type = Type,
+						Name = Name.Replace("\\\\", "\\"), // Be kind to people cut/copy/paste escaped strings around
+						LineNumber = LineNumber,
+						Header = new StringView(ScriptSourceFile.Data.Memory.Slice(HeaderStartPos, HeaderEndPos - HeaderStartPos)),
+						Body = new StringView(ScriptSourceFile.Data.Memory.Slice(BodyStartPos, BodyEndPos - BodyStartPos)),
+						External = false,
+					});
+				}
+
+				// Search for the manifest and any output.  Add fragments to the session
+				for (int i = 0; i < ScriptFragments.Count; ++i)
+				{
+					switch (ScriptFragments[i].Type)
+					{
+						case ScriptFragmentType.Manifest:
+							if (ManifestIndex != -1)
+							{
+								ScriptSourceFile.LogError(ScriptFragments[i].LineNumber, "There can be only one manifest section in a test script");
+								break;
+							}
+							ManifestIndex = i;
+							if (ScriptFragments[i].Name.Length == 0)
+							{
+								ScriptSourceFile.LogError(ScriptFragments[i].LineNumber, "Manifest name can not be blank");
+								break;
+							}
+							TestFileManager.AddSourceFragment(ScriptSourceFile, ScriptFragments[i].Name, ScriptFragments[i].LineNumber, ScriptFragments[i].Body);
+							break;
+
+						case ScriptFragmentType.Console:
+							if (ConsoleIndex != -1)
+							{
+								ScriptSourceFile.LogError(ScriptFragments[i].LineNumber, "There can be only one console section in a test script");
+								break;
+							}
+							ConsoleIndex = i;
+							break;
+
+						case ScriptFragmentType.Header:
+							if (ScriptFragments[i].Name.Length == 0)
+							{
+								ScriptSourceFile.LogError(ScriptFragments[i].LineNumber, "Header name can not be blank");
+								break;
+							}
+							if (ScriptFragments[i].Body.Length == 0)
+							{
+								// Read the NoExportTypes.h file from the engine source so we don't have to keep a copy
+								if (Path.GetFileName(ScriptFragments[i].Name).Equals("NoExportTypes.h", StringComparison.OrdinalIgnoreCase))
+								{
+									string ExternalPath = Path.Combine(Unreal.EngineDirectory.FullName, ScriptFragments[i].Name);
+									if (File.Exists(ExternalPath))
+									{
+										ScriptFragment Copy = ScriptFragments[i];
+										Copy.Body = new StringView(File.ReadAllText(ExternalPath));
+										Copy.External = true;
+										ScriptFragments[i] = Copy;
+									}
+								}
+							}
+							TestFileManager.AddSourceFragment(ScriptSourceFile, ScriptFragments[i].Name, ScriptFragments[i].LineNumber, ScriptFragments[i].Body);
+							break;
+
+						case ScriptFragmentType.Output:
+							OutputFragments.Add(ScriptFragments[i].Name, i);
+							break;
+					}
+				}
+
+				if (ManifestIndex == -1)
+				{
+					ScriptSourceFile.LogError("There must be a manifest section in a test script");
+				}
+
+				if (ConsoleIndex == -1)
+				{
+					ScriptSourceFile.LogError("There must be a console section in a test script");
+				}
+			});
+
+			// Run UHT
+			if (!Session.bHasErrors)
+			{
+				Session.Run(ScriptFragments[ManifestIndex].Name);
+			}
+
+			// If we have no console index, then there is nothing we can do.  This is a fatal error than can not be tested
+			bool bSuccess = true;
+			if (ConsoleIndex == -1)
+			{
+				ScriptSourceFile.LogError("Unable to do any verification without a console section");
+				Session.LogMessages();
+				File.Copy(InPath, OutPath, true);
+				bSuccess = false;
+			}
+			else
+			{
+
+				// Generate the console block
+				List<string> ConsoleLines = Session.CollectMessages();
+				StringBuilder SBConsole = new StringBuilder();
+				foreach (string Line in ConsoleLines)
+				{
+					SBConsole.AppendLine(Line);
+				}
+
+				// Verify the console block 
+				// We trim the ends because it is too easy to leave off the ending CRLF in the script file.
+				if (ScriptFragments[ConsoleIndex].Body.ToString().TrimEnd() != SBConsole.ToString().TrimEnd())
+				{
+					Log.TraceError("Console output failed to match");
+					bSuccess = false;
+				}
+
+				// Check the output
+				foreach (KeyValuePair<string, string> KVP in TestFileManager.Outputs)
+				{
+					if (OutputFragments.TryGetValue(KVP.Key, out int Index))
+					{
+						if (ScriptFragments[Index].Body.ToString().TrimEnd() != KVP.Value.TrimEnd())
+						{
+							Log.TraceError($"Output \"{KVP.Key}\" failed to match");
+							bSuccess = false;
+						}
+						OutputFragments.Remove(KVP.Key);
+					}
+					else
+					{
+						Log.TraceError($"Output \"{KVP.Key}\" not found in test script");
+					}
+				}
+				foreach (KeyValuePair<string, int> KVP in OutputFragments)
+				{
+					Log.TraceError($"Output \"{KVP.Key}\" in test script but not generated");
+				}
+
+				// Create the complete output.  Output includes all of the source fragments and console fragments
+				// and followed the output data sorted by file name.
+				StringBuilder SBTest = new StringBuilder();
+				for (int i = 0; i < ScriptFragments.Count; ++i)
+				{
+					if (ScriptFragments[i].Type != ScriptFragmentType.Output)
+					{
+						SBTest.Append(ScriptFragments[i].Header);
+						if (i == ConsoleIndex)
+						{
+							SBTest.Append(SBConsole);
+						}
+						else if (!ScriptFragments[i].External)
+						{
+							SBTest.Append(ScriptFragments[i].Body);
+						}
+					}
+				}
+
+				// Add the output
+				foreach (KeyValuePair<string, string> KVP in TestFileManager.Outputs)
+				{
+					SBTest.Append($"!output {KVP.Key}\r\n");
+					SBTest.Append(KVP.Value);
+				}
+
+				// Write the final content
+				try
+				{
+					File.WriteAllText(OutPath, SBTest.ToString());
+				}
+				catch (Exception E)
+				{
+					Log.TraceError("Unable to write test result to \"{0}\"", E.Message);
+				}
+			}
+
+			if (bSuccess)
+			{
+				Log.TraceInformation("Test {0} succeeded", InPath);
+			}
+			else
+			{
+				Log.TraceError("Test {0} failed", InPath);
+			}
+			return bSuccess;
+		}
+
+		private static bool RunManifestTests(UhtGlobalOptions Options, string TestDirectory, string TestOutputDirectory, List<string> Manifests)
+		{
+			return true;
+		}
+
+		private static bool RunScriptTests(UhtGlobalOptions Options, string TestDirectory, string TestOutputDirectory, List<string> Scripts)
+		{
+			bool bResult = true;
+			foreach (string Script in Scripts)
+			{
+				bResult &= RunScriptTest(Options, TestDirectory, TestOutputDirectory, Script);
+			}
+			return bResult;
+		}
+
+		private static bool RunDirectoryTests(UhtGlobalOptions Options, string TestDirectory, string TestOutputDirectory, List<string> Directories)
+		{
+			bool bResult = true;
+			foreach (string Directory in Directories)
+			{
+				bResult &= RunTests(Options, Path.Combine(TestDirectory, Directory), Path.Combine(TestOutputDirectory, Directory));
+			}
+			return bResult;
+		}
+
+		private static bool RunTests(UhtGlobalOptions Options, string TestDirectory, string TestOutputDirectory)
+		{
+			// Create output directory
+			Directory.CreateDirectory(TestOutputDirectory);
+
+			List<string> Scripts = new List<string>();
+			foreach (string Script in Directory.EnumerateFiles(TestDirectory, "*.uhttest"))
+			{
+				Scripts.Add(Path.GetFileName(Script));
+			}
+			Scripts.Sort(StringComparer.OrdinalIgnoreCase);
+
+			List<string> Directories = new List<string>();
+			foreach (string Directory in Directory.EnumerateDirectories(TestDirectory))
+			{
+				Directories.Add(Path.GetFileName(Directory));
+			}
+			Directories.Sort(StringComparer.OrdinalIgnoreCase);
+
+			List<string> Manifests = new List<string>();
+			foreach (string Manifest in Directory.EnumerateFiles(TestDirectory, "*.uhtmanifest"))
+			{
+				Manifests.Add(Path.GetFileName(Manifest));
+			}
+			Manifests.Sort(StringComparer.OrdinalIgnoreCase);
+
+			return
+				RunManifestTests(Options, TestDirectory, TestOutputDirectory, Manifests) &&
+				RunScriptTests(Options, TestDirectory, TestOutputDirectory, Scripts) &&
+				RunDirectoryTests(Options, TestDirectory, TestOutputDirectory, Directories);
+		}
+
+		public static bool RunTests(UhtGlobalOptions Options)
+		{
+			DirectoryReference EngineSourceProgramDirectory = DirectoryReference.Combine(Unreal.EngineDirectory, "Source", "Programs");
+			string TestDirectory = FileReference.Combine(EngineSourceProgramDirectory, "UnrealBuildTool.Tests", "UHT").FullName;
+			string TestOutputDirectory = FileReference.Combine(EngineSourceProgramDirectory, "UnrealBuildTool.Tests", "UHT.Out").FullName;
+
+			// Clear the output directory
+			try
+			{
+				Directory.Delete(TestOutputDirectory, true);
+			}
+			catch (Exception)
+			{ }
+
+			// Collect a list of all the test scripts
+			Log.TraceInformation("Running tests in {0}", TestDirectory);
+			Log.TraceInformation("Output can be compared in {0}", TestOutputDirectory);
+
+			// Run the tests on the directory
+			return RunTests(Options, TestDirectory, TestOutputDirectory);
+		}
+	}
+
+	/// <summary>
+	/// Invoke UHT
+	/// </summary>
+	[ToolMode("UnrealHeaderTool", ToolModeOptions.SingleInstance | ToolModeOptions.ShowExecutionTime)]
+	class UnrealHeaderToolMode : ToolMode
+	{
+		/// <summary>
+		/// Directory for saved application settings (typically Engine/Programs)
+		/// </summary>
+		static DirectoryReference? CachedEngineProgramSavedDirectory;
+
+		/// <summary>
+		/// The engine programs directory
+		/// </summary>
+		public static DirectoryReference EngineProgramSavedDirectory
+		{
+			get
+			{
+				if (CachedEngineProgramSavedDirectory == null)
+				{
+					if (Unreal.IsEngineInstalled())
+					{
+						CachedEngineProgramSavedDirectory = Utils.GetUserSettingDirectory() ?? DirectoryReference.Combine(Unreal.EngineDirectory, "Programs");
+					}
+					else
+					{
+						CachedEngineProgramSavedDirectory = DirectoryReference.Combine(Unreal.EngineDirectory, "Programs");
+					}
+				}
+				return CachedEngineProgramSavedDirectory;
+			}
+		}
+
+		/// <summary>
+		/// Print (incomplete) usage information
+		/// </summary>
+		private static void PrintUsage()
+		{
+			Console.WriteLine("UnrealBuildTool -Mode=UnrealHeaderTool ProjectFile ManifestFile [Options]");
+			Console.WriteLine("");
+			Console.WriteLine("Options:");
+			int LongestPrefix = 0;
+			foreach (FieldInfo Info in typeof(UhtGlobalOptions).GetFields())
+			{
+				foreach (CommandLineAttribute Att in Info.GetCustomAttributes<CommandLineAttribute>())
+				{
+					if (Att.Prefix != null && Att.Description != null)
+					{
+						LongestPrefix = Att.Prefix.Length > LongestPrefix ? Att.Prefix.Length : LongestPrefix;
+					}
+				}
+			}
+
+			foreach (UhtExporter Generator in UhtExporterTable.Instance)
+			{
+				LongestPrefix = Generator.Name.Length + 2 > LongestPrefix ? Generator.Name.Length + 2 : LongestPrefix;
+			}
+
+			foreach (FieldInfo Info in typeof(UhtGlobalOptions).GetFields())
+			{
+				foreach (CommandLineAttribute Att in Info.GetCustomAttributes<CommandLineAttribute>())
+				{
+					if (Att.Prefix != null && Att.Description != null)
+					{
+						Console.WriteLine($"  {Att.Prefix.PadRight(LongestPrefix)} :  {Att.Description}");
+					}
+				}
+			}
+
+			Console.WriteLine("");
+			Console.WriteLine("Generators: Prefix with 'no' to disable a generator");
+			foreach (UhtExporter Generator in UhtExporterTable.Instance)
+			{
+				string IsDefault = UhtConfig.Instance.IsExporterEnabled(Generator.Name) || Generator.Options.HasAnyFlags(UhtExporterOptions.Default) ? " (Default)" : "";
+				Console.WriteLine($"  -{Generator.Name.PadRight(LongestPrefix)} :  {Generator.Description}{IsDefault}");
+			}
+			Console.WriteLine("");
+		}
+
+		/// <summary>
+		/// Execute the command
+		/// </summary>
+		/// <param name="Arguments">Command line arguments</param>
+		/// <returns>Exit code</returns>
+		public override int Execute(CommandLineArguments Arguments)
+		{
+			try
+			{
+
+				// Initialize the attributes
+				UhtAttributeScanner.Scan();
+
+				// Initialize the config
+				UhtConfigImpl.Read(Arguments);
+
+				// Parse the global options
+				UhtGlobalOptions Options = new UhtGlobalOptions(Arguments);
+				if (Arguments.GetRawArray().Length <= 1 || Options.bGetHelp)
+				{
+					PrintUsage();
+					return Options.bGetHelp ? (int)CompilationResult.Succeeded : (int)CompilationResult.OtherCompilationError;
+				}
+
+				// Configure the log system
+				Log.OutputLevel = Options.LogOutputLevel;
+				Log.IncludeTimestamps = Options.bLogTimestamps;
+				Log.IncludeProgramNameWithSeverityPrefix = Options.bLogFromMsBuild;
+
+				// Add the log writer if requested. When building a target, we'll create the writer for the default log file later.
+				if (!Options.bNoLog)
+				{
+					if (Options.LogFileName != null)
+					{
+						Log.AddFileWriter("LogTraceListener", Options.LogFileName);
+					}
+
+					if (!Log.HasFileWriter())
+					{
+						string BaseLogFileName = FileReference.Combine(EngineProgramSavedDirectory, "UnrealBuildTool", "Log_UHT.txt").FullName;
+
+						FileReference LogFile = new FileReference(BaseLogFileName);
+						foreach (string LogSuffix in Arguments.GetValues("-LogSuffix="))
+						{
+							LogFile = LogFile.ChangeExtension(null) + "_" + LogSuffix + LogFile.GetExtension();
+						}
+
+						Log.AddFileWriter("DefaultLogTraceListener", LogFile);
+					}
+				}
+
+				// If we are running test scripts
+				if (Options.bTest)
+				{
+					return UhtTestHarness.RunTests(Options) ? (int)CompilationResult.Succeeded : (int)CompilationResult.OtherCompilationError;
+				}
+
+				// If we don't have a manifest name, then generate an error
+				if (Arguments.GetPositionalArgumentCount() < 2)
+				{
+					Log.TraceError("Manifest file name is required");
+					return (int)CompilationResult.OtherCompilationError;
+				}
+
+				UhtSession Session = new UhtSession
+				{
+					FileManager = new UhtStdFileManager(),
+					EngineDirectory = Unreal.EngineDirectory.FullName,
+					ReferenceDirectory = FileReference.Combine(EngineProgramSavedDirectory, "UnrealBuildTool", "ReferenceGeneratedCode").FullName,
+					VerifyDirectory = FileReference.Combine(EngineProgramSavedDirectory, "UnrealBuildTool", "VerifyGeneratedCode").FullName,
+					bWarningsAsErrors = Options.bWarningsAsErrors,
+					bGoWide = !Options.bNoGoWide,
+					bFailIfGeneratedCodeChanges = Options.bFailIfGeneratedCodeChanges,
+					bNoOutput = Options.bNoOutput,
+					bIncludeDebugOutput = Options.bIncludeDebugOutput,
+				};
+
+				// Extract the project directory from the project file name
+				{
+					string? ProjectPath = Path.GetDirectoryName(Arguments.GetPositionalArguments()[0]);
+					if (!string.IsNullOrEmpty(ProjectPath))
+					{
+						Session.ProjectDirectory = ProjectPath;
+					}
+				}				
+
+				if (Options.bWriteRef)
+				{
+					Session.ReferenceMode = UhtReferenceMode.Reference;
+				}
+				else if (Options.bVerifyRef)
+				{
+					Session.ReferenceMode = UhtReferenceMode.Verify;
+				}
+
+				foreach(UhtExporter Exporter in UhtExporterTable.Instance)
+				{
+					if (Arguments.HasOption($"-{Exporter.Name}"))
+					{
+						Session.SetExporterStatus(Exporter.Name, true);
+					}
+					else if (Arguments.HasOption($"-no{Exporter.Name}"))
+					{
+						Session.SetExporterStatus(Exporter.Name, false);
+					}
+				}
+
+				// Read and parse
+				Session.Run(Arguments.GetPositionalArguments()[1]);
+				Session.LogMessages();
+				return (int)(Session.bHasErrors ? CompilationResult.OtherCompilationError : CompilationResult.Succeeded);
+			}
+			catch (Exception Ex)
+			{
+				// Unhandled exception.
+				Log.TraceError("Unhandled exception: {0}", ExceptionUtils.FormatException(Ex));
+				Log.TraceLog(ExceptionUtils.FormatExceptionDetails(Ex));
+				return (int)CompilationResult.OtherCompilationError;
+			}
+		}
+	}
+}
diff --git a/Engine/Source/Programs/UnrealBuildTool/ProjectFiles/ProjectFileGenerator.cs b/Engine/Source/Programs/UnrealBuildTool/ProjectFiles/ProjectFileGenerator.cs
--- a/Engine/Source/Programs/UnrealBuildTool/ProjectFiles/ProjectFileGenerator.cs
+++ b/Engine/Source/Programs/UnrealBuildTool/ProjectFiles/ProjectFileGenerator.cs
@@ -1173,6 +1173,13 @@ namespace UnrealBuildTool
 		/// <param name="EngineProject">Engine project to add files to.</param>
 		private void AddUBTConfigFilesToEngineProject(ProjectFile EngineProject)
 		{
+			// UHT in UBT has a configuration file that needs to be included.
+			DirectoryReference UBTConfigDirectory = DirectoryReference.Combine(Unreal.EngineDirectory, "Programs", "UnrealBuildTool", "Config");
+			if (DirectoryReference.Exists(UBTConfigDirectory))
+			{
+				EngineProject.AddFilesToProject(SourceFileSearch.FindFiles(UBTConfigDirectory), Unreal.EngineDirectory);
+			}
+
 			EngineProject.AddAliasedFileToProject(new AliasedFile(
 					XmlConfig.GetSchemaLocation(),
 					XmlConfig.GetSchemaLocation().FullName,
diff --git a/Engine/Source/Programs/UnrealBuildTool/System/ExternalExecution.cs b/Engine/Source/Programs/UnrealBuildTool/System/ExternalExecution.cs
--- a/Engine/Source/Programs/UnrealBuildTool/System/ExternalExecution.cs
+++ b/Engine/Source/Programs/UnrealBuildTool/System/ExternalExecution.cs
@@ -18,6 +18,7 @@ using System.Collections.Concurrent;
 using System.Threading.Tasks;
 using OpenTracing;
 using OpenTracing.Util;
+using UnrealBuildTool.Modes;
 
 namespace UnrealBuildTool
 {
@@ -908,6 +909,22 @@ namespace UnrealBuildTool
 		/// Performs any early outs if headers need no changes, given the UObject modules, tool path, game name, and configuration
 		/// </summary>
 		public static void ExecuteHeaderToolIfNecessary(BuildConfiguration BuildConfiguration, FileReference? ProjectFile, TargetMakefile Makefile, string TargetName, ISourceFileWorkingSet WorkingSet)
+		{
+			if (!BuildConfiguration.bUseBuiltInUnrealHeaderTool || Makefile.bHasProjectScriptPlugin || TargetName.Equals("UnrealHeaderTool", StringComparison.InvariantCultureIgnoreCase))
+			{
+				ExecuteExternalHeaderToolIfNecessary(BuildConfiguration, ProjectFile, Makefile, TargetName, WorkingSet);
+			}
+			else
+			{
+				ExecuteInternalHeaderToolIfNecessary(BuildConfiguration, ProjectFile, Makefile, TargetName, WorkingSet);
+			}
+		}
+
+		/// <summary>
+		/// Builds and runs the header tool and touches the header directories.
+		/// Performs any early outs if headers need no changes, given the UObject modules, tool path, game name, and configuration
+		/// </summary>
+		private static void ExecuteExternalHeaderToolIfNecessary(BuildConfiguration BuildConfiguration, FileReference? ProjectFile, TargetMakefile Makefile, string TargetName, ISourceFileWorkingSet WorkingSet)
 		{
 			if (ProgressWriter.bWriteMarkup)
 			{
@@ -1055,7 +1072,7 @@ namespace UnrealBuildTool
 
 					// Disable extensions when serializing to remove the $type fields
 					Directory.CreateDirectory(ModuleInfoFileName.Directory.FullName);
-					System.IO.File.WriteAllText(ModuleInfoFileName.FullName, JsonSerializer.Serialize(Manifest));
+					System.IO.File.WriteAllText(ModuleInfoFileName.FullName, JsonSerializer.Serialize(Manifest, new JsonSerializerOptions { WriteIndented = true }));
 
 					string CmdLine = (ProjectFile != null) ? "\"" + ProjectFile.FullName + "\"" : TargetName;
 					CmdLine += " \"" + ModuleInfoFileName + "\" -LogCmds=\"loginit warning, logexit warning, logdatabase error\" -Unattended -WarningsAsErrors";
@@ -1108,7 +1125,215 @@ namespace UnrealBuildTool
 							UHTResult = ((int)(UHTResult) == 130) ? CompilationResult.Canceled : CompilationResult.CrashOrAssert;
 						}
 
-						if (BuildHostPlatform.Current.Platform == UnrealTargetPlatform.Win64 && 
+						if (BuildHostPlatform.Current.Platform == UnrealTargetPlatform.Win64 &&
+							(int)(UHTResult) < 0)
+						{
+							Log.TraceError(String.Format("UnrealHeaderTool failed with exit code 0x{0:X} - check that Unreal Engine prerequisites are installed.", (int)UHTResult));
+						}
+
+						throw new CompilationResultException(UHTResult);
+					}
+
+					Log.TraceInformation("Reflection code generated for {0} in {1} seconds", ActualTargetName, s.Elapsed.TotalSeconds);
+
+					// Update the tool info file
+					DirectoryReference.CreateDirectory(ToolInfoFile.Directory);
+					FileReference.WriteAllText(ToolInfoFile, HeaderToolReceipt.FullName);
+
+					// Now that UHT has successfully finished generating code, we need to update all cached FileItems in case their last write time has changed.
+					// Otherwise UBT might not detect changes UHT made.
+					using (GlobalTracer.Instance.BuildSpan("ExternalExecution.ResetCachedHeaderInfo()").StartActive())
+					{
+						ResetCachedHeaderInfo(Makefile.UObjectModules);
+					}
+				}
+				else
+				{
+					Log.TraceVerbose("Generated code is up to date.");
+				}
+
+				Progress.Write(2, 3);
+
+				using (GlobalTracer.Instance.BuildSpan("ExternalExecution.UpdateTimestamps()").StartActive())
+				{
+					UpdateTimestamps(Makefile.UObjectModules);
+				}
+
+				Progress.Write(3, 3);
+			}
+			if (ProgressWriter.bWriteMarkup)
+			{
+				Log.WriteLine(LogEventType.Console, "@progress pop");
+			}
+		}
+
+		/// <summary>
+		/// Builds and runs the header tool and touches the header directories.
+		/// Performs any early outs if headers need no changes, given the UObject modules, tool path, game name, and configuration
+		/// </summary>
+		private static void ExecuteInternalHeaderToolIfNecessary(BuildConfiguration BuildConfiguration, FileReference? ProjectFile, TargetMakefile Makefile, string TargetName, ISourceFileWorkingSet WorkingSet)
+		{
+			if (ProgressWriter.bWriteMarkup)
+			{
+				Log.WriteLine(LogEventType.Console, "@progress push 5%");
+			}
+			using (ProgressWriter Progress = new ProgressWriter("Generating code...", false))
+			{
+				string RootLocalPath = Unreal.RootDirectory.FullName;
+
+				UnrealTargetConfiguration UHTConfig = BuildConfiguration.bForceDebugUnrealHeaderTool ? UnrealTargetConfiguration.Debug : UnrealTargetConfiguration.Development;
+
+				// Figure out the receipt path
+				FileReference HeaderToolReceipt = GetHeaderToolReceiptFile(ProjectFile, UHTConfig, Makefile.bHasProjectScriptPlugin);
+
+				// Get UHT assembly timestamp
+				DateTime UnrealBuildToolTimestamp = new FileInfo(Assembly.GetExecutingAssembly().Location).LastWriteTimeUtc;
+				//DateTime UnrealHeaderToolTimestamp = DateTime.MinValue;
+				//Assembly? UHTAssembly = AppDomain.CurrentDomain.GetAssemblies().Single(x => x.GetName().Name == "EpicGames.UHT");
+				//if (UHTAssembly != null)
+				//{
+				//	UnrealHeaderToolTimestamp = new FileInfo(UHTAssembly.Location).LastWriteTimeUtc;
+				//}
+				DateTime ComboTimeStamp = UnrealBuildToolTimestamp;
+				//if (UnrealHeaderToolTimestamp > ComboTimeStamp)
+				//{
+				//	ComboTimeStamp = UnrealHeaderToolTimestamp;
+				//}
+
+				// ensure the headers are up to date
+				bool bUHTNeedsToRun = false;
+				if (BuildConfiguration.bForceHeaderGeneration)
+				{
+					bUHTNeedsToRun = true;
+				}
+				else if (AreGeneratedCodeFilesOutOfDate(BuildConfiguration, Makefile.UObjectModules, ComboTimeStamp))
+				{
+					bUHTNeedsToRun = true;
+				}
+
+				// Check we're not using a different version of UHT
+				FileReference ModuleInfoFileName = FileReference.Combine(Makefile.ProjectIntermediateDirectory, TargetName + ".uhtmanifest");
+				FileReference ToolInfoFile = ModuleInfoFileName.ChangeExtension(".uhtpath");
+				if (!bUHTNeedsToRun)
+				{
+					if (!FileReference.Exists(ToolInfoFile))
+					{
+						bUHTNeedsToRun = true;
+					}
+					else if (FileReference.ReadAllText(ToolInfoFile) != HeaderToolReceipt.FullName)
+					{
+						bUHTNeedsToRun = true;
+					}
+				}
+
+				// Get the file containing dependencies for the generated code
+				FileReference ExternalDependenciesFile = ModuleInfoFileName.ChangeExtension(".deps");
+				if (AreExternalDependenciesOutOfDate(ExternalDependenciesFile))
+				{
+					bUHTNeedsToRun = true;
+				}
+
+				// @todo ubtmake: Optimization: Ideally we could avoid having to generate this data in the case where UHT doesn't even need to run!  Can't we use the existing copy?  (see below use of Manifest)
+
+				List<UHTManifest.Module> Modules = new List<UHTManifest.Module>();
+				foreach(UHTModuleInfo UObjectModule in Makefile.UObjectModules)
+				{
+					Modules.Add(
+						new UHTManifest.Module
+						{
+							Name = UObjectModule.ModuleName,
+							ModuleType = (UHTModuleType)Enum.Parse(typeof(UHTModuleType), UObjectModule.ModuleType),
+							OverrideModuleType = (EPackageOverrideType)Enum.Parse(typeof(EPackageOverrideType), UObjectModule.OverrideModuleType),
+							BaseDirectory = UObjectModule.ModuleDirectories[0].FullName,
+							IncludeBase = UObjectModule.ModuleDirectories[0].ParentDirectory!.FullName,
+							OutputDirectory = Path.GetDirectoryName(UObjectModule.GeneratedCPPFilenameBase)!,
+							ClassesHeaders = UObjectModule.PublicUObjectClassesHeaders.Select((Header) => Header.AbsolutePath).ToList(),
+							PublicHeaders = UObjectModule.PublicUObjectHeaders.Select((Header) => Header.AbsolutePath).ToList(),
+							InternalHeaders = UObjectModule.InternalUObjectHeaders.Select((Header) => Header.AbsolutePath).ToList(),
+							PrivateHeaders = UObjectModule.PrivateUObjectHeaders.Select((Header) => Header.AbsolutePath).ToList(),
+							GeneratedCPPFilenameBase = UObjectModule.GeneratedCPPFilenameBase,
+							SaveExportedHeaders = !UObjectModule.bIsReadOnly,
+							GeneratedCodeVersion = UObjectModule.GeneratedCodeVersion,
+						});
+				}
+				UHTManifest Manifest = new UHTManifest
+				{
+					TargetName = TargetName,
+					IsGameTarget = Makefile.TargetType != TargetType.Program,
+					RootLocalPath = RootLocalPath,
+					ExternalDependenciesFile = ExternalDependenciesFile.FullName,
+					Modules = Modules,
+				};
+
+				if (bUHTNeedsToRun)
+				{
+					Progress.Write(1, 3);
+
+					// Disable extensions when serializing to remove the $type fields
+					Directory.CreateDirectory(ModuleInfoFileName.Directory.FullName);
+					System.IO.File.WriteAllText(ModuleInfoFileName.FullName, JsonSerializer.Serialize(Manifest, new JsonSerializerOptions { WriteIndented = true }));
+
+					string ActualTargetName = String.IsNullOrEmpty(TargetName) ? "UE5" : TargetName;
+					Log.TraceInformation("Parsing headers for {0}", ActualTargetName);
+
+					// Generate the command line
+					string CmdLine = (ProjectFile != null) ? "\"" + ProjectFile.FullName + "\"" : TargetName;
+					CmdLine += " \"" + ModuleInfoFileName + "\" -WarningsAsErrors";
+
+					if (Makefile.UHTAdditionalArguments != null)
+					{
+						foreach (string UHTArgument in Makefile.UHTAdditionalArguments)
+						{
+							CmdLine += " " + UHTArgument;
+						}
+					}
+
+					if (Unreal.IsEngineInstalled())
+					{
+						CmdLine += " -installed";
+					}
+
+					if (BuildConfiguration.bFailIfGeneratedCodeChanges)
+					{
+						CmdLine += " -FailIfGeneratedCodeChanges";
+					}
+
+					Log.TraceInformation("  Running Internal UnrealHeaderTool {0}", CmdLine);
+
+					// Convert the command line to an argument list
+					string[] Args = CmdLine.Split();
+					for (int Index = 0; Index < Args.Length; ++Index)
+					{
+						string Arg = Args[Index];
+						if (Arg[0] == '"' && Arg[Arg.Length - 1] == '"')
+						{
+							Args[Index] = Arg.Substring(1, Arg.Length - 2);
+						}
+					}
+					CommandLineArguments Arguments = new CommandLineArguments(Args);
+
+					// Run UHT
+					Stopwatch s = new Stopwatch();
+					s.Start();
+					IScope Timer = GlobalTracer.Instance.BuildSpan("Executing UnrealHeaderTool").StartActive();
+					UnrealHeaderToolMode UHTTool = new UnrealHeaderToolMode();
+					CompilationResult UHTResult = (CompilationResult)UHTTool.Execute(Arguments);
+					Timer.Span.Finish();
+					s.Stop();
+
+					if (UHTResult != CompilationResult.Succeeded)
+					{
+						// On Linux and Mac, the shell will return 128+signal number exit codes if UHT gets a signal (e.g. crashes or is interrupted)
+						if ((BuildHostPlatform.Current.Platform == UnrealTargetPlatform.Linux ||
+							BuildHostPlatform.Current.Platform == UnrealTargetPlatform.Mac) &&
+							(int)(UHTResult) >= 128
+							)
+						{
+							// SIGINT is 2, so 128 + SIGINT is 130
+							UHTResult = ((int)(UHTResult) == 130) ? CompilationResult.Canceled : CompilationResult.CrashOrAssert;
+						}
+
+						if (BuildHostPlatform.Current.Platform == UnrealTargetPlatform.Win64 &&
 							(int)(UHTResult) < 0)
 						{
 							Log.TraceError(String.Format("UnrealHeaderTool failed with exit code 0x{0:X} - check that Unreal Engine prerequisites are installed.", (int)UHTResult));
diff --git a/Engine/Source/Programs/UnrealBuildTool/UnrealBuildTool.csproj b/Engine/Source/Programs/UnrealBuildTool/UnrealBuildTool.csproj
--- a/Engine/Source/Programs/UnrealBuildTool/UnrealBuildTool.csproj
+++ b/Engine/Source/Programs/UnrealBuildTool/UnrealBuildTool.csproj
@@ -106,6 +106,7 @@
   <ItemGroup>
     <ProjectReference Include="..\Shared\EpicGames.Build\EpicGames.Build.csproj" />
     <ProjectReference Include="..\Shared\EpicGames.Core\EpicGames.Core.csproj" />
+    <ProjectReference Include="..\Shared\EpicGames.UHT\EpicGames.UHT.csproj" />
   </ItemGroup>
   <ItemGroup>
     <PackageReference Include="Microsoft.CodeAnalysis.CSharp" Version="3.7.0" />
diff --git a/Engine/Source/Runtime/Core/Private/Math/UnitConversion.cpp b/Engine/Source/Runtime/Core/Private/Math/UnitConversion.cpp
--- a/Engine/Source/Runtime/Core/Private/Math/UnitConversion.cpp
+++ b/Engine/Source/Runtime/Core/Private/Math/UnitConversion.cpp
@@ -17,6 +17,7 @@ struct FParseCandidate
 	EUnit Unit;
 };
 
+//TODO: Add note about updating
 FParseCandidate ParseCandidates[] = {
 	
 	{ TEXT("Micrometers"),			EUnit::Micrometers },			{ TEXT("um"),		EUnit::Micrometers }, 			{ TEXT("\u00B5m"),	EUnit::Micrometers },
