Updating RunUHT to also check for EpicGames.UHT.
Added comment to UnrealBuildTool.csproj to make people aware the change should be done when adding new project references.

#rnx

diff --git a/Engine/Build/BatchFiles/BuildUAT.bat b/Engine/Build/BatchFiles/BuildUAT.bat
--- a/Engine/Build/BatchFiles/BuildUAT.bat
+++ b/Engine/Build/BatchFiles/BuildUAT.bat
@@ -25,6 +25,8 @@ dir /s^
  Programs\Shared\EpicGames.Core\*.csproj^
  Programs\Shared\EpicGames.Build\*.cs^
  Programs\Shared\EpicGames.Build\*.csproj^
+ Programs\Shared\EpicGames.UHT\*.cs^
+ Programs\Shared\EpicGames.UHT\*.csproj^
  Programs\UnrealBuildTool\*.cs^
  Programs\UnrealBuildTool\*.csproj^
  | find ".cs" >..\Intermediate\Build\AutomationToolFiles.txt
diff --git a/Engine/Build/BatchFiles/BuildUAT.sh b/Engine/Build/BatchFiles/BuildUAT.sh
--- a/Engine/Build/BatchFiles/BuildUAT.sh
+++ b/Engine/Build/BatchFiles/BuildUAT.sh
@@ -37,7 +37,7 @@ elif [ ! -f ../Binaries/DotNET/AutomationTool/AutomationTool.dll ]; then
 
 elif [ -f ../Intermediate/Build/AutomationToolLastBuildTime ]; then
   UPDATED_DEP_FILES="$(find \
-    Programs/Shared/EpicGames.Core Programs/Shared/EpicGames.Build Programs/UnrealBuildTool \
+    Programs/Shared/EpicGames.Core Programs/Shared/EpicGames.Build Programs/Shared/EpicGames.UHT Programs/UnrealBuildTool \
     -type f \
     \( -iname \*.cs -or -iname \*.csproj \) \
     -newer ../Intermediate/Build/AutomationToolLastBuildTime)"
diff --git a/Engine/Source/Programs/UnrealBuildTool/UnrealBuildTool.csproj b/Engine/Source/Programs/UnrealBuildTool/UnrealBuildTool.csproj
--- a/Engine/Source/Programs/UnrealBuildTool/UnrealBuildTool.csproj
+++ b/Engine/Source/Programs/UnrealBuildTool/UnrealBuildTool.csproj
@@ -104,6 +104,8 @@
   </ItemGroup>
 
   <ItemGroup>
+    <!-- If you add a project reference, BuildUAT.bat/.sh need to be updated to also test for those projects.
+    While similar checks are in RunUBT.bat/.sh and GenerateProjectFiles.bat, they don't currently check for these projects. -->
     <ProjectReference Include="..\Shared\EpicGames.Build\EpicGames.Build.csproj" />
     <ProjectReference Include="..\Shared\EpicGames.Core\EpicGames.Core.csproj" />
     <ProjectReference Include="..\Shared\EpicGames.UHT\EpicGames.UHT.csproj" />
