Engine targets for Mac installed builds

This forces UBT to always include an engine target for Mac installed builds

#rnx

diff --git a/Engine/Source/Programs/UnrealBuildTool/ProjectFiles/Xcode/XcodeProjectFileGenerator.cs b/Engine/Source/Programs/UnrealBuildTool/ProjectFiles/Xcode/XcodeProjectFileGenerator.cs
--- a/Engine/Source/Programs/UnrealBuildTool/ProjectFiles/Xcode/XcodeProjectFileGenerator.cs
+++ b/Engine/Source/Programs/UnrealBuildTool/ProjectFiles/Xcode/XcodeProjectFileGenerator.cs
@@ -301,9 +301,10 @@ namespace UnrealBuildTool
 
 			if (bGeneratingGameProjectFiles)
 			{
-				if (bGeneratingRunIOSProject || bGeneratingRunTVOSProject)
+				if (bGeneratingRunIOSProject || bGeneratingRunTVOSProject || UnrealBuildBase.Unreal.IsEngineInstalled())
 				{
 					// an Engine target is required in order to be able to get Xcode to sign blueprint projects
+					// always include the engine target for installed builds.
 					IncludeEnginePrograms = true;
 				}
 				bIncludeEngineSource = true;
