Escape brackets and parenthesis in IOS Remote Toolchain remote paths to prevent the shell trying to match using them.

diff --git a/Engine/Source/Programs/UnrealBuildTool/ToolChain/RemoteMac.cs b/Engine/Source/Programs/UnrealBuildTool/ToolChain/RemoteMac.cs
--- a/Engine/Source/Programs/UnrealBuildTool/ToolChain/RemoteMac.cs
+++ b/Engine/Source/Programs/UnrealBuildTool/ToolChain/RemoteMac.cs
@@ -740,7 +740,7 @@ namespace UnrealBuildTool
 		/// <returns>Equivalent remote path</returns>
 		private string GetRemotePath(string LocalPath)
 		{
-			return String.Format("{0}/{1}", RemoteBaseDir, LocalPath.Replace(":", "").Replace("\\", "/").Replace(" ", "_"));
+			return String.Format("{0}/{1}", RemoteBaseDir, LocalPath.Replace(":", "").Replace('\\', '/').Replace(' ', '_')).Replace('(', '_').Replace(')', '_').Replace('[', '_').Replace(']', '_');
 		}
 
 		/// <summary>
