Add path to find asan headers if needed when enabling ASan

diff --git a/Engine/Source/Programs/UnrealBuildTool/Platform/Windows/VCToolChain.cs b/Engine/Source/Programs/UnrealBuildTool/Platform/Windows/VCToolChain.cs
--- a/Engine/Source/Programs/UnrealBuildTool/Platform/Windows/VCToolChain.cs
+++ b/Engine/Source/Programs/UnrealBuildTool/Platform/Windows/VCToolChain.cs
@@ -455,6 +455,10 @@ namespace UnrealBuildTool
 				{
 					AddDefinition(Arguments, "USING_ADDRESS_SANITISER=1");
 				}
+
+				// Currently the ASan headers are not default around. They can be found at this location so lets use this until this is resolved in the toolchain
+				// Jira with some more info and the MSVC bug at UE-144727
+				AddSystemIncludePath(Arguments, DirectoryReference.Combine(EnvVars.CompilerDir, "crt", "src"), Target.WindowsPlatform.Compiler);
 			}
 
 			//
