Adding a simpler command line option to UHT mode where it will generate the manifest instead of depending on one already being generated.  This will greatly simplify testing in the future by allowing A/B testing to be done without the need of building the target (to get the manifest)

#rnx

diff --git a/Engine/Source/Programs/UnrealBuildTool/Configuration/UEBuildTarget.cs b/Engine/Source/Programs/UnrealBuildTool/Configuration/UEBuildTarget.cs
--- a/Engine/Source/Programs/UnrealBuildTool/Configuration/UEBuildTarget.cs
+++ b/Engine/Source/Programs/UnrealBuildTool/Configuration/UEBuildTarget.cs
@@ -1705,7 +1705,7 @@ namespace UnrealBuildTool
 		/// <summary>
 		/// Builds the target, appending list of output files and returns building result.
 		/// </summary>
-		public TargetMakefile Build(BuildConfiguration BuildConfiguration, ISourceFileWorkingSet WorkingSet, TargetDescriptor TargetDescriptor)
+		public TargetMakefile Build(BuildConfiguration BuildConfiguration, ISourceFileWorkingSet WorkingSet, TargetDescriptor TargetDescriptor, bool bInitOnly = false)
 		{
 			CppConfiguration CppConfiguration = GetCppConfiguration(Configuration);
 
@@ -1844,6 +1844,12 @@ namespace UnrealBuildTool
 				Makefile.UHTAdditionalArguments = UHTAdditionalArguments.ToArray();
 			}
 
+			// UHT mode uses this to create the makefile at least to the point where UHT would have valid manifest information
+			if (bInitOnly)
+			{
+				return Makefile;
+			}
+
 			// NOTE: Even in Gather mode, we need to run UHT to make sure the files exist for the static action graph to be setup correctly.  This is because UHT generates .cpp
 			// files that are injected as top level prerequisites.  If UHT only emitted included header files, we wouldn't need to run it during the Gather phase at all.
 			if (Makefile.UObjectModules.Count > 0)
diff --git a/Engine/Source/Programs/UnrealBuildTool/Modes/UnrealHeaderToolMode.cs b/Engine/Source/Programs/UnrealBuildTool/Modes/UnrealHeaderToolMode.cs
--- a/Engine/Source/Programs/UnrealBuildTool/Modes/UnrealHeaderToolMode.cs
+++ b/Engine/Source/Programs/UnrealBuildTool/Modes/UnrealHeaderToolMode.cs
@@ -898,7 +898,7 @@ namespace UnrealBuildTool.Modes
 	/// <summary>
 	/// Invoke UHT
 	/// </summary>
-	[ToolMode("UnrealHeaderTool", ToolModeOptions.SingleInstance | ToolModeOptions.ShowExecutionTime)]
+	[ToolMode("UnrealHeaderTool", ToolModeOptions.XmlConfig | ToolModeOptions.BuildPlatforms | ToolModeOptions.SingleInstance | ToolModeOptions.ShowExecutionTime)]
 	class UnrealHeaderToolMode : ToolMode
 	{
 		/// <summary>
@@ -933,7 +933,7 @@ namespace UnrealBuildTool.Modes
 		/// </summary>
 		private static void PrintUsage()
 		{
-			Console.WriteLine("UnrealBuildTool -Mode=UnrealHeaderTool ProjectFile ManifestFile [Options]");
+			Console.WriteLine("UnrealBuildTool -Mode=UnrealHeaderTool [ProjectFile ManifestFile] -OR [\"-Target...\"] [Options]");
 			Console.WriteLine("");
 			Console.WriteLine("Options:");
 			int LongestPrefix = 0;
@@ -992,7 +992,20 @@ namespace UnrealBuildTool.Modes
 
 				// Parse the global options
 				UhtGlobalOptions Options = new UhtGlobalOptions(Arguments);
-				if (Arguments.GetRawArray().Length <= 1 || Options.bGetHelp)
+				int TargetArgumentIndex = -1;
+				if (Arguments.GetPositionalArgumentCount() == 0)
+				{
+					for (int Index = 0; Index < Arguments.Count; ++Index)
+					{
+						if (Arguments[Index].StartsWith("-Target", StringComparison.OrdinalIgnoreCase))
+						{
+							TargetArgumentIndex = Index;
+							break;
+						}
+					}
+				}
+				int RequiredArgCount = TargetArgumentIndex >= 0 ? 0 : 2;
+				if (Arguments.GetPositionalArgumentCount() != RequiredArgCount || Options.bGetHelp)
 				{
 					PrintUsage();
 					return Options.bGetHelp ? (int)CompilationResult.Succeeded : (int)CompilationResult.OtherCompilationError;
@@ -1031,17 +1044,57 @@ namespace UnrealBuildTool.Modes
 					return UhtTestHarness.RunTests(Options) ? (int)CompilationResult.Succeeded : (int)CompilationResult.OtherCompilationError;
 				}
 
-				// If we don't have a manifest name, then generate an error
-				if (Arguments.GetPositionalArgumentCount() < 2)
+				string? ProjectPath = null;
+				string? ManifestPath = null;
+
+				if (TargetArgumentIndex >= 0)
+				{
+					CommandLineArguments LocalArguments = new CommandLineArguments(new string[] { Arguments[TargetArgumentIndex] });
+					List<TargetDescriptor> TargetDescriptors = TargetDescriptor.ParseCommandLine(LocalArguments, false, false, false);
+					if (TargetDescriptors.Count == 0)
+					{
+						Log.TraceError("No target descriptors found.");
+						return (int)CompilationResult.OtherCompilationError;
+					}
+
+					TargetDescriptor TargetDesc = TargetDescriptors[0];
+
+					// Create the target
+					UEBuildTarget Target = UEBuildTarget.Create(TargetDesc, false, false, false);
+
+					// Create the makefile for the target and export the module information
+					using (ISourceFileWorkingSet WorkingSet = new EmptySourceFileWorkingSet())
+					{
+						// Create the build configuration object, and read the settings
+						BuildConfiguration BuildConfiguration = new BuildConfiguration();
+						XmlConfig.ApplyTo(BuildConfiguration);
+						Arguments.ApplyTo(BuildConfiguration);
+
+						// Create the makefile
+						TargetMakefile Makefile = Target.Build(BuildConfiguration, WorkingSet, TargetDesc, true);
+
+						FileReference ModuleInfoFileName = ExternalExecution.GetUHTModuleInfoFileName(Makefile, Target.TargetName);
+						FileReference DepsFileName = ExternalExecution.GetUHTDepsFileName(ModuleInfoFileName);
+						ManifestPath = ModuleInfoFileName.FullName;
+						ExternalExecution.WriteUHTManifest(Makefile, Target.TargetName, ModuleInfoFileName, DepsFileName);
+
+						if (Target.ProjectFile != null)
+						{
+							ProjectPath = Path.GetDirectoryName(Target.ProjectFile.FullName);
+						}
+					}
+				}
+				else
 				{
-					Log.TraceError("Manifest file name is required");
-					return (int)CompilationResult.OtherCompilationError;
+					ProjectPath = Path.GetDirectoryName(Arguments.GetPositionalArguments()[0]);
+					ManifestPath = Arguments.GetPositionalArguments()[1];
 				}
 
 				UhtSession Session = new UhtSession
 				{
 					FileManager = new UhtStdFileManager(),
 					EngineDirectory = Unreal.EngineDirectory.FullName,
+					ProjectDirectory = string.IsNullOrEmpty(ProjectPath) ? null : ProjectPath,
 					ReferenceDirectory = FileReference.Combine(EngineProgramSavedDirectory, "UnrealBuildTool", "ReferenceGeneratedCode").FullName,
 					VerifyDirectory = FileReference.Combine(EngineProgramSavedDirectory, "UnrealBuildTool", "VerifyGeneratedCode").FullName,
 					bWarningsAsErrors = Options.bWarningsAsErrors,
@@ -1051,15 +1104,6 @@ namespace UnrealBuildTool.Modes
 					bIncludeDebugOutput = Options.bIncludeDebugOutput,
 				};
 
-				// Extract the project directory from the project file name
-				{
-					string? ProjectPath = Path.GetDirectoryName(Arguments.GetPositionalArguments()[0]);
-					if (!string.IsNullOrEmpty(ProjectPath))
-					{
-						Session.ProjectDirectory = ProjectPath;
-					}
-				}				
-
 				if (Options.bWriteRef)
 				{
 					Session.ReferenceMode = UhtReferenceMode.Reference;
@@ -1082,7 +1126,7 @@ namespace UnrealBuildTool.Modes
 				}
 
 				// Read and parse
-				Session.Run(Arguments.GetPositionalArguments()[1]);
+				Session.Run(ManifestPath!);
 				Session.LogMessages();
 				return (int)(Session.bHasErrors ? CompilationResult.OtherCompilationError : CompilationResult.Succeeded);
 			}
diff --git a/Engine/Source/Programs/UnrealBuildTool/System/ExternalExecution.cs b/Engine/Source/Programs/UnrealBuildTool/System/ExternalExecution.cs
--- a/Engine/Source/Programs/UnrealBuildTool/System/ExternalExecution.cs
+++ b/Engine/Source/Programs/UnrealBuildTool/System/ExternalExecution.cs
@@ -920,6 +920,84 @@ namespace UnrealBuildTool
 			}
 		}
 
+		/// <summary>
+		/// Return the file reference for the UHT manifest file
+		/// </summary>
+		/// <param name="Makefile">Input makefile</param>
+		/// <param name="TargetName">Name of the target</param>
+		/// <returns>Manifest file name</returns>
+		public static FileReference GetUHTModuleInfoFileName(TargetMakefile Makefile, string TargetName)
+		{
+			return FileReference.Combine(Makefile.ProjectIntermediateDirectory, TargetName + ".uhtmanifest");
+		}
+
+		/// <summary>
+		/// Return the file reference for the UHT deps file
+		/// </summary>
+		/// <param name="ModuleInfoFileName">Manifest info file name</param>
+		/// <returns>UHT dependency file name</returns>
+		public static FileReference GetUHTDepsFileName(FileReference ModuleInfoFileName)
+		{
+			return ModuleInfoFileName.ChangeExtension(".deps");
+		}
+
+		/// <summary>
+		/// Convert the makefile to a UHTManifest object
+		/// </summary>
+		/// <param name="Makefile">Input makefile</param>
+		/// <param name="TargetName">Name of the target</param>
+		/// <param name="DepsFileName">Name of the dependencies file.</param>
+		/// <returns>Output UHT manifest</returns>
+		public static UHTManifest CreateUHTManifest(TargetMakefile Makefile, string TargetName, FileReference DepsFileName)
+		{
+			List<UHTManifest.Module> Modules = new List<UHTManifest.Module>();
+			foreach (UHTModuleInfo UObjectModule in Makefile.UObjectModules)
+			{
+				Modules.Add(
+					new UHTManifest.Module
+					{
+						Name = UObjectModule.ModuleName,
+						ModuleType = (UHTModuleType)Enum.Parse(typeof(UHTModuleType), UObjectModule.ModuleType),
+						OverrideModuleType = (EPackageOverrideType)Enum.Parse(typeof(EPackageOverrideType), UObjectModule.OverrideModuleType),
+						BaseDirectory = UObjectModule.ModuleDirectories[0].FullName,
+						IncludeBase = UObjectModule.ModuleDirectories[0].ParentDirectory!.FullName,
+						OutputDirectory = Path.GetDirectoryName(UObjectModule.GeneratedCPPFilenameBase)!,
+						ClassesHeaders = UObjectModule.PublicUObjectClassesHeaders.Select((Header) => Header.AbsolutePath).ToList(),
+						PublicHeaders = UObjectModule.PublicUObjectHeaders.Select((Header) => Header.AbsolutePath).ToList(),
+						InternalHeaders = UObjectModule.InternalUObjectHeaders.Select((Header) => Header.AbsolutePath).ToList(),
+						PrivateHeaders = UObjectModule.PrivateUObjectHeaders.Select((Header) => Header.AbsolutePath).ToList(),
+						GeneratedCPPFilenameBase = UObjectModule.GeneratedCPPFilenameBase,
+						SaveExportedHeaders = !UObjectModule.bIsReadOnly,
+						GeneratedCodeVersion = UObjectModule.GeneratedCodeVersion,
+					});
+			}
+			UHTManifest Manifest = new UHTManifest
+			{
+				TargetName = TargetName,
+				IsGameTarget = Makefile.TargetType != TargetType.Program,
+				RootLocalPath = Unreal.RootDirectory.FullName,
+				ExternalDependenciesFile = DepsFileName.FullName,
+				Modules = Modules,
+			};
+
+			return Manifest;
+		}
+
+		/// <summary>
+		/// Write the manifest to disk
+		/// </summary>
+		/// <param name="Makefile">Input makefile</param>
+		/// <param name="TargetName">Name of the target</param>
+		/// <param name="ModuleInfoFileName">Destination file name.  If not supplied, it will be generated.</param>
+		/// <param name="DepsFileName">Name of the dependencies file.  If not supplied, it will be generated.</param>
+		public static void WriteUHTManifest(TargetMakefile Makefile, string TargetName, FileReference ModuleInfoFileName, FileReference DepsFileName)
+		{
+			// @todo ubtmake: Optimization: Ideally we could avoid having to generate this data in the case where UHT doesn't even need to run!  Can't we use the existing copy?  (see below use of Manifest)
+			UHTManifest Manifest = CreateUHTManifest(Makefile, TargetName, DepsFileName);
+			Directory.CreateDirectory(ModuleInfoFileName.Directory.FullName);
+			System.IO.File.WriteAllText(ModuleInfoFileName.FullName, JsonSerializer.Serialize(Manifest, new JsonSerializerOptions { WriteIndented = true }));
+		}
+
 		/// <summary>
 		/// Builds and runs the header tool and touches the header directories.
 		/// Performs any early outs if headers need no changes, given the UObject modules, tool path, game name, and configuration
@@ -962,7 +1040,7 @@ namespace UnrealBuildTool
 				}
 
 				// Check we're not using a different version of UHT
-				FileReference ModuleInfoFileName = FileReference.Combine(Makefile.ProjectIntermediateDirectory, TargetName + ".uhtmanifest");
+				FileReference ModuleInfoFileName = GetUHTModuleInfoFileName(Makefile, TargetName);
 				FileReference ToolInfoFile = ModuleInfoFileName.ChangeExtension(".uhtpath");
 				if(!bUHTNeedsToRun)
 				{
@@ -977,45 +1055,13 @@ namespace UnrealBuildTool
 				}
 
 				// Get the file containing dependencies for the generated code
-				FileReference ExternalDependenciesFile = ModuleInfoFileName.ChangeExtension(".deps");
+				FileReference ExternalDependenciesFile = GetUHTDepsFileName(ModuleInfoFileName);
 				if (AreExternalDependenciesOutOfDate(ExternalDependenciesFile))
 				{
 					bUHTNeedsToRun = true;
 					bHaveHeaderTool = false; // Force UHT to build until dependency checking is fast enough to run all the time
 				}
 
-				// @todo ubtmake: Optimization: Ideally we could avoid having to generate this data in the case where UHT doesn't even need to run!  Can't we use the existing copy?  (see below use of Manifest)
-
-				List<UHTManifest.Module> Modules = new List<UHTManifest.Module>();
-				foreach(UHTModuleInfo UObjectModule in Makefile.UObjectModules)
-				{
-					Modules.Add(
-						new UHTManifest.Module
-						{
-							Name = UObjectModule.ModuleName,
-							ModuleType = (UHTModuleType)Enum.Parse(typeof(UHTModuleType), UObjectModule.ModuleType),
-							OverrideModuleType = (EPackageOverrideType)Enum.Parse(typeof(EPackageOverrideType), UObjectModule.OverrideModuleType),
-							BaseDirectory = UObjectModule.ModuleDirectories[0].FullName,
-							IncludeBase = UObjectModule.ModuleDirectories[0].ParentDirectory!.FullName,
-							OutputDirectory = Path.GetDirectoryName(UObjectModule.GeneratedCPPFilenameBase)!,
-							ClassesHeaders = UObjectModule.PublicUObjectClassesHeaders.Select((Header) => Header.AbsolutePath).ToList(),
-							PublicHeaders = UObjectModule.PublicUObjectHeaders.Select((Header) => Header.AbsolutePath).ToList(),
-							InternalHeaders = UObjectModule.InternalUObjectHeaders.Select((Header) => Header.AbsolutePath).ToList(),
-							PrivateHeaders = UObjectModule.PrivateUObjectHeaders.Select((Header) => Header.AbsolutePath).ToList(),
-							GeneratedCPPFilenameBase = UObjectModule.GeneratedCPPFilenameBase,
-							SaveExportedHeaders = !UObjectModule.bIsReadOnly,
-							GeneratedCodeVersion = UObjectModule.GeneratedCodeVersion,
-						});
-				}
-				UHTManifest Manifest = new UHTManifest
-				{
-					TargetName = TargetName,
-					IsGameTarget = Makefile.TargetType != TargetType.Program,
-					RootLocalPath = RootLocalPath,
-					ExternalDependenciesFile = ExternalDependenciesFile.FullName,
-					Modules = Modules,
-				};
-
 				if (!bIsBuildingUHT && bUHTNeedsToRun)
 				{
 					// Always build UnrealHeaderTool if header regeneration is required, unless we're running within an installed ecosystem or hot-reloading
@@ -1070,9 +1116,7 @@ namespace UnrealBuildTool
 						throw new BuildException("Unable to generate headers because UnrealHeaderTool binary was not found ({0}).", HeaderToolPath);
 					}
 
-					// Disable extensions when serializing to remove the $type fields
-					Directory.CreateDirectory(ModuleInfoFileName.Directory.FullName);
-					System.IO.File.WriteAllText(ModuleInfoFileName.FullName, JsonSerializer.Serialize(Manifest, new JsonSerializerOptions { WriteIndented = true }));
+					WriteUHTManifest(Makefile, TargetName, ModuleInfoFileName, ExternalDependenciesFile);
 
 					string CmdLine = (ProjectFile != null) ? "\"" + ProjectFile.FullName + "\"" : TargetName;
 					CmdLine += " \"" + ModuleInfoFileName + "\" -LogCmds=\"loginit warning, logexit warning, logdatabase error\" -Unattended -WarningsAsErrors";
@@ -1188,17 +1232,6 @@ namespace UnrealBuildTool
 
 				// Get UHT assembly timestamp
 				DateTime UnrealBuildToolTimestamp = new FileInfo(Assembly.GetExecutingAssembly().Location).LastWriteTimeUtc;
-				//DateTime UnrealHeaderToolTimestamp = DateTime.MinValue;
-				//Assembly? UHTAssembly = AppDomain.CurrentDomain.GetAssemblies().Single(x => x.GetName().Name == "EpicGames.UHT");
-				//if (UHTAssembly != null)
-				//{
-				//	UnrealHeaderToolTimestamp = new FileInfo(UHTAssembly.Location).LastWriteTimeUtc;
-				//}
-				DateTime ComboTimeStamp = UnrealBuildToolTimestamp;
-				//if (UnrealHeaderToolTimestamp > ComboTimeStamp)
-				//{
-				//	ComboTimeStamp = UnrealHeaderToolTimestamp;
-				//}
 
 				// ensure the headers are up to date
 				bool bUHTNeedsToRun = false;
@@ -1206,13 +1239,13 @@ namespace UnrealBuildTool
 				{
 					bUHTNeedsToRun = true;
 				}
-				else if (AreGeneratedCodeFilesOutOfDate(BuildConfiguration, Makefile.UObjectModules, ComboTimeStamp))
+				else if (AreGeneratedCodeFilesOutOfDate(BuildConfiguration, Makefile.UObjectModules, UnrealBuildToolTimestamp))
 				{
 					bUHTNeedsToRun = true;
 				}
 
 				// Check we're not using a different version of UHT
-				FileReference ModuleInfoFileName = FileReference.Combine(Makefile.ProjectIntermediateDirectory, TargetName + ".uhtmanifest");
+				FileReference ModuleInfoFileName = GetUHTModuleInfoFileName(Makefile, TargetName);
 				FileReference ToolInfoFile = ModuleInfoFileName.ChangeExtension(".uhtpath");
 				if (!bUHTNeedsToRun)
 				{
@@ -1227,51 +1260,17 @@ namespace UnrealBuildTool
 				}
 
 				// Get the file containing dependencies for the generated code
-				FileReference ExternalDependenciesFile = ModuleInfoFileName.ChangeExtension(".deps");
+				FileReference ExternalDependenciesFile = GetUHTDepsFileName(ModuleInfoFileName);
 				if (AreExternalDependenciesOutOfDate(ExternalDependenciesFile))
 				{
 					bUHTNeedsToRun = true;
 				}
 
-				// @todo ubtmake: Optimization: Ideally we could avoid having to generate this data in the case where UHT doesn't even need to run!  Can't we use the existing copy?  (see below use of Manifest)
-
-				List<UHTManifest.Module> Modules = new List<UHTManifest.Module>();
-				foreach(UHTModuleInfo UObjectModule in Makefile.UObjectModules)
-				{
-					Modules.Add(
-						new UHTManifest.Module
-						{
-							Name = UObjectModule.ModuleName,
-							ModuleType = (UHTModuleType)Enum.Parse(typeof(UHTModuleType), UObjectModule.ModuleType),
-							OverrideModuleType = (EPackageOverrideType)Enum.Parse(typeof(EPackageOverrideType), UObjectModule.OverrideModuleType),
-							BaseDirectory = UObjectModule.ModuleDirectories[0].FullName,
-							IncludeBase = UObjectModule.ModuleDirectories[0].ParentDirectory!.FullName,
-							OutputDirectory = Path.GetDirectoryName(UObjectModule.GeneratedCPPFilenameBase)!,
-							ClassesHeaders = UObjectModule.PublicUObjectClassesHeaders.Select((Header) => Header.AbsolutePath).ToList(),
-							PublicHeaders = UObjectModule.PublicUObjectHeaders.Select((Header) => Header.AbsolutePath).ToList(),
-							InternalHeaders = UObjectModule.InternalUObjectHeaders.Select((Header) => Header.AbsolutePath).ToList(),
-							PrivateHeaders = UObjectModule.PrivateUObjectHeaders.Select((Header) => Header.AbsolutePath).ToList(),
-							GeneratedCPPFilenameBase = UObjectModule.GeneratedCPPFilenameBase,
-							SaveExportedHeaders = !UObjectModule.bIsReadOnly,
-							GeneratedCodeVersion = UObjectModule.GeneratedCodeVersion,
-						});
-				}
-				UHTManifest Manifest = new UHTManifest
-				{
-					TargetName = TargetName,
-					IsGameTarget = Makefile.TargetType != TargetType.Program,
-					RootLocalPath = RootLocalPath,
-					ExternalDependenciesFile = ExternalDependenciesFile.FullName,
-					Modules = Modules,
-				};
-
 				if (bUHTNeedsToRun)
 				{
 					Progress.Write(1, 3);
 
-					// Disable extensions when serializing to remove the $type fields
-					Directory.CreateDirectory(ModuleInfoFileName.Directory.FullName);
-					System.IO.File.WriteAllText(ModuleInfoFileName.FullName, JsonSerializer.Serialize(Manifest, new JsonSerializerOptions { WriteIndented = true }));
+					WriteUHTManifest(Makefile, TargetName, ModuleInfoFileName, ExternalDependenciesFile);
 
 					string ActualTargetName = String.IsNullOrEmpty(TargetName) ? "UE5" : TargetName;
 					Log.TraceInformation("Parsing headers for {0}", ActualTargetName);
