UnrealBuildTool: bPreprocessDepends for LinuxToolChain

#rnx

diff --git a/Engine/Source/Programs/UnrealBuildTool/Platform/Linux/LinuxToolChain.cs b/Engine/Source/Programs/UnrealBuildTool/Platform/Linux/LinuxToolChain.cs
--- a/Engine/Source/Programs/UnrealBuildTool/Platform/Linux/LinuxToolChain.cs
+++ b/Engine/Source/Programs/UnrealBuildTool/Platform/Linux/LinuxToolChain.cs
@@ -79,6 +79,8 @@ namespace UnrealBuildTool
 		/** Pass --gdb-index option to linker to generate .gdb_index section. */
 		protected bool bGdbIndexSection = true;
 
+		protected bool bPreprocessDepends = false;
+
 		/** Allows you to override the maximum binary size allowed to be passed to objcopy.exe when cross building on Windows. */
 		/** Max value is 2GB, due to bat file limitation */
 		protected UInt64 MaxBinarySizeOverrideForObjcopy = 0;
@@ -195,6 +197,13 @@ namespace UnrealBuildTool
 			bPreservePSYM = InPreservePSYM;
 		}
 
+		public override void SetUpGlobalEnvironment(ReadOnlyTargetRules Target)
+		{
+			base.SetUpGlobalEnvironment(Target);
+
+			bPreprocessDepends = Target.LinuxPlatform.bPreprocessDepends;
+		}
+
 		private string GetHostPlatformBinarySuffix()
 		{
 			if (BuildHostPlatform.Current.Platform == UnrealTargetPlatform.Win64)
@@ -1224,6 +1233,22 @@ namespace UnrealBuildTool
 			Log.TraceInformation("------------------------------");
 		}
 
+		public static string NormalizeCommandLinePath(FileSystemReference Reference)
+		{
+			// Try to use a relative path to shorten command line length.
+			if (Reference.IsUnderDirectory(Unreal.RootDirectory))
+			{
+				return Reference.MakeRelativeTo(UnrealBuildTool.EngineSourceDirectory).Replace("\\", "/");
+			}
+
+			return Reference.FullName.Replace("\\", "/");
+		}
+
+		public static string NormalizeCommandLinePath(FileItem Item)
+		{
+			return NormalizeCommandLinePath(Item.Location);
+		}
+
 		public override CPPOutput CompileCPPFiles(CppCompileEnvironment CompileEnvironment, List<FileItem> InputFiles, DirectoryReference OutputDir, string ModuleName, IActionGraphBuilder Graph)
 		{
 			string Arguments = GetCLArguments_Global(CompileEnvironment);
@@ -1238,26 +1263,17 @@ namespace UnrealBuildTool
 
 			if (CompileEnvironment.PrecompiledHeaderAction == PrecompiledHeaderAction.Include)
 			{
-				PCHArguments += string.Format(" -include \"{0}\"", CompileEnvironment.PrecompiledHeaderIncludeFilename!.FullName.Replace('\\', '/'));
+				PCHArguments += string.Format(" -include \"{0}\"", NormalizeCommandLinePath(CompileEnvironment.PrecompiledHeaderIncludeFilename!));
 			}
 
 			// Add include paths to the argument list.
 			foreach (DirectoryReference IncludePath in CompileEnvironment.UserIncludePaths)
 			{
-				string IncludePathString;
-				if (IncludePath.IsUnderDirectory(Unreal.RootDirectory))
-				{
-					IncludePathString = IncludePath.MakeRelativeTo(UnrealBuildTool.EngineSourceDirectory);
-				}
-				else
-				{
-					IncludePathString = IncludePath.FullName;
-				}
-				Arguments += string.Format(" -I\"{0}\"", IncludePathString.Replace('\\', '/'));
+				Arguments += string.Format(" -I\"{0}\"", NormalizeCommandLinePath(IncludePath));
 			}
 			foreach (DirectoryReference IncludePath in CompileEnvironment.SystemIncludePaths)
 			{
-				Arguments += string.Format(" -I\"{0}\"", IncludePath.FullName.Replace('\\', '/'));
+				Arguments += string.Format(" -I\"{0}\"", NormalizeCommandLinePath(IncludePath));
 			}
 
 			// Add preprocessor definitions to the argument list.
@@ -1308,7 +1324,7 @@ namespace UnrealBuildTool
 
 				foreach (FileItem ForceIncludeFile in CompileEnvironment.ForceIncludeFiles)
 				{
-					FileArguments += String.Format(" -include \"{0}\"", ForceIncludeFile.Location.FullName.Replace('\\', '/'));
+					FileArguments += String.Format(" -include \"{0}\"", NormalizeCommandLinePath(ForceIncludeFile));
 				}
 
 				// Add the C++ source file and its included files to the prerequisite item list.
@@ -1323,7 +1339,7 @@ namespace UnrealBuildTool
 					Result.PrecompiledHeaderFile = PrecompiledHeaderFile;
 
 					// Add the parameters needed to compile the precompiled header file to the command-line.
-					FileArguments += string.Format(" -o \"{0}\"", PrecompiledHeaderFile.AbsolutePath.Replace('\\', '/'));
+					FileArguments += string.Format(" -o \"{0}\"", NormalizeCommandLinePath(PrecompiledHeaderFile));
 				}
 				else
 				{
@@ -1337,11 +1353,11 @@ namespace UnrealBuildTool
 					CompileAction.ProducedItems.Add(ObjectFile);
 					Result.ObjectFiles.Add(ObjectFile);
 
-					FileArguments += string.Format(" -o \"{0}\"", ObjectFile.AbsolutePath.Replace('\\', '/'));
+					FileArguments += string.Format(" -o \"{0}\"", NormalizeCommandLinePath(ObjectFile));
 				}
 
 				// Add the source file path to the command-line.
-				FileArguments += string.Format(" \"{0}\"", SourceFile.AbsolutePath.Replace('\\', '/'));
+				FileArguments += string.Format(" \"{0}\"", NormalizeCommandLinePath(SourceFile));
 
 				// Generate the timing info
 				if (CompileEnvironment.bPrintTimingInfo)
@@ -1352,10 +1368,10 @@ namespace UnrealBuildTool
 				}
 
 				// Generate the included header dependency list
-				if(CompileEnvironment.bGenerateDependenciesFile)
+				if(!bPreprocessDepends && CompileEnvironment.bGenerateDependenciesFile)
 				{
 					FileItem DependencyListFile = FileItem.GetItemByFileReference(FileReference.Combine(OutputDir, Path.GetFileName(SourceFile.AbsolutePath) + ".d"));
-					FileArguments += string.Format(" -MD -MF\"{0}\"", DependencyListFile.AbsolutePath.Replace('\\', '/'));
+					FileArguments += string.Format(" -MD -MF\"{0}\"", NormalizeCommandLinePath(DependencyListFile));
 					CompileAction.DependencyListFile = DependencyListFile;
 					CompileAction.ProducedItems.Add(DependencyListFile);
 				}
@@ -1371,7 +1387,7 @@ namespace UnrealBuildTool
 				FileReference CompilerResponseFileName = CompileAction.ProducedItems[0].Location + ".rsp";
 				FileItem CompilerResponseFileItem = Graph.CreateIntermediateTextFile(CompilerResponseFileName, AllArguments);
 
-				CompileAction.CommandArguments = string.Format(" @\"{0}\"", CompilerResponseFileName);
+				CompileAction.CommandArguments = string.Format(" @\"{0}\"", NormalizeCommandLinePath(CompilerResponseFileName));
 				CompileAction.PrerequisiteItems.Add(CompilerResponseFileItem);
 				CompileAction.CommandDescription = "Compile";
 				CompileAction.CommandVersion = CompilerVersionString!;
@@ -1382,6 +1398,46 @@ namespace UnrealBuildTool
 				CompileAction.bCanExecuteRemotely =
 					CompileEnvironment.PrecompiledHeaderAction != PrecompiledHeaderAction.Create ||
 					CompileEnvironment.bAllowRemotelyCompiledPCHs;
+
+				if (bPreprocessDepends && CompileEnvironment.bGenerateDependenciesFile)
+				{
+					Action PrepassAction = Graph.CreateAction(ActionType.Compile);
+					PrepassAction.PrerequisiteItems.AddRange(CompileAction.PrerequisiteItems);
+					PrepassAction.PrerequisiteItems.Remove(CompilerResponseFileItem);
+					PrepassAction.CommandDescription = "Preprocess Depends";
+					PrepassAction.StatusDescription = CompileAction.StatusDescription;
+					PrepassAction.bIsGCCCompiler = true;
+					PrepassAction.bCanExecuteRemotely = false;
+					PrepassAction.bShouldOutputStatusDescription = true;
+					PrepassAction.CommandPath = CompileAction.CommandPath;
+					PrepassAction.CommandVersion = CompileAction.CommandVersion;
+					PrepassAction.WorkingDirectory = CompileAction.WorkingDirectory;
+
+					string PreprocessArguments = Arguments.ToString();
+					string PreprocessFileArguments = FileArguments;
+					PreprocessArguments = PreprocessArguments.Replace(" -c ", " ");
+
+					FileItem DependencyListFile = FileItem.GetItemByFileReference(FileReference.Combine(OutputDir, Path.GetFileName(SourceFile.AbsolutePath) + ".d"));
+					PreprocessFileArguments += string.Format(" -M -MF\"{0}\"", NormalizeCommandLinePath(DependencyListFile));
+					PrepassAction.DependencyListFile = DependencyListFile;
+					PrepassAction.ProducedItems.Add(DependencyListFile);
+
+					PreprocessFileArguments = PreprocessFileArguments.Replace(" -ftime-trace", string.Empty);
+					PreprocessFileArguments = PreprocessFileArguments.Replace(string.Format(" -o \"{0}\"", NormalizeCommandLinePath(CompileAction.ProducedItems.First())), String.Empty);
+
+					PrepassAction.DeleteItems.AddRange(PrepassAction.ProducedItems);
+
+					string AllPreprocessArgs = PreprocessArguments + PreprocessFileArguments + CompileEnvironment.AdditionalArguments;
+
+					FileReference CompilerPrepassResponseFileName = PrepassAction.ProducedItems[0].Location + ".rsp";
+					FileItem CompilerPrepassResponseFileItem = Graph.CreateIntermediateTextFile(CompilerPrepassResponseFileName, AllPreprocessArgs);
+
+					PrepassAction.CommandArguments = string.Format(" @\"{0}\"", NormalizeCommandLinePath(CompilerPrepassResponseFileName));
+					PrepassAction.PrerequisiteItems.Add(CompilerResponseFileItem);
+
+					CompileAction.DependencyListFile = DependencyListFile;
+					CompileAction.PrerequisiteItems.Add(DependencyListFile);
+				}
 			}
 
 			return Result;
diff --git a/Engine/Source/Programs/UnrealBuildTool/Platform/Linux/UEBuildLinux.cs b/Engine/Source/Programs/UnrealBuildTool/Platform/Linux/UEBuildLinux.cs
--- a/Engine/Source/Programs/UnrealBuildTool/Platform/Linux/UEBuildLinux.cs
+++ b/Engine/Source/Programs/UnrealBuildTool/Platform/Linux/UEBuildLinux.cs
@@ -77,6 +77,13 @@ namespace UnrealBuildTool
 		[CommandLine("-EnableLLDB")]
 		[XmlConfigFile(Category = "BuildConfiguration", Name = "bTuneDebugInfoForLLDB")]
 		public bool bTuneDebugInfoForLLDB = false;
+
+		/// <summary>
+		/// Generate dependency files by preprocessing. This is only recommended when distributing builds as it adds additional overhead.
+		/// </summary>
+		[CommandLine("-PreprocessDepends")]
+		[XmlConfigFile(Category = "BuildConfiguration", Name = "bPreprocessDepends")]
+		public bool bPreprocessDepends = false;
 	}
 
 	/// <summary>
@@ -131,6 +138,11 @@ namespace UnrealBuildTool
 			get { return Inner.bTuneDebugInfoForLLDB; }
 		}
 
+		public bool bPreprocessDepends
+		{
+			get { return Inner.bPreprocessDepends; }
+		}
+
 		#pragma warning restore CS1591
 		#endregion
 	}
