Re-add printing build details

diff --git a/Engine/Source/Programs/UnrealBuildTool/Platform/Linux/LinuxToolChain.cs b/Engine/Source/Programs/UnrealBuildTool/Platform/Linux/LinuxToolChain.cs
--- a/Engine/Source/Programs/UnrealBuildTool/Platform/Linux/LinuxToolChain.cs
+++ b/Engine/Source/Programs/UnrealBuildTool/Platform/Linux/LinuxToolChain.cs
@@ -1140,6 +1140,11 @@ namespace UnrealBuildTool
 		/// </summary>
 		private static List<FileItem> AllBinaries = new List<FileItem>();
 
+		/// <summary>
+		/// Tracks that information about used C++ library is only printed once
+		/// </summary>
+		private bool bHasPrintedBuildDetails = false;
+
 		/// <summary>
 		/// Checks if we actually can use LTO/PGO with this set of tools
 		/// </summary>
@@ -1256,6 +1261,13 @@ namespace UnrealBuildTool
 
 			//var BuildPlatform = UEBuildPlatform.GetBuildPlatform(CompileEnvironment.Platform);
 
+			if (!bHasPrintedBuildDetails)
+			{
+				PrintBuildDetails(CompileEnvironment);
+
+				bHasPrintedBuildDetails = true;
+			}
+
 			if ((CompileEnvironment.bAllowLTCG || CompileEnvironment.bPGOOptimize || CompileEnvironment.bPGOProfile) && !CanUseAdvancedLinkerFeatures(CompileEnvironment.Architecture))
 			{
 				throw new BuildException(ExplainWhyCannotUseAdvancedLinkerFeatures(CompileEnvironment.Architecture));
