Refactor Apple Platforms Compilation

This work aims to make it so more of the work to compile files between the various Apple platforms is shared.

#rnx

diff --git a/Engine/Source/Programs/UnrealBuildTool/Platform/IOS/IOSToolChain.cs b/Engine/Source/Programs/UnrealBuildTool/Platform/IOS/IOSToolChain.cs
--- a/Engine/Source/Programs/UnrealBuildTool/Platform/IOS/IOSToolChain.cs
+++ b/Engine/Source/Programs/UnrealBuildTool/Platform/IOS/IOSToolChain.cs
@@ -226,13 +226,35 @@ namespace UnrealBuildTool
 		}
 
 
-		string GetCompileArguments_Global(CppCompileEnvironment CompileEnvironment)
+		protected override string GetCompileArgumentsForCompileEnvironment(CppCompileEnvironment CompileEnvironment)
 		{
+			StringBuilder Arguments = new StringBuilder();
+			Arguments.Append(base.GetCompileArgumentsForCompileEnvironment(CompileEnvironment));
+
 			string Result = "";
+			// Add additional frameworks so that their headers can be found
+			foreach (UEBuildFramework Framework in CompileEnvironment.AdditionalFrameworks)
+			{
+				if (Framework.FrameworkDirectory != null)
+				{
+					string FrameworkDir = Framework.FrameworkDirectory.FullName;
+					// embedded frameworks have a framework inside of this directory, so we use this directory. regular frameworks need to go one up to point to the 
+					// directory containing the framework. -F gives a path to look for the -framework
+					if (FrameworkDir.EndsWith(".framework"))
+					{
+						FrameworkDir = Path.GetDirectoryName(FrameworkDir)!;
+					}
+					Result += String.Format(" -F\"{0}\"", FrameworkDir);
+				}
+			}
 
-			Result += " -fmessage-length=0";
-			Result += " -pipe";
-			Result += " -fpascal-strings";
+			Arguments.Append(Result);
+			return Arguments.ToString();
+		}
+
+		protected override string GetCompileArguments_Global(CppCompileEnvironment CompileEnvironment)
+		{
+			string Result = base.GetCompileArguments_Global(CompileEnvironment);
 
 			if (RuntimePlatform.IsWindows)
 			{
@@ -277,27 +299,6 @@ namespace UnrealBuildTool
 				Result += " -flto";
 			}
 
-			Result += " -Wall -Werror";
-			Result += " -Wdelete-non-virtual-dtor";
-
-			// clang 12.00 has a new warning for copies in ranged loops. Instances have all been fixed up (2020/6/26) but
-			// are likely to be reintroduced due to no equivalent on other platforms at this time so disable the warning
-			// See also MacToolChain.cs
-			if (GetClangVersion().Major >= 12)
-			{
-				Result += " -Wno-range-loop-analysis";
-			}
-
-			if (CompileEnvironment.ShadowVariableWarningLevel != WarningLevel.Off)
-			{
-				Result += " -Wshadow" + ((CompileEnvironment.ShadowVariableWarningLevel == WarningLevel.Error) ? "" : " -Wno-error=shadow");
-			}
-
-			if (CompileEnvironment.bEnableUndefinedIdentifierWarnings)
-			{
-				Result += " -Wundef" + (CompileEnvironment.bUndefinedIdentifierWarningsAsErrors ? "" : " -Wno-error=undef");
-			}
-
 			// fix for Xcode 8.3 enabling nonportable include checks, but p4 has some invalid cases in it
 			if (Settings.Value.IOSSDKVersionFloat >= 10.3)
 			{
@@ -309,8 +310,6 @@ namespace UnrealBuildTool
 				Result += " -fembed-bitcode";
 			}
 
-			Result += " -c";
-
 			// What architecture(s) to build for
 			Result += GetArchitectureArgument(CompileEnvironment.Configuration, CompileEnvironment.Architecture);
 
@@ -342,102 +341,15 @@ namespace UnrealBuildTool
 				Result += " -O0";
 			}
 
-			if (!CompileEnvironment.bUseInlining)
-			{
-				Result += " -fno-inline-functions";
-			}
-
-			// Create DWARF format debug info if wanted,
-			if (CompileEnvironment.bCreateDebugInfo)
-			{
-				Result += " -gdwarf-2";
-			}
-
-			// Add additional frameworks so that their headers can be found
-			foreach (UEBuildFramework Framework in CompileEnvironment.AdditionalFrameworks)
-			{
-				if (Framework.FrameworkDirectory != null)
-				{
-					string FrameworkDir = Framework.FrameworkDirectory.FullName;
-					// embedded frameworks have a framework inside of this directory, so we use this directory. regular frameworks need to go one up to point to the 
-					// directory containing the framework. -F gives a path to look for the -framework
-					if (FrameworkDir.EndsWith(".framework"))
-					{
-						FrameworkDir = Path.GetDirectoryName(FrameworkDir)!;
-					}
-					Result += String.Format(" -F\"{0}\"", FrameworkDir);
-				}
-			}
-
-			return Result;
-		}
-
-		static string GetCompileArguments_CPP(CppCompileEnvironment CompileEnvironment)
-		{
-			string Result = "";
-			Result += " -x objective-c++";
-			Result += GetCppStandardCompileArgument(CompileEnvironment);
-			Result += " -stdlib=libc++";
-
-			return Result;
-		}
-
-		static string GetCompileArguments_MM(CppCompileEnvironment CompileEnvironment)
-		{
-			string Result = "";
-			Result += " -x objective-c++";
-			Result += GetCppStandardCompileArgument(CompileEnvironment);
-			Result += " -stdlib=libc++";
-			return Result;
-		}
-
-		static string GetCompileArguments_M(CppCompileEnvironment CompileEnvironment)
-		{
-			string Result = "";
-			Result += " -x objective-c";
-			Result += GetCppStandardCompileArgument(CompileEnvironment);
-			Result += " -stdlib=libc++";
-			return Result;
-		}
-
-		static string GetCompileArguments_C()
-		{
-			string Result = "";
-			Result += " -x c";
-			return Result;
-		}
-
-		static string GetCompileArguments_PCH(CppCompileEnvironment CompileEnvironment)
-		{
-			string Result = "";
-			Result += " -x objective-c++-header";
-			Result += GetCppStandardCompileArgument(CompileEnvironment);
-			Result += " -stdlib=libc++";
-
-			return Result;
-		}
-
-		// Conditionally enable (default disabled) generation of information about every class with virtual functions for use by the C++ runtime type identification features 
-		// (`dynamic_cast' and `typeid'). If you don't use those parts of the language, you can save some space by using -fno-rtti. 
-		// Note that exception handling uses the same information, but it will generate it as needed. 
-		static string GetRTTIFlag(CppCompileEnvironment CompileEnvironment)
-		{
-			string Result = "";
-
-			if (CompileEnvironment.bUseRTTI)
-			{
-				Result = " -frtti";
-			}
-			else
-			{
-				Result = " -fno-rtti";
-			}
-
+			// Added to cope with errors on the index target when building with Xcode
+			Result += " -Wno-implicit-int-conversion";
+			Result += " -Wno-shorten-64-to-32";
+			Result += " -Wno-inconsistent-missing-override";
 			return Result;
 		}
 
 		// Conditionally enable (default disabled) Objective-C exceptions
-		static string GetObjCExceptionsFlag(CppCompileEnvironment CompileEnvironment)
+		protected override string GetObjCExceptionsFlag(CppCompileEnvironment CompileEnvironment)
 		{
 			string Result = "";
 
@@ -622,77 +534,16 @@ namespace UnrealBuildTool
 
 		public override CPPOutput CompileCPPFiles(CppCompileEnvironment CompileEnvironment, List<FileItem> InputFiles, DirectoryReference OutputDir, string ModuleName, IActionGraphBuilder Graph)
 		{
-			string Arguments = GetCompileArguments_Global(CompileEnvironment);
-			string PCHArguments = "";
-
-			if (CompileEnvironment.PrecompiledHeaderAction == PrecompiledHeaderAction.Include)
-			{
-				PCHArguments += string.Format(" -include \"{0}\"", CompileEnvironment.PrecompiledHeaderIncludeFilename!.FullName);
-			}
-
-			foreach (FileItem ForceIncludeFile in CompileEnvironment.ForceIncludeFiles)
-			{
-				Arguments += String.Format(" -include \"{0}\"", ForceIncludeFile.Location.FullName);
-			}
-
-			// Add include paths to the argument list.
-			HashSet<DirectoryReference> AllIncludes = new HashSet<DirectoryReference>(CompileEnvironment.UserIncludePaths);
-			AllIncludes.UnionWith(CompileEnvironment.SystemIncludePaths);
-			foreach (DirectoryReference IncludePath in AllIncludes)
-			{
-				Arguments += string.Format(" -I\"{0}\"", IncludePath.FullName);
-			}
-
-			foreach (string Definition in CompileEnvironment.Definitions)
-			{
-				string DefinitionArgument = Definition.Contains("\"") ? Definition.Replace("\"", "\\\"") : Definition;
-				Arguments += string.Format(" -D \"{0}\"", DefinitionArgument);
-			}
+			string Arguments = GetCompileArgumentsForCompileEnvironment(CompileEnvironment);
+			string PCHArguments = GetPCHCompileArgumentsForCompileEnvironment(CompileEnvironment);
 
 			CPPOutput Result = new CPPOutput();
 			// Create a compile action for each source file.
 			foreach (FileItem SourceFile in InputFiles)
 			{
 				Action CompileAction = Graph.CreateAction(ActionType.Compile);
-				string FilePCHArguments = "";
-				string FileArguments = "";
-				string Extension = Path.GetExtension(SourceFile.AbsolutePath).ToUpperInvariant();
 
-				if (CompileEnvironment.PrecompiledHeaderAction == PrecompiledHeaderAction.Create)
-				{
-					// Compile the file as a C++ PCH.
-					FileArguments += GetCompileArguments_PCH(CompileEnvironment);
-					FileArguments += GetRTTIFlag(CompileEnvironment);
-					FileArguments += GetObjCExceptionsFlag(CompileEnvironment);
-				}
-				else if (Extension == ".C")
-				{
-					// Compile the file as C code.
-					FileArguments += GetCompileArguments_C();
-				}
-				else if (Extension == ".MM")
-				{
-					// Compile the file as Objective-C++ code.
-					FileArguments += GetCompileArguments_MM(CompileEnvironment);
-					FileArguments += GetRTTIFlag(CompileEnvironment);
-					FileArguments += GetObjCExceptionsFlag(CompileEnvironment);
-				}
-				else if (Extension == ".M")
-				{
-					// Compile the file as Objective-C code.
-					FileArguments += GetCompileArguments_M(CompileEnvironment);
-					FileArguments += GetObjCExceptionsFlag(CompileEnvironment);
-				}
-				else
-				{
-					// Compile the file as C++ code.
-					FileArguments += GetCompileArguments_CPP(CompileEnvironment);
-					FileArguments += GetRTTIFlag(CompileEnvironment);
-					FileArguments += GetObjCExceptionsFlag(CompileEnvironment);
-
-					// only use PCH for .cpp files
-					FilePCHArguments = PCHArguments;
-				}
+				string AllArgs = GetCompileArgumentsForAbsoluteFilePathWithCompileEnvironment(SourceFile.AbsolutePath, CompileEnvironment);
 
 				// Add the C++ source file and its included files to the prerequisite item list.
 				CompileAction.PrerequisiteItems.Add(SourceFile);
@@ -719,7 +570,7 @@ namespace UnrealBuildTool
 					Result.PrecompiledHeaderFile = PrecompiledHeaderFile;
 
 					// Add the parameters needed to compile the precompiled header file to the command-line.
-					FileArguments += string.Format(" -o \"{0}\"", PrecompiledHeaderFile.AbsolutePath);
+					AllArgs += string.Format(" -o \"{0}\"", PrecompiledHeaderFile.AbsolutePath);
 				}
 				else
 				{
@@ -743,36 +594,24 @@ namespace UnrealBuildTool
 
 					CompileAction.ProducedItems.Add(ObjectFile);
 					Result.ObjectFiles.Add(ObjectFile);
-					FileArguments += string.Format(" -o \"{0}\"", ObjectFile.AbsolutePath);
+					AllArgs += string.Format(" -o \"{0}\"", ObjectFile.AbsolutePath);
 					OutputFilePath = ObjectFile.AbsolutePath;
 				}
 
 				// Add the source file path to the command-line.
-				FileArguments += string.Format(" \"{0}\"", SourceFile.AbsolutePath);
+				AllArgs += string.Format(" \"{0}\"", SourceFile.AbsolutePath);
 
 				// Generate the included header dependency list
 				if (CompileEnvironment.bGenerateDependenciesFile)
 				{
 					FileItem DependencyListFile = FileItem.GetItemByFileReference(FileReference.Combine(OutputDir, Path.GetFileName(SourceFile.AbsolutePath) + ".d"));
-					FileArguments += string.Format(" -MD -MF\"{0}\"", DependencyListFile.AbsolutePath.Replace('\\', '/'));
+					AllArgs += string.Format(" -MD -MF\"{0}\"", DependencyListFile.AbsolutePath.Replace('\\', '/'));
 					CompileAction.DependencyListFile = DependencyListFile;
 					CompileAction.ProducedItems.Add(DependencyListFile);
 				}
 
 				string CompilerPath = Settings.Value.ToolchainDir + IOSCompiler;
 
-				string AllArgs = FilePCHArguments + Arguments + FileArguments + CompileEnvironment.AdditionalArguments;
-				/*				string SourceText = System.IO.File.ReadAllText(SourceFile.AbsolutePath);
-								if (CompileEnvironment.bOptimizeForSize && (SourceFile.AbsolutePath.Contains("ElementBatcher.cpp") || SourceText.Contains("ElementBatcher.cpp") || SourceFile.AbsolutePath.Contains("AnimationRuntime.cpp") || SourceText.Contains("AnimationRuntime.cpp")
-									|| SourceFile.AbsolutePath.Contains("AnimEncoding.cpp") || SourceText.Contains("AnimEncoding.cpp") || SourceFile.AbsolutePath.Contains("TextRenderComponent.cpp") || SourceText.Contains("TextRenderComponent.cpp")
-									|| SourceFile.AbsolutePath.Contains("SWidget.cpp") || SourceText.Contains("SWidget.cpp") || SourceFile.AbsolutePath.Contains("SCanvas.cpp") || SourceText.Contains("SCanvas.cpp") || SourceFile.AbsolutePath.Contains("ShaderCore.cpp") || SourceText.Contains("ShaderCore.cpp")
-									|| SourceFile.AbsolutePath.Contains("ParticleSystemRender.cpp") || SourceText.Contains("ParticleSystemRender.cpp")))
-								{
-									Log.TraceInformation("Forcing {0} to --O3!", SourceFile.AbsolutePath);
-
-									AllArgs = AllArgs.Replace("-Oz", "-O3");
-								}*/
-
 				// Analyze and then compile using the shell to perform the indirection
 				string? StaticAnalysisMode = Environment.GetEnvironmentVariable("CLANG_STATIC_ANALYZER_MODE");
 				if (StaticAnalysisMode != null && StaticAnalysisMode != "" && OutputFilePath != null)
diff --git a/Engine/Source/Programs/UnrealBuildTool/Platform/Mac/MacToolChain.cs b/Engine/Source/Programs/UnrealBuildTool/Platform/Mac/MacToolChain.cs
--- a/Engine/Source/Programs/UnrealBuildTool/Platform/Mac/MacToolChain.cs
+++ b/Engine/Source/Programs/UnrealBuildTool/Platform/Mac/MacToolChain.cs
@@ -191,13 +191,9 @@ namespace UnrealBuildTool
 			return ArchArg;
 		}
 
-		string GetCompileArguments_Global(CppCompileEnvironment CompileEnvironment)
+		protected override string GetCompileArguments_Global(CppCompileEnvironment CompileEnvironment)
 		{
-			string Result = "";
-
-			Result += " -fmessage-length=0";
-			Result += " -pipe";
-			Result += " -fpascal-strings";
+			string Result = base.GetCompileArguments_Global(CompileEnvironment);
 
 			Result += " -fexceptions";
 			Result += " -DPLATFORM_EXCEPTIONS_DISABLED=0";
@@ -222,34 +218,9 @@ namespace UnrealBuildTool
 				Result += " -fsanitize=undefined";
 			}			
 
-			Result += " -Wall -Werror";
-			Result += " -Wdelete-non-virtual-dtor";
-
-			// clang 12.00 has a new warning for copies in ranged loops. Instances have all been fixed up (2020/6/26) but
-			// are likely to be reintroduced due to no equivalent on other platforms at this time so disable the warning
-			if (GetClangVersion().Major >= 12)
-			{
-				Result += " -Wno-range-loop-analysis ";
-			}			
-
 			//Result += " -Wsign-compare"; // fed up of not seeing the signed/unsigned warnings we get on Windows - lets enable them here too.
 
-			if (CompileEnvironment.ShadowVariableWarningLevel != WarningLevel.Off)
-			{
-				Result += " -Wshadow" + ((CompileEnvironment.ShadowVariableWarningLevel == WarningLevel.Error) ? "" : " -Wno-error=shadow");
-			}
-			
-			if (CompileEnvironment.bEnableUndefinedIdentifierWarnings)
-			{
-				Result += " -Wundef" + (CompileEnvironment.bUndefinedIdentifierWarningsAsErrors ? "" : " -Wno-error=undef");
-			}
 
-			if (CompileEnvironment.bEnableOSX109Support)
-			{
-				Result += " -faligned-new"; // aligned operator new is supported only on macOS 10.14 and above
-			}
-
-			Result += " -c";
 
 			// Pass through architecture and OS info
 			Result += " " + FormatArchitectureArg(CompileEnvironment.Architecture);	
@@ -289,80 +260,6 @@ namespace UnrealBuildTool
 				Result += " -O0";
 			}
 
-			if (!CompileEnvironment.bUseInlining)
-			{
-				Result += " -fno-inline-functions";
-			}
-
-			// Create DWARF format debug info if wanted,
-			if (CompileEnvironment.bCreateDebugInfo)
-			{
-				Result += " -gdwarf-2";
-			}
-
-			return Result;
-		}
-		
-		static string GetCompileArguments_CPP(CppCompileEnvironment CompileEnvironment)
-		{
-			string Result = "";
-			Result += " -x objective-c++";
-			Result += GetCppStandardCompileArgument(CompileEnvironment);
-			Result += " -stdlib=libc++";
-
-			return Result;
-		}
-
-		static string GetCompileArguments_MM(CppCompileEnvironment CompileEnvironment)
-		{
-			string Result = "";
-			Result += " -x objective-c++";
-			Result += GetCppStandardCompileArgument(CompileEnvironment);
-			Result += " -stdlib=libc++";
-			return Result;
-		}
-
-		static string GetCompileArguments_M(CppCompileEnvironment CompileEnvironment)
-		{
-			string Result = "";
-			Result += " -x objective-c";
-			Result += " -stdlib=libc++";
-			return Result;
-		}
-
-		static string GetCompileArguments_C()
-		{
-			string Result = "";
-			Result += " -x c";
-			return Result;
-		}
-
-		static string GetCompileArguments_PCH(CppCompileEnvironment CompileEnvironment)
-		{
-			string Result = "";
-			Result += " -x objective-c++-header";
-			Result += GetCppStandardCompileArgument(CompileEnvironment);
-			Result += " -stdlib=libc++";
-
-			return Result;
-		}
-
-		// Conditionally enable (default disabled) generation of information about every class with virtual functions for use by the C++ runtime type identification features 
-		// (`dynamic_cast' and `typeid'). If you don't use those parts of the language, you can save some space by using -fno-rtti. 
-		// Note that exception handling uses the same information, but it will generate it as needed. 
-		static string GetRTTIFlag(CppCompileEnvironment CompileEnvironment)
-		{
-			string Result = "";
-
-			if (CompileEnvironment.bUseRTTI)
-			{
-				Result = " -frtti";
-			}
-			else
-			{
-				Result = " -fno-rtti";
-			}
-
 			return Result;
 		}
 		
@@ -432,53 +329,58 @@ namespace UnrealBuildTool
 			return Result;
 		}
 
-		public override CPPOutput CompileCPPFiles(CppCompileEnvironment CompileEnvironment, List<FileItem> InputFiles, DirectoryReference OutputDir, string ModuleName, IActionGraphBuilder Graph)
+		protected override string GetCompileArgumentsForCompileEnvironment(CppCompileEnvironment CompileEnvironment)
 		{
 			StringBuilder Arguments = new StringBuilder();
-			StringBuilder PCHArguments = new StringBuilder();
+			Arguments.Append(base.GetCompileArgumentsForCompileEnvironment(CompileEnvironment));
 
-			Arguments.Append(GetCompileArguments_Global(CompileEnvironment));
-
-			if (CompileEnvironment.PrecompiledHeaderAction == PrecompiledHeaderAction.Include)
+			List<string> FrameworksSearchPaths = new List<string>();
+			foreach (UEBuildFramework Framework in CompileEnvironment.AdditionalFrameworks)
 			{
-				// Add the precompiled header file's path to the include path so GCC can find it.
-				// This needs to be before the other include paths to ensure GCC uses it instead of the source header file.
-				PCHArguments.Append(" -include \"");
-				PCHArguments.Append(CompileEnvironment.PrecompiledHeaderIncludeFilename);
-				PCHArguments.Append("\"");
+				if (Framework.FrameworkDirectory != null)
+				{
+					string FrameworkPath = Path.GetDirectoryName(Path.GetFullPath(Framework.Name))!;
+					if (!FrameworksSearchPaths.Contains(FrameworkPath))
+					{
+						Arguments.Append(" -F \"");
+						Arguments.Append(FrameworkPath);
+						Arguments.Append("\"");
+						FrameworksSearchPaths.Add(FrameworkPath);
+					}
+				}
 			}
 
-			// Add include paths to the argument list.
-			HashSet<DirectoryReference> AllIncludes = new HashSet<DirectoryReference>(CompileEnvironment.UserIncludePaths);
-			AllIncludes.UnionWith(CompileEnvironment.SystemIncludePaths);
-			foreach (DirectoryReference IncludePath in AllIncludes)
-			{
-				Arguments.Append(" -I\"");
-				Arguments.Append(IncludePath);
-				Arguments.Append("\"");
-			}
+			return Arguments.ToString();
+		}
 
-			foreach (string Definition in CompileEnvironment.Definitions)
-			{
-				string DefinitionArgument = Definition.Contains("\"") ? Definition.Replace("\"", "\\\"") : Definition;
-				Arguments.Append(" -D\"");
-				Arguments.Append(DefinitionArgument);
-				Arguments.Append("\"");
-			}
+		public override string GetCompileArgumentsForAbsoluteFilePathWithCompileEnvironment(string FilePath, CppCompileEnvironment CompileEnvironment)
+		{
+			string Arguments = base.GetCompileArgumentsForAbsoluteFilePathWithCompileEnvironment(FilePath, CompileEnvironment);
 
-			List<string> FrameworksSearchPaths = new List<string>();
-			foreach (UEBuildFramework Framework in CompileEnvironment.AdditionalFrameworks)
+			string EscapedAdditionalArgs = "";
+			if (!string.IsNullOrWhiteSpace(CompileEnvironment.AdditionalArguments))
 			{
-				string FrameworkPath = Path.GetDirectoryName(Path.GetFullPath(Framework.Name))!;
-				if (!FrameworksSearchPaths.Contains(FrameworkPath))
+				foreach (string AdditionalArg in CompileEnvironment.AdditionalArguments.Split(new char[] { ' ' }, StringSplitOptions.RemoveEmptyEntries))
 				{
-					Arguments.Append(" -F \"");
-					Arguments.Append(FrameworkPath);
-					Arguments.Append("\"");
-					FrameworksSearchPaths.Add(FrameworkPath);
+					Match DefinitionMatch = Regex.Match(AdditionalArg, "-D\"?(?<Name>.*)=(?<Value>.*)\"?");
+					if (DefinitionMatch.Success)
+					{
+						EscapedAdditionalArgs += string.Format(" -D{0}=\"{1}\"", DefinitionMatch.Groups["Name"].Value, DefinitionMatch.Groups["Value"].Value);
+					}
+					else
+					{
+						EscapedAdditionalArgs += " " + AdditionalArg;
+					}
 				}
 			}
 
+			string Result = Arguments + EscapedAdditionalArgs;
+			
+			return Result;
+		}
+
+		public override CPPOutput CompileCPPFiles(CppCompileEnvironment CompileEnvironment, List<FileItem> InputFiles, DirectoryReference OutputDir, string ModuleName, IActionGraphBuilder Graph)
+		{
 			CPPOutput Result = new CPPOutput();
 			// Create a compile action for each source file.
 			foreach (FileItem SourceFile in InputFiles)
@@ -487,46 +389,9 @@ namespace UnrealBuildTool
 				CompileAction.PrerequisiteItems.AddRange(CompileEnvironment.ForceIncludeFiles);
 				CompileAction.PrerequisiteItems.AddRange(CompileEnvironment.AdditionalPrerequisites);
 
-				string FileArguments = "";
+				string AllArgs = GetCompileArgumentsForAbsoluteFilePathWithCompileEnvironment(SourceFile.AbsolutePath, CompileEnvironment);
 				string Extension = Path.GetExtension(SourceFile.AbsolutePath).ToUpperInvariant();
 
-				if (CompileEnvironment.PrecompiledHeaderAction == PrecompiledHeaderAction.Create)
-				{
-					// Compile the file as a C++ PCH.
-					FileArguments += GetCompileArguments_PCH(CompileEnvironment);
-					FileArguments += GetRTTIFlag(CompileEnvironment);
-				}
-				else if (Extension == ".C")
-				{
-					// Compile the file as C code.
-					FileArguments += GetCompileArguments_C();
-				}
-				else if (Extension == ".MM")
-				{
-					// Compile the file as Objective-C++ code.
-					FileArguments += GetCompileArguments_MM(CompileEnvironment);
-					FileArguments += GetRTTIFlag(CompileEnvironment);
-				}
-				else if (Extension == ".M")
-				{
-					// Compile the file as Objective-C++ code.
-					FileArguments += GetCompileArguments_M(CompileEnvironment);
-				}
-				else
-				{
-					// Compile the file as C++ code.
-					FileArguments += GetCompileArguments_CPP(CompileEnvironment);
-					FileArguments += GetRTTIFlag(CompileEnvironment);
-
-					// only use PCH for .cpp files
-					FileArguments += PCHArguments.ToString();
-				}
-
-				foreach (FileItem ForceIncludeFile in CompileEnvironment.ForceIncludeFiles)
-				{
-					FileArguments += String.Format(" -include \"{0}\"", ForceIncludeFile.Location);
-				}
-
 				// Add the C++ source file and its included files to the prerequisite item list.
 				CompileAction.PrerequisiteItems.Add(SourceFile);
 
@@ -539,7 +404,7 @@ namespace UnrealBuildTool
 					Result.PrecompiledHeaderFile = PrecompiledHeaderFile;
 
 					// Add the parameters needed to compile the precompiled header file to the command-line.
-					FileArguments += string.Format(" -o \"{0}\"", PrecompiledHeaderFile.AbsolutePath);
+					AllArgs += string.Format(" -o \"{0}\"", PrecompiledHeaderFile.AbsolutePath);
 				}
 				else
 				{
@@ -552,49 +417,31 @@ namespace UnrealBuildTool
 
 					CompileAction.ProducedItems.Add(ObjectFile);
 					Result.ObjectFiles.Add(ObjectFile);
-					FileArguments += string.Format(" -o \"{0}\"", ObjectFile.AbsolutePath);
+					AllArgs += string.Format(" -o \"{0}\"", ObjectFile.AbsolutePath);
 					OutputFilePath = ObjectFile.AbsolutePath;
 				}
 
 				// Add the source file path to the command-line.
-				FileArguments += string.Format(" \"{0}\"", SourceFile.AbsolutePath);
+				AllArgs += string.Format(" \"{0}\"", SourceFile.AbsolutePath);
 
 				// Generate the timing info
 				if (CompileEnvironment.bPrintTimingInfo)
 				{
 					FileItem TraceFile = FileItem.GetItemByFileReference(FileReference.Combine(OutputDir, Path.GetFileName(SourceFile.AbsolutePath) + ".json"));
-					FileArguments += " -ftime-trace";
+					AllArgs += " -ftime-trace";
 					CompileAction.ProducedItems.Add(TraceFile);
 				}
 
 				// Generate the included header dependency list
-				if(!bPreprocessDepends && CompileEnvironment.bGenerateDependenciesFile)
+				if (!bPreprocessDepends && CompileEnvironment.bGenerateDependenciesFile)
+
 				{
 					FileItem DependencyListFile = FileItem.GetItemByFileReference(FileReference.Combine(OutputDir, Path.GetFileName(SourceFile.AbsolutePath) + ".d"));
-					FileArguments += string.Format(" -MD -MF\"{0}\"", DependencyListFile.AbsolutePath.Replace('\\', '/'));
+					AllArgs += string.Format(" -MD -MF\"{0}\"", DependencyListFile.AbsolutePath.Replace('\\', '/'));
 					CompileAction.DependencyListFile = DependencyListFile;
 					CompileAction.ProducedItems.Add(DependencyListFile);
 				}
 
-				string EscapedAdditionalArgs = "";
-				if(!string.IsNullOrWhiteSpace(CompileEnvironment.AdditionalArguments))
-				{
-					foreach(string AdditionalArg in CompileEnvironment.AdditionalArguments.Split(new char[] { ' ' }, StringSplitOptions.RemoveEmptyEntries))
-					{
-						Match DefinitionMatch = Regex.Match(AdditionalArg, "-D\"?(?<Name>.*)=(?<Value>.*)\"?");
-						if (DefinitionMatch.Success)
-						{
-							EscapedAdditionalArgs += string.Format(" -D{0}=\"{1}\"", DefinitionMatch.Groups["Name"].Value, DefinitionMatch.Groups["Value"].Value);
-						}
-						else
-						{
-							EscapedAdditionalArgs += " " + AdditionalArg;
-						}
-					}
-				}
-
-				string AllArgs = Arguments + FileArguments + EscapedAdditionalArgs;
-
 				string CompilerPath = Settings.ToolchainDir + MacCompiler;
 				
 				// Analyze and then compile using the shell to perform the indirection
@@ -613,6 +460,7 @@ namespace UnrealBuildTool
 					string ArchPath = "/usr/bin/arch";
 					CompileAction.CommandPath = new FileReference(ArchPath);
 					CompileAction.CommandArguments = string.Format("-{0} {1} {2}", MacExports.HostArchitecture, CompilerPath, AllArgs);
+
 				}
 				else
 				{
@@ -642,40 +490,35 @@ namespace UnrealBuildTool
 					PrepassAction.CommandVersion = CompileAction.CommandVersion;
 					PrepassAction.WorkingDirectory = CompileAction.WorkingDirectory;
 
-					string PreprocessArguments = Arguments.ToString();
-					string PreprocessFileArguments = FileArguments;
-					PreprocessArguments = PreprocessArguments.Replace(" -c ", " ");
+					string PreprocessArguments = AllArgs.ToString().Replace(" -c ", " ");
 
 					FileItem DependencyListFile = FileItem.GetItemByFileReference(FileReference.Combine(OutputDir, Path.GetFileName(SourceFile.AbsolutePath) + ".d"));
-					PreprocessFileArguments += string.Format(" -M -MF\"{0}\"", DependencyListFile.AbsolutePath.Replace('\\', '/'));
+					PreprocessArguments += string.Format(" -M -MF\"{0}\"", DependencyListFile.AbsolutePath.Replace('\\', '/'));
 					PrepassAction.DependencyListFile = DependencyListFile;
 					PrepassAction.ProducedItems.Add(DependencyListFile);
 
-					PreprocessFileArguments = PreprocessFileArguments.Replace(" -ftime-trace", string.Empty);
-					PreprocessFileArguments = PreprocessFileArguments.Replace(string.Format(" -o \"{0}\"", CompileAction.ProducedItems.First().AbsolutePath), String.Empty);
+					PreprocessArguments = PreprocessArguments.Replace(" -ftime-trace", string.Empty);
+					PreprocessArguments = PreprocessArguments.Replace(string.Format(" -o \"{0}\"", CompileAction.ProducedItems.First().AbsolutePath), String.Empty);
 
 					PrepassAction.DeleteItems.AddRange(PrepassAction.ProducedItems);
 
-					string AllPreprocessArgs = PreprocessArguments + PreprocessFileArguments + EscapedAdditionalArgs;
-
 					PrepassAction.WorkingDirectory = GetMacDevSrcRoot();
 
 					if (MacExports.IsRunningUnderRosetta)
 					{
 						string ArchPath = "/usr/bin/arch";
 						PrepassAction.CommandPath = new FileReference(ArchPath);
-						PrepassAction.CommandArguments = string.Format("-{0} {1} {2}", MacExports.HostArchitecture, CompilerPath, AllPreprocessArgs);
+						PrepassAction.CommandArguments = string.Format("-{0} {1} {2}", MacExports.HostArchitecture, CompilerPath, PreprocessArguments);
 					}
 					else
 					{
 						PrepassAction.CommandPath = new FileReference(CompilerPath);
-						PrepassAction.CommandArguments = AllPreprocessArgs;
+						PrepassAction.CommandArguments = PreprocessArguments;
 					}
 
 					CompileAction.DependencyListFile = DependencyListFile;
 					CompileAction.PrerequisiteItems.Add(DependencyListFile);
 				}
-
 			}
 			return Result;
 		}
@@ -1302,22 +1145,22 @@ namespace UnrealBuildTool
 
 			string ArgumentString = "-c \"";
 			ArgumentString += string.Format("for i in {{1..30}}; ");
-				ArgumentString += string.Format("do if [ -f \\\"{0}\\\" ] ; ", MachOBinary.AbsolutePath);
-				ArgumentString += string.Format("then ");
-					ArgumentString += string.Format("break; ");
-				ArgumentString += string.Format("else ");
-					ArgumentString += string.Format("sleep 1; ");
-				ArgumentString += string.Format("fi; ");
+			ArgumentString += string.Format("do if [ -f \\\"{0}\\\" ] ; ", MachOBinary.AbsolutePath);
+			ArgumentString += string.Format("then ");
+			ArgumentString += string.Format("break; ");
+			ArgumentString += string.Format("else ");
+			ArgumentString += string.Format("sleep 1; ");
+			ArgumentString += string.Format("fi; ");
 			ArgumentString += string.Format("done; ");
 
 			ArgumentString += string.Format("if [ ! -f \\\"{1}\\\" ] || [ \\\"{0}\\\" -nt \\\"{1}\\\" ] ; ", MachOBinary.AbsolutePath, OutputFile.AbsolutePath);
 			ArgumentString += string.Format("then ");
-				ArgumentString += string.Format("rm -rf \\\"{0}\\\"; ", OutputFile.AbsolutePath);
-				ArgumentString += string.Format(" \\\"{0}\\\" {3} -f \\\"{1}\\\" -o \\\"{2}\\\"; ",
-					DsymutilPath,
-					MachOBinary.AbsolutePath,
-					OutputFile.AbsolutePath,
-					ExtraOptions);
+			ArgumentString += string.Format("rm -rf \\\"{0}\\\"; ", OutputFile.AbsolutePath);
+			ArgumentString += string.Format(" \\\"{0}\\\" {3} -f \\\"{1}\\\" -o \\\"{2}\\\"; ",
+				DsymutilPath,
+				MachOBinary.AbsolutePath,
+				OutputFile.AbsolutePath,
+				ExtraOptions);
 			ArgumentString += string.Format("fi; ");
 			ArgumentString += "\"";
 
diff --git a/Engine/Source/Programs/UnrealBuildTool/ToolChain/AppleToolChain.cs b/Engine/Source/Programs/UnrealBuildTool/ToolChain/AppleToolChain.cs
--- a/Engine/Source/Programs/UnrealBuildTool/ToolChain/AppleToolChain.cs
+++ b/Engine/Source/Programs/UnrealBuildTool/ToolChain/AppleToolChain.cs
@@ -231,6 +231,106 @@ namespace UnrealBuildTool
 			return Result;
 		}
 
+		protected virtual string GetObjCExceptionsFlag(CppCompileEnvironment CompileEnvironment)
+		{
+			return "";
+		}
+
+		protected string GetPCHCompileArgumentsForCompileEnvironment(CppCompileEnvironment CompileEnvironment)
+		{
+			StringBuilder PCHArguments = new StringBuilder();
+			if (CompileEnvironment.PrecompiledHeaderAction == PrecompiledHeaderAction.Include)
+			{
+				// Add the precompiled header file's path to the include path so GCC can find it.
+				// This needs to be before the other include paths to ensure GCC uses it instead of the source header file.
+				PCHArguments.Append(" -include \"");
+				PCHArguments.Append(CompileEnvironment.PrecompiledHeaderIncludeFilename!.FullName);
+				PCHArguments.Append("\"");
+			}
+
+			return PCHArguments.ToString();
+		}
+
+		protected virtual string GetCompileArgumentsForCompileEnvironment(CppCompileEnvironment CompileEnvironment)
+		{
+			StringBuilder Arguments = new StringBuilder();
+			Arguments.Append(GetCompileArguments_Global(CompileEnvironment));
+
+			// Add include paths to the argument list.
+			HashSet<DirectoryReference> AllIncludes = new HashSet<DirectoryReference>(CompileEnvironment.UserIncludePaths);
+			AllIncludes.UnionWith(CompileEnvironment.SystemIncludePaths);
+			foreach (DirectoryReference IncludePath in AllIncludes)
+			{
+				Arguments.Append(" -I\"");
+				Arguments.Append(IncludePath.FullName);
+				Arguments.Append("\"");
+			}
+
+			foreach (string Definition in CompileEnvironment.Definitions)
+			{
+				string DefinitionArgument = Definition.Contains("\"") ? Definition.Replace("\"", "\\\"") : Definition;
+				Arguments.Append(" -D\"");
+				Arguments.Append(DefinitionArgument);
+				Arguments.Append("\"");
+			}
+
+			return Arguments.ToString();
+		}
+
+		public virtual string GetCompileArgumentsForAbsoluteFilePathWithCompileEnvironment(string FilePath, CppCompileEnvironment CompileEnvironment)
+		{
+			string Arguments = GetCompileArgumentsForCompileEnvironment(CompileEnvironment);
+			string PCHArguments = GetPCHCompileArgumentsForCompileEnvironment(CompileEnvironment);
+
+			foreach (FileItem ForceIncludeFile in CompileEnvironment.ForceIncludeFiles)
+			{
+				Arguments += String.Format(" -include \"{0}\"", ForceIncludeFile.Location.FullName);
+			}
+
+			string FileArguments = "";
+			string Extension = Path.GetExtension(FilePath).ToUpperInvariant();
+
+			if (CompileEnvironment.PrecompiledHeaderAction == PrecompiledHeaderAction.Create)
+			{
+				// Compile the file as a C++ PCH.
+				FileArguments += GetCompileArguments_PCH(CompileEnvironment);
+				FileArguments += GetRTTIFlag(CompileEnvironment);
+				FileArguments += GetObjCExceptionsFlag(CompileEnvironment);
+			}
+			else if (Extension == ".C")
+			{
+				// Compile the file as C code.
+				FileArguments += GetCompileArguments_C();
+			}
+			else if (Extension == ".MM")
+			{
+				// Compile the file as Objective-C++ code.
+				FileArguments += GetCompileArguments_MM(CompileEnvironment);
+				FileArguments += GetRTTIFlag(CompileEnvironment);
+				FileArguments += GetObjCExceptionsFlag(CompileEnvironment);
+			}
+			else if (Extension == ".M")
+			{
+				// Compile the file as Objective-C++ code.
+				FileArguments += GetCompileArguments_M(CompileEnvironment);
+				FileArguments += GetObjCExceptionsFlag(CompileEnvironment);
+			}
+			else
+			{
+				// Compile the file as C++ code.
+				FileArguments += GetCompileArguments_CPP(CompileEnvironment);
+				FileArguments += GetRTTIFlag(CompileEnvironment);
+				FileArguments += GetObjCExceptionsFlag(CompileEnvironment);
+
+				// only use PCH for .cpp files
+				Arguments = PCHArguments + Arguments;
+			}
+
+			string AllArgs = Arguments + FileArguments;
+
+			return AllArgs;
+		}
+
 		protected string GetDsymutilPath(out string ExtraOptions, bool bIsForLTOBuild=false)
 		{
 			FileReference DsymutilLocation = new FileReference("/usr/bin/dsymutil");
@@ -299,5 +399,113 @@ namespace UnrealBuildTool
 			ExtraOptions = (bIsForLTOBuild && Major == 10 && Minor == 0 && Patch == 1) ? "-j 1" : "";
 			return DsymutilLocation.FullName;
 		}
+
+		protected virtual string GetCompileArguments_Global(CppCompileEnvironment CompileEnvironment)
+		{
+			string Result = "";
+
+			Result += " -fmessage-length=0";
+			Result += " -pipe";
+			Result += " -fpascal-strings";
+			Result += " -Wall -Werror";
+			Result += " -Wdelete-non-virtual-dtor";
+
+			// clang 12.00 has a new warning for copies in ranged loops. Instances have all been fixed up (2020/6/26) but
+			// are likely to be reintroduced due to no equivalent on other platforms at this time so disable the warning
+			// See also MacToolChain.cs
+			if (GetClangVersion().Major >= 12)
+			{
+				Result += " -Wno-range-loop-analysis";
+			}
+
+			if (CompileEnvironment.ShadowVariableWarningLevel != WarningLevel.Off)
+			{
+				Result += " -Wshadow" + ((CompileEnvironment.ShadowVariableWarningLevel == WarningLevel.Error) ? "" : " -Wno-error=shadow");
+			}
+
+			if (CompileEnvironment.bEnableUndefinedIdentifierWarnings)
+			{
+				Result += " -Wundef" + (CompileEnvironment.bUndefinedIdentifierWarningsAsErrors ? "" : " -Wno-error=undef");
+			}
+
+			Result += " -c";
+
+			if (!CompileEnvironment.bUseInlining)
+			{
+				Result += " -fno-inline-functions";
+			}
+
+			// Create DWARF format debug info if wanted,
+			if (CompileEnvironment.bCreateDebugInfo)
+			{
+				Result += " -gdwarf-2";
+			}
+
+			return Result;
+		}
+
+		static string GetCompileArguments_CPP(CppCompileEnvironment CompileEnvironment)
+		{
+			string Result = "";
+			Result += " -x objective-c++";
+			Result += GetCppStandardCompileArgument(CompileEnvironment);
+			Result += " -stdlib=libc++";
+
+			return Result;
+		}
+
+		static string GetCompileArguments_MM(CppCompileEnvironment CompileEnvironment)
+		{
+			string Result = "";
+			Result += " -x objective-c++";
+			Result += GetCppStandardCompileArgument(CompileEnvironment);
+			Result += " -stdlib=libc++";
+			return Result;
+		}
+
+		static string GetCompileArguments_M(CppCompileEnvironment CompileEnvironment)
+		{
+			string Result = "";
+			Result += " -x objective-c";
+			Result += " -stdlib=libc++";
+			return Result;
+		}
+
+		static string GetCompileArguments_C()
+		{
+			string Result = "";
+			Result += " -x c";
+			return Result;
+		}
+
+		static string GetCompileArguments_PCH(CppCompileEnvironment CompileEnvironment)
+		{
+			string Result = "";
+			Result += " -x objective-c++-header";
+			Result += GetCppStandardCompileArgument(CompileEnvironment);
+			Result += " -stdlib=libc++";
+
+			return Result;
+		}
+
+		// Conditionally enable (default disabled) generation of information about every class with virtual functions for use by the C++ runtime type identification features 
+		// (`dynamic_cast' and `typeid'). If you don't use those parts of the language, you can save some space by using -fno-rtti. 
+		// Note that exception handling uses the same information, but it will generate it as needed. 
+		protected static string GetRTTIFlag(CppCompileEnvironment CompileEnvironment)
+		{
+			string Result = "";
+
+			if (CompileEnvironment.bUseRTTI)
+			{
+				Result = " -frtti";
+			}
+			else
+			{
+				Result = " -fno-rtti";
+			}
+
+			return Result;
+		}
+
 	};
 }
