UnrealBuildTool: use colon prefix with -l linker flag to support linking libraries on Linux without "lib" prefix

This addresses an issue on Linux where libraries that do not have the "lib"
prefix can go missing at link time because the linker ends up looking for the
library with the prefix. For example, a library with the filename "foo.so"
would result in the "-lfoo" flag being added to the link command, which would
then cause the linker to look for a library with the name "libfoo" that does
not exist.

diff --git a/Engine/Source/Programs/UnrealBuildTool/Platform/Linux/LinuxToolChain.cs b/Engine/Source/Programs/UnrealBuildTool/Platform/Linux/LinuxToolChain.cs
--- a/Engine/Source/Programs/UnrealBuildTool/Platform/Linux/LinuxToolChain.cs
+++ b/Engine/Source/Programs/UnrealBuildTool/Platform/Linux/LinuxToolChain.cs
@@ -1767,6 +1767,13 @@ namespace UnrealBuildTool
 						// Remove lib prefix
 						LibName = LibName.Remove(0, 3);
 					}
+					else if (LibraryDependency.Exists)
+					{
+						// The library exists, but it is not prefixed with "lib", so force the
+						// linker to find it without that prefix by prepending a colon to
+						// the file name.
+						LibName = string.Format(":{0}", LibraryDependency.Name);
+					}
 					string LibLinkFlag = string.Format(" -l{0}", LibName);
 
 					if (LinkEnvironment.bIsBuildingDLL && LinkEnvironment.bIsCrossReferenced)
