Fixed client crash at FAndroidOpenGLFramePacer::InitSwappy

#localization none
#tests tested locally

diff --git a/Engine/Source/Programs/UnrealBuildTool/Platform/Android/AndroidToolChain.cs b/Engine/Source/Programs/UnrealBuildTool/Platform/Android/AndroidToolChain.cs
--- a/Engine/Source/Programs/UnrealBuildTool/Platform/Android/AndroidToolChain.cs
+++ b/Engine/Source/Programs/UnrealBuildTool/Platform/Android/AndroidToolChain.cs
@@ -951,7 +951,7 @@ namespace UnrealBuildTool
 			Result += " -Wl,-soname,libUnreal.so";
 
 			// hide all symbols by default
-			Result += " -Wl,--exclude-libs,ALL";
+			//Result += " -Wl,--exclude-libs,ALL";
 
 			Result += " -Wl,--build-id=sha1";               // add build-id to make debugging easier
 
