Changing UBT VNI warning to be a log rather than a warning

diff --git a/Engine/Source/Programs/UnrealBuildTool/Configuration/UEBuildTarget.cs b/Engine/Source/Programs/UnrealBuildTool/Configuration/UEBuildTarget.cs
--- a/Engine/Source/Programs/UnrealBuildTool/Configuration/UEBuildTarget.cs
+++ b/Engine/Source/Programs/UnrealBuildTool/Configuration/UEBuildTarget.cs
@@ -4332,7 +4332,7 @@ namespace UnrealBuildTool
 				// Module must not have Verse if Verse is not enabled
 				if (Module.bHasVerse && !Rules.bUseVerse && !Rules.bBuildAllModules && !ProjectFileGenerator.bGenerateProjectFiles)
 				{
-					Log.TraceWarning("Module '{0}' has associated Verse code but target '{1}' does not have Verse enabled. C++ include errors are likely to follow.", Module.Name, TargetName);
+					Log.TraceInformation("Module '{0}' has associated Verse code but target '{1}' does not have Verse enabled. C++ include errors are likely to follow.", Module.Name, TargetName);
 				}
 			}
 
