PropagateAlpha  fix for Android

diff --git a/Engine/Build/Android/Java/src/com/epicgames/unreal/GameActivity.java.template b/Engine/Build/Android/Java/src/com/epicgames/unreal/GameActivity.java.template
--- a/Engine/Build/Android/Java/src/com/epicgames/unreal/GameActivity.java.template
+++ b/Engine/Build/Android/Java/src/com/epicgames/unreal/GameActivity.java.template
@@ -2660,9 +2660,9 @@ public class GameActivity extends $${gameActivitySuperClass}$$ implements Surfac
 				Log.debug( "Did not find DepthBufferPreference, using default.");
 			}
 
-			if (bundle.containsKey("com.epicgames.ue4.GameActivity.PropagateAlpha"))
+			if (bundle.containsKey("com.epicgames.unreal.GameActivity.PropagateAlpha"))
 			{
-				PropagateAlpha = bundle.getInt("com.epicgames.ue4.GameActivity.PropagateAlpha");
+				PropagateAlpha = bundle.getInt("com.epicgames.unreal.GameActivity.PropagateAlpha");
 				Log.debug( "Found PropagateAlpha = " + PropagateAlpha);
 			}
 			else
diff --git a/Engine/Source/Programs/UnrealBuildTool/Platform/Android/UEDeployAndroid.cs b/Engine/Source/Programs/UnrealBuildTool/Platform/Android/UEDeployAndroid.cs
--- a/Engine/Source/Programs/UnrealBuildTool/Platform/Android/UEDeployAndroid.cs
+++ b/Engine/Source/Programs/UnrealBuildTool/Platform/Android/UEDeployAndroid.cs
@@ -2493,6 +2493,9 @@ namespace UnrealBuildTool
 			bool bSupportsVulkan = false;
 			Ini.GetBool("/Script/AndroidRuntimeSettings.AndroidRuntimeSettings", "bSupportsVulkan", out bSupportsVulkan);
 
+			int PropagateAlpha = 0;
+			Ini.GetInt32("/Script/Engine.RendererSettings", "r.Mobile.PropagateAlpha", out PropagateAlpha);
+
 			bool bAllowIMU = true;
 			Ini.GetBool("/Script/AndroidRuntimeSettings.AndroidRuntimeSettings", "bAllowIMU", out bAllowIMU);
 			if (IsPackagingForDaydream(Ini) && bAllowIMU)
@@ -2800,6 +2803,7 @@ namespace UnrealBuildTool
 			Text.AppendLine(string.Format("\t\t<meta-data android:name=\"com.epicgames.unreal.GameActivity.bUseDisplayCutout\" android:value=\"{0}\"/>", bUseDisplayCutout ? "true" : "false"));
 			Text.AppendLine(string.Format("\t\t<meta-data android:name=\"com.epicgames.unreal.GameActivity.bAllowIMU\" android:value=\"{0}\"/>", bAllowIMU ? "true" : "false"));
 			Text.AppendLine(string.Format("\t\t<meta-data android:name=\"com.epicgames.unreal.GameActivity.bSupportsVulkan\" android:value=\"{0}\"/>", bSupportsVulkan ? "true" : "false"));
+			Text.AppendLine(string.Format("\t\t<meta-data android:name=\"com.epicgames.unreal.GameActivity.PropagateAlpha\" android:value=\"{0}\"/>", PropagateAlpha));
 			Text.AppendLine(string.Format("\t\t<meta-data android:name=\"com.epicgames.unreal.GameActivity.StartupPermissions\" android:value=\"{0}\"/>", StartupPermissions));
 			if (bPackageForDaydream)
 			{
