UnrealBuildTool: New option -IgnoreInvalidFiles

When files are specified via -File= -SingleFile= or -FileList=, any file not valid for the target will be ignored.
If no individual files are found for the target, nothing will be built.

#rnx

diff --git a/Engine/Source/Programs/UnrealBuildTool/Configuration/BuildConfiguration.cs b/Engine/Source/Programs/UnrealBuildTool/Configuration/BuildConfiguration.cs
--- a/Engine/Source/Programs/UnrealBuildTool/Configuration/BuildConfiguration.cs
+++ b/Engine/Source/Programs/UnrealBuildTool/Configuration/BuildConfiguration.cs
@@ -159,5 +159,14 @@ namespace UnrealBuildTool
 		/// </summary>
 		[XmlConfigFile(Category = "WindowsPlatform")]
 		public int MaxNestedPathLength = 200;
+
+		/// <summary>
+		/// When single file targets are specified, via -File=, -SingleFile=, or -FileList=
+		/// If this option is set, no error will be produced if the source file is not included in the target.
+		/// Additionally, if any file or file list is specified, the target will not be built if none of the specified files are part of that target,
+		/// including the case where a file specified via -FileList= is empty.
+		/// </summary>
+		[CommandLine("-IgnoreInvalidFiles")]
+		public bool bIgnoreInvalidFiles;
 	}
 }
diff --git a/Engine/Source/Programs/UnrealBuildTool/Configuration/UEBuildTarget.cs b/Engine/Source/Programs/UnrealBuildTool/Configuration/UEBuildTarget.cs
--- a/Engine/Source/Programs/UnrealBuildTool/Configuration/UEBuildTarget.cs
+++ b/Engine/Source/Programs/UnrealBuildTool/Configuration/UEBuildTarget.cs
@@ -1904,18 +1904,32 @@ namespace UnrealBuildTool
 
 			// Find the set of binaries to build. If we're compiling only specific files, filter the list of binaries to only include the files we're interested in.
 			List<UEBuildBinary> BuildBinaries = Binaries;
-			foreach (FileReference SpecificFile in TargetDescriptor.SpecificFilesToCompile)
+
 			{
-				UEBuildBinary? Binary = Binaries.Find(x => x.Modules.Any(y => y.ContainsFile(SpecificFile)));
+				// If we've been asked to ignore invalid files, don't build anything if there are no valid individual files to build
+				bool bBuildShouldContinue = !BuildConfiguration.bIgnoreInvalidFiles;
 
-				if (Binary == null)
+				foreach (FileReference SpecificFile in TargetDescriptor.SpecificFilesToCompile)
 				{
-					throw new BuildException("Couldn't find any module containing {0} in {1}.", SpecificFile, TargetName);
+					UEBuildBinary? Binary = Binaries.Find(x => x.Modules.Any(y => y.ContainsFile(SpecificFile)));
+
+					if (Binary == null)
+					{
+						if (!BuildConfiguration.bIgnoreInvalidFiles)
+						{
+							throw new BuildException("Couldn't find any module containing {0} in {1}.", SpecificFile, TargetName);
+						}
+					}
+					else if (!BuildBinaries.Contains(Binary))
+					{
+						bBuildShouldContinue = true;
+						BuildBinaries.Add(Binary);
+					}
 				}
 
-				if (!BuildBinaries.Contains(Binary))
+				if (!bBuildShouldContinue)
 				{
-					BuildBinaries.Add(Binary);
+					return Makefile;
 				}
 			}
 
