Added UPL registerBuildSettings node to list plugin ini settings that affect packaging
Added registerBuildSettings node to AndroidFileServer UPL

#android

diff --git a/Engine/Plugins/Runtime/AndroidFileServer/Source/AndroidFileServer/AndroidFileServer_UPL.xml b/Engine/Plugins/Runtime/AndroidFileServer/Source/AndroidFileServer/AndroidFileServer_UPL.xml
--- a/Engine/Plugins/Runtime/AndroidFileServer/Source/AndroidFileServer/AndroidFileServer_UPL.xml
+++ b/Engine/Plugins/Runtime/AndroidFileServer/Source/AndroidFileServer/AndroidFileServer_UPL.xml
@@ -49,6 +49,12 @@
 		</if>
 	</init>
 
+	<!-- register the ini keys that modify the packaging -->
+	<registerBuildSettings>
+		<insertValue value="/Script/AndroidFileServerEditor.AndroidFileServerRuntimeSettings[bEnablePlugin,bAllowNetworkConnection,SecurityToken,bIncludeInShipping,bAllowExternalStartInShipping]"/>
+		<insertNewline/>
+	</registerBuildSettings>
+
 	<androidManifestUpdates>
 		<if condition="bEnabled">
 			<true>
diff --git a/Engine/Source/Programs/UnrealBuildTool/Platform/Android/UEDeployAndroid.cs b/Engine/Source/Programs/UnrealBuildTool/Platform/Android/UEDeployAndroid.cs
--- a/Engine/Source/Programs/UnrealBuildTool/Platform/Android/UEDeployAndroid.cs
+++ b/Engine/Source/Programs/UnrealBuildTool/Platform/Android/UEDeployAndroid.cs
@@ -1888,7 +1888,7 @@ namespace UnrealBuildTool
 			return true;
 		}
 
-		private string GetAllBuildSettings(AndroidToolChain ToolChain, bool bForDistribution, bool bMakeSeparateApks, bool bPackageDataInsideApk, bool bDisableVerifyOBBOnStartUp, bool bUseExternalFilesDir, string TemplatesHashCode)
+		private string GetAllBuildSettings(AndroidToolChain ToolChain, UnrealPluginLanguage UPL, bool bForDistribution, bool bMakeSeparateApks, bool bPackageDataInsideApk, bool bDisableVerifyOBBOnStartUp, bool bUseExternalFilesDir, string TemplatesHashCode)
 		{
 			// make the settings string - this will be char by char compared against last time
 			StringBuilder CurrentSettings = new StringBuilder();
@@ -1980,6 +1980,52 @@ namespace UnrealBuildTool
 				}
 			}
 
+			// get a list of the ini settings in UPL files that may affect the build
+			// architecture doesn't matter here since this node does not use init logic
+			string UPLBuildSettings = UPL.ProcessPluginNode("arm64-v8a", "registerBuildSettings", "");
+			foreach (string Line in UPLBuildSettings.Split('\n', StringSplitOptions.RemoveEmptyEntries))
+			{
+				string SectionName = Line.Trim();
+
+				// needed keys are provided in [ ] separated by commas
+				string[]? NeededKeys = null;
+				int KeyIndex = SectionName.IndexOf('[');
+				if (KeyIndex > 0)
+				{
+					string KeyList = SectionName.Substring(KeyIndex + 1);
+					SectionName = SectionName.Substring(0, KeyIndex);
+					int CloseIndex = KeyList.IndexOf("]");
+					if (CloseIndex > 1)
+					{
+						NeededKeys = KeyList.Substring(0, CloseIndex).Split(',', StringSplitOptions.RemoveEmptyEntries);
+						if (NeededKeys.Length == 0)
+						{
+							NeededKeys = null;
+						}
+					}
+				}
+
+				// write the values for the requested keys (or all if none specified)
+				Section = Ini.FindSection(SectionName);
+				if (Section != null)
+				{
+					foreach (string Key in Section.KeyNames)
+					{
+						if (NeededKeys != null && !NeededKeys.Contains(Key))
+						{
+							continue;
+						}
+						if (Section.TryGetValues(Key, out IReadOnlyList<string>? Values))
+						{
+							foreach (string Value in Values)
+							{
+								CurrentSettings.AppendLine(string.Format("{0}:{1}={2}", SectionName, Key, Value));
+							}
+						}
+					}
+				}
+			}
+
 			return CurrentSettings.ToString();
 		}
 
@@ -3914,7 +3960,7 @@ namespace UnrealBuildTool
 
 			// check to see if any "meta information" is newer than last time we build
 			string TemplatesHashCode = GenerateTemplatesHashCode(EngineDirectory);
-			string CurrentBuildSettings = GetAllBuildSettings(ToolChain, bForDistribution, bMakeSeparateApks, bPackageDataInsideApk, bDisableVerifyOBBOnStartUp, bUseExternalFilesDir, TemplatesHashCode);
+			string CurrentBuildSettings = GetAllBuildSettings(ToolChain, UPL!, bForDistribution, bMakeSeparateApks, bPackageDataInsideApk, bDisableVerifyOBBOnStartUp, bUseExternalFilesDir, TemplatesHashCode);
 			string BuildSettingsCacheFile = Path.Combine(IntermediateAndroidPath, "UEBuildSettings.txt");
 
 			// Architecture remapping
