Protect against additional configChanges which could terminate an Android application

#android

diff --git a/Engine/Source/Programs/UnrealBuildTool/Platform/Android/UEDeployAndroid.cs b/Engine/Source/Programs/UnrealBuildTool/Platform/Android/UEDeployAndroid.cs
--- a/Engine/Source/Programs/UnrealBuildTool/Platform/Android/UEDeployAndroid.cs
+++ b/Engine/Source/Programs/UnrealBuildTool/Platform/Android/UEDeployAndroid.cs
@@ -2750,8 +2750,8 @@ namespace UnrealBuildTool
 				Text.AppendLine("\t\t<activity android:name=\"com.epicgames.unreal.GameActivity\"");
 				Text.AppendLine("\t\t          android:label=\"@string/app_name\"");
 				Text.AppendLine("\t\t          android:theme=\"@style/UnrealSplashTheme\"");
-				Text.AppendLine(bAddDensity ? "\t\t          android:configChanges=\"mcc|mnc|uiMode|density|screenSize|smallestScreenSize|screenLayout|orientation|keyboardHidden|keyboard\""
-											: "\t\t          android:configChanges=\"mcc|mnc|uiMode|screenSize|smallestScreenSize|screenLayout|orientation|keyboardHidden|keyboard\"");
+				Text.AppendLine(bAddDensity ? "\t\t          android:configChanges=\"mcc|mnc|uiMode|density|screenSize|smallestScreenSize|screenLayout|orientation|keyboardHidden|keyboard|navigation|touchscreen|locale|fontScale|layoutDirection\""
+											: "\t\t          android:configChanges=\"mcc|mnc|uiMode|screenSize|smallestScreenSize|screenLayout|orientation|keyboardHidden|keyboard|navigation|touchscreen|locale|fontScale|layoutDirection\"");
 			}
 			else
 			{
@@ -2759,8 +2759,8 @@ namespace UnrealBuildTool
 				Text.AppendLine("\t\t          android:exported=\"true\"");
 				Text.AppendLine("\t\t          android:label=\"@string/app_name\"");
 				Text.AppendLine("\t\t          android:theme=\"@android:style/Theme.Black.NoTitleBar.Fullscreen\"");
-				Text.AppendLine(bAddDensity ? "\t\t          android:configChanges=\"mcc|mnc|uiMode|density|screenSize|smallestScreenSize|screenLayout|orientation|keyboardHidden|keyboard\""
-											: "\t\t          android:configChanges=\"mcc|mnc|uiMode|screenSize|smallestScreenSize|screenLayout|orientation|keyboardHidden|keyboard\"");
+				Text.AppendLine(bAddDensity ? "\t\t          android:configChanges=\"mcc|mnc|uiMode|density|screenSize|smallestScreenSize|screenLayout|orientation|keyboardHidden|keyboard|navigation|touchscreen|locale|fontScale|layoutDirection\""
+											: "\t\t          android:configChanges=\"mcc|mnc|uiMode|screenSize|smallestScreenSize|screenLayout|orientation|keyboardHidden|keyboard|navigation|touchscreen|locale|fontScale|layoutDirection\"");
 
 			}
 			if (SDKLevelInt >= 24)
