UnrealBuildTool: Add /fastfail to cl.exe & link.exe to enable Windows Error Reporting to save crash dumps. Microsoft noted this won't have any adverse effect when no faults occur.

#rnx

diff --git a/Engine/Source/Programs/UnrealBuildTool/Platform/Windows/VCToolChain.cs b/Engine/Source/Programs/UnrealBuildTool/Platform/Windows/VCToolChain.cs
--- a/Engine/Source/Programs/UnrealBuildTool/Platform/Windows/VCToolChain.cs
+++ b/Engine/Source/Programs/UnrealBuildTool/Platform/Windows/VCToolChain.cs
@@ -439,6 +439,11 @@ namespace UnrealBuildTool
 				}
 			}
 
+			if (Target.WindowsPlatform.Compiler.IsMSVC())
+			{
+				Arguments.Add("/fastfail");
+			}
+
 			// Address sanitizer
 			if (Target.WindowsPlatform.bEnableAddressSanitizer)
 			{
@@ -1028,6 +1033,11 @@ namespace UnrealBuildTool
 				}
 			}
 
+			if (Target.WindowsPlatform.Compiler.IsMSVC())
+			{
+				Arguments.Add("/fastfail");
+			}
+
 			// for monolithic editor builds, add the PDBPAGESIZE option, (VS 16.11, VC toolchain 14.29.30133), but the pdb will be too large without this
 			// some monolithic game builds could be too large as well, but they can be added in a .Target.cs with:
 			//   			WindowsPlatform.AdditionalLinkerOptions = "/PDBPAGESIZE:8192";
