UnrealBuildTool: Add -SaveCrashDumps option that will attempt to find and copy crash dump files to a specified location.

#rnx

diff --git a/Engine/Source/Programs/UnrealBuildTool/Modes/BuildMode.cs b/Engine/Source/Programs/UnrealBuildTool/Modes/BuildMode.cs
--- a/Engine/Source/Programs/UnrealBuildTool/Modes/BuildMode.cs
+++ b/Engine/Source/Programs/UnrealBuildTool/Modes/BuildMode.cs
@@ -87,6 +87,12 @@ namespace UnrealBuildTool
 		[CommandLine("-WriteOutdatedActions=")]
 		public FileReference? WriteOutdatedActionsFile = null;
 
+		/// <summary>
+		/// An optional directory to copy crash dump files into
+		/// </summary>
+		[CommandLine("-SaveCrashDumps=")]
+		public DirectoryReference? SaveCrashDumpDirectory = null;
+
 		/// <summary>
 		/// Main entry point
 		/// </summary>
@@ -236,6 +242,9 @@ namespace UnrealBuildTool
 			}
 			finally
 			{
+				// Check if anything failed during our run, and act accordingly.
+				ProcessCoreDumps(SaveCrashDumpDirectory);
+
 				// Save all the caches
 				SourceFileMetadataCache.SaveAll();
 				CppDependencyCache.SaveAll();
@@ -854,6 +863,107 @@ namespace UnrealBuildTool
 			}
 			return new List<LinkedAction>(OutputItemToProducingAction.Values);
 		}
+
+		void ProcessCoreDumps(DirectoryReference? SaveCrashDumpDirectory)
+		{
+			if (SaveCrashDumpDirectory == null)
+			{
+				return;
+			}
+
+			if (!RuntimePlatform.IsWindows)
+			{
+				return;
+			}
+
+			// set to true to have this code create some files in expected locations to report 
+			bool bTestDumpDetection = false;
+
+			using OpenTracing.IScope Scope = GlobalTracer.Instance.BuildSpan("Processing crash dump files").StartActive();
+
+			DateTime UBTStartTime = Process.GetCurrentProcess().StartTime;
+
+			List<FileReference> FoundCrashDumps = new List<FileReference>();
+
+			// examine the contents of %LOCALAPPDATA%\CrashDumps
+			string? LocalAppData = Environment.GetEnvironmentVariable("LOCALAPPDATA");
+			if (LocalAppData != null)
+			{
+				DirectoryReference CrashDumpsDirectory = DirectoryReference.Combine(DirectoryReference.FromString(LocalAppData)!, "CrashDumps");
+				if (DirectoryReference.Exists(CrashDumpsDirectory))
+				{
+					if (bTestDumpDetection)
+					{
+						System.IO.File.Create(System.IO.Path.Combine(CrashDumpsDirectory.FullName, Guid.NewGuid().ToString() + ".dmp")).Close();
+					}
+
+					foreach (FileReference CrashDump in DirectoryReference.EnumerateFiles(CrashDumpsDirectory))
+					{
+						if (FileReference.GetLastWriteTime(CrashDump) > UBTStartTime)
+						{
+							Log.TraceWarning($"Crash dump {CrashDump} was created duing UnrealBuildTool execution");
+							FoundCrashDumps.Add(CrashDump);
+						}
+					}
+				}
+			}
+
+			// examine the contents of %TMP% (on CI agents, this should not be unreasonably slow as the tmp dir gets cleaned between builds
+			string? TMP = Environment.GetEnvironmentVariable("TMP");
+			if (TMP != null)
+			{
+				DirectoryReference TmpDir = DirectoryReference.FromString(TMP)!;
+
+				if (bTestDumpDetection)
+				{
+					System.IO.File.Create(System.IO.Path.Combine(TmpDir.FullName, Guid.NewGuid().ToString() + ".dmp")); //.Close(); Intentionally not closed, to trigger catch() block below
+				}
+
+				List<DirectoryReference> AccessibleTmpDirectories = new List<DirectoryReference>();
+
+				AccessibleTmpDirectories.Add(TmpDir);
+
+				for (int I = 0; I < AccessibleTmpDirectories.Count; ++I)
+				{
+					try
+					{
+						// Manual recursion to avoid the case where an inaccessible file prevents us from iterating reachable parts of the dir
+						AccessibleTmpDirectories.AddRange(DirectoryReference.EnumerateDirectories(AccessibleTmpDirectories[I]));
+
+						foreach (FileReference TmpFile in DirectoryReference.EnumerateFiles(AccessibleTmpDirectories[I]))
+						{
+							if (TmpFile.HasExtension(".dmp") && FileReference.GetLastWriteTime(TmpFile) > UBTStartTime)
+							{
+								Log.TraceWarning($"Crash dump {TmpFile} was created duing UnrealBuildTool execution");
+								FoundCrashDumps.Add(TmpFile);
+							}
+						}
+					}
+					catch
+					{
+						// silently ignore inaccessible directories
+					}
+				}
+			}
+
+			if (FoundCrashDumps.Count > 0)
+			{
+				DirectoryReference.CreateDirectory(SaveCrashDumpDirectory);
+				foreach (FileReference CrashDump in FoundCrashDumps)
+				{
+					Log.TraceInformation($"Copying {CrashDump} to {SaveCrashDumpDirectory}");
+					try
+					{
+						FileReference.Copy(CrashDump, FileReference.Combine(SaveCrashDumpDirectory, CrashDump.GetFileName()));
+					}
+					catch(Exception Ex)
+					{
+						// don't stop if there was a problem copying one of the files
+						Log.TraceWarning($"Failed to copy crash dump {CrashDump}: {ExceptionUtils.FormatException(Ex)}");
+					}
+				}
+			}
+		}
 	}
 }
 
