Hide all exportable symbols on Android, except the ones called from Java

diff --git a/Engine/Source/Programs/UnrealBuildTool/Platform/Android/AndroidToolChain.cs b/Engine/Source/Programs/UnrealBuildTool/Platform/Android/AndroidToolChain.cs
--- a/Engine/Source/Programs/UnrealBuildTool/Platform/Android/AndroidToolChain.cs
+++ b/Engine/Source/Programs/UnrealBuildTool/Platform/Android/AndroidToolChain.cs
@@ -431,6 +431,11 @@ namespace UnrealBuildTool
 			return Ini.GetBool("/Script/AndroidRuntimeSettings.AndroidRuntimeSettings", "bEnableAdvancedBinaryCompression", out bEnableAdvancedBinaryCompression) && bEnableAdvancedBinaryCompression;
 		}
 
+		private string GetVersionScriptFilename(LinkEnvironment LinkEnvironment)
+		{
+			return Path.Combine(LinkEnvironment.IntermediateDirectory!.FullName, "ExportSymbols.ldscript");
+		}
+
 		public override void SetUpGlobalEnvironment(ReadOnlyTargetRules Target)
 		{
 			base.SetUpGlobalEnvironment(Target);
@@ -950,8 +955,16 @@ namespace UnrealBuildTool
 			// make sure the DT_SONAME field is set properly (or we can a warning toast at startup on new Android)
 			Result += " -Wl,-soname,libUnreal.so";
 
-			// hide all symbols by default
-			//Result += " -Wl,--exclude-libs,ALL";
+			// exclude exported symbols from all static libs
+			Result += " -Wl,--exclude-libs,ALL";
+
+			string VersionScriptFile = GetVersionScriptFilename(LinkEnvironment);
+			using (StreamWriter Writer = File.CreateText(VersionScriptFile))
+			{
+				// Make all symbols (except ones called from Java) hidden
+				Writer.WriteLine("{ global: Java_*; ANativeActivity_onCreate; JNI_OnLoad; local: *; };");
+			}
+			Result += " -Wl,--version-script=" + VersionScriptFile;
 
 			Result += " -Wl,--build-id=sha1";               // add build-id to make debugging easier
 
@@ -2070,6 +2083,9 @@ namespace UnrealBuildTool
 				// Only execute linking on the local PC.
 				LinkAction.bCanExecuteRemotely = false;
 
+				string VersionScriptFileItem = GetVersionScriptFilename(LinkEnvironment);
+				LinkAction.PrerequisiteItems.Add(FileItem.GetItemByPath(VersionScriptFileItem));
+				
 				if(bExecuteCompilerThroughShell)
 				{
 					SetupActionToExecuteCompilerThroughShell(ref LinkAction, LinkAction.CommandPath.FullName, LinkAction.CommandArguments, "Link");
