UBT: Tweak how single-file compiles work.

- When a file list is specified, any headers identified via a scan are added to a separate list of files to be compiled. These are always considered optional; since we don't know for sure whether they're part of the target, they are excluded from checks that require a matching binary to be identified for them.
- The -IgnoreInvalidFiles argument doesn't seem to have been working as intended, since the BuildBinaries will already contain all binaries for the target despite claims made in a comment (so bBuildShouldContinue will default to, and remain set to, false).
- While we could fix the logic to filter the binaries built for single-file compiles, it seems simpler to keep the same codepath as regular builds.

diff --git a/Engine/Source/Programs/UnrealBuildTool/Configuration/TargetDescriptor.cs b/Engine/Source/Programs/UnrealBuildTool/Configuration/TargetDescriptor.cs
--- a/Engine/Source/Programs/UnrealBuildTool/Configuration/TargetDescriptor.cs
+++ b/Engine/Source/Programs/UnrealBuildTool/Configuration/TargetDescriptor.cs
@@ -56,6 +56,12 @@ namespace UnrealBuildTool
 		[CommandLine("-SingleFile=")]
 		public List<FileReference> SpecificFilesToCompile = new List<FileReference>();
 
+		/// <summary>
+		/// Individual files to compile which may or may not be part of the target. This file set is built from a header scan, and may include files
+		/// not part of the target by design.
+		/// </summary>
+		public List<FileReference> OptionalFilesToCompile = new List<FileReference>();
+
 		/// <summary>
 		/// Whether to perform hot reload for this target
 		/// </summary>
diff --git a/Engine/Source/Programs/UnrealBuildTool/Configuration/UEBuildTarget.cs b/Engine/Source/Programs/UnrealBuildTool/Configuration/UEBuildTarget.cs
--- a/Engine/Source/Programs/UnrealBuildTool/Configuration/UEBuildTarget.cs
+++ b/Engine/Source/Programs/UnrealBuildTool/Configuration/UEBuildTarget.cs
@@ -837,8 +837,7 @@ namespace UnrealBuildTool
 					Log.TraceLog("  {0}", NewFile);
 				}
 
-				Descriptor.SpecificFilesToCompile.AddRange(NewFiles.Except(Descriptor.SpecificFilesToCompile));
-
+				Descriptor.OptionalFilesToCompile.AddRange(NewFiles.Except(Descriptor.SpecificFilesToCompile));
 			}
 
 			// If we're compiling a plugin, and this target is monolithic, just create the object files
@@ -1903,41 +1902,23 @@ namespace UnrealBuildTool
 			}
 
 			// Find the set of binaries to build. If we're compiling only specific files, filter the list of binaries to only include the files we're interested in.
-			List<UEBuildBinary> BuildBinaries = Binaries;
-
+			if (TargetDescriptor.SpecificFilesToCompile.Count > 0 && !BuildConfiguration.bIgnoreInvalidFiles)
 			{
-				// If we've been asked to ignore invalid files, don't build anything if there are no valid individual files to build
-				bool bBuildShouldContinue = !BuildConfiguration.bIgnoreInvalidFiles;
-
 				foreach (FileReference SpecificFile in TargetDescriptor.SpecificFilesToCompile)
 				{
 					UEBuildBinary? Binary = Binaries.Find(x => x.Modules.Any(y => y.ContainsFile(SpecificFile)));
-
 					if (Binary == null)
 					{
-						if (!BuildConfiguration.bIgnoreInvalidFiles)
-						{
-							throw new BuildException("Couldn't find any module containing {0} in {1}.", SpecificFile, TargetName);
-						}
-					}
-					else if (!BuildBinaries.Contains(Binary))
-					{
-						bBuildShouldContinue = true;
-						BuildBinaries.Add(Binary);
+						throw new BuildException("Couldn't find any module containing {0} in {1}.", SpecificFile, TargetName);
 					}
 				}
-
-				if (!bBuildShouldContinue)
-				{
-					return Makefile;
-				}
 			}
 
 			// Build the target's binaries.
 			DirectoryReference ExeDir = GetExecutableDir();
 			using (GlobalTracer.Instance.BuildSpan("UEBuildBinary.Build()").StartActive())
 			{
-				foreach (UEBuildBinary Binary in BuildBinaries)
+				foreach (UEBuildBinary Binary in Binaries)
 				{
 					List<FileItem> BinaryOutputItems = Binary.Build(Rules, TargetToolChain, GlobalCompileEnvironment, GlobalLinkEnvironment, TargetDescriptor.SpecificFilesToCompile, WorkingSet, ExeDir, Makefile);
 					Makefile.OutputItems.AddRange(BinaryOutputItems);
diff --git a/Engine/Source/Programs/UnrealBuildTool/Modes/BuildMode.cs b/Engine/Source/Programs/UnrealBuildTool/Modes/BuildMode.cs
--- a/Engine/Source/Programs/UnrealBuildTool/Modes/BuildMode.cs
+++ b/Engine/Source/Programs/UnrealBuildTool/Modes/BuildMode.cs
@@ -485,12 +485,13 @@ namespace UnrealBuildTool
 					}
 					else
 					{
-						foreach (FileReference SpecificFile in TargetDescriptor.SpecificFilesToCompile)
+						HashSet<FileReference> ForceBuildFiles = new HashSet<FileReference>();
+						ForceBuildFiles.UnionWith(TargetDescriptor.SpecificFilesToCompile);
+						ForceBuildFiles.UnionWith(TargetDescriptor.OptionalFilesToCompile);
+
+						foreach (LinkedAction PrerequisiteAction in PrerequisiteActions.Where(x => x.PrerequisiteItems.Any(y => ForceBuildFiles.Contains(y.Location))))
 						{
-							foreach (LinkedAction PrerequisiteAction in PrerequisiteActions.Where(x => x.PrerequisiteItems.Any(y => y.Location == SpecificFile)))
-							{
-								ActionToOutdatedFlag[PrerequisiteAction] = true;
-							}
+							ActionToOutdatedFlag[PrerequisiteAction] = true;
 						}
 					}
 				}
@@ -810,7 +811,7 @@ namespace UnrealBuildTool
 			if(TargetDescriptor.SpecificFilesToCompile.Count > 0)
 			{
 				// If we're just compiling a specific files, set the target items to be all the derived items
-				List<FileItem> FilesToCompile = TargetDescriptor.SpecificFilesToCompile.ConvertAll(x => FileItem.GetItemByFileReference(x));
+				List<FileItem> FilesToCompile = TargetDescriptor.SpecificFilesToCompile.Union(TargetDescriptor.OptionalFilesToCompile).Select(x => FileItem.GetItemByFileReference(x)).ToList();
 				OutputItems.UnionWith(
 					Makefile.Actions.Where(x => x.PrerequisiteItems.Any(y => FilesToCompile.Contains(y)))
 					.SelectMany(x => x.ProducedItems));
