UnrealBuildTool: Bump preferred Windows clang verison to 13.
Please note there are still some outstanding runtime issues when building with clang 13 and this is intended to help address them.

diff --git a/Engine/Source/Programs/UnrealBuildTool/Platform/Windows/UEBuildWindows.cs b/Engine/Source/Programs/UnrealBuildTool/Platform/Windows/UEBuildWindows.cs
--- a/Engine/Source/Programs/UnrealBuildTool/Platform/Windows/UEBuildWindows.cs
+++ b/Engine/Source/Programs/UnrealBuildTool/Platform/Windows/UEBuildWindows.cs
@@ -832,6 +832,8 @@ namespace UnrealBuildTool
 		/// </summary>
 		static readonly VersionNumberRange[] PreferredClangVersions =
 		{
+			VersionNumberRange.Parse("13.0.0", "13.999"),
+			VersionNumberRange.Parse("12.0.0", "12.999"),
 			VersionNumberRange.Parse("11.0.0", "11.999"),
 			VersionNumberRange.Parse("10.0.0", "10.999")
 		};
