UnrealBuildTool: Don't pass /fastfail to link.exe if linking with PGO options

#rnx

diff --git a/Engine/Source/Programs/UnrealBuildTool/Platform/Windows/VCToolChain.cs b/Engine/Source/Programs/UnrealBuildTool/Platform/Windows/VCToolChain.cs
--- a/Engine/Source/Programs/UnrealBuildTool/Platform/Windows/VCToolChain.cs
+++ b/Engine/Source/Programs/UnrealBuildTool/Platform/Windows/VCToolChain.cs
@@ -1033,7 +1033,7 @@ namespace UnrealBuildTool
 				}
 			}
 
-			if (Target.WindowsPlatform.Compiler.IsMSVC())
+			if (Target.WindowsPlatform.Compiler.IsMSVC() && !LinkEnvironment.bPGOOptimize && !LinkEnvironment.bPGOProfile)
 			{
 				Arguments.Add("/fastfail");
 			}
