UnrealBuildTool: Fix engine modules being filtered out of generated projects when programs should not be added

#rnx

diff --git a/Engine/Source/Programs/UnrealBuildTool/ProjectFiles/ProjectFileGenerator.cs b/Engine/Source/Programs/UnrealBuildTool/ProjectFiles/ProjectFileGenerator.cs
--- a/Engine/Source/Programs/UnrealBuildTool/ProjectFiles/ProjectFileGenerator.cs
+++ b/Engine/Source/Programs/UnrealBuildTool/ProjectFiles/ProjectFileGenerator.cs
@@ -294,7 +294,7 @@ namespace UnrealBuildTool
 		/// True if we should include program projects in the generated solution.
 		/// </summary>
 		[XmlConfigFile]
-		protected bool IncludeEnginePrograms = true;
+		protected bool bIncludeEnginePrograms = true;
 
 		/// <summary>
 		/// Whether to include C++ targets
@@ -940,7 +940,7 @@ namespace UnrealBuildTool
 						if (bIncludeConfigFiles)
 						{
 							AddEngineConfigFiles(EngineProject);
-							if (IncludeEnginePrograms)
+							if (bIncludeEnginePrograms)
 							{
 								AddUnrealHeaderToolConfigFiles(EngineProject);
 								AddUBTConfigFilesToEngineProject(EngineProject);
@@ -1038,10 +1038,14 @@ namespace UnrealBuildTool
 						RootFolder.AddSubFolder(DebugProjectFile.Item2).ChildProjects.Add(DebugProjectFile.Item1);
 					}
 
-
 					foreach (KeyValuePair<FileReference, ProjectFile> CurProgramProject in ProgramProjects)
 					{
+						FileReference UnrealProjectFile = CurProgramProject.Value.ProjectTargets.First().UnrealProjectFilePath!;
 						Project Target = CurProgramProject.Value.ProjectTargets.FirstOrDefault(t => !String.IsNullOrEmpty(t.TargetRules!.SolutionDirectory));
+						if (!bIncludeEnginePrograms && UnrealProjectFile.IsUnderDirectory(Unreal.EngineDirectory))
+						{
+							continue;
+						}
 
 						if (Target != null)
 						{
@@ -1285,7 +1289,7 @@ namespace UnrealBuildTool
 							case "-GAME":
 								// Generates project files for a single game
 								bIncludeDotNetPrograms = false;
-								IncludeEnginePrograms = false;
+								bIncludeEnginePrograms = false;
 								bGeneratingGameProjectFiles = true;
 								break;
 
@@ -1368,13 +1372,13 @@ namespace UnrealBuildTool
 
 				bool bInstalledEngineWithSource = Unreal.IsEngineInstalled() && DirectoryReference.Exists(UnrealBuildTool.EngineSourceDirectory);
 
-				bIncludeEngineSource = bAlwaysIncludeEngineModules || bInstalledEngineWithSource;
+				bIncludeEngineSource = !Unreal.IsEngineInstalled() || bInstalledEngineWithSource;
 				bIncludeDocumentation = false;
 				bIncludeBuildSystemFiles = false;
 				bIncludeShaderSource = true;
 				bIncludeTemplateFiles = false;
 				bIncludeConfigFiles = true;
-				IncludeEnginePrograms = bAlwaysIncludeEngineModules;
+				bIncludeEnginePrograms = bAlwaysIncludeEngineModules;
 				bIncludeDotNetPrograms = bIncludeDotNetPrograms || bAlwaysIncludeDotNetPrograms;
 			}
 			else
@@ -1556,7 +1560,7 @@ namespace UnrealBuildTool
 			List<string> UnsupportedPlatformNameStrings = Utils.MakeListOfUnsupportedPlatforms(SupportedPlatforms, bIncludeUnbuildablePlatforms: true);
 
 			// Locate all targets (*.Target.cs files)
-			List<FileReference> FoundTargetFiles = Rules.FindAllRulesSourceFiles(Rules.RulesFileType.Target, AllGameProjects.Select(x => x.Directory).ToList(), ForeignPlugins: null, AdditionalSearchPaths: null, bIncludeEngine: IncludeEnginePrograms, bIncludeTempTargets: bIncludeTempTargets);
+			List<FileReference> FoundTargetFiles = Rules.FindAllRulesSourceFiles(Rules.RulesFileType.Target, AllGameProjects.Select(x => x.Directory).ToList(), ForeignPlugins: null, AdditionalSearchPaths: null, bIncludeEngine: bIncludeEngineSource, bIncludeTempTargets: bIncludeTempTargets);
 			foreach (FileReference CurTargetFile in FoundTargetFiles)
 			{
 				string CleanTargetFileName = Utils.CleanDirectorySeparators(CurTargetFile.FullName);
@@ -2362,7 +2366,7 @@ namespace UnrealBuildTool
 
 					if (Unreal.GetExtensionDirs(Unreal.EngineDirectory, "Source/Programs").Any(x => TargetFilePath.IsUnderDirectory(x)))
 					{
-						WantProjectFileForTarget = IncludeEnginePrograms;
+						WantProjectFileForTarget = bIncludeEnginePrograms;
 					}
 					else if (Unreal.GetExtensionDirs(Unreal.EngineDirectory, "Source").Any(x => TargetFilePath.IsUnderDirectory(x)))
 					{
diff --git a/Engine/Source/Programs/UnrealBuildTool/ProjectFiles/Rider/RiderProjectFileGenerator.cs b/Engine/Source/Programs/UnrealBuildTool/ProjectFiles/Rider/RiderProjectFileGenerator.cs
--- a/Engine/Source/Programs/UnrealBuildTool/ProjectFiles/Rider/RiderProjectFileGenerator.cs
+++ b/Engine/Source/Programs/UnrealBuildTool/ProjectFiles/Rider/RiderProjectFileGenerator.cs
@@ -81,7 +81,7 @@ namespace UnrealBuildTool
 			bGenerateProjectFiles = true;
 			bIncludeTemplateFiles = false;
 			bIncludeTestAndShippingConfigs = true;
-			IncludeEnginePrograms = true; // It's true by default, but I like to have things explicit
+			bIncludeEnginePrograms = true; // It's true by default, but I like to have things explicit
 		}
 
 		protected override ProjectFile AllocateProjectFile(FileReference InitFilePath, DirectoryReference BaseDir)
@@ -196,7 +196,7 @@ namespace UnrealBuildTool
 
 					if (TargetFilePath.IsUnderDirectory(EngineSourceProgramsDirectory))
 					{
-						WantProjectFileForTarget = IncludeEnginePrograms;
+						WantProjectFileForTarget = bIncludeEnginePrograms;
 					}
 					else if (TargetFilePath.IsUnderDirectory(UnrealBuildTool.EngineSourceDirectory))
 					{
diff --git a/Engine/Source/Programs/UnrealBuildTool/ProjectFiles/Xcode/XcodeProjectFileGenerator.cs b/Engine/Source/Programs/UnrealBuildTool/ProjectFiles/Xcode/XcodeProjectFileGenerator.cs
--- a/Engine/Source/Programs/UnrealBuildTool/ProjectFiles/Xcode/XcodeProjectFileGenerator.cs
+++ b/Engine/Source/Programs/UnrealBuildTool/ProjectFiles/Xcode/XcodeProjectFileGenerator.cs
@@ -305,7 +305,7 @@ namespace UnrealBuildTool
 				{
 					// an Engine target is required in order to be able to get Xcode to sign blueprint projects
 					// always include the engine target for installed builds.
-					IncludeEnginePrograms = true;
+					bIncludeEnginePrograms = true;
 				}
 				bIncludeEngineSource = true;
 			}
