Add XGE config bUnavailableIfInUse to allow falling back to a different executor if XGE is in use.
Add XGE config MinActions to avoid using XGE for small groups of tasks (e.g. single file compile).

diff --git a/Engine/Source/Programs/UnrealBuildTool/Executors/XGE.cs b/Engine/Source/Programs/UnrealBuildTool/Executors/XGE.cs
--- a/Engine/Source/Programs/UnrealBuildTool/Executors/XGE.cs
+++ b/Engine/Source/Programs/UnrealBuildTool/Executors/XGE.cs
@@ -64,6 +64,19 @@ namespace UnrealBuildTool
 		[XmlConfigFile(Category = "XGE")]
 		static bool bUseVCCompilerMode = false;
 
+		/// <summary>
+		/// Minimum number of actions to use XGE execution.
+		/// </summary>
+		[XmlConfigFile(Category = "XGE")]
+		static public int MinActions = 2;
+
+		/// <summary>
+		/// Check for a concurrent XGE build and treat the XGE executor as unavailable if it's in use.
+		/// This will allow UBT to fall back to another executor such as the parallel executor. 
+		/// </summary>
+		[XmlConfigFile(Category = "XGE")]
+		static bool bUnavailableIfInUse = false;
+
 		private const string ProgressMarkupPrefix = "@action";
 
 		public XGE()
@@ -315,6 +328,26 @@ namespace UnrealBuildTool
 				}
 			}
 
+			// Check if there's an XGE build already running 
+			if (bUnavailableIfInUse)
+			{
+				Process XGEProcess = new Process()
+				{
+					StartInfo = new ProcessStartInfo(
+					XgConsoleExe,
+					"/Command=Unused /nowait /silent")  // The actual command here doesn't matter - it will fail with a different error code (1) than "in use" (4)
+					{
+						UseShellExecute = false
+					}
+				};
+				if (Utils.RunLocalProcess(XGEProcess) == 4)
+				{
+					Log.TraceWarning("Unable to use Incredibuild executor because a build is already in progress");
+					return false;
+
+				}
+			}
+
 			return true;
 		}
 
diff --git a/Engine/Source/Programs/UnrealBuildTool/System/ActionGraph.cs b/Engine/Source/Programs/UnrealBuildTool/System/ActionGraph.cs
--- a/Engine/Source/Programs/UnrealBuildTool/System/ActionGraph.cs
+++ b/Engine/Source/Programs/UnrealBuildTool/System/ActionGraph.cs
@@ -344,7 +344,7 @@ namespace UnrealBuildTool
 				{
 					Executor = new HybridExecutor(TargetDescriptors, BuildConfiguration.MaxParallelActions);
 				}
-				else if (BuildConfiguration.bAllowXGE && XGE.IsAvailable())
+				else if (BuildConfiguration.bAllowXGE && XGE.IsAvailable() && ActionsToExecute.Count >= XGE.MinActions)
 				{
 					Executor = new XGE();
 				}
