UnrealBuildTool: Add -PackagePath to set destination folder of input files used at link time (normally used to debug a linker crash for platforms that support it)

#rnx

diff --git a/Engine/Source/Programs/UnrealBuildTool/Configuration/TargetRules.cs b/Engine/Source/Programs/UnrealBuildTool/Configuration/TargetRules.cs
--- a/Engine/Source/Programs/UnrealBuildTool/Configuration/TargetRules.cs
+++ b/Engine/Source/Programs/UnrealBuildTool/Configuration/TargetRules.cs
@@ -1330,6 +1330,14 @@ namespace UnrealBuildTool
 		[XmlConfigFile(Category = "BuildConfiguration")]
 		public bool bAllowRuntimeSymbolFiles = true;
 
+		/// <summary>
+		/// Package full path (directory + filename) where to store input files used at link time 
+		/// Normally used to debug a linker crash for platforms that support it
+		/// </summary>
+		[CommandLine("-PackagePath")]
+		[XmlConfigFile(Category = "BuildConfiguration")]
+		public string? PackagePath = null;
+
 		/// <summary>
 		/// Bundle version for Mac apps.
 		/// </summary>
@@ -2788,6 +2796,11 @@ namespace UnrealBuildTool
 			get { return Inner.bAllowRuntimeSymbolFiles; }
 		}
 
+		public string? PackagePath
+		{
+			get { return Inner.PackagePath; }
+		}
+
 		public string? BundleVersion
 		{
 			get { return Inner.BundleVersion; }
diff --git a/Engine/Source/Programs/UnrealBuildTool/Configuration/UEBuildTarget.cs b/Engine/Source/Programs/UnrealBuildTool/Configuration/UEBuildTarget.cs
--- a/Engine/Source/Programs/UnrealBuildTool/Configuration/UEBuildTarget.cs
+++ b/Engine/Source/Programs/UnrealBuildTool/Configuration/UEBuildTarget.cs
@@ -3712,6 +3712,7 @@ namespace UnrealBuildTool
 			GlobalLinkEnvironment.bCreateMapFile = Rules.bCreateMapFile;
 			GlobalLinkEnvironment.bHasExports = Rules.bHasExports;
 			GlobalLinkEnvironment.bUsePDBFiles = Rules.bUsePDBFiles;
+			GlobalLinkEnvironment.PackagePath = Rules.PackagePath;
 			GlobalLinkEnvironment.BundleDirectory = BuildPlatform.GetBundleDirectory(Rules, Binaries[0].OutputFilePaths);
 			GlobalLinkEnvironment.BundleVersion = Rules.BundleVersion;
 			GlobalLinkEnvironment.bAllowLTCG = Rules.bAllowLTCG;
diff --git a/Engine/Source/Programs/UnrealBuildTool/System/LinkEnvironment.cs b/Engine/Source/Programs/UnrealBuildTool/System/LinkEnvironment.cs
--- a/Engine/Source/Programs/UnrealBuildTool/System/LinkEnvironment.cs
+++ b/Engine/Source/Programs/UnrealBuildTool/System/LinkEnvironment.cs
@@ -266,6 +266,12 @@ namespace UnrealBuildTool
 		/// </summary>
 		public bool bPrintTimingInfo;
 
+		/// <summary>
+		/// Package full path (directory + filename) where to store input files used at link time 
+		/// Normally used to debug a linker crash for platforms that support it
+		/// </summary>
+		public string? PackagePath;
+
 		/// <summary>
 		/// Bundle version for Mac apps
 		/// </summary>
@@ -369,6 +375,7 @@ namespace UnrealBuildTool
 			bUseFastPDBLinking = Other.bUseFastPDBLinking;
 			bIgnoreUnresolvedSymbols = Other.bIgnoreUnresolvedSymbols;
 			bPrintTimingInfo = Other.bPrintTimingInfo;
+			PackagePath = Other.PackagePath;
 			BundleVersion = Other.BundleVersion;
 			InstallName = Other.InstallName;
 			InputFiles.AddRange(Other.InputFiles);
