- Moved ini struct/map parsing out of Turnkey/WinPlatform to ConfigHierarchy so they can be shared, and used elsewhere

diff --git a/Engine/Source/Programs/AutomationTool/Turnkey/Commands/CreateBuild.cs b/Engine/Source/Programs/AutomationTool/Turnkey/Commands/CreateBuild.cs
--- a/Engine/Source/Programs/AutomationTool/Turnkey/Commands/CreateBuild.cs
+++ b/Engine/Source/Programs/AutomationTool/Turnkey/Commands/CreateBuild.cs
@@ -58,7 +58,7 @@ namespace Turnkey.Commands
 		{
 			ConfigHierarchy Config = ConfigCache.ReadHierarchy(ConfigHierarchyType.Game, ProjectFile.Directory, BuildHostPlatform.Current.Platform);
 
-			string StagingDirectory = ExecuteBuild.GetStructEntryForSetting(Config, SettingsSection, "StagingDirectory", "Path");
+			string StagingDirectory = Config.GetStructEntryForSetting(SettingsSection, "StagingDirectory", "Path");
 			if (string.IsNullOrEmpty(StagingDirectory))
 			{
 				StagingDirectory = Path.Combine(ProjectFile.Directory.FullName, "Saved", "Archives");
@@ -99,7 +99,7 @@ namespace Turnkey.Commands
 		private static string GetPlatformSetting(ConfigHierarchy Config, UnrealTargetPlatform Platform, string PerPlatformKey, string GenericKey, string Default)
 		{
 			string IniPlatformName = ConfigHierarchy.GetIniPlatformName(Platform);
-			string Value = ExecuteBuild.GetMapValueForSetting(Config, SettingsSection, PerPlatformKey, Platform.ToString());
+			string Value = Config.GetMapValueForSetting(SettingsSection, PerPlatformKey, Platform.ToString());
 			if (Value == null && !string.IsNullOrEmpty(GenericKey))
 			{
 				Config.GetString(SettingsSection, GenericKey, out Value);
@@ -331,7 +331,7 @@ namespace Turnkey.Commands
 				Config.GetBool(SettingsSection, "bBuildHttpChunkInstallData", out bBuildHttpChunkInstallData);
 				if (bBuildHttpChunkInstallData)
 				{
-					string HttpChunkInstallDataDirectory = ExecuteBuild.GetStructEntryForSetting(Config, SettingsSection, "HttpChunkInstallDataDirectory", "Path");
+					string HttpChunkInstallDataDirectory = Config.GetStructEntryForSetting(SettingsSection, "HttpChunkInstallDataDirectory", "Path");
 					string HttpChunkInstallDataVersion;
 					Config.GetString(SettingsSection, "HttpChunkInstallDataVersion", out HttpChunkInstallDataVersion);
 
@@ -705,7 +705,7 @@ namespace Turnkey.Commands
 							// we need a device here, if canceled, quit out
 							return;
 						}
-						PlatformParams += " -device=" + string.Join("+", Devices);
+						PlatformParams += " -device=" + string.Join("+", Devices.Select(x => x.Id));
 					}
 				}
 
@@ -726,8 +726,8 @@ namespace Turnkey.Commands
 						TurnkeyUtils.Log($"Here are the existing builds your project already has (and their descrptions):");
 						foreach (string Build in ExistingBuilds)
 						{
-							string BuildName = ExecuteBuild.GetStructEntry(Build, "Name", false);
-							string BuildHelp = ExecuteBuild.GetStructEntry(Build, "HelpText", false);
+							string BuildName = ConfigHierarchy.GetStructEntry(Build, "Name", false);
+							string BuildHelp = ConfigHierarchy.GetStructEntry(Build, "HelpText", false);
 							TurnkeyUtils.Log($"  {BuildName} - {BuildHelp}");
 						}
 						Name = TurnkeyUtils.ReadInput("Enter unique name for this build, it must be different than the above names", "");
diff --git a/Engine/Source/Programs/AutomationTool/Turnkey/Commands/ExecuteBuild.cs b/Engine/Source/Programs/AutomationTool/Turnkey/Commands/ExecuteBuild.cs
--- a/Engine/Source/Programs/AutomationTool/Turnkey/Commands/ExecuteBuild.cs
+++ b/Engine/Source/Programs/AutomationTool/Turnkey/Commands/ExecuteBuild.cs
@@ -15,84 +15,6 @@ namespace Turnkey.Commands
 	{
 		protected override CommandGroup Group => CommandGroup.Builds;
 
-		// could move this to TurnkeyUtils!
-		public static string GetStructEntryForSetting(ConfigHierarchy Config, string Section, string Setting, string Property)
-		{
-			string ConfigEntry;
-			if (Config.GetString(Section, Setting, out ConfigEntry))
-			{
-				return GetStructEntry(ConfigEntry, Property, false);
-			}
-			return null;
-		}
-
-		public static string GetStructEntry(string Input, string Property, bool bIsArrayProperty)
-		{
-			string PrimaryRegex;
-			string AltRegex = null;
-			if (bIsArrayProperty)
-			{
-				PrimaryRegex = string.Format("{0}\\s*=\\s*\\((.*?)\\)", Property);
-			}
-			else
-			{
-				// handle quoted strings, allowing for escaped quotation marks (basically doing " followed by whatever, until we see a quote that was not proceeded by a \, and gather the whole mess in an outer group)
-				PrimaryRegex = string.Format("{0}\\s*=\\s*\"((.*?)[^\\\\])\"", Property);
-				// when no quotes, we skip over whitespace, and we end when we see whitespace, a comma or a ). This will handle (Ip = 192.168.0.1 , Name=....) , and return only '192.168.0.1'
-				AltRegex = string.Format("{0}\\s*=\\s*(.*?)[\\s,\\)]", Property);
-			}
-
-			// attempt to match it!
-			Match Result = Regex.Match(Input, PrimaryRegex);
-			if (!Result.Success && AltRegex != null)
-			{
-				Result = Regex.Match(Input, AltRegex);
-			}
-
-			// if we got a success, return the main match value
-			if (Result.Success)
-			{
-				return Result.Groups[1].Value.ToString();
-			}
-
-			return null;
-		}
-
-		public static string GetMapValueForSetting(ConfigHierarchy Config, string Section, string Setting, string Key)
-		{
-			string ConfigEntry;
-			if (Config.GetString(Section, Setting, out ConfigEntry))
-			{
-				return GetMapValue(ConfigEntry, Key);
-			}
-			return null;
-		}
-
-		// Key cannot have escaped quotes or commas
-		public static string GetMapValue(string Input, string Key)
-		{
-			string PrimaryRegex;
-			string AltRegex = null;
-			// handle quoted strings, allowing for escaped quotation marks (and possibly the key in quotes as well)
-			PrimaryRegex = string.Format("{0}\"?\\s*,\\s*\"((.*?)[^\\\\])\"", Key);
-			AltRegex = string.Format("{0}\"?\\s*,\\s*(.*?)[\\s,\\)]", Key);
-
-			// attempt to match it!
-			Match Result = Regex.Match(Input, PrimaryRegex);
-			if (!Result.Success && AltRegex != null)
-			{
-				Result = Regex.Match(Input, AltRegex);
-			}
-
-			// if we got a success, return the main match value
-			if (Result.Success)
-			{
-				return Result.Groups[1].Value.ToString();
-			}
-
-			return null;
-		}
-
 		protected override void Execute(string[] CommandOptions)
 		{
 			// we need a platform to execute
@@ -135,10 +57,10 @@ namespace Turnkey.Commands
 				{
 					foreach (string Build in Builds)
 					{
-						string Name = GetStructEntry(Build, "Name", false);
-						string Help = GetStructEntry(Build, "HelpText", false);
-						string SpecificPlatforms = GetStructEntry(Build, "SpecificPlatforms", true);
-						string Params = GetStructEntry(Build, "BuildCookRunParams", false);
+						string Name = ConfigHierarchy.GetStructEntry(Build, "Name", false);
+						string Help = ConfigHierarchy.GetStructEntry(Build, "HelpText", false);
+						string SpecificPlatforms = ConfigHierarchy.GetStructEntry(Build, "SpecificPlatforms", true);
+						string Params = ConfigHierarchy.GetStructEntry(Build, "BuildCookRunParams", false);
 
 						// make sure required entries are there
 						if (Name == null || Params == null)
diff --git a/Engine/Source/Programs/AutomationTool/Win/WinPlatform.Automation.cs b/Engine/Source/Programs/AutomationTool/Win/WinPlatform.Automation.cs
--- a/Engine/Source/Programs/AutomationTool/Win/WinPlatform.Automation.cs
+++ b/Engine/Source/Programs/AutomationTool/Win/WinPlatform.Automation.cs
@@ -77,9 +77,9 @@ public static class SteamDeckSupport
 			// Expected ini format: +SteamDeckDevice=(IpAddr=10.1.33.19,Name=MySteamDeck,UserName=deck)
 			foreach (string DeckDevice in SteamDeckDevices)
 			{
-				string IpAddr     = GetStructEntry(DeckDevice, "IpAddr", false);
-				string DeviceName = GetStructEntry(DeckDevice, "Name", false);
-				string UserName   = GetStructEntry(DeckDevice, "UserName", false);
+				string IpAddr     = ConfigHierarchy.GetStructEntry(DeckDevice, "IpAddr", false);
+				string DeviceName = ConfigHierarchy.GetStructEntry(DeckDevice, "Name", false);
+				string UserName   = ConfigHierarchy.GetStructEntry(DeckDevice, "UserName", false);
 
 				// Name is optional, if its empty/not found lets just use the IpAddr for the Name
 				if (string.IsNullOrEmpty(DeviceName))
@@ -102,39 +102,6 @@ public static class SteamDeckSupport
 		return Devices;
 	}
 
-	// GetStructEntry copied from ExecuteBuild.cs TODO move to a better to share this
-	private static string GetStructEntry(string Input, string Property, bool bIsArrayProperty)
-	{
-		string PrimaryRegex;
-		string AltRegex = null;
-		if (bIsArrayProperty)
-		{
-			PrimaryRegex = string.Format("{0}\\s*=\\s*\\((.*?)\\)", Property);
-		}
-		else
-		{
-			// handle quoted strings, allowing for escaped quotation marks (basically doing " followed by whatever, until we see a quote that was not proceeded by a \, and gather the whole mess in an outer group)
-			PrimaryRegex = string.Format("{0}\\s*=\\s*\"((.*?)[^\\\\])\"", Property);
-			// when no quotes, we skip over whitespace, and we end when we see whitespace, a comma or a ). This will handle (Ip = 192.168.0.1 , Name=....) , and return only '192.168.0.1'
-			AltRegex = string.Format("{0}\\s*=\\s*(.*?)[\\s,\\)]", Property);
-		}
-
-		// attempt to match it!
-		Match Result = Regex.Match(Input, PrimaryRegex);
-		if (!Result.Success && AltRegex != null)
-		{
-			Result = Regex.Match(Input, AltRegex);
-		}
-
-		// if we got a success, return the main match value
-		if (Result.Success)
-		{
-			return Result.Groups[1].Value.ToString();
-		}
-
-		return null;
-	}
-
 	/**
 	 * Get the ipaddr and username of the device from the params and/or device registry
 	 */
diff --git a/Engine/Source/Programs/UnrealBuildTool/System/ConfigHierarchy.cs b/Engine/Source/Programs/UnrealBuildTool/System/ConfigHierarchy.cs
--- a/Engine/Source/Programs/UnrealBuildTool/System/ConfigHierarchy.cs
+++ b/Engine/Source/Programs/UnrealBuildTool/System/ConfigHierarchy.cs
@@ -8,6 +8,7 @@ using System.Threading.Tasks;
 using EpicGames.Core;
 using System.IO;
 using System.Diagnostics.CodeAnalysis;
+using System.Text.RegularExpressions;
 using UnrealBuildBase;
 
 namespace UnrealBuildTool
@@ -1196,5 +1197,117 @@ namespace UnrealBuildTool
 				return TargetPlatform.ToString();
 			}
 		}
+
+
+		#region Unreal struct/map parsing helpers
+
+		/// <summary>
+		/// Gets an ini setting, and then pulls the value for a property out of a struct, in the format:
+		///   [SomeSection]
+		///   SomeStruct=(Foo=Bar,Prop="My Value")
+		/// </summary>
+		/// <param name="Section">Ini section ('SomeSection' in this example)</param>
+		/// <param name="Setting">Name of the struct setting ('SomeStruct' in this example)</param>
+		/// <param name="Property">Name of the property inside the struct ('Prop' in this example)</param>
+		/// <returns>The value retrieved from the struct ('My Value' in this example), or null if anything was not found</returns>
+		public string? GetStructEntryForSetting(string Section, string Setting, string Property)
+		{
+			string ConfigEntry;
+			if (GetString(Section, Setting, out ConfigEntry))
+			{
+				return GetStructEntry(ConfigEntry, Property, false);
+			}
+			return null;
+		}
+
+		/// <summary>
+		/// Pulls the value for a property out of a struct in the given input, in the format:
+		///   (Foo=Bar,Prop="My Value")
+		/// </summary>
+		/// <param name="Input">The entire struct as retrieved from the ini via GetString()</param>
+		/// <param name="Property">Name of the property inside the struct ('Prop' in this example)</param>
+		/// <param name="bIsArrayProperty">Pass true when the value pulled is an array, like (Foo=(X=1,Y=2)), this would return X=1,Y=2</param>
+		/// <returns>The value retrieved from the struct ('My Value' in this example), or null if anything was not found</returns>
+		public static string? GetStructEntry(string Input, string Property, bool bIsArrayProperty)
+		{
+			string PrimaryRegex;
+			string? AltRegex = null;
+			if (bIsArrayProperty)
+			{
+				PrimaryRegex = string.Format("{0}\\s*=\\s*\\((.*?)\\)", Property);
+			}
+			else
+			{
+				// handle quoted strings, allowing for escaped quotation marks (basically doing " followed by whatever, until we see a quote that was not proceeded by a \, and gather the whole mess in an outer group)
+				PrimaryRegex = string.Format("{0}\\s*=\\s*\"((.*?)[^\\\\])\"", Property);
+				// when no quotes, we skip over whitespace, and we end when we see whitespace, a comma or a ). This will handle (Ip = 192.168.0.1 , Name=....) , and return only '192.168.0.1'
+				AltRegex = string.Format("{0}\\s*=\\s*(.*?)[\\s,\\)]", Property);
+			}
+
+			// attempt to match it!
+			Match Result = Regex.Match(Input, PrimaryRegex);
+			if (!Result.Success && AltRegex != null)
+			{
+				Result = Regex.Match(Input, AltRegex);
+			}
+
+			// if we got a success, return the main match value
+			if (Result.Success)
+			{
+				return Result.Groups[1].Value.ToString();
+			}
+
+			return null;
+		}
+
+		/// <summary>
+		/// Gets an ini setting, and then pulls the value for a property out of a map, in the format:
+		///   [SomeSection]
+		///   SomeMap=((Foo=Bar),(SomeKey="My Value"))
+		/// </summary>
+		/// <param name="Section">Ini section ('SomeSection' in this example)</param>
+		/// <param name="Setting">Name of the struct setting ('SomeMap' in this example)</param>
+		/// <param name="Key">Name of the key inside the struct ('SomeKey' in this example)</param>
+		/// <returns>The value retrieved from the map ('My Value' in this example), or null if anything was not found</returns>
+		public string? GetMapValueForSetting(string Section, string Setting, string Key)
+		{
+			string ConfigEntry;
+			if (GetString(Section, Setting, out ConfigEntry))
+			{
+				return GetMapValue(ConfigEntry, Key);
+			}
+			return null;
+		}
+
+		/// <summary>
+		/// Pulls the value for a property out of a struct in the given input, in the format:
+		///   ((Foo=Bar),(SomeKey="My Value"))
+		/// </summary>
+		/// <param name="Input">The entire struct as retrieved from the ini via GetString()</param>
+		/// <param name="Key">Name of the key inside the struct ('SomeKey' in this example). Key cannot have escaped quotes or commas</param>
+		/// <returns>The value retrieved from the map ('My Value' in this example), or null if anything was not found</returns>
+		public static string? GetMapValue(string Input, string Key)
+		{
+			// handle quoted strings, allowing for escaped quotation marks (and possibly the key in quotes as well)
+			string PrimaryRegex = string.Format("{0}\"?\\s*,\\s*\"((.*?)[^\\\\])\"", Key);
+			string AltRegex = string.Format("{0}\"?\\s*,\\s*(.*?)[\\s,\\)]", Key);
+
+			// attempt to match it!
+			Match Result = Regex.Match(Input, PrimaryRegex);
+			if (!Result.Success && AltRegex != null)
+			{
+				Result = Regex.Match(Input, AltRegex);
+			}
+
+			// if we got a success, return the main match value
+			if (Result.Success)
+			{
+				return Result.Groups[1].Value.ToString();
+			}
+
+			return null;
+		}
+
+		#endregion
 	}
 }
