Remove bForceCompressNativeLibs Android setting as it's a dublicate of bExtractNativeLibs

diff --git a/Engine/Source/Programs/UnrealBuildTool/Platform/Android/UEDeployAndroid.cs b/Engine/Source/Programs/UnrealBuildTool/Platform/Android/UEDeployAndroid.cs
--- a/Engine/Source/Programs/UnrealBuildTool/Platform/Android/UEDeployAndroid.cs
+++ b/Engine/Source/Programs/UnrealBuildTool/Platform/Android/UEDeployAndroid.cs
@@ -2698,11 +2698,8 @@ namespace UnrealBuildTool
 			Text.AppendLine("\t<application android:label=\"@string/app_name\"");
 			Text.AppendLine("\t             android:icon=\"@drawable/icon\"");
 
-			bool bForceCompressNativeLibs = false;
-			Ini.GetBool("/Script/AndroidRuntimeSettings.AndroidRuntimeSettings", "bForceCompressNativeLibs", out bForceCompressNativeLibs);
-
 			AndroidToolChain.ClangSanitizer Sanitizer = ToolChain.BuildWithSanitizer();
-			if ((Sanitizer != AndroidToolChain.ClangSanitizer.None && Sanitizer != AndroidToolChain.ClangSanitizer.HwAddress) || bForceCompressNativeLibs)
+			if ((Sanitizer != AndroidToolChain.ClangSanitizer.None && Sanitizer != AndroidToolChain.ClangSanitizer.HwAddress))
 			{
 				bExtractNativeLibs = true;
 			}
