Update AutoSDK to NDK r21e

diff --git a/Engine/Source/Programs/UnrealBuildTool/Platform/Android/AndroidPlatformSDK.cs b/Engine/Source/Programs/UnrealBuildTool/Platform/Android/AndroidPlatformSDK.cs
--- a/Engine/Source/Programs/UnrealBuildTool/Platform/Android/AndroidPlatformSDK.cs
+++ b/Engine/Source/Programs/UnrealBuildTool/Platform/Android/AndroidPlatformSDK.cs
@@ -13,11 +13,11 @@ namespace UnrealBuildTool
 	{
 		public override string GetMainVersion()
 		{
-			return "r21b";
+			return "r21e";
 		}
 		public override string GetAutoSDKDirectoryForMainVersion()
 		{
-			return "-23";
+			return "-24";
 		}
 
 		public override void GetValidVersionRange(out string MinVersion, out string MaxVersion)
