Generated code for UHT and VNI now stored in separate subfolders of the generated code directory
- so they can be independently wiped
- fix for build errors caused by stale files leftover from moved/renamed UBT module

#rnx

diff --git a/Engine/Source/Programs/UnrealBuildTool/Configuration/UEBuildModuleCPP.cs b/Engine/Source/Programs/UnrealBuildTool/Configuration/UEBuildModuleCPP.cs
--- a/Engine/Source/Programs/UnrealBuildTool/Configuration/UEBuildModuleCPP.cs
+++ b/Engine/Source/Programs/UnrealBuildTool/Configuration/UEBuildModuleCPP.cs
@@ -34,15 +34,36 @@ namespace UnrealBuildTool
 			public readonly List<FileItem> ISPCFiles = new List<FileItem>();
 		}
 
+		/// <summary>
+		/// If UHT found any associated UObjects in this module's source files
+		/// </summary>
+		public bool bHasUObjects;
+
 		/// <summary>
 		/// The directory for this module's generated code
 		/// </summary>
 		public readonly DirectoryReference? GeneratedCodeDirectory;
 
 		/// <summary>
-		/// Set for modules that have generated code
+		/// The directory for this module's generated UHT code
+		/// </summary>
+		public DirectoryReference? GeneratedCodeDirectoryUHT
+		{
+			get { return GeneratedCodeDirectory != null ? DirectoryReference.Combine(GeneratedCodeDirectory!, "UHT") : null; }
+		}
+
+		/// <summary>
+		/// The directory for this module's generated VNI code
+		/// </summary>
+		public DirectoryReference? GeneratedCodeDirectoryVNI
+		{
+			get { return GeneratedCodeDirectory != null ? DirectoryReference.Combine(GeneratedCodeDirectory!, "VNI") : null; }
+		}
+
+		/// <summary>
+		/// Global override to force all include paths to be always added
 		/// </summary>
-		public bool bAddGeneratedCodeIncludePath;
+		static public bool bForceAddGeneratedCodeIncludePath;
 
 		/// <summary>
 		/// Paths containing *.gen.cpp files for this module.  If this is null then this module doesn't have any generated code.
@@ -321,27 +342,30 @@ namespace UnrealBuildTool
 			bool bLegacyPublicIncludePaths
 			)
 		{
-			// This directory may not exist for this module (or ever exist, if it doesn't contain any generated headers), but we want the project files
-			// to search it so we can pick up generated code definitions after UHT is run for the first time.
-			if(bAddGeneratedCodeIncludePath || (ProjectFileGenerator.bGenerateProjectFiles && GeneratedCodeDirectory != null))
+			if (GeneratedCodeDirectory != null)
 			{
-				IncludePaths.Add(GeneratedCodeDirectory!);
+				// This directory may not exist for this module (or ever exist, if it doesn't contain any generated headers), but we want the project files
+				// to search it so we can pick up generated code definitions after UHT is run for the first time.
+				bool bForceAddIncludePath = bForceAddGeneratedCodeIncludePath || ProjectFileGenerator.bGenerateProjectFiles;
+
+				if (bHasUObjects || bForceAddIncludePath)
+				{
+					IncludePaths.Add(GeneratedCodeDirectoryUHT!);
+				}
 
-				// If this module has Verse code, also add a VNI subdirectory in the generated code directory
-				if (bHasVerse)
+				if (bHasVerse || bForceAddIncludePath)
 				{
-					DirectoryReference VNIGeneratedCodeDirectory = DirectoryReference.Combine(GeneratedCodeDirectory!, "VNI");
-					IncludePaths.Add(VNIGeneratedCodeDirectory);
+					IncludePaths.Add(GeneratedCodeDirectoryVNI!);
 				}
+			}
 
-				if (Rules.AdditionalCodeGenDirectories != null)
+			if (Rules.AdditionalCodeGenDirectories != null)
+			{
+				foreach (string CodeGenDir in Rules.AdditionalCodeGenDirectories)
 				{
-					foreach (string CodeGenDir in Rules.AdditionalCodeGenDirectories)
+					if (Directory.Exists(CodeGenDir))
 					{
-						if (Directory.Exists(CodeGenDir))
-						{
-							IncludePaths.Add(new DirectoryReference(CodeGenDir));
-						}
+						IncludePaths.Add(new DirectoryReference(CodeGenDir));
 					}
 				}
 			}
diff --git a/Engine/Source/Programs/UnrealBuildTool/Modes/GenerateClangDatabase.cs b/Engine/Source/Programs/UnrealBuildTool/Modes/GenerateClangDatabase.cs
--- a/Engine/Source/Programs/UnrealBuildTool/Modes/GenerateClangDatabase.cs
+++ b/Engine/Source/Programs/UnrealBuildTool/Modes/GenerateClangDatabase.cs
@@ -47,6 +47,9 @@ namespace UnrealBuildTool
 				}
 			}
 
+			// Force C++ modules to always include their generated code directories
+			UEBuildModuleCPP.bForceAddGeneratedCodeIncludePath = true;
+
 			// Parse all the target descriptors
 			List<TargetDescriptor> TargetDescriptors = TargetDescriptor.ParseCommandLine(Arguments, BuildConfiguration.bUsePrecompiled, BuildConfiguration.bSkipRulesCompile, BuildConfiguration.bForceRulesCompile);
 
@@ -68,15 +71,6 @@ namespace UnrealBuildTool
 
 					bool IsWindowsClang = ClangPath.GetFileName().Equals("clang-cl.exe", StringComparison.OrdinalIgnoreCase);
 
-					// Convince each module to output its generated code include path
-					foreach (UEBuildBinary Binary in Target.Binaries)
-					{
-						foreach (UEBuildModuleCPP Module in Binary.Modules.OfType<UEBuildModuleCPP>())
-						{
-							Module.bAddGeneratedCodeIncludePath = true;
-						}
-					}
-
 					// Create all the binaries and modules
 					CppCompileEnvironment GlobalCompileEnvironment = Target.CreateCompileEnvironmentForProjectFiles();
 					foreach (UEBuildBinary Binary in Target.Binaries)
diff --git a/Engine/Source/Programs/UnrealBuildTool/System/ExternalExecution.cs b/Engine/Source/Programs/UnrealBuildTool/System/ExternalExecution.cs
--- a/Engine/Source/Programs/UnrealBuildTool/System/ExternalExecution.cs
+++ b/Engine/Source/Programs/UnrealBuildTool/System/ExternalExecution.cs
@@ -461,7 +461,7 @@ namespace UnrealBuildTool
 				{
 					UEBuildModuleCPP Module = ModulesSortedByType[Idx];
 
-					UHTModuleInfo Info = new UHTModuleInfo(Module.Name, Module.RulesFile, Module.ModuleDirectories, ModuleToType[Module], DirectoryItem.GetItemByDirectoryReference(Module.GeneratedCodeDirectory!), GeneratedCodeVersion, Module.Rules.bUsePrecompiled, Module.Rules.OverridePackageType);
+					UHTModuleInfo Info = new UHTModuleInfo(Module.Name, Module.RulesFile, Module.ModuleDirectories, ModuleToType[Module], DirectoryItem.GetItemByDirectoryReference(Module.GeneratedCodeDirectoryUHT!), GeneratedCodeVersion, Module.Rules.bUsePrecompiled, Module.Rules.OverridePackageType);
 					ModuleInfoArray[Idx] = Info;
 
 					Queue.Enqueue(() => SetupUObjectModule(Info, ExcludedFolders, MetadataCache, Queue));
@@ -473,18 +473,33 @@ namespace UnrealBuildTool
 			{
 				UEBuildModuleCPP Module = ModulesSortedByType[Idx];
 				UHTModuleInfo Info = ModuleInfoArray[Idx];
+
+				// Delete any obsolete files at the root of the generated code directory
+				// Remove this code after engine version 5.2 (SOL-2551)
+				if (Module.GeneratedCodeDirectory != null && DirectoryReference.Exists(Module.GeneratedCodeDirectory))
+				{
+					DirectoryItem GeneratedCodeDirectory = DirectoryItem.GetItemByDirectoryReference(Module.GeneratedCodeDirectory);
+					foreach (FileItem File in GeneratedCodeDirectory.EnumerateFiles())
+					{
+						File.Delete();
+					}
+				}
+
 				if (Info.PublicUObjectClassesHeaders.Count > 0 || Info.PrivateUObjectHeaders.Count > 0 || Info.PublicUObjectHeaders.Count > 0 || Info.InternalUObjectHeaders.Count > 0)
 				{
-					// Set a flag indicating that we need to add the generated headers directory
-					Module.bAddGeneratedCodeIncludePath = true;
+					Module.bHasUObjects = true;
 
 					// If we've got this far and there are no source files then it's likely we're installed and ignoring
 					// engine files, so we don't need a .gen.cpp either
-					Info.GeneratedCPPFilenameBase = Path.Combine(Module.GeneratedCodeDirectory!.FullName, Info.ModuleName) + ".gen";
+					DirectoryReference GeneratedCodeDirectoryUHT = Module.GeneratedCodeDirectoryUHT!;
+					Info.GeneratedCPPFilenameBase = Path.Combine(GeneratedCodeDirectoryUHT.FullName, Info.ModuleName) + ".gen";
 					if (!Module.Rules.bUsePrecompiled)
 					{
-						Module.GeneratedCppDirectories = new List<string>();
-						Module.GeneratedCppDirectories.Add(Module.GeneratedCodeDirectory.FullName);
+						if (Module.GeneratedCppDirectories == null)
+						{
+							Module.GeneratedCppDirectories = new List<string>();
+						}
+						Module.GeneratedCppDirectories.Add(GeneratedCodeDirectoryUHT.FullName);
 
 						if (Module.Rules.AdditionalCodeGenDirectories != null)
 						{
@@ -512,12 +527,16 @@ namespace UnrealBuildTool
 				else
 				{
 					// Remove any stale generated code directory
-					// Don't destroy the directory if we have Verse
-					if(!Module.bHasVerse && Module.GeneratedCodeDirectory != null && !Module.Rules.bUsePrecompiled)
+					if (Module.GeneratedCodeDirectoryUHT != null && !Module.Rules.bUsePrecompiled)
 					{
-						if (DirectoryReference.Exists(Module.GeneratedCodeDirectory))
+						if (DirectoryReference.Exists(Module.GeneratedCodeDirectoryUHT))
 						{
-							Directory.Delete(Module.GeneratedCodeDirectory.FullName, true);
+							Directory.Delete(Module.GeneratedCodeDirectoryUHT.FullName, true);
+							// Also delete parent directory if now empty
+							if (!Directory.EnumerateFileSystemEntries(Module.GeneratedCodeDirectory!.FullName).Any())
+							{
+								Directory.Delete(Module.GeneratedCodeDirectory!.FullName, true);
+							}
 						}
 					}
 				}
