Removed obsolete module rules property `AdditionalCodeGenDirectories`

#rnx

diff --git a/Engine/Source/Programs/UnrealBuildTool/Configuration/ModuleRules.cs b/Engine/Source/Programs/UnrealBuildTool/Configuration/ModuleRules.cs
--- a/Engine/Source/Programs/UnrealBuildTool/Configuration/ModuleRules.cs
+++ b/Engine/Source/Programs/UnrealBuildTool/Configuration/ModuleRules.cs
@@ -1027,11 +1027,6 @@ namespace UnrealBuildTool
 		/// </summary>
 		public List<string> ExternalDependencies = new List<string>();
 
-		/// <summary>
-		/// External directories containing generated interop files.
-		/// </summary>
-		public List<string> AdditionalCodeGenDirectories = new List<string>();
-
 		/// <summary>
 		/// Subclass rules files which invalidate the makefile if modified.
 		/// </summary>
diff --git a/Engine/Source/Programs/UnrealBuildTool/Configuration/UEBuildModuleCPP.cs b/Engine/Source/Programs/UnrealBuildTool/Configuration/UEBuildModuleCPP.cs
--- a/Engine/Source/Programs/UnrealBuildTool/Configuration/UEBuildModuleCPP.cs
+++ b/Engine/Source/Programs/UnrealBuildTool/Configuration/UEBuildModuleCPP.cs
@@ -359,17 +359,6 @@ namespace UnrealBuildTool
 				}
 			}
 
-			if (Rules.AdditionalCodeGenDirectories != null)
-			{
-				foreach (string CodeGenDir in Rules.AdditionalCodeGenDirectories)
-				{
-					if (Directory.Exists(CodeGenDir))
-					{
-						IncludePaths.Add(new DirectoryReference(CodeGenDir));
-					}
-				}
-			}
-
 			ModuleInterfacePaths.Add(UEToolChain.GetModuleInterfaceDir(IntermediateDirectory));
 
 			base.AddModuleToCompileEnvironment(SourceModule, SourceBinary, IncludePaths, SystemIncludePaths, ModuleInterfacePaths, Definitions, AdditionalFrameworks, AdditionalPrerequisites, bLegacyPublicIncludePaths);
diff --git a/Engine/Source/Programs/UnrealBuildTool/System/ExternalExecution.cs b/Engine/Source/Programs/UnrealBuildTool/System/ExternalExecution.cs
--- a/Engine/Source/Programs/UnrealBuildTool/System/ExternalExecution.cs
+++ b/Engine/Source/Programs/UnrealBuildTool/System/ExternalExecution.cs
@@ -500,17 +500,6 @@ namespace UnrealBuildTool
 							Module.GeneratedCppDirectories = new List<string>();
 						}
 						Module.GeneratedCppDirectories.Add(GeneratedCodeDirectoryUHT.FullName);
-
-						if (Module.Rules.AdditionalCodeGenDirectories != null)
-						{
-							foreach (string DirPath in Module.Rules.AdditionalCodeGenDirectories)
-							{
-								if (Directory.Exists(DirPath))
-								{
-									Module.GeneratedCppDirectories.Add(Path.GetFullPath(DirPath));
-								}
-							}
-						}
 					}
 
 					UObjectModules.Add(Info);
