UnrealBuildTool: Don't use a remote distribution executor if the total number of actions is less than the number of max local parallel actions

diff --git a/Engine/Source/Programs/UnrealBuildTool/System/ActionGraph.cs b/Engine/Source/Programs/UnrealBuildTool/System/ActionGraph.cs
--- a/Engine/Source/Programs/UnrealBuildTool/System/ActionGraph.cs
+++ b/Engine/Source/Programs/UnrealBuildTool/System/ActionGraph.cs
@@ -328,38 +328,46 @@ namespace UnrealBuildTool
 		}
 
 		/// <summary>
-		/// Executes a list of actions.
+		/// Selects an ActionExecutor
 		/// </summary>
-		public static void ExecuteActions(BuildConfiguration BuildConfiguration, List<LinkedAction> ActionsToExecute, List<TargetDescriptor> TargetDescriptors)
+		private static ActionExecutor SelectExecutor(BuildConfiguration BuildConfiguration, int ActionCount, List<TargetDescriptor> TargetDescriptors)
 		{
-			if (ActionsToExecute.Count == 0)
-			{
-				Log.TraceInformation("Target is up to date");
-			}
-			else
+			if (ActionCount > ParallelExecutor.GetDefaultNumParallelProcesses())
 			{
-				// Figure out which executor to use
-				ActionExecutor Executor;
 				if (BuildConfiguration.bAllowHybridExecutor && HybridExecutor.IsAvailable())
 				{
-					Executor = new HybridExecutor(TargetDescriptors, BuildConfiguration.MaxParallelActions);
+					return new HybridExecutor(TargetDescriptors, BuildConfiguration.MaxParallelActions);
 				}
-				else if (BuildConfiguration.bAllowXGE && XGE.IsAvailable() && ActionsToExecute.Count >= XGE.MinActions)
+				else if (BuildConfiguration.bAllowXGE && XGE.IsAvailable() && ActionCount >= XGE.MinActions)
 				{
-					Executor = new XGE();
+					return new XGE();
 				}
 				else if (BuildConfiguration.bAllowFASTBuild && FASTBuild.IsAvailable())
 				{
-					Executor = new FASTBuild(BuildConfiguration.MaxParallelActions);
+					return new FASTBuild(BuildConfiguration.MaxParallelActions);
 				}
 				else if (BuildConfiguration.bAllowSNDBS && SNDBS.IsAvailable())
 				{
-					Executor = new SNDBS(TargetDescriptors);
-				}
-				else
-				{
-					Executor = new ParallelExecutor(BuildConfiguration.MaxParallelActions);
+					return new SNDBS(TargetDescriptors);
 				}
+			}
+
+			return new ParallelExecutor(BuildConfiguration.MaxParallelActions);
+		}
+
+		/// <summary>
+		/// Executes a list of actions.
+		/// </summary>
+		public static void ExecuteActions(BuildConfiguration BuildConfiguration, List<LinkedAction> ActionsToExecute, List<TargetDescriptor> TargetDescriptors)
+		{
+			if (ActionsToExecute.Count == 0)
+			{
+				Log.TraceInformation("Target is up to date");
+			}
+			else
+			{
+				// Figure out which executor to use
+				ActionExecutor Executor = SelectExecutor(BuildConfiguration, ActionsToExecute.Count, TargetDescriptors);
 
 				// Execute the build
 				Stopwatch Timer = Stopwatch.StartNew();
diff --git a/Engine/Source/Programs/UnrealBuildTool/System/Utils.cs b/Engine/Source/Programs/UnrealBuildTool/System/Utils.cs
--- a/Engine/Source/Programs/UnrealBuildTool/System/Utils.cs
+++ b/Engine/Source/Programs/UnrealBuildTool/System/Utils.cs
@@ -1255,7 +1255,7 @@ namespace UnrealBuildTool
 				NumPhysicalCores = NumLogicalCores;
 			}
 
-			Log.TraceInformation($"Determining max actions to execute in parallel ({NumPhysicalCores} physical cores, {NumLogicalCores} logical cores)");
+			Log.TraceInformationOnce($"Determining max actions to execute in parallel ({NumPhysicalCores} physical cores, {NumLogicalCores} logical cores)");
 
 			// The number of actions to execute in parallel is trying to keep the CPU busy enough in presence of I/O stalls.
 			int MaxActionsToExecuteInParallel;
@@ -1264,12 +1264,12 @@ namespace UnrealBuildTool
 				// The CPU has more logical cores than physical ones, aka uses hyper-threading. 
 				// Use multiplier if provided
 				MaxActionsToExecuteInParallel = (int)(NumPhysicalCores * ProcessorCountMultiplier);
-				Log.TraceInformation($"  Requested {ProcessorCountMultiplier} process count multiplier: limiting max parallel actions to {MaxActionsToExecuteInParallel}");
+				Log.TraceInformationOnce($"  Requested {ProcessorCountMultiplier} process count multiplier: limiting max parallel actions to {MaxActionsToExecuteInParallel}");
 			}
 			// kick off a task per physical core - evidence suggests that, in general, using more cores does not yield significantly better throughput
 			else
 			{
-				Log.TraceInformation($"  Executing up to {NumPhysicalCores} processes, one per physical core");
+				Log.TraceInformationOnce($"  Executing up to {NumPhysicalCores} processes, one per physical core");
 				MaxActionsToExecuteInParallel = NumPhysicalCores;
 			}
 
@@ -1283,7 +1283,7 @@ namespace UnrealBuildTool
 					if (TotalMemoryActions < MaxActionsToExecuteInParallel)
 					{
 						MaxActionsToExecuteInParallel = Math.Max(1, Math.Min(MaxActionsToExecuteInParallel, TotalMemoryActions));
-						Log.TraceInformation($"  Requested {StringUtils.FormatBytesString(MemoryPerActionBytes)} free memory per action, {StringUtils.FormatBytesString(FreeMemoryBytes)} available: limiting max parallel actions to {MaxActionsToExecuteInParallel}");
+						Log.TraceInformationOnce($"  Requested {StringUtils.FormatBytesString(MemoryPerActionBytes)} free memory per action, {StringUtils.FormatBytesString(FreeMemoryBytes)} available: limiting max parallel actions to {MaxActionsToExecuteInParallel}");
 					}
 				}
 			}
@@ -1291,7 +1291,7 @@ namespace UnrealBuildTool
 			if (MaxProcessorCount < MaxActionsToExecuteInParallel)
 			{
 				MaxActionsToExecuteInParallel = Math.Max(1, Math.Min(MaxActionsToExecuteInParallel, MaxProcessorCount));
-				Log.TraceInformation($"  Requested max {MaxProcessorCount} action(s): limiting max parallel actions to {MaxActionsToExecuteInParallel}");
+				Log.TraceInformationOnce($"  Requested max {MaxProcessorCount} action(s): limiting max parallel actions to {MaxActionsToExecuteInParallel}");
 			}
 			return MaxActionsToExecuteInParallel;
 		}
