Toolchains can modify the target receipt, for adding custom properties etc.

#rnx

diff --git a/Engine/Source/Programs/UnrealBuildTool/Configuration/UEBuildTarget.cs b/Engine/Source/Programs/UnrealBuildTool/Configuration/UEBuildTarget.cs
--- a/Engine/Source/Programs/UnrealBuildTool/Configuration/UEBuildTarget.cs
+++ b/Engine/Source/Programs/UnrealBuildTool/Configuration/UEBuildTarget.cs
@@ -1653,6 +1653,8 @@ namespace UnrealBuildTool
 				Receipt.AdditionalProperties.Add(new ReceiptProperty("CustomConfig", Rules.CustomConfig));
 			}
 
+			ToolChain.ModifyTargetReceipt(Receipt);
+
 			return Receipt;
 		}
 
diff --git a/Engine/Source/Programs/UnrealBuildTool/ToolChain/UEToolChain.cs b/Engine/Source/Programs/UnrealBuildTool/ToolChain/UEToolChain.cs
--- a/Engine/Source/Programs/UnrealBuildTool/ToolChain/UEToolChain.cs
+++ b/Engine/Source/Programs/UnrealBuildTool/ToolChain/UEToolChain.cs
@@ -91,6 +91,10 @@ namespace UnrealBuildTool
 		{
 		}
 
+		public virtual void ModifyTargetReceipt(TargetReceipt Receipt)
+		{
+		}
+
 		public virtual void FinalizeOutput(ReadOnlyTargetRules Target, TargetMakefile Makefile)
 		{
 		}
