Debugger arguments are set up correctly for Windows-based platform extensions

#rnx

diff --git a/Engine/Source/Programs/UnrealBuildTool/Platform/Windows/WindowsProjectGenerator.cs b/Engine/Source/Programs/UnrealBuildTool/Platform/Windows/WindowsProjectGenerator.cs
--- a/Engine/Source/Programs/UnrealBuildTool/Platform/Windows/WindowsProjectGenerator.cs
+++ b/Engine/Source/Programs/UnrealBuildTool/Platform/Windows/WindowsProjectGenerator.cs
@@ -49,6 +49,34 @@ namespace UnrealBuildTool
 			return InPlatform.ToString();
 		}
 
+		public override string GetVisualStudioUserFileStrings(UnrealTargetPlatform InPlatform, UnrealTargetConfiguration InConfiguration, string InConditionString, TargetRules InTargetRules, FileReference TargetRulesPath, FileReference ProjectFilePath, string ProjectName, string? ForeignUProjectPath)
+		{
+			StringBuilder VCUserFileContent = new StringBuilder();
+
+			VCUserFileContent.AppendLine("  <PropertyGroup {0}>", InConditionString);
+			if (InTargetRules.Type != TargetType.Game)
+			{
+				string DebugOptions = "";
+
+				if (ForeignUProjectPath != null)
+				{
+					DebugOptions += ForeignUProjectPath;
+					DebugOptions += " -skipcompile";
+				}
+				else if (InTargetRules.Type == TargetType.Editor && ProjectName != ProjectFileGenerator.EngineProjectFileNameBase)
+				{
+					DebugOptions += ProjectName;
+				}
+
+				VCUserFileContent.AppendLine("    <LocalDebuggerCommandArguments>{0}</LocalDebuggerCommandArguments>", DebugOptions);
+			}
+			VCUserFileContent.AppendLine("    <DebuggerFlavor>WindowsLocalDebugger</DebuggerFlavor>");
+			VCUserFileContent.AppendLine("  </PropertyGroup>");
+
+
+			return VCUserFileContent.ToString();
+		}
+
 		public override bool RequiresVSUserFileGeneration()
 		{
 			return true;
diff --git a/Engine/Source/Programs/UnrealBuildTool/ProjectFiles/PlatformProjectGenerator.cs b/Engine/Source/Programs/UnrealBuildTool/ProjectFiles/PlatformProjectGenerator.cs
--- a/Engine/Source/Programs/UnrealBuildTool/ProjectFiles/PlatformProjectGenerator.cs
+++ b/Engine/Source/Programs/UnrealBuildTool/ProjectFiles/PlatformProjectGenerator.cs
@@ -257,6 +257,24 @@ namespace UnrealBuildTool
 			return "";
 		}
 
+		/// <summary>
+		/// Get the text to insert into the user file for the given platform/configuration/target
+		/// </summary>
+		/// <param name="InPlatform">The platform being added</param>
+		/// <param name="InConfiguration">The configuration being added</param>
+		/// <param name="InConditionString">The condition string </param>
+		/// <param name="InTargetRules">The target rules </param>
+		/// <param name="TargetRulesPath">The target rules path</param>
+		/// <param name="ProjectFilePath">The project file path</param>
+		/// <param name="ProjectName">The name of the project</param>
+		/// <param name="ForeignUProjectPath">Path to foreign .uproject file, if any</param>
+		/// <returns>The string to append to the user file</returns>
+		public virtual string GetVisualStudioUserFileStrings(UnrealTargetPlatform InPlatform, UnrealTargetConfiguration InConfiguration,
+			string InConditionString, TargetRules InTargetRules, FileReference TargetRulesPath, FileReference ProjectFilePath, string ProjectName, string? ForeignUProjectPath)
+		{
+			return GetVisualStudioUserFileStrings(InPlatform, InConfiguration, InConditionString, InTargetRules, TargetRulesPath, ProjectFilePath);
+		}
+
 		/// <summary>
 		/// For Additional Project Property files that need to be written out.  This is currently used only on Android. 
 		/// </summary>
diff --git a/Engine/Source/Programs/UnrealBuildTool/ProjectFiles/VisualStudio/VCProject.cs b/Engine/Source/Programs/UnrealBuildTool/ProjectFiles/VisualStudio/VCProject.cs
--- a/Engine/Source/Programs/UnrealBuildTool/ProjectFiles/VisualStudio/VCProject.cs
+++ b/Engine/Source/Programs/UnrealBuildTool/ProjectFiles/VisualStudio/VCProject.cs
@@ -1782,7 +1782,7 @@ namespace UnrealBuildTool
 				{
 					TargetRules TargetRulesObject = Combination.ProjectTarget.TargetRules!;
 
-					if (Combination.Platform!.Value.IsInGroup(UnrealPlatformGroup.Windows) || Combination.Platform!.Value.IsInGroup(UnrealPlatformGroup.HoloLens))
+					if (Combination.Platform!.Value.IsInGroup(UnrealPlatformGroup.HoloLens)) // this logic should be moved into HoloLensProjectGenerator.GetVisualStudioUserFileStrings
 					{
 						VCUserFileContent.AppendLine("  <PropertyGroup {0}>", ConditionString);
 						if (TargetRulesObject.Type != TargetType.Game)
@@ -1801,13 +1801,14 @@ namespace UnrealBuildTool
 
 							VCUserFileContent.AppendLine("    <LocalDebuggerCommandArguments>{0}</LocalDebuggerCommandArguments>", DebugOptions);
 						}
-						VCUserFileContent.AppendLine("    <DebuggerFlavor>{0}</DebuggerFlavor>", Combination.Platform == UnrealTargetPlatform.HoloLens ? "AppHostLocalDebugger " : "WindowsLocalDebugger");
+						VCUserFileContent.AppendLine("    <DebuggerFlavor>{0}</DebuggerFlavor>", "AppHostLocalDebugger ");
 						VCUserFileContent.AppendLine("  </PropertyGroup>");
 					}
 
 					if (ProjGenerator != null)
 					{
-						VCUserFileContent.Append(ProjGenerator.GetVisualStudioUserFileStrings(Combination.Platform.Value, Configuration, ConditionString, TargetRulesObject, Combination.ProjectTarget.TargetFilePath, ProjectFilePath));
+						string? ForeignUProjectPath = (IsForeignProject && !string.IsNullOrEmpty(UProjectPath)) ? UProjectPath : null;
+						VCUserFileContent.Append(ProjGenerator.GetVisualStudioUserFileStrings(Combination.Platform.Value, Configuration, ConditionString, TargetRulesObject, Combination.ProjectTarget.TargetFilePath, ProjectFilePath, ProjectName, ForeignUProjectPath));
 					}
 				}
 			}
