Wrap path cmd provided to clang into quotes to fix issue with path containing whitespaces

diff --git a/Engine/Source/Programs/UnrealBuildTool/Platform/Android/AndroidToolChain.cs b/Engine/Source/Programs/UnrealBuildTool/Platform/Android/AndroidToolChain.cs
--- a/Engine/Source/Programs/UnrealBuildTool/Platform/Android/AndroidToolChain.cs
+++ b/Engine/Source/Programs/UnrealBuildTool/Platform/Android/AndroidToolChain.cs
@@ -964,7 +964,7 @@ namespace UnrealBuildTool
 				// Make all symbols (except ones called from Java) hidden
 				Writer.WriteLine("{ global: Java_*; ANativeActivity_onCreate; JNI_OnLoad; local: *; };");
 			}
-			Result += " -Wl,--version-script=" + VersionScriptFile;
+			Result += " -Wl,--version-script=\"" + VersionScriptFile + "\"";
 
 			Result += " -Wl,--build-id=sha1";               // add build-id to make debugging easier
 
