Manifest generator can use TargetName if there is no project, so Programs are not named DefaultUEProject, for example.

#rnx

diff --git a/Engine/Source/Programs/UnrealBuildTool/Platform/Windows/VCManifestGenerator.cs b/Engine/Source/Programs/UnrealBuildTool/Platform/Windows/VCManifestGenerator.cs
--- a/Engine/Source/Programs/UnrealBuildTool/Platform/Windows/VCManifestGenerator.cs
+++ b/Engine/Source/Programs/UnrealBuildTool/Platform/Windows/VCManifestGenerator.cs
@@ -685,10 +685,10 @@ namespace UnrealBuildTool
 		/// <summary>
 		/// Get the package identity name string
 		/// </summary>
-		protected string GetIdentityPackageName()
+		protected string GetIdentityPackageName( string? TargetName)
 		{
             // Read the PackageName from config
-			var DefaultName = (ProjectFile != null) ? ProjectFile.GetFileNameWithoutAnyExtensions() : "DefaultUEProject";
+			var DefaultName = (ProjectFile != null) ? ProjectFile.GetFileNameWithoutAnyExtensions() : (TargetName ?? "DefaultUEProject");
             var PackageName = Regex.Replace(GetConfigString("PackageName", "ProjectName", DefaultName), "[^-.A-Za-z0-9]", "");
             if (string.IsNullOrWhiteSpace(PackageName))
             {
@@ -737,9 +737,9 @@ namespace UnrealBuildTool
 		/// <summary>
 		/// Get the package identity element
 		/// </summary>
-		protected XElement GetIdentity(out string IdentityName)
+		protected XElement GetIdentity(string? TargetName, out string IdentityName)
         {
-            var PackageName = GetIdentityPackageName();
+            var PackageName = GetIdentityPackageName(TargetName);
             var PublisherName = GetIdentityPublisherName();
             var VersionNumber = GetIdentityVersionNumber();
 
@@ -787,7 +787,7 @@ namespace UnrealBuildTool
 		/// <summary>
 		/// Return the entire manifest element
 		/// </summary>
-		protected abstract XElement GetManifest(List<UnrealTargetConfiguration> TargetConfigs, List<string> Executables, out string IdentityName);
+		protected abstract XElement GetManifest(List<UnrealTargetConfiguration> TargetConfigs, List<string> Executables, string? TargetName, out string IdentityName);
 
 		/// <summary>
 		/// Perform any platform-specific processing on the manifest before it is saved
@@ -799,7 +799,7 @@ namespace UnrealBuildTool
 		/// <summary>
 		/// Create a manifest and return the list of modified files
 		/// </summary>
-		public List<string>? CreateManifest(string InManifestName, string InOutputPath, string InIntermediatePath, FileReference? InProjectFile, string InProjectDirectory, List<UnrealTargetConfiguration> InTargetConfigs, List<string> InExecutables)
+		public List<string>? CreateManifest(string InManifestName, string InOutputPath, string InIntermediatePath, string? InTargetName, FileReference? InProjectFile, string InProjectDirectory, List<UnrealTargetConfiguration> InTargetConfigs, List<string> InExecutables)
 		{
 			// Check parameter values are valid.
 			if (InTargetConfigs.Count != InExecutables.Count)
@@ -926,7 +926,7 @@ namespace UnrealBuildTool
 
 			// Create the manifest document
 			string? IdentityName = null;
-			var ManifestXmlDocument = new XDocument(GetManifest(InTargetConfigs, InExecutables, out IdentityName));
+			var ManifestXmlDocument = new XDocument(GetManifest(InTargetConfigs, InExecutables, InTargetName, out IdentityName));
 
 			// Export manifest to the intermediate directory then compare the contents to any existing target manifest
 			// and replace if there are differences.
