Adding a "deprecation" path for engine includes to allow types to be moved without breaking licensee module builds.
- Non-engine modules/targets will have to specify the "version" of includes via IncludeOrderVersion in TargetRules or ModuleRules.
- This setting will control the value of UE_ENABLE_INCLUDE_ORDER_DEPRECATED_IN_XXX where XXX is the version of the engine.
- When moving types out of a header, users will need to include the new location of the type in the header it was removed from but only if UE_ENABLE_INCLUDE_ORDER_DEPRECATED_IN_XXX is set.
- If a target does not change its IncludeOrderVersion to the latest version, UBT will print out a message telling users how to upgrade.
- This change introduces a new set of SharedPCH permutations to make sure modules with older versions get the PCH with UE_ENABLE_INCLUDE_ORDER_DEPRECATED_IN_XXX set correctly.

diff --git a/Engine/Source/Programs/InterchangeWorker/InterchangeWorker.Target.cs b/Engine/Source/Programs/InterchangeWorker/InterchangeWorker.Target.cs
--- a/Engine/Source/Programs/InterchangeWorker/InterchangeWorker.Target.cs
+++ b/Engine/Source/Programs/InterchangeWorker/InterchangeWorker.Target.cs
@@ -10,6 +10,7 @@ public class InterchangeWorkerTarget : TargetRules
 	{
 		Type = TargetType.Program;
 		LinkType = TargetLinkType.Monolithic;
+		IncludeOrderVersion = EngineIncludeOrderVersion.Latest;
 		LaunchModuleName = "InterchangeWorker";
 		SolutionDirectory = "Programs/InterchangeWorker";
 
diff --git a/Engine/Source/Programs/ShaderCompileWorker/ShaderCompileWorker.Target.cs b/Engine/Source/Programs/ShaderCompileWorker/ShaderCompileWorker.Target.cs
--- a/Engine/Source/Programs/ShaderCompileWorker/ShaderCompileWorker.Target.cs
+++ b/Engine/Source/Programs/ShaderCompileWorker/ShaderCompileWorker.Target.cs
@@ -11,6 +11,7 @@ public class ShaderCompileWorkerTarget : TargetRules
 	{
 		Type = TargetType.Program;
 		LinkType = TargetLinkType.Modular;
+		IncludeOrderVersion = EngineIncludeOrderVersion.Latest;
 
 		LaunchModuleName = "ShaderCompileWorker";
 
diff --git a/Engine/Source/Programs/UnrealBuildTool/Configuration/ModuleRules.cs b/Engine/Source/Programs/UnrealBuildTool/Configuration/ModuleRules.cs
--- a/Engine/Source/Programs/UnrealBuildTool/Configuration/ModuleRules.cs
+++ b/Engine/Source/Programs/UnrealBuildTool/Configuration/ModuleRules.cs
@@ -703,6 +703,27 @@ namespace UnrealBuildTool
 		}
 		private BuildSettingsVersion? DefaultBuildSettingsPrivate;
 
+		/// <summary>
+		/// What version of include order to use when compiling this module. Can be overridden via -ForceIncludeOrder on the command line or in a module's rules.
+		/// </summary>
+		public EngineIncludeOrderVersion IncludeOrderVersion
+		{
+			get
+			{
+				if (Target.ForcedIncludeOrder != null)
+				{
+					return Target.ForcedIncludeOrder.Value;
+				}
+				if (bTreatAsEngineModule)
+				{
+					return EngineIncludeOrderVersion.Latest;
+				}
+				return IncludeOrderVersionPrivate ?? Target.IncludeOrderVersion;
+			}
+			set { IncludeOrderVersionPrivate = value; }
+		}
+		private EngineIncludeOrderVersion? IncludeOrderVersionPrivate;
+
 		/// <summary>
 		/// Use run time type information
 		/// </summary>
diff --git a/Engine/Source/Programs/UnrealBuildTool/Configuration/TargetRules.cs b/Engine/Source/Programs/UnrealBuildTool/Configuration/TargetRules.cs
--- a/Engine/Source/Programs/UnrealBuildTool/Configuration/TargetRules.cs
+++ b/Engine/Source/Programs/UnrealBuildTool/Configuration/TargetRules.cs
@@ -130,6 +130,65 @@ namespace UnrealBuildTool
 		Latest
 	}
 
+	/// <summary>
+	/// What version of include order to use when compiling.
+	/// </summary>
+	public enum EngineIncludeOrderVersion
+	{
+		/// <summary>
+		/// Include order used in Unreal 5.0
+		/// </summary>
+		Unreal5_0,
+
+		/// <summary>
+		/// Include order used in Unreal 5.1
+		/// </summary>
+		Unreal5_1,
+
+		// *** When adding new entries here, be sure to update UEBuildModuleCPP.CurrentIncludeOrderDefine to ensure that the correct guard is used. ***
+
+		/// <summary>
+		/// Always use the latest version of include order.
+		/// </summary>
+		Latest = Unreal5_1,
+
+		/// <summary>
+		/// Contains the oldest version of include order that the engine supports.
+		/// </summary>
+		Oldest = Unreal5_0,
+	}
+
+	/// <summary>
+	/// Utility class for EngineIncludeOrderVersion defines
+	/// </summary>
+	public class EngineIncludeOrderHelper
+	{
+		/// <summary>
+		/// Returns a list of every deprecation define available.
+		/// </summary>
+		/// <returns></returns>
+		public static List<string> GetAllDeprecationDefines()
+		{
+			return new List<string>()
+			{
+				"UE_ENABLE_INCLUDE_ORDER_DEPRECATED_IN_5_1",
+			};
+		}
+
+		/// <summary>
+		/// Get a list of every deprecation define and their value for the specified engine include order.
+		/// </summary>
+		/// <param name="InVersion"></param>
+		/// <returns></returns>
+		public static List<string> GetDeprecationDefines(EngineIncludeOrderVersion InVersion)
+		{
+			return new List<string>()
+			{
+				"UE_ENABLE_INCLUDE_ORDER_DEPRECATED_IN_5_1=" + (InVersion < EngineIncludeOrderVersion.Unreal5_1 ? "1" : "0"),
+			};
+		}
+	}
+
 	/// <summary>
 	/// Attribute used to mark fields which much match between targets in the shared build environment
 	/// </summary>
@@ -281,6 +340,29 @@ namespace UnrealBuildTool
 		}
 		private BuildSettingsVersion? DefaultBuildSettingsPrivate; // Cannot be initialized inline; potentially overridden before the constructor is called.
 
+		/// <summary>
+		/// Force the include order to a specific version. Overrides any Target and Module rules.
+		/// </summary>
+		[CommandLine("-ForceIncludeOrder=")]
+		public EngineIncludeOrderVersion? ForcedIncludeOrder = null;
+
+		/// <summary>
+		/// What version of include order to use when compiling this target. Can be overridden via -ForceIncludeOrder on the command line. ModuleRules.IncludeOrderVersion takes precedence.
+		/// </summary>
+		public EngineIncludeOrderVersion IncludeOrderVersion
+		{
+			get
+			{
+				if (ForcedIncludeOrder != null)
+				{
+					return ForcedIncludeOrder.Value;
+				}
+				return IncludeOrderVersionPrivate ?? EngineIncludeOrderVersion.Oldest;
+			}
+			set { IncludeOrderVersionPrivate = value; }
+		}
+		private EngineIncludeOrderVersion? IncludeOrderVersionPrivate;
+
 		/// <summary>
 		/// Tracks a list of config values read while constructing this target
 		/// </summary>
@@ -2023,6 +2105,18 @@ namespace UnrealBuildTool
 				Diagnostics.Add(String.Format("[Upgrade] Suppress this message by setting 'DefaultBuildSettings = BuildSettingsVersion.{1};' in {2}, and explicitly overriding settings that differ from the new defaults.", Version, (BuildSettingsVersion)(BuildSettingsVersion.Latest - 1), File!.GetFileName()));
 				Diagnostics.Add("[Upgrade]");
 			}
+
+			if (IncludeOrderVersion <= (EngineIncludeOrderVersion)(EngineIncludeOrderVersion.Latest - 1) && ForcedIncludeOrder == null)
+			{
+				Diagnostics.Add("[Upgrade]");
+				Diagnostics.Add("[Upgrade] Using backward-compatible include order. The latest version of UE has changed the order of includes, which may require code changes. The current setting is:");
+				Diagnostics.Add(String.Format("[Upgrade]     IncludeOrderVersion = EngineIncludeOrderVersion.{0}", IncludeOrderVersion));
+				Diagnostics.Add(String.Format("[Upgrade] Suppress this message by setting 'IncludeOrderVersion = EngineIncludeOrderVersion.{0};' in {1}.", EngineIncludeOrderVersion.Latest, File!.GetFileName()));
+				Diagnostics.Add("[Upgrade] Alternatively you can set this to 'EngineIncludeOrderVersion.Latest' to always use the latest include order. This will potentially cause compile errors when integrating new versions of the engine.");
+				Diagnostics.Add("[Upgrade]");
+			}
+
+			Log.TraceLog("Using EngineIncludeOrderVersion.{0} for target {1}", IncludeOrderVersion, File!.GetFileName());
 		}
 	}
 
@@ -2120,6 +2214,16 @@ namespace UnrealBuildTool
 			get { return Inner.DefaultBuildSettings; }
 		}
 
+		public EngineIncludeOrderVersion? ForcedIncludeOrder
+		{
+			get { return Inner.ForcedIncludeOrder; }
+		}
+
+		public EngineIncludeOrderVersion IncludeOrderVersion
+		{
+			get { return Inner.IncludeOrderVersion; }
+		}
+
 		internal ConfigValueTracker ConfigValueTracker
 		{
 			get { return Inner.ConfigValueTracker; }
diff --git a/Engine/Source/Programs/UnrealBuildTool/Configuration/UEBuildModuleCPP.cs b/Engine/Source/Programs/UnrealBuildTool/Configuration/UEBuildModuleCPP.cs
--- a/Engine/Source/Programs/UnrealBuildTool/Configuration/UEBuildModuleCPP.cs
+++ b/Engine/Source/Programs/UnrealBuildTool/Configuration/UEBuildModuleCPP.cs
@@ -739,12 +739,25 @@ namespace UnrealBuildTool
 			PrecompiledHeaderInstance? Instance = Template.Instances.Find(x => IsCompatibleForSharedPCH(x.CompileEnvironment, ModuleCompileEnvironment));
 			if(Instance == null)
 			{
+				List<string> Definitions = Template.BaseCompileEnvironment.Definitions;
+
+				// Modify definitions if we need to create a new shared pch for the include order
+				if (ModuleCompileEnvironment.IncludeOrderVersion != Template.BaseCompileEnvironment.IncludeOrderVersion)
+				{
+					Definitions = new List<string>(Definitions);
+					foreach (string OldDefine in EngineIncludeOrderHelper.GetDeprecationDefines(Template.BaseCompileEnvironment.IncludeOrderVersion))
+					{
+						Definitions.Remove(OldDefine);
+					}
+					Definitions.AddRange(EngineIncludeOrderHelper.GetDeprecationDefines(ModuleCompileEnvironment.IncludeOrderVersion));
+				}
+
 				// Create a suffix to distinguish this shared PCH variant from any others. Currently only optimized and non-optimized shared PCHs are supported.
 				string Variant = GetSuffixForSharedPCH(ModuleCompileEnvironment, Template.BaseCompileEnvironment);
 
 				// Create the wrapper file, which sets all the definitions needed to compile it
 				FileReference WrapperLocation = FileReference.Combine(Template.OutputDir, String.Format("SharedPCH.{0}{1}.h", Template.Module.Name, Variant));
-				FileItem WrapperFile = CreatePCHWrapperFile(WrapperLocation, Template.BaseCompileEnvironment.Definitions, Template.HeaderFile, Graph);
+				FileItem WrapperFile = CreatePCHWrapperFile(WrapperLocation, Definitions, Template.HeaderFile, Graph);
 
 				// Create the compile environment for this PCH
 				CppCompileEnvironment CompileEnvironment = new CppCompileEnvironment(Template.BaseCompileEnvironment);
@@ -805,6 +818,10 @@ namespace UnrealBuildTool
 			{
 				return false;
 			}
+			if (ModuleCompileEnvironment.IncludeOrderVersion != CompileEnvironment.IncludeOrderVersion)
+			{
+				return false;
+			}
 			return true;
 		}
 
@@ -899,6 +916,14 @@ namespace UnrealBuildTool
 				Variant += String.Format(".{0}", CompileEnvironment.CppStandard);
 			}
 
+			if (CompileEnvironment.IncludeOrderVersion != BaseCompileEnvironment.IncludeOrderVersion)
+			{
+				if (CompileEnvironment.IncludeOrderVersion != EngineIncludeOrderVersion.Latest)
+				{
+					Variant += ".InclOrder" + CompileEnvironment.IncludeOrderVersion.ToString();
+				}
+			}
+
 			return Variant;
 		}
 
@@ -916,6 +941,7 @@ namespace UnrealBuildTool
 			CompileEnvironment.UnsafeTypeCastWarningLevel = ModuleCompileEnvironment.UnsafeTypeCastWarningLevel;
 			CompileEnvironment.bEnableUndefinedIdentifierWarnings = ModuleCompileEnvironment.bEnableUndefinedIdentifierWarnings;
 			CompileEnvironment.CppStandard = ModuleCompileEnvironment.CppStandard;
+			CompileEnvironment.IncludeOrderVersion = ModuleCompileEnvironment.IncludeOrderVersion;
 		}
 
 		/// <summary>
@@ -1153,6 +1179,10 @@ namespace UnrealBuildTool
 								Writer.WriteLine("#define DEPRECATED_FORGAME DEPRECATED");
 								Writer.WriteLine("#undef UE_DEPRECATED_FORGAME");
 								Writer.WriteLine("#define UE_DEPRECATED_FORGAME UE_DEPRECATED");
+								foreach (string DeprecationDefine in EngineIncludeOrderHelper.GetAllDeprecationDefines())
+								{
+									Writer.WriteLine("#undef " + DeprecationDefine);
+								}
 							}
 
 							WriteDefinitions(CompileEnvironment.Definitions, Writer);
@@ -1405,6 +1435,7 @@ namespace UnrealBuildTool
 			Result.ShadowVariableWarningLevel = Rules.ShadowVariableWarningLevel;
 			Result.UnsafeTypeCastWarningLevel = Rules.UnsafeTypeCastWarningLevel;
 			Result.bEnableUndefinedIdentifierWarnings = Rules.bEnableUndefinedIdentifierWarnings;
+			Result.IncludeOrderVersion = Rules.IncludeOrderVersion;
 
 			// If the module overrides the C++ language version, override it on the compile environment
 			if (Rules.CppStandard != CppStandardVersion.Default)
@@ -1432,6 +1463,8 @@ namespace UnrealBuildTool
 				Result.Definitions.Add("UE_IS_ENGINE_MODULE=0");
 			}
 
+			Result.Definitions.AddRange(EngineIncludeOrderHelper.GetDeprecationDefines(Rules.IncludeOrderVersion));
+
 			// For game modules, set the define for the project and target names, which will be used by the IMPLEMENT_PRIMARY_GAME_MODULE macro.
 			if (!Rules.bTreatAsEngineModule)
 			{
@@ -1487,6 +1520,8 @@ namespace UnrealBuildTool
 				CompileEnvironment.Definitions.Add("UE_IS_ENGINE_MODULE=0");
 			}
 
+			CompileEnvironment.Definitions.AddRange(EngineIncludeOrderHelper.GetDeprecationDefines(Rules.IncludeOrderVersion));
+
 			// Add the module's private definitions.
 			CompileEnvironment.Definitions.AddRange(PublicDefinitions);
 
diff --git a/Engine/Source/Programs/UnrealBuildTool/System/CppCompileEnvironment.cs b/Engine/Source/Programs/UnrealBuildTool/System/CppCompileEnvironment.cs
--- a/Engine/Source/Programs/UnrealBuildTool/System/CppCompileEnvironment.cs
+++ b/Engine/Source/Programs/UnrealBuildTool/System/CppCompileEnvironment.cs
@@ -410,6 +410,11 @@ namespace UnrealBuildTool
 		/// </summary>
 		public bool bEnableCoroutines = false;
 
+		/// <summary>
+		/// What version of include order specified by the module rules. Used to determine shared PCH variants.
+		/// </summary>
+		public EngineIncludeOrderVersion IncludeOrderVersion = EngineIncludeOrderVersion.Latest;
+
 		/// <summary>
 		/// Default constructor.
 		/// </summary>
@@ -492,6 +497,7 @@ namespace UnrealBuildTool
 			bHideSymbolsByDefault = Other.bHideSymbolsByDefault;
 			CppStandard = Other.CppStandard;
 			bEnableCoroutines = Other.bEnableCoroutines;
+			IncludeOrderVersion = Other.IncludeOrderVersion;
 		}
 	}
 }
diff --git a/Engine/Source/Programs/UnrealPak/UnrealPak.Target.cs b/Engine/Source/Programs/UnrealPak/UnrealPak.Target.cs
--- a/Engine/Source/Programs/UnrealPak/UnrealPak.Target.cs
+++ b/Engine/Source/Programs/UnrealPak/UnrealPak.Target.cs
@@ -9,6 +9,7 @@ public class UnrealPakTarget : TargetRules
 	{
 		Type = TargetType.Program;
 		LinkType = TargetLinkType.Modular;
+		IncludeOrderVersion = EngineIncludeOrderVersion.Latest;
 		LaunchModuleName = "UnrealPak";
 
 		bBuildDeveloperTools = false;
diff --git a/Engine/Source/UnrealClient.Target.cs b/Engine/Source/UnrealClient.Target.cs
--- a/Engine/Source/UnrealClient.Target.cs
+++ b/Engine/Source/UnrealClient.Target.cs
@@ -8,6 +8,7 @@ public class UnrealClientTarget : TargetRules
     public UnrealClientTarget(TargetInfo Target) : base(Target)
 	{
 		Type = TargetType.Client;
+		IncludeOrderVersion = EngineIncludeOrderVersion.Latest;
 		BuildEnvironment = TargetBuildEnvironment.Shared;
 		ExtraModuleNames.Add("UnrealGame");
 	}
diff --git a/Engine/Source/UnrealEditor.Target.cs b/Engine/Source/UnrealEditor.Target.cs
--- a/Engine/Source/UnrealEditor.Target.cs
+++ b/Engine/Source/UnrealEditor.Target.cs
@@ -8,6 +8,7 @@ public class UnrealEditorTarget : TargetRules
 	public UnrealEditorTarget( TargetInfo Target ) : base(Target)
 	{
 		Type = TargetType.Editor;
+		IncludeOrderVersion = EngineIncludeOrderVersion.Latest;
 		BuildEnvironment = TargetBuildEnvironment.Shared;
 		bBuildAllModules = true;
 		ExtraModuleNames.Add("UnrealGame");
diff --git a/Engine/Source/UnrealGame.Target.cs b/Engine/Source/UnrealGame.Target.cs
--- a/Engine/Source/UnrealGame.Target.cs
+++ b/Engine/Source/UnrealGame.Target.cs
@@ -9,6 +9,7 @@ public class UnrealGameTarget : TargetRules
 	public UnrealGameTarget( TargetInfo Target ) : base(Target)
 	{
 		Type = TargetType.Game;
+		IncludeOrderVersion = EngineIncludeOrderVersion.Latest;
 		BuildEnvironment = TargetBuildEnvironment.Shared;
 
 		ExtraModuleNames.Add("UnrealGame");
diff --git a/Engine/Source/UnrealServer.Target.cs b/Engine/Source/UnrealServer.Target.cs
--- a/Engine/Source/UnrealServer.Target.cs
+++ b/Engine/Source/UnrealServer.Target.cs
@@ -9,6 +9,7 @@ public class UnrealServerTarget : TargetRules
     public UnrealServerTarget(TargetInfo Target) : base(Target)
 	{
 		Type = TargetType.Server;
+		IncludeOrderVersion = EngineIncludeOrderVersion.Latest;
 		BuildEnvironment = TargetBuildEnvironment.Shared;
 		ExtraModuleNames.Add("UnrealGame");
 	}
