TargetDescriptor.cs: Warning suppression

#trivial
#rnx

diff --git a/Engine/Source/Programs/UnrealBuildTool/Configuration/TargetDescriptor.cs b/Engine/Source/Programs/UnrealBuildTool/Configuration/TargetDescriptor.cs
--- a/Engine/Source/Programs/UnrealBuildTool/Configuration/TargetDescriptor.cs
+++ b/Engine/Source/Programs/UnrealBuildTool/Configuration/TargetDescriptor.cs
@@ -477,11 +477,13 @@ namespace UnrealBuildTool
 
 		public override int GetHashCode()
 		{
+#pragma warning disable RS1024
 			return String.GetHashCode(ProjectFile?.FullName) + 
 				Name.GetHashCode() + 
 				Platform.GetHashCode() + 
 				Configuration.GetHashCode() + 
 				Architecture.GetHashCode();
+#pragma warning restore RE1024
 		}
 
 		public override bool Equals(object? Obj) 
