nullability annotations, rethrow fixes, net6.0 warning suppression

#trivial
#rnx

diff --git a/Engine/Source/Programs/AutomationTool/AutomationUtils/DeviceReservation.cs b/Engine/Source/Programs/AutomationTool/AutomationUtils/DeviceReservation.cs
--- a/Engine/Source/Programs/AutomationTool/AutomationUtils/DeviceReservation.cs
+++ b/Engine/Source/Programs/AutomationTool/AutomationUtils/DeviceReservation.cs
@@ -10,6 +10,8 @@ using UnrealBuildTool;
 using System.Text;
 using System.Text.RegularExpressions;
 
+#pragma warning disable SYSLIB0014
+
 namespace AutomationTool.DeviceReservation
 {
 	/// <summary>
diff --git a/Engine/Source/Programs/AutomationTool/AutomationUtils/Platform.cs b/Engine/Source/Programs/AutomationTool/AutomationUtils/Platform.cs
--- a/Engine/Source/Programs/AutomationTool/AutomationUtils/Platform.cs
+++ b/Engine/Source/Programs/AutomationTool/AutomationUtils/Platform.cs
@@ -145,7 +145,7 @@ namespace AutomationTool
 					}
 				}
 				// Re-throw, this is still a critical error!
-				throw Ex;
+				throw;
 			}
 			foreach (var PotentialPlatformType in AllTypes)
 			{
diff --git a/Engine/Source/Programs/AutomationTool/BuildGraph/Tasks/CallMethodTask.cs b/Engine/Source/Programs/AutomationTool/BuildGraph/Tasks/CallMethodTask.cs
--- a/Engine/Source/Programs/AutomationTool/BuildGraph/Tasks/CallMethodTask.cs
+++ b/Engine/Source/Programs/AutomationTool/BuildGraph/Tasks/CallMethodTask.cs
@@ -224,7 +224,7 @@ namespace AutomationTool.Tasks
 				}
 				else
 				{
-					Arguments[Idx] = TypeDescriptor.GetConverter(ParameterType).ConvertFromString(ArgumentString);
+					Arguments[Idx] = TypeDescriptor.GetConverter(ParameterType).ConvertFromString(ArgumentString)!;
 				}
 			}
 
diff --git a/Engine/Source/Programs/AutomationTool/BuildGraph/Tasks/NotarizeTask.cs b/Engine/Source/Programs/AutomationTool/BuildGraph/Tasks/NotarizeTask.cs
--- a/Engine/Source/Programs/AutomationTool/BuildGraph/Tasks/NotarizeTask.cs
+++ b/Engine/Source/Programs/AutomationTool/BuildGraph/Tasks/NotarizeTask.cs
@@ -14,6 +14,8 @@ using System.Xml;
 using EpicGames.Core;
 using UnrealBuildTool;
 
+#pragma warning disable SYSLIB0014
+
 namespace AutomationTool.Tasks
 {
 	/// <summary>
diff --git a/Engine/Source/Programs/AutomationTool/Gauntlet/Unreal/BuildSource/Gauntlet.StagedBuildSource.cs b/Engine/Source/Programs/AutomationTool/Gauntlet/Unreal/BuildSource/Gauntlet.StagedBuildSource.cs
--- a/Engine/Source/Programs/AutomationTool/Gauntlet/Unreal/BuildSource/Gauntlet.StagedBuildSource.cs
+++ b/Engine/Source/Programs/AutomationTool/Gauntlet/Unreal/BuildSource/Gauntlet.StagedBuildSource.cs
@@ -156,14 +156,14 @@ namespace Gauntlet
 					Log.Verbose("\tCopying {0} to {1}", BuildPath, DestPath);
 					Utils.SystemHelpers.CopyDirectory(BuildPath, DestPath, Utils.SystemHelpers.CopyOptions.Mirror);
 				}
-				catch (Exception Ex)
+				catch
 				{
 					lock (Globals.MainLock)
 					{
 						LocalInstalls[InBuild].Status = InstallStatus.Error;
 					}
 
-					throw Ex;
+					throw;
 				}
 
 				lock (Globals.MainLock)
diff --git a/Engine/Source/Programs/AutomationTool/Scripts/RebuildHLODCommand.Automation.cs b/Engine/Source/Programs/AutomationTool/Scripts/RebuildHLODCommand.Automation.cs
--- a/Engine/Source/Programs/AutomationTool/Scripts/RebuildHLODCommand.Automation.cs
+++ b/Engine/Source/Programs/AutomationTool/Scripts/RebuildHLODCommand.Automation.cs
@@ -67,7 +67,7 @@ namespace AutomationScripts.Automation
                 LogInformation("********** REBUILD HLODS COMMAND FAILED **********");
                 LogInformation("Error message: {0}", ProcessEx.Message);
                 HandleFailure(ProcessEx.Message);
-                throw ProcessEx;
+                throw;
             }
 
             // The processes steps have completed successfully.
diff --git a/Engine/Source/Programs/AutomationTool/Scripts/RebuildLightMapsCommand.Automation.cs b/Engine/Source/Programs/AutomationTool/Scripts/RebuildLightMapsCommand.Automation.cs
--- a/Engine/Source/Programs/AutomationTool/Scripts/RebuildLightMapsCommand.Automation.cs
+++ b/Engine/Source/Programs/AutomationTool/Scripts/RebuildLightMapsCommand.Automation.cs
@@ -55,7 +55,7 @@ namespace AutomationScripts.Automation
 				LogInformation("********** REBUILD LIGHT MAPS COMMAND FAILED **********");
                 LogInformation("Error message: {0}", ProcessEx.Message);
 				HandleFailure(ProcessEx.Message);
-				throw ProcessEx;
+				throw;
 			}
 
 			// The processes steps have completed successfully.
diff --git a/Engine/Source/Programs/AutomationTool/Scripts/ResavePackagesCommand.Automation.cs b/Engine/Source/Programs/AutomationTool/Scripts/ResavePackagesCommand.Automation.cs
--- a/Engine/Source/Programs/AutomationTool/Scripts/ResavePackagesCommand.Automation.cs
+++ b/Engine/Source/Programs/AutomationTool/Scripts/ResavePackagesCommand.Automation.cs
@@ -57,7 +57,7 @@ namespace AutomationScripts.Automation
 				LogInformation("********** RESAVE PACKAGES COMMAND FAILED **********");
                 LogInformation("Error message: {0}", ProcessEx.Message);
 				HandleFailure(ProcessEx.Message);
-				throw ProcessEx;
+				throw;
 			}
 
 			// The processes steps have completed successfully.
diff --git a/Engine/Source/Programs/AutomationTool/Turnkey/Providers/HttpCopyProvider.cs b/Engine/Source/Programs/AutomationTool/Turnkey/Providers/HttpCopyProvider.cs
--- a/Engine/Source/Programs/AutomationTool/Turnkey/Providers/HttpCopyProvider.cs
+++ b/Engine/Source/Programs/AutomationTool/Turnkey/Providers/HttpCopyProvider.cs
@@ -8,6 +8,8 @@ using System.Text.RegularExpressions;
 using EpicGames.Core;
 using System.Net;
 
+#pragma warning disable SYSLIB0014
+
 namespace Turnkey
 {
 	// there is an Https subclass below - so the operation will supply http: or https: and this class will handle both, by using the ProviderToken for Uris
diff --git a/Engine/Source/Programs/AutomationTool/XLocLocalization/XLocLocalizationProvider.cs b/Engine/Source/Programs/AutomationTool/XLocLocalization/XLocLocalizationProvider.cs
--- a/Engine/Source/Programs/AutomationTool/XLocLocalization/XLocLocalizationProvider.cs
+++ b/Engine/Source/Programs/AutomationTool/XLocLocalization/XLocLocalizationProvider.cs
@@ -13,6 +13,8 @@ using UnrealBuildTool;
 using EpicGames.Localization;
 using System.Threading.Tasks;
 
+#pragma warning disable SYSLIB0014
+
 namespace EpicGames.XLocLocalization
 {
 	public struct XLocConfig
@@ -367,7 +369,7 @@ namespace EpicGames.XLocLocalization
 				{
 					return RequestAuthToken(XLocApiClient);
 				}
-				catch (Exception Ex)
+				catch
 				{
 					if (++Count < MAX_COUNT)
 					{
@@ -376,7 +378,7 @@ namespace EpicGames.XLocLocalization
 					}
 
 					BuildCommand.LogWarning("RequestAuthToken attempt {0}/{1} failed.", Count, MAX_COUNT);
-					throw Ex;
+					throw;
 				}
 			}
 		}
diff --git a/Engine/Source/Programs/UnrealBuildTool/Configuration/UEBuildBinary.cs b/Engine/Source/Programs/UnrealBuildTool/Configuration/UEBuildBinary.cs
--- a/Engine/Source/Programs/UnrealBuildTool/Configuration/UEBuildBinary.cs
+++ b/Engine/Source/Programs/UnrealBuildTool/Configuration/UEBuildBinary.cs
@@ -605,7 +605,7 @@ namespace UnrealBuildTool
 			foreach (FileReference OutputFilePath in OutputFilePaths)
 			{
 				// Find the base directory for this binary
-				DirectoryReference BaseDir = RootDirectories.FirstOrDefault(x => OutputFilePath.IsUnderDirectory(x));
+				DirectoryReference? BaseDir = RootDirectories.FirstOrDefault(x => OutputFilePath.IsUnderDirectory(x));
 				if (BaseDir == null)
 				{
 					continue;
diff --git a/Engine/Source/Programs/UnrealBuildTool/Configuration/UEBuildModule.cs b/Engine/Source/Programs/UnrealBuildTool/Configuration/UEBuildModule.cs
--- a/Engine/Source/Programs/UnrealBuildTool/Configuration/UEBuildModule.cs
+++ b/Engine/Source/Programs/UnrealBuildTool/Configuration/UEBuildModule.cs
@@ -360,7 +360,7 @@ namespace UnrealBuildTool
 		/// <returns>An enumerable containing the dependencies of the module.</returns>
 		public IEnumerable<UEBuildModule> GetDirectDependencyModules()
 		{
-			return PublicDependencyModules.Concat(PrivateDependencyModules).Concat(DynamicallyLoadedModules);
+			return PublicDependencyModules!.Concat(PrivateDependencyModules!).Concat(DynamicallyLoadedModules!);
 		}
 
 		/// <summary>
@@ -432,14 +432,14 @@ namespace UnrealBuildTool
 
 				// Remove all the allow listed folders
 				ReferencedDirs.ExceptWith(RestrictedFoldersAllowList);
-				ReferencedDirs.ExceptWith(PublicDependencyModules.SelectMany(x => x.RestrictedFoldersAllowList));
-				ReferencedDirs.ExceptWith(PrivateDependencyModules.SelectMany(x => x.RestrictedFoldersAllowList));
+				ReferencedDirs.ExceptWith(PublicDependencyModules!.SelectMany(x => x.RestrictedFoldersAllowList));
+				ReferencedDirs.ExceptWith(PrivateDependencyModules!.SelectMany(x => x.RestrictedFoldersAllowList));
 
 				// Add flags for each of them
 				foreach(DirectoryReference ReferencedDir in ReferencedDirs)
 				{
 					// Find the base directory containing this reference
-					DirectoryReference BaseDir = RootDirectories.FirstOrDefault(x => ReferencedDir.IsUnderDirectory(x));
+					DirectoryReference? BaseDir = RootDirectories.FirstOrDefault(x => ReferencedDir.IsUnderDirectory(x));
 					// @todo platplug does this need to check platform extension engine directories? what are ReferencedDir's here?
 					if (BaseDir == null)
 					{
@@ -896,8 +896,8 @@ namespace UnrealBuildTool
 		public List<UEBuildModule> GetUnboundReferences()
 		{
 			List<UEBuildModule> Modules = new List<UEBuildModule>();
-			Modules.AddRange(PrivateDependencyModules.Where(x => x.Binary == null));
-			Modules.AddRange(PublicDependencyModules.Where(x => x.Binary == null));
+			Modules.AddRange(PrivateDependencyModules!.Where(x => x.Binary == null));
+			Modules.AddRange(PublicDependencyModules!.Where(x => x.Binary == null));
 			return Modules;
 		}
 
diff --git a/Engine/Source/Programs/UnrealBuildTool/Configuration/UEBuildModuleCPP.cs b/Engine/Source/Programs/UnrealBuildTool/Configuration/UEBuildModuleCPP.cs
--- a/Engine/Source/Programs/UnrealBuildTool/Configuration/UEBuildModuleCPP.cs
+++ b/Engine/Source/Programs/UnrealBuildTool/Configuration/UEBuildModuleCPP.cs
@@ -1158,7 +1158,7 @@ namespace UnrealBuildTool
 					GetAllDependencyModules(new List<UEBuildModule>(), ReferencedModules, bIncludeDynamicallyLoaded: false, bForceCircular: false, bOnlyDirectDependencies: true);
 
 					// Find the first shared PCH module we can use
-					PrecompiledHeaderTemplate Template = CompileEnvironment.SharedPCHs.FirstOrDefault(x => ReferencedModules.Contains(x.Module));
+					PrecompiledHeaderTemplate? Template = CompileEnvironment.SharedPCHs.FirstOrDefault(x => ReferencedModules.Contains(x.Module));
 					if(Template != null && Template.IsValidFor(CompileEnvironment))
 					{
 						PrecompiledHeaderInstance Instance = FindOrCreateSharedPCH(ToolChain, Template, CompileEnvironment, Graph);
diff --git a/Engine/Source/Programs/UnrealBuildTool/Configuration/UEBuildTarget.cs b/Engine/Source/Programs/UnrealBuildTool/Configuration/UEBuildTarget.cs
--- a/Engine/Source/Programs/UnrealBuildTool/Configuration/UEBuildTarget.cs
+++ b/Engine/Source/Programs/UnrealBuildTool/Configuration/UEBuildTarget.cs
@@ -289,7 +289,7 @@ namespace UnrealBuildTool
 
 	internal class UnrealTargetPlatformTypeConverter : TypeConverter
 	{
-		public override bool CanConvertFrom(ITypeDescriptorContext context, Type sourceType)
+		public override bool CanConvertFrom(ITypeDescriptorContext? context, Type sourceType)
 		{
 			if (sourceType == typeof(string))
 				return true;
@@ -297,7 +297,7 @@ namespace UnrealBuildTool
 			return base.CanConvertFrom(context, sourceType);
 		}
 
-		public override bool CanConvertTo(ITypeDescriptorContext context, Type destinationType)
+		public override bool CanConvertTo(ITypeDescriptorContext? context, Type? destinationType)
 		{
 			if (destinationType == typeof(string))
 				return true;
@@ -305,7 +305,7 @@ namespace UnrealBuildTool
 			return base.CanConvertTo(context, destinationType);
 		}
 
-		public override object ConvertFrom(ITypeDescriptorContext context, CultureInfo culture, object value)
+		public override object? ConvertFrom(ITypeDescriptorContext? context, CultureInfo? culture, object value)
 		{
 			if (value.GetType() == typeof(string))
 			{
@@ -314,9 +314,9 @@ namespace UnrealBuildTool
 			return base.ConvertFrom(context, culture, value);
 		}
 
-		public override object ConvertTo(ITypeDescriptorContext context, CultureInfo culture, object value, Type destinationType)
+		public override object? ConvertTo(ITypeDescriptorContext? context, CultureInfo? culture, object? value, Type destinationType)
 		{
-			if (destinationType == typeof(string))
+			if (destinationType == typeof(string) && value != null)
 			{
 				UnrealTargetPlatform Platform = (UnrealTargetPlatform)value;
 				return Platform.ToString();
@@ -1725,7 +1725,7 @@ namespace UnrealBuildTool
 			// For installed builds, filter out all the binaries that aren't in mods
 			if (UnrealBuildTool.IsProjectInstalled())
 			{
-				List<DirectoryReference> ModDirectories = EnabledPlugins.Where(x => x.Type == PluginType.Mod).Select(x => x.Directory).ToList();
+				List<DirectoryReference> ModDirectories = EnabledPlugins!.Where(x => x.Type == PluginType.Mod).Select(x => x.Directory).ToList();
 
 				List<UEBuildBinary> FilteredBinaries = new List<UEBuildBinary>();
 				foreach (UEBuildBinary DLLBinary in Binaries)
@@ -2675,7 +2675,7 @@ namespace UnrealBuildTool
 			{
 				PreBuildCommandBatches.Add(new Tuple<string[], UEBuildPlugin?>(Rules.PreBuildSteps.ToArray(), null));
 			}
-			foreach (UEBuildPlugin BuildPlugin in BuildPlugins.Where(x => x.Descriptor.PreBuildSteps != null))
+			foreach (UEBuildPlugin BuildPlugin in BuildPlugins!.Where(x => x.Descriptor.PreBuildSteps != null))
 			{
 				AddCustomBuildSteps(BuildPlugin.Descriptor.PreBuildSteps!, BuildPlugin, PreBuildCommandBatches);
 			}
@@ -2700,7 +2700,7 @@ namespace UnrealBuildTool
 				{
 					PostBuildCommandBatches.Add(new Tuple<string[], UEBuildPlugin?>(Rules.PostBuildSteps.ToArray(), null));
 				}
-				foreach (UEBuildPlugin BuildPlugin in BuildPlugins.Where(x => x.Descriptor.PostBuildSteps != null))
+				foreach (UEBuildPlugin BuildPlugin in BuildPlugins!.Where(x => x.Descriptor.PostBuildSteps != null))
 				{
 					AddCustomBuildSteps(BuildPlugin.Descriptor.PostBuildSteps!, BuildPlugin, PostBuildCommandBatches);
 				}
@@ -2816,7 +2816,7 @@ namespace UnrealBuildTool
 					List<UEBuildModule> GameModules = Binary.FindHotReloadModules();
 					if (GameModules != null && GameModules.Count > 0)
 					{
-						if (!UnrealBuildTool.IsProjectInstalled() || EnabledPlugins.Where(x => x.Type == PluginType.Mod).Any(x => Binary.OutputFilePaths[0].IsUnderDirectory(x.Directory)))
+						if (!UnrealBuildTool.IsProjectInstalled() || EnabledPlugins!.Where(x => x.Type == PluginType.Mod).Any(x => Binary.OutputFilePaths[0].IsUnderDirectory(x.Directory)))
 						{
 							HotReloadModuleNames.UnionWith(GameModules.OfType<UEBuildModuleCPP>().Where(x => !x.Rules.bUsePrecompiled).Select(x => x.Name));
 						}
@@ -4112,7 +4112,7 @@ namespace UnrealBuildTool
 			// Reads additional dependencies array for project module from project file and fills PrivateDependencyModuleNames.
 			if (ProjectDescriptor != null && ProjectDescriptor.Modules != null)
 			{
-				ModuleDescriptor Module = ProjectDescriptor.Modules.FirstOrDefault(x => x.Name.Equals(ModuleName, StringComparison.InvariantCultureIgnoreCase));
+				ModuleDescriptor? Module = ProjectDescriptor.Modules.FirstOrDefault(x => x.Name.Equals(ModuleName, StringComparison.InvariantCultureIgnoreCase));
 				if (Module != null && Module.AdditionalDependencies != null)
 				{
 					RulesObject.PrivateDependencyModuleNames.AddRange(Module.AdditionalDependencies);
@@ -4148,13 +4148,13 @@ namespace UnrealBuildTool
 				if (RulesObject.PCHUsage != ModuleRules.PCHUsageMode.NoPCHs && RulesObject.PCHUsage != ModuleRules.PCHUsageMode.UseExplicitOrSharedPCHs && RulesObject.PrivatePCHHeaderFile == null)
 				{
 					// Try to figure out the legacy PCH file
-					FileReference CppFile = DirectoryReference.EnumerateFiles(RulesObject.Directory, "*.cpp", SearchOption.AllDirectories).FirstOrDefault();
+					FileReference? CppFile = DirectoryReference.EnumerateFiles(RulesObject.Directory, "*.cpp", SearchOption.AllDirectories).FirstOrDefault();
 					if (CppFile != null)
 					{
 						string? IncludeFile = MetadataCache.GetFirstInclude(FileItem.GetItemByFileReference(CppFile));
 						if (IncludeFile != null)
 						{
-							FileReference PchIncludeFile = DirectoryReference.EnumerateFiles(RulesObject.Directory, Path.GetFileName(IncludeFile), SearchOption.AllDirectories).FirstOrDefault();
+							FileReference? PchIncludeFile = DirectoryReference.EnumerateFiles(RulesObject.Directory, Path.GetFileName(IncludeFile), SearchOption.AllDirectories).FirstOrDefault();
 							if (PchIncludeFile != null)
 							{
 								RulesObject.PrivatePCHHeaderFile = PchIncludeFile.MakeRelativeTo(RulesObject.Directory).Replace(Path.DirectorySeparatorChar, '/');
diff --git a/Engine/Source/Programs/UnrealBuildTool/Executors/SNDBS.cs b/Engine/Source/Programs/UnrealBuildTool/Executors/SNDBS.cs
--- a/Engine/Source/Programs/UnrealBuildTool/Executors/SNDBS.cs
+++ b/Engine/Source/Programs/UnrealBuildTool/Executors/SNDBS.cs
@@ -76,10 +76,13 @@ namespace UnrealBuildTool
 				LocalProcess.StartInfo = new ProcessStartInfo(SNDBSUtilExe, $"-connected");
 				LocalProcess.OutputDataReceived += (Sender, Args) =>
 				{
-					Match Result = FindHost.Match(Args.Data);
-					if (Result.Success)
+					if (Args.Data != null)
 					{
-						BrokerHostName = Result.Groups[1].Value;
+						Match Result = FindHost.Match(Args.Data);
+						if (Result.Success)
+						{
+							BrokerHostName = Result.Groups[1].Value;
+						}
 					}
 				};
 				if (Utils.RunLocalProcess(LocalProcess) == 1 && BrokerHostName.Length > 0)
diff --git a/Engine/Source/Programs/UnrealBuildTool/Modes/GenerateProjectFilesMode.cs b/Engine/Source/Programs/UnrealBuildTool/Modes/GenerateProjectFilesMode.cs
--- a/Engine/Source/Programs/UnrealBuildTool/Modes/GenerateProjectFilesMode.cs
+++ b/Engine/Source/Programs/UnrealBuildTool/Modes/GenerateProjectFilesMode.cs
@@ -274,7 +274,7 @@ namespace UnrealBuildTool
 			// so now validate it
 			if (!string.IsNullOrEmpty(CandidateProjectPath))
 			{
-				FileReference CandidateProjectFile = new FileReference(CandidateProjectPath);
+				FileReference? CandidateProjectFile = new FileReference(CandidateProjectPath);
 
 				// if the path doesn't exist then check native paths (ueprojectdirs)
 				if (!FileReference.Exists(CandidateProjectFile))
diff --git a/Engine/Source/Programs/UnrealBuildTool/Platform/Android/AndroidToolChain.cs b/Engine/Source/Programs/UnrealBuildTool/Platform/Android/AndroidToolChain.cs
--- a/Engine/Source/Programs/UnrealBuildTool/Platform/Android/AndroidToolChain.cs
+++ b/Engine/Source/Programs/UnrealBuildTool/Platform/Android/AndroidToolChain.cs
@@ -2176,7 +2176,7 @@ namespace UnrealBuildTool
 			{
 				ObjectFileDirectories.Add(new DirectoryReference(LibraryPath));
 			}
-			Directory.CreateDirectory(Path.GetDirectoryName(FileName));
+			Directory.CreateDirectory(Path.GetDirectoryName(FileName)!);
 			File.WriteAllLines(FileName, ObjectFileDirectories.Select(x => x.FullName).OrderBy(x => x).ToArray());
 		}
 
diff --git a/Engine/Source/Programs/UnrealBuildTool/Platform/IOS/IOSToolChain.cs b/Engine/Source/Programs/UnrealBuildTool/Platform/IOS/IOSToolChain.cs
--- a/Engine/Source/Programs/UnrealBuildTool/Platform/IOS/IOSToolChain.cs
+++ b/Engine/Source/Programs/UnrealBuildTool/Platform/IOS/IOSToolChain.cs
@@ -2197,7 +2197,7 @@ namespace UnrealBuildTool
 			PSI.RedirectStandardOutput = true;
 			PSI.UseShellExecute = false;
 			PSI.CreateNoWindow = true;
-			Process NewProcess = Process.Start(PSI);
+			Process NewProcess = Process.Start(PSI)!;
 
 			// Wait for the process to exit and grab it's output
 			StdOutResults = NewProcess.StandardOutput.ReadToEnd();
diff --git a/Engine/Source/Programs/UnrealBuildTool/Platform/Linux/LinuxToolChain.cs b/Engine/Source/Programs/UnrealBuildTool/Platform/Linux/LinuxToolChain.cs
--- a/Engine/Source/Programs/UnrealBuildTool/Platform/Linux/LinuxToolChain.cs
+++ b/Engine/Source/Programs/UnrealBuildTool/Platform/Linux/LinuxToolChain.cs
@@ -366,15 +366,18 @@ namespace UnrealBuildTool
 						// read just the first string
 						string? VersionString = Proc.StandardOutput.ReadLine();
 
-						Regex VersionPattern = new Regex("version \\d+(\\.\\d+)+");
-						Match VersionMatch = VersionPattern.Match(VersionString);
-
-						// version match will be like "version 3.3", so remove the "version"
-						if (VersionMatch.Value.StartsWith("version "))
+						if (VersionString != null)
 						{
-							CompilerVersionString = VersionMatch.Value.Replace("version ", "");
+							Regex VersionPattern = new Regex("version \\d+(\\.\\d+)+");
+							Match VersionMatch = VersionPattern.Match(VersionString);
+
+							// version match will be like "version 3.3", so remove the "version"
+							if (VersionMatch.Value.StartsWith("version "))
+							{
+								CompilerVersionString = VersionMatch.Value.Replace("version ", "");
 
-							DetermineCompilerMajMinPatchFromVersionString();
+								DetermineCompilerMajMinPatchFromVersionString();
+							}
 						}
 					}
 				}
@@ -803,8 +806,8 @@ namespace UnrealBuildTool
 		internal static string EscapeArgument(string definition)
 		{
 			string[] splitData = definition.Split('=');
-			string myKey = splitData.ElementAtOrDefault(0);
-			string myValue = splitData.ElementAtOrDefault(1);
+			string? myKey = splitData.ElementAtOrDefault(0);
+			string? myValue = splitData.ElementAtOrDefault(1);
 
 			if (string.IsNullOrEmpty(myKey)) { return ""; }
 			if (!string.IsNullOrEmpty(myValue))
@@ -1861,7 +1864,7 @@ namespace UnrealBuildTool
 			string LinkScriptName = string.Format((bUseCmdExe ? "Link-{0}.link.bat" : "Link-{0}.link.sh"), OutputFile.Location.GetFileName());
 			string LinkScriptFullPath = Path.Combine(LinkEnvironment.LocalShadowDirectory!.FullName, LinkScriptName);
 			Log.TraceVerbose("Creating link script: {0}", LinkScriptFullPath);
-			Directory.CreateDirectory(Path.GetDirectoryName(LinkScriptFullPath));
+			Directory.CreateDirectory(Path.GetDirectoryName(LinkScriptFullPath)!);
 			using (StreamWriter LinkWriter = File.CreateText(LinkScriptFullPath))
 			{
 				if (bUseCmdExe)
@@ -1931,7 +1934,7 @@ namespace UnrealBuildTool
 					{
 						bHasWipedFixDepsScript = true;
 						Log.TraceVerbose("Creating script: {0}", FixDepsScriptPath);
-						Directory.CreateDirectory(Path.GetDirectoryName(FixDepsScriptPath));
+						Directory.CreateDirectory(Path.GetDirectoryName(FixDepsScriptPath)!);
 						using (StreamWriter Writer = File.CreateText(FixDepsScriptPath))
 						{
 						if (bUseCmdExe)
@@ -2022,7 +2025,7 @@ namespace UnrealBuildTool
 					string RelinkScriptFullPath = Path.Combine(LinkEnvironment.LocalShadowDirectory.FullName, RelinkScriptName);
 
 					Log.TraceVerbose("Creating script: {0}", RelinkScriptFullPath);
-					Directory.CreateDirectory(Path.GetDirectoryName(RelinkScriptFullPath));
+					Directory.CreateDirectory(Path.GetDirectoryName(RelinkScriptFullPath)!);
 					using (StreamWriter RelinkWriter = File.CreateText(RelinkScriptFullPath))
 					{
 					string RelinkInvocation = LinkCommandString;
diff --git a/Engine/Source/Programs/UnrealBuildTool/Platform/Mac/ApplePlatformSDK.cs b/Engine/Source/Programs/UnrealBuildTool/Platform/Mac/ApplePlatformSDK.cs
--- a/Engine/Source/Programs/UnrealBuildTool/Platform/Mac/ApplePlatformSDK.cs
+++ b/Engine/Source/Programs/UnrealBuildTool/Platform/Mac/ApplePlatformSDK.cs
@@ -39,6 +39,13 @@ namespace UnrealBuildTool
 
 		public override bool TryConvertVersionToInt(string? StringValue, out UInt64 OutValue)
 		{
+			OutValue = 0;
+
+			if (StringValue == null)
+			{
+				return false;
+			}
+
 			// 8 bits per component, with high getting extra from high 32
 			Match Result = Regex.Match(StringValue, @"^(\d+).(\d+)(.(\d+))?(.(\d+))?$");
 			if (Result.Success)
@@ -54,7 +61,7 @@ namespace UnrealBuildTool
 				}
 				return true;
 			}
-			OutValue = 0;
+
 			return false;
 		}
 
diff --git a/Engine/Source/Programs/UnrealBuildTool/Preprocessor/Token.cs b/Engine/Source/Programs/UnrealBuildTool/Preprocessor/Token.cs
--- a/Engine/Source/Programs/UnrealBuildTool/Preprocessor/Token.cs
+++ b/Engine/Source/Programs/UnrealBuildTool/Preprocessor/Token.cs
@@ -384,7 +384,7 @@ namespace UnrealBuildTool
 			}
 			else 
 			{
-				return Enumerable.SequenceEqual(Literal, Other.Literal);
+				return Other.Literal != null && Enumerable.SequenceEqual(Literal, Other.Literal);
 			}
 		}
 
diff --git a/Engine/Source/Programs/UnrealBuildTool/ProjectFiles/Rider/RiderProjectFile.cs b/Engine/Source/Programs/UnrealBuildTool/ProjectFiles/Rider/RiderProjectFile.cs
--- a/Engine/Source/Programs/UnrealBuildTool/ProjectFiles/Rider/RiderProjectFile.cs
+++ b/Engine/Source/Programs/UnrealBuildTool/ProjectFiles/Rider/RiderProjectFile.cs
@@ -401,7 +401,7 @@ namespace UnrealBuildTool
 			
 			Writer.WriteValue("File", Plugin.File.FullName );
 			Writer.WriteValue("Type", Plugin.Type.ToString());
-			if(Plugin.Dependencies.Any())
+			if(Plugin.Dependencies != null && Plugin.Dependencies.Any())
 			{
 				Writer.WriteStringArrayField("Dependencies", Plugin.Dependencies.Select(it => it.Name));
 			}
@@ -421,7 +421,7 @@ namespace UnrealBuildTool
 		private static void ExportPluginsFromTarget(UEBuildTarget Target, JsonWriter Writer)
 		{
 			Target.SetupPlugins();
-			if (!Target.BuildPlugins.Any()) return;
+			if (Target.BuildPlugins == null || !Target.BuildPlugins.Any()) return;
 			
 			Writer.WriteObjectStart("Plugins");
 			foreach (UEBuildPlugin plugin in Target.BuildPlugins!)
diff --git a/Engine/Source/Programs/UnrealBuildTool/ProjectFiles/Rider/RiderProjectFileGenerator.cs b/Engine/Source/Programs/UnrealBuildTool/ProjectFiles/Rider/RiderProjectFileGenerator.cs
--- a/Engine/Source/Programs/UnrealBuildTool/ProjectFiles/Rider/RiderProjectFileGenerator.cs
+++ b/Engine/Source/Programs/UnrealBuildTool/ProjectFiles/Rider/RiderProjectFileGenerator.cs
@@ -208,7 +208,7 @@ namespace UnrealBuildTool
 				{
 					RulesAssembly RulesAssembly;
 
-					FileReference CheckProjectFile =
+					FileReference? CheckProjectFile =
 						AllGames.FirstOrDefault(x => TargetFilePath.IsUnderDirectory(x.Directory));
 					if (CheckProjectFile == null)
 					{
diff --git a/Engine/Source/Programs/UnrealBuildTool/ProjectFiles/Xcode/XcodeProject.cs b/Engine/Source/Programs/UnrealBuildTool/ProjectFiles/Xcode/XcodeProject.cs
--- a/Engine/Source/Programs/UnrealBuildTool/ProjectFiles/Xcode/XcodeProject.cs
+++ b/Engine/Source/Programs/UnrealBuildTool/ProjectFiles/Xcode/XcodeProject.cs
@@ -1364,7 +1364,7 @@ namespace UnrealBuildTool
 
 						if (bCreateMacInfoPlist)
 						{
-							Directory.CreateDirectory(Path.GetDirectoryName(MacInfoPlistPath));
+							Directory.CreateDirectory(Path.GetDirectoryName(MacInfoPlistPath)!);
 							UEDeployMac.GeneratePList(ProjectPath.FullName, bIsUnrealGame, GameName, Config.BuildTarget, EngineDir.FullName, MacExecutableFileName);
 						}
 						if (bCreateIOSInfoPlist)
@@ -1379,7 +1379,7 @@ namespace UnrealBuildTool
 							{
 								ReceiptFilename = TargetReceipt.GetDefaultPath(ProjectPath, GameName, UnrealTargetPlatform.IOS, Config.BuildConfig, "");
 							}
-							Directory.CreateDirectory(Path.GetDirectoryName(IOSInfoPlistPath));
+							Directory.CreateDirectory(Path.GetDirectoryName(IOSInfoPlistPath)!);
                             bool bSupportPortrait, bSupportLandscape;
 							TargetReceipt? Receipt;
 							TargetReceipt.TryRead(ReceiptFilename, out Receipt);
@@ -1388,7 +1388,7 @@ namespace UnrealBuildTool
 						}
 						if (bCreateTVOSInfoPlist)
 						{
-							Directory.CreateDirectory(Path.GetDirectoryName(TVOSInfoPlistPath));
+							Directory.CreateDirectory(Path.GetDirectoryName(TVOSInfoPlistPath)!);
 							UEDeployTVOS.GenerateTVOSPList(ProjectPath.FullName, bIsUnrealGame, GameName, bIsClient, Config.BuildTarget, EngineDir.FullName, ProjectPath + "/Binaries/TVOS/Payload", null, BundleIdentifier);
 						}
 					}
diff --git a/Engine/Source/Programs/UnrealBuildTool/System/ActionHistory.cs b/Engine/Source/Programs/UnrealBuildTool/System/ActionHistory.cs
--- a/Engine/Source/Programs/UnrealBuildTool/System/ActionHistory.cs
+++ b/Engine/Source/Programs/UnrealBuildTool/System/ActionHistory.cs
@@ -114,7 +114,7 @@ namespace UnrealBuildTool
 		{
 			string InvariantText = Text.ToUpperInvariant();
 			byte[] InvariantBytes = Encoding.Unicode.GetBytes(InvariantText);
-			return new MD5CryptoServiceProvider().ComputeHash(InvariantBytes);
+			return MD5.Create().ComputeHash(InvariantBytes);
 		}
 
 		/// <summary>
@@ -275,7 +275,7 @@ namespace UnrealBuildTool
 		{
 			FileReference LayerLocation = GetLayerLocationForFile(File.Location);
 
-			ActionHistoryLayer Layer = Layers.FirstOrDefault(x => x.Location == LayerLocation);
+			ActionHistoryLayer? Layer = Layers.FirstOrDefault(x => x.Location == LayerLocation);
 			if (Layer == null)
 			{
 				lock (LockObject)
@@ -412,7 +412,7 @@ namespace UnrealBuildTool
 		{
 			lock (LockObject)
 			{
-				ActionHistoryPartition Partition = Partitions.FirstOrDefault(x => x.BaseDir == BaseDir);
+				ActionHistoryPartition? Partition = Partitions.FirstOrDefault(x => x.BaseDir == BaseDir);
 				if(Partition == null)
 				{
 					Partition = new ActionHistoryPartition(BaseDir);
diff --git a/Engine/Source/Programs/UnrealBuildTool/System/BuildHostPlatform.cs b/Engine/Source/Programs/UnrealBuildTool/System/BuildHostPlatform.cs
--- a/Engine/Source/Programs/UnrealBuildTool/System/BuildHostPlatform.cs
+++ b/Engine/Source/Programs/UnrealBuildTool/System/BuildHostPlatform.cs
@@ -120,7 +120,7 @@ namespace UnrealBuildTool
 			{
 				PID = Proc.Id;
 				Name = Proc.ProcessName;
-				Filename = Path.GetFullPath(Proc.MainModule.FileName);
+				Filename = Proc.MainModule?.FileName != null ? Path.GetFullPath(Proc.MainModule.FileName) : String.Empty;
 			}
 
 			/// <summary>
@@ -208,7 +208,10 @@ namespace UnrealBuildTool
 				{
 					foreach (ProcessModule Module in Proc.Modules.Cast<System.Diagnostics.ProcessModule>())
 					{
-						Modules.Add(Path.GetFullPath(Module.FileName));
+						if (Module.FileName != null)
+						{
+							Modules.Add(Path.GetFullPath(Module.FileName));
+						}
 					}
 				}
 			}
diff --git a/Engine/Source/Programs/UnrealBuildTool/System/ExternalExecution.cs b/Engine/Source/Programs/UnrealBuildTool/System/ExternalExecution.cs
--- a/Engine/Source/Programs/UnrealBuildTool/System/ExternalExecution.cs
+++ b/Engine/Source/Programs/UnrealBuildTool/System/ExternalExecution.cs
@@ -377,7 +377,7 @@ namespace UnrealBuildTool
 				}
 				if (RulesObject.Plugin != null)
 				{
-					ModuleDescriptor Module = RulesObject.Plugin.Descriptor.Modules.FirstOrDefault(x => x.Name == RulesObject.Name);
+					ModuleDescriptor? Module = RulesObject.Plugin.Descriptor.Modules?.FirstOrDefault(x => x.Name == RulesObject.Name);
 					if(Module != null)
 					{
 						return GetGameModuleTypeFromDescriptor(Module);
@@ -385,7 +385,7 @@ namespace UnrealBuildTool
 				}
 				if(ProjectDescriptor != null && ProjectDescriptor.Modules != null)
 				{
-					ModuleDescriptor Module = ProjectDescriptor.Modules.FirstOrDefault(x => x.Name == RulesObject.Name);
+					ModuleDescriptor? Module = ProjectDescriptor.Modules.FirstOrDefault(x => x.Name == RulesObject.Name);
 					if(Module != null)
 					{
 						return UHTModuleTypeExtensions.GameModuleTypeFromHostType(Module.Type) ?? UHTModuleType.GameRuntime;
@@ -405,7 +405,7 @@ namespace UnrealBuildTool
 				}
 				if (RulesObject.Plugin != null)
 				{
-					ModuleDescriptor Module = RulesObject.Plugin.Descriptor.Modules.FirstOrDefault(x => x.Name == RulesObject.Name);
+					ModuleDescriptor? Module = RulesObject.Plugin.Descriptor.Modules?.FirstOrDefault(x => x.Name == RulesObject.Name);
 					if (Module != null)
 					{
 						return GetEngineModuleTypeFromDescriptor(Module);
diff --git a/Engine/Source/Programs/UnrealBuildTool/System/HotReload.cs b/Engine/Source/Programs/UnrealBuildTool/System/HotReload.cs
--- a/Engine/Source/Programs/UnrealBuildTool/System/HotReload.cs
+++ b/Engine/Source/Programs/UnrealBuildTool/System/HotReload.cs
@@ -382,7 +382,7 @@ namespace UnrealBuildTool
 							FileReference? MainModuleFile;
 							try
 							{
-								MainModuleFile = new FileReference(RunningProcess!.MainModule.FileName);
+								MainModuleFile = new FileReference(RunningProcess!.MainModule!.FileName!);
 							}
 							catch
 							{
@@ -1307,7 +1307,7 @@ namespace UnrealBuildTool
 			{
 				Writer.WriteObjectStart();
 
-				IExternalAction LinkAction = Actions.FirstOrDefault(x => x.ActionType == ActionType.Link && x.ProducedItems.Any(y => y.HasExtension(".exe") || y.HasExtension(".dll")));
+				IExternalAction? LinkAction = Actions.FirstOrDefault(x => x.ActionType == ActionType.Link && x.ProducedItems.Any(y => y.HasExtension(".exe") || y.HasExtension(".dll")));
 				if(LinkAction != null)
 				{
 					Writer.WriteValue("LinkerPath", LinkAction.CommandPath.FullName);
@@ -1328,7 +1328,7 @@ namespace UnrealBuildTool
 				{
 					if(Action.ActionType == ActionType.Link)
 					{
-						FileItem OutputFile = Action.ProducedItems.FirstOrDefault(x => x.HasExtension(".exe") || x.HasExtension(".dll"));
+						FileItem? OutputFile = Action.ProducedItems.FirstOrDefault(x => x.HasExtension(".exe") || x.HasExtension(".dll"));
 						if(OutputFile != null && Action.PrerequisiteItems.Any(x => OriginalFileToPatchedFile.ContainsKey(x.Location)))
 						{
 							Writer.WriteObjectStart();
diff --git a/Engine/Source/Programs/UnrealBuildTool/System/Plugins.cs b/Engine/Source/Programs/UnrealBuildTool/System/Plugins.cs
--- a/Engine/Source/Programs/UnrealBuildTool/System/Plugins.cs
+++ b/Engine/Source/Programs/UnrealBuildTool/System/Plugins.cs
@@ -378,7 +378,7 @@ namespace UnrealBuildTool
 				{
 					foreach (ModuleDescriptor ChildModule in Child.Descriptor.Modules)
 					{
-						ModuleDescriptor ParentModule = Parent.Descriptor.Modules.FirstOrDefault(x => x.Name.Equals(ChildModule.Name) && x.Type == ChildModule.Type);
+						ModuleDescriptor? ParentModule = Parent.Descriptor.Modules.FirstOrDefault(x => x.Name.Equals(ChildModule.Name) && x.Type == ChildModule.Type);
 						if (ParentModule != null)
 						{
 							// merge allow/deny lists (if the parent had a list, and child didn't specify a list, just add the child platform to the parent list - for allow/deny lists!)
@@ -424,7 +424,7 @@ namespace UnrealBuildTool
 				{ 
 					foreach (PluginReferenceDescriptor ChildPluginReference in Child.Descriptor.Plugins)
 					{
-						PluginReferenceDescriptor ParentPluginReference = Parent.Descriptor.Plugins.FirstOrDefault(x => x.Name.Equals(ChildPluginReference.Name));
+						PluginReferenceDescriptor? ParentPluginReference = Parent.Descriptor.Plugins.FirstOrDefault(x => x.Name.Equals(ChildPluginReference.Name));
 						if (ParentPluginReference != null)
 						{
 							// we only need to explicitly list the platform in an allow list if the parent also had an allow list (otherwise, we could mistakenly remove all other platforms)
diff --git a/Engine/Source/Programs/UnrealBuildTool/System/TargetMakefile.cs b/Engine/Source/Programs/UnrealBuildTool/System/TargetMakefile.cs
--- a/Engine/Source/Programs/UnrealBuildTool/System/TargetMakefile.cs
+++ b/Engine/Source/Programs/UnrealBuildTool/System/TargetMakefile.cs
@@ -434,7 +434,7 @@ namespace UnrealBuildTool
 			using (GlobalTracer.Instance.BuildSpan("Checking makefile validity").StartActive())
 			{
 				// Check if the arguments are different
-				if(!Enumerable.SequenceEqual(Makefile.AdditionalArguments, Arguments))
+				if(!Enumerable.SequenceEqual(Makefile.AdditionalArguments!, Arguments))
 				{
 					ReasonNotLoaded = "command line arguments changed";
 					return null;
diff --git a/Engine/Source/Programs/UnrealBuildTool/System/XmlConfigData.cs b/Engine/Source/Programs/UnrealBuildTool/System/XmlConfigData.cs
--- a/Engine/Source/Programs/UnrealBuildTool/System/XmlConfigData.cs
+++ b/Engine/Source/Programs/UnrealBuildTool/System/XmlConfigData.cs
@@ -98,7 +98,7 @@ namespace UnrealBuildTool
 					string TypeName = Reader.ReadString();
 
 					// Try to find it in the list of configurable types
-					Type Type = Types.FirstOrDefault(x => x.Name == TypeName);
+					Type? Type = Types.FirstOrDefault(x => x.Name == TypeName);
 					if(Type == null)
 					{
 						Data = null;
diff --git a/Engine/Source/Programs/UnrealBuildTool/ToolChain/RemoteMac.cs b/Engine/Source/Programs/UnrealBuildTool/ToolChain/RemoteMac.cs
--- a/Engine/Source/Programs/UnrealBuildTool/ToolChain/RemoteMac.cs
+++ b/Engine/Source/Programs/UnrealBuildTool/ToolChain/RemoteMac.cs
@@ -994,7 +994,7 @@ namespace UnrealBuildTool
 		/// <param name="LocalFile">The file to download</param>
 		void DownloadFile(FileReference LocalFile)
 		{
-			RemoteMapping Mapping = Mappings.FirstOrDefault(x => LocalFile.IsUnderDirectory(x.LocalDirectory));
+			RemoteMapping? Mapping = Mappings.FirstOrDefault(x => LocalFile.IsUnderDirectory(x.LocalDirectory));
 			if(Mapping == null)
 			{
 				throw new BuildException("File for download '{0}' is not under any mapped directory.", LocalFile);
