LowLevelTests manual integration with UBT
#rnx

diff --git a/Engine/Source/Developer/LowLevelTestsRunner/Public/TestCommon/EditorUtilities.h b/Engine/Source/Developer/LowLevelTestsRunner/Public/TestCommon/EditorUtilities.h
--- a/Engine/Source/Developer/LowLevelTestsRunner/Public/TestCommon/EditorUtilities.h
+++ b/Engine/Source/Developer/LowLevelTestsRunner/Public/TestCommon/EditorUtilities.h
@@ -4,7 +4,10 @@
 
 #include "CoreMinimal.h"
 #include "Misc/QueuedThreadPool.h"
+
+#if WITH_COREUOBJECT
 #include "UObject/PackageResourceManager.h"
+#endif
 
 #if WITH_EDITORONLY_DATA && defined(DERIVEDDATACACHE_API)
 #include "DerivedDataBuild.h"
diff --git a/Engine/Source/Programs/UnrealBuildTool/Configuration/TestTargetRules.cs b/Engine/Source/Programs/UnrealBuildTool/Configuration/TestTargetRules.cs
--- a/Engine/Source/Programs/UnrealBuildTool/Configuration/TestTargetRules.cs
+++ b/Engine/Source/Programs/UnrealBuildTool/Configuration/TestTargetRules.cs
@@ -33,7 +33,7 @@ namespace UnrealBuildTool
 			DefaultBuildSettings = BuildSettingsVersion.V2;
 
 			ExeBinariesSubFolder = Name = TestedTarget.Name + "Tests";
-			TargetSourceFile = File = FileReference.Combine(UnrealBuildTool.EngineSourceDirectory, "UnrealBuildTool", "Configuration", "TestTargetRules.cs");
+			File = FileReference.Combine(UnrealBuildTool.EngineSourceDirectory, "UnrealBuildTool", "Configuration", "TestTargetRules.cs");
 			LaunchModuleName = TestedTarget.LaunchModuleName + "Tests";
 
 			ManifestFileNames = TestedTarget.ManifestFileNames;
