LLT TargetSourceFile - fixes incremental build
#rnx

diff --git a/Engine/Source/Programs/UnrealBuildTool/Configuration/TestTargetRules.cs b/Engine/Source/Programs/UnrealBuildTool/Configuration/TestTargetRules.cs
--- a/Engine/Source/Programs/UnrealBuildTool/Configuration/TestTargetRules.cs
+++ b/Engine/Source/Programs/UnrealBuildTool/Configuration/TestTargetRules.cs
@@ -33,7 +33,7 @@ namespace UnrealBuildTool
 			DefaultBuildSettings = BuildSettingsVersion.V2;
 
 			ExeBinariesSubFolder = Name = TestedTarget.Name + "Tests";
-			File = FileReference.Combine(UnrealBuildTool.EngineSourceDirectory, "UnrealBuildTool", "Configuration", "TestTargetRules.cs");
+			TargetSourceFile = File = FileReference.Combine(UnrealBuildTool.EngineSourceDirectory, "UnrealBuildTool", "Configuration", "TestTargetRules.cs");
 			LaunchModuleName = TestedTarget.LaunchModuleName + "Tests";
 
 			ManifestFileNames = TestedTarget.ManifestFileNames;
