UBT: Add -CrashDiagnosticDirectory to set where to put crash report files (for platforms that support it)

#rnx

diff --git a/Engine/Source/Programs/UnrealBuildTool/Configuration/TargetRules.cs b/Engine/Source/Programs/UnrealBuildTool/Configuration/TargetRules.cs
--- a/Engine/Source/Programs/UnrealBuildTool/Configuration/TargetRules.cs
+++ b/Engine/Source/Programs/UnrealBuildTool/Configuration/TargetRules.cs
@@ -1426,6 +1426,13 @@ namespace UnrealBuildTool
 		[XmlConfigFile(Category = "BuildConfiguration")]
 		public string? PackagePath = null;
 
+		/// <summary>
+		/// Directory where to put crash report files for platforms that support it
+		/// </summary>
+		[CommandLine("-CrashDiagnosticDirectory")]
+		[XmlConfigFile(Category = "BuildConfiguration")]
+		public string? CrashDiagnosticDirectory = null;
+
 		/// <summary>
 		/// Bundle version for Mac apps.
 		/// </summary>
@@ -2916,6 +2923,11 @@ namespace UnrealBuildTool
 			get { return Inner.PackagePath; }
 		}
 
+		public string? CrashDiagnosticDirectory
+		{
+			get { return Inner.CrashDiagnosticDirectory; }
+		}
+
 		public string? BundleVersion
 		{
 			get { return Inner.BundleVersion; }
diff --git a/Engine/Source/Programs/UnrealBuildTool/Configuration/UEBuildTarget.cs b/Engine/Source/Programs/UnrealBuildTool/Configuration/UEBuildTarget.cs
--- a/Engine/Source/Programs/UnrealBuildTool/Configuration/UEBuildTarget.cs
+++ b/Engine/Source/Programs/UnrealBuildTool/Configuration/UEBuildTarget.cs
@@ -3704,6 +3704,7 @@ namespace UnrealBuildTool
 			GlobalCompileEnvironment.bHideSymbolsByDefault = !Rules.bPublicSymbolsByDefault;
 			GlobalCompileEnvironment.CppStandard = Rules.CppStandard;
 			GlobalCompileEnvironment.AdditionalArguments = Rules.AdditionalCompilerArguments ?? String.Empty;
+			GlobalCompileEnvironment.CrashDiagnosticDirectory = Rules.CrashDiagnosticDirectory;
 
 			GlobalLinkEnvironment.bUseDebugCRT = bUseDebugCRT;
 			GlobalLinkEnvironment.bUseStaticCRT = Rules.bUseStaticCRT;
@@ -3715,6 +3716,7 @@ namespace UnrealBuildTool
 			GlobalLinkEnvironment.bHasExports = Rules.bHasExports;
 			GlobalLinkEnvironment.bUsePDBFiles = Rules.bUsePDBFiles;
 			GlobalLinkEnvironment.PackagePath = Rules.PackagePath;
+			GlobalLinkEnvironment.CrashDiagnosticDirectory = Rules.CrashDiagnosticDirectory;
 			GlobalLinkEnvironment.BundleDirectory = BuildPlatform.GetBundleDirectory(Rules, Binaries[0].OutputFilePaths);
 			GlobalLinkEnvironment.BundleVersion = Rules.BundleVersion;
 			GlobalLinkEnvironment.bAllowLTCG = Rules.bAllowLTCG;
diff --git a/Engine/Source/Programs/UnrealBuildTool/System/CppCompileEnvironment.cs b/Engine/Source/Programs/UnrealBuildTool/System/CppCompileEnvironment.cs
--- a/Engine/Source/Programs/UnrealBuildTool/System/CppCompileEnvironment.cs
+++ b/Engine/Source/Programs/UnrealBuildTool/System/CppCompileEnvironment.cs
@@ -415,6 +415,11 @@ namespace UnrealBuildTool
 		/// </summary>
 		public EngineIncludeOrderVersion IncludeOrderVersion = EngineIncludeOrderVersion.Latest;
 
+		/// <summary>
+		/// Directory where to put crash report files for platforms that support it
+		/// </summary>
+		public string? CrashDiagnosticDirectory;
+
 		/// <summary>
 		/// Default constructor.
 		/// </summary>
@@ -498,6 +503,7 @@ namespace UnrealBuildTool
 			CppStandard = Other.CppStandard;
 			bEnableCoroutines = Other.bEnableCoroutines;
 			IncludeOrderVersion = Other.IncludeOrderVersion;
+			CrashDiagnosticDirectory = Other.CrashDiagnosticDirectory;
 		}
 	}
 }
diff --git a/Engine/Source/Programs/UnrealBuildTool/System/LinkEnvironment.cs b/Engine/Source/Programs/UnrealBuildTool/System/LinkEnvironment.cs
--- a/Engine/Source/Programs/UnrealBuildTool/System/LinkEnvironment.cs
+++ b/Engine/Source/Programs/UnrealBuildTool/System/LinkEnvironment.cs
@@ -272,6 +272,11 @@ namespace UnrealBuildTool
 		/// </summary>
 		public string? PackagePath;
 
+		/// <summary>
+		/// Directory where to put crash report files for platforms that support it
+		/// </summary>
+		public string? CrashDiagnosticDirectory;
+
 		/// <summary>
 		/// Bundle version for Mac apps
 		/// </summary>
@@ -376,6 +381,7 @@ namespace UnrealBuildTool
 			bIgnoreUnresolvedSymbols = Other.bIgnoreUnresolvedSymbols;
 			bPrintTimingInfo = Other.bPrintTimingInfo;
 			PackagePath = Other.PackagePath;
+			CrashDiagnosticDirectory = Other.CrashDiagnosticDirectory;
 			BundleVersion = Other.BundleVersion;
 			InstallName = Other.InstallName;
 			InputFiles.AddRange(Other.InputFiles);
