UnrealBuildTool: Use module's ShortName for code generation directory to help avoid overly long paths

Before:
Intermediate/Build/Win64/UnrealEditor/Inc/MaterialShaderQualitySettings/

After:
Intermediate/Build/Win64/UnrealEditor/Inc/MSQS/

#rnx

diff --git a/Engine/Source/Programs/UnrealBuildTool/Configuration/UEBuildTarget.cs b/Engine/Source/Programs/UnrealBuildTool/Configuration/UEBuildTarget.cs
--- a/Engine/Source/Programs/UnrealBuildTool/Configuration/UEBuildTarget.cs
+++ b/Engine/Source/Programs/UnrealBuildTool/Configuration/UEBuildTarget.cs
@@ -4277,8 +4277,8 @@ namespace UnrealBuildTool
 						GeneratedCodeDirectory = DirectoryReference.Combine(GeneratedCodeDirectory, RulesObject.BinariesSubFolder);
 					}
 
-					// Finally, append the module name.
-					GeneratedCodeDirectory = DirectoryReference.Combine(GeneratedCodeDirectory, ModuleName);
+					// Finally, append the module name (using the ShortName if it has been set)
+					GeneratedCodeDirectory = DirectoryReference.Combine(GeneratedCodeDirectory, RulesObject.ShortName ?? ModuleName);
 				}
 
 				// For legacy modules, add a bunch of default include paths.
