UnrealBuildTool: Respect -NoPCH flag in mac toolchain

#rnx

diff --git a/Engine/Source/Programs/UnrealBuildTool/Platform/Mac/UEBuildMac.cs b/Engine/Source/Programs/UnrealBuildTool/Platform/Mac/UEBuildMac.cs
--- a/Engine/Source/Programs/UnrealBuildTool/Platform/Mac/UEBuildMac.cs
+++ b/Engine/Source/Programs/UnrealBuildTool/Platform/Mac/UEBuildMac.cs
@@ -160,7 +160,7 @@ namespace UnrealBuildTool
 			bool bCompilingMultipleArchitectures = Target.Architecture.Contains("+");
 			Target.bCompileISPC = !bCompilingForArm;
 
-			Target.bUsePCHFiles = !bCompilingMultipleArchitectures;
+			Target.bUsePCHFiles = Target.bUsePCHFiles && !bCompilingMultipleArchitectures;
 		}
 
 		/// <summary>
