UnrealBuildTool: Relative pathing for Mac & Linux toolchains

#rnx

diff --git a/Engine/Source/Programs/IncludeTool/IncludeTool/CompileEnvironment.cs b/Engine/Source/Programs/IncludeTool/IncludeTool/CompileEnvironment.cs
--- a/Engine/Source/Programs/IncludeTool/IncludeTool/CompileEnvironment.cs
+++ b/Engine/Source/Programs/IncludeTool/IncludeTool/CompileEnvironment.cs
@@ -281,7 +281,7 @@ namespace IncludeTool
 
 						// Parse a list of tokens
 						List<string> Tokens = new List<string>();
-						ParseArgumentTokens(CompilerType, Params, Tokens);
+						ParseArgumentTokens(CompilerType, Params, Tokens, BaseDir);
 
 						// Parse it into a list of options, removing any that don't apply
 						List<FileReference> SourceFiles = new List<FileReference>();
@@ -290,11 +290,21 @@ namespace IncludeTool
 						{
 							if(Tokens[Idx] == "/Fo" || Tokens[Idx] == "/Fp" || Tokens[Idx] == "-o")
 							{
-								OutputFiles.Add(new FileReference(Tokens[++Idx]));
+								string OutputPath = Tokens[++Idx];
+								if (!Path.IsPathRooted(OutputPath))
+								{
+									OutputPath = Path.Combine(BaseDir.FullName, OutputPath);
+								}
+								OutputFiles.Add(new FileReference(OutputPath));
 							}
 							else if(Tokens[Idx].StartsWith("/Fo") || Tokens[Idx].StartsWith("/Fp"))
 							{
-								OutputFiles.Add(new FileReference(Tokens[Idx].Substring(3)));
+								string OutputPath = Tokens[Idx].Substring(3);
+								if (!Path.IsPathRooted(OutputPath))
+								{
+									OutputPath = Path.Combine(BaseDir.FullName, OutputPath);
+								}
+								OutputFiles.Add(new FileReference(OutputPath));
 							}
 							else if (Tokens[Idx] == "/D" || Tokens[Idx] == "-D")
 							{
@@ -306,19 +316,39 @@ namespace IncludeTool
 							}
 							else if (Tokens[Idx] == "/I" || Tokens[Idx] == "-I")
 							{
-								Environment.IncludePaths.Add(new DirectoryReference(DirectoryReference.Combine(BaseDir, Tokens[++Idx].Replace("//", "/")).FullName.ToLowerInvariant()));
+								string IncludePath = Tokens[++Idx].Replace("//", "/");
+								if (!Path.IsPathRooted(IncludePath))
+								{
+									IncludePath = Path.Combine(BaseDir.FullName, IncludePath);
+								}
+								Environment.IncludePaths.Add(new DirectoryReference(IncludePath.ToLowerInvariant()));
 							}
 							else if (Tokens[Idx].StartsWith("-I"))
 							{
-								Environment.IncludePaths.Add(new DirectoryReference(DirectoryReference.Combine(BaseDir, Tokens[Idx].Substring(2).Replace("//", "/")).FullName.ToLowerInvariant()));
+								string IncludePath = Tokens[Idx].Substring(2).Replace("//", "/");
+								if (!Path.IsPathRooted(IncludePath))
+								{
+									IncludePath = Path.Combine(BaseDir.FullName, IncludePath);
+								}
+								Environment.IncludePaths.Add(new DirectoryReference(IncludePath.ToLowerInvariant()));
 							}
 							else if(Tokens[Idx].StartsWith("/FI"))
 							{
-								Environment.ForceIncludeFiles.Add(new FileReference(Tokens[Idx].Substring(3)));
+								string ForceIncludePath = Tokens[Idx].Substring(3);
+								if (!Path.IsPathRooted(ForceIncludePath))
+								{
+									ForceIncludePath = Path.Combine(BaseDir.FullName, ForceIncludePath);
+								}
+								Environment.ForceIncludeFiles.Add(new FileReference(ForceIncludePath));
 							}
 							else if(Tokens[Idx] == "-include")
 							{
-								Environment.ForceIncludeFiles.Add(new FileReference(Tokens[++Idx]));
+								string ForceIncludePath = Tokens[++Idx];
+								if (!Path.IsPathRooted(ForceIncludePath))
+								{
+									ForceIncludePath = Path.Combine(BaseDir.FullName, ForceIncludePath);
+								}
+								Environment.ForceIncludeFiles.Add(new FileReference(ForceIncludePath));
 							}
 							else if (Tokens[Idx] == "-Xclang" || Tokens[Idx] == "-target" || Tokens[Idx] == "--target" || Tokens[Idx] == "-x" || Tokens[Idx] == "-o")
 							{
@@ -387,13 +417,14 @@ namespace IncludeTool
 				return false;
 			}
 		}
-		
+
 		/// <summary>
 		/// Parse a parameter list into a series of tokens, reading response files as appropriate
 		/// </summary>
 		/// <param name="Text">The line of text to parse</param>
 		/// <param name="Tokens">List to be filled with the parsed tokens</param>
-		static void ParseArgumentTokens(CompilerType CompilerType, string Text, List<string> Tokens)
+		/// <param name="BaseDir">Base directory to root relative pathed response files from</param>
+		static void ParseArgumentTokens(CompilerType CompilerType, string Text, List<string> Tokens, DirectoryReference BaseDir)
 		{
 			for (int Idx = 0; Idx < Text.Length; Idx++)
 			{
@@ -425,7 +456,12 @@ namespace IncludeTool
 					// Add the token to the list. If it's a response file, recursively parse the tokens out of that.
 					if(Token[0] == '@')
 					{
-						ParseArgumentTokens(CompilerType, File.ReadAllText(Token.ToString().Substring(1)), Tokens);
+						string ResponsePath = Token.ToString().Substring(1);
+						if (!Path.IsPathRooted(ResponsePath))
+						{
+							ResponsePath = Path.Combine(BaseDir.FullName, ResponsePath);
+						}
+						ParseArgumentTokens(CompilerType, File.ReadAllText(ResponsePath), Tokens, BaseDir);
 					}
 					else
 					{
diff --git a/Engine/Source/Programs/UnrealBuildTool/Platform/Linux/LinuxToolChain.cs b/Engine/Source/Programs/UnrealBuildTool/Platform/Linux/LinuxToolChain.cs
--- a/Engine/Source/Programs/UnrealBuildTool/Platform/Linux/LinuxToolChain.cs
+++ b/Engine/Source/Programs/UnrealBuildTool/Platform/Linux/LinuxToolChain.cs
@@ -79,6 +79,8 @@ namespace UnrealBuildTool
 		/** Pass --gdb-index option to linker to generate .gdb_index section. */
 		protected bool bGdbIndexSection = true;
 
+		protected bool bPreprocessDepends = false;
+
 		/** Allows you to override the maximum binary size allowed to be passed to objcopy.exe when cross building on Windows. */
 		/** Max value is 2GB, due to bat file limitation */
 		protected UInt64 MaxBinarySizeOverrideForObjcopy = 0;
@@ -195,6 +197,13 @@ namespace UnrealBuildTool
 			bPreservePSYM = InPreservePSYM;
 		}
 
+		public override void SetUpGlobalEnvironment(ReadOnlyTargetRules Target)
+		{
+			base.SetUpGlobalEnvironment(Target);
+
+			bPreprocessDepends = Target.LinuxPlatform.bPreprocessDepends;
+		}
+
 		private string GetHostPlatformBinarySuffix()
 		{
 			if (BuildHostPlatform.Current.Platform == UnrealTargetPlatform.Win64)
@@ -1232,6 +1241,22 @@ namespace UnrealBuildTool
 			Log.TraceInformation("------------------------------");
 		}
 
+		public static string NormalizeCommandLinePath(FileSystemReference Reference)
+		{
+			// Try to use a relative path to shorten command line length.
+			if (Reference.IsUnderDirectory(Unreal.RootDirectory))
+			{
+				return Reference.MakeRelativeTo(UnrealBuildTool.EngineSourceDirectory).Replace("\\", "/");
+			}
+
+			return Reference.FullName.Replace("\\", "/");
+		}
+
+		public static string NormalizeCommandLinePath(FileItem Item)
+		{
+			return NormalizeCommandLinePath(Item.Location);
+		}
+
 		public override CPPOutput CompileCPPFiles(CppCompileEnvironment CompileEnvironment, List<FileItem> InputFiles, DirectoryReference OutputDir, string ModuleName, IActionGraphBuilder Graph)
 		{
 			string Arguments = GetCLArguments_Global(CompileEnvironment);
@@ -1253,26 +1278,17 @@ namespace UnrealBuildTool
 
 			if (CompileEnvironment.PrecompiledHeaderAction == PrecompiledHeaderAction.Include)
 			{
-				PCHArguments += string.Format(" -include \"{0}\"", CompileEnvironment.PrecompiledHeaderIncludeFilename!.FullName.Replace('\\', '/'));
+				PCHArguments += string.Format(" -include \"{0}\"", NormalizeCommandLinePath(CompileEnvironment.PrecompiledHeaderIncludeFilename!));
 			}
 
 			// Add include paths to the argument list.
 			foreach (DirectoryReference IncludePath in CompileEnvironment.UserIncludePaths)
 			{
-				string IncludePathString;
-				if (IncludePath.IsUnderDirectory(Unreal.RootDirectory))
-				{
-					IncludePathString = IncludePath.MakeRelativeTo(UnrealBuildTool.EngineSourceDirectory);
-				}
-				else
-				{
-					IncludePathString = IncludePath.FullName;
-				}
-				Arguments += string.Format(" -I\"{0}\"", IncludePathString.Replace('\\', '/'));
+				Arguments += string.Format(" -I\"{0}\"", NormalizeCommandLinePath(IncludePath));
 			}
 			foreach (DirectoryReference IncludePath in CompileEnvironment.SystemIncludePaths)
 			{
-				Arguments += string.Format(" -I\"{0}\"", IncludePath.FullName.Replace('\\', '/'));
+				Arguments += string.Format(" -I\"{0}\"", NormalizeCommandLinePath(IncludePath));
 			}
 
 			// Add preprocessor definitions to the argument list.
@@ -1323,7 +1339,7 @@ namespace UnrealBuildTool
 
 				foreach (FileItem ForceIncludeFile in CompileEnvironment.ForceIncludeFiles)
 				{
-					FileArguments += String.Format(" -include \"{0}\"", ForceIncludeFile.Location.FullName.Replace('\\', '/'));
+					FileArguments += String.Format(" -include \"{0}\"", NormalizeCommandLinePath(ForceIncludeFile));
 				}
 
 				// Add the C++ source file and its included files to the prerequisite item list.
@@ -1338,7 +1354,7 @@ namespace UnrealBuildTool
 					Result.PrecompiledHeaderFile = PrecompiledHeaderFile;
 
 					// Add the parameters needed to compile the precompiled header file to the command-line.
-					FileArguments += string.Format(" -o \"{0}\"", PrecompiledHeaderFile.AbsolutePath.Replace('\\', '/'));
+					FileArguments += string.Format(" -o \"{0}\"", NormalizeCommandLinePath(PrecompiledHeaderFile));
 				}
 				else
 				{
@@ -1352,11 +1368,11 @@ namespace UnrealBuildTool
 					CompileAction.ProducedItems.Add(ObjectFile);
 					Result.ObjectFiles.Add(ObjectFile);
 
-					FileArguments += string.Format(" -o \"{0}\"", ObjectFile.AbsolutePath.Replace('\\', '/'));
+					FileArguments += string.Format(" -o \"{0}\"", NormalizeCommandLinePath(ObjectFile));
 				}
 
 				// Add the source file path to the command-line.
-				FileArguments += string.Format(" \"{0}\"", SourceFile.AbsolutePath.Replace('\\', '/'));
+				FileArguments += string.Format(" \"{0}\"", NormalizeCommandLinePath(SourceFile));
 
 				// Generate the timing info
 				if (CompileEnvironment.bPrintTimingInfo)
@@ -1367,10 +1383,10 @@ namespace UnrealBuildTool
 				}
 
 				// Generate the included header dependency list
-				if(CompileEnvironment.bGenerateDependenciesFile)
+				if(!bPreprocessDepends && CompileEnvironment.bGenerateDependenciesFile)
 				{
 					FileItem DependencyListFile = FileItem.GetItemByFileReference(FileReference.Combine(OutputDir, Path.GetFileName(SourceFile.AbsolutePath) + ".d"));
-					FileArguments += string.Format(" -MD -MF\"{0}\"", DependencyListFile.AbsolutePath.Replace('\\', '/'));
+					FileArguments += string.Format(" -MD -MF\"{0}\"", NormalizeCommandLinePath(DependencyListFile));
 					CompileAction.DependencyListFile = DependencyListFile;
 					CompileAction.ProducedItems.Add(DependencyListFile);
 				}
@@ -1386,7 +1402,7 @@ namespace UnrealBuildTool
 				FileReference CompilerResponseFileName = CompileAction.ProducedItems[0].Location + ".rsp";
 				FileItem CompilerResponseFileItem = Graph.CreateIntermediateTextFile(CompilerResponseFileName, AllArguments);
 
-				CompileAction.CommandArguments = string.Format(" @\"{0}\"", CompilerResponseFileName);
+				CompileAction.CommandArguments = string.Format(" @\"{0}\"", NormalizeCommandLinePath(CompilerResponseFileName));
 				CompileAction.PrerequisiteItems.Add(CompilerResponseFileItem);
 				CompileAction.CommandDescription = "Compile";
 				CompileAction.CommandVersion = CompilerVersionString!;
@@ -1397,6 +1413,46 @@ namespace UnrealBuildTool
 				CompileAction.bCanExecuteRemotely =
 					CompileEnvironment.PrecompiledHeaderAction != PrecompiledHeaderAction.Create ||
 					CompileEnvironment.bAllowRemotelyCompiledPCHs;
+
+				if (bPreprocessDepends && CompileEnvironment.bGenerateDependenciesFile)
+				{
+					Action PrepassAction = Graph.CreateAction(ActionType.Compile);
+					PrepassAction.PrerequisiteItems.AddRange(CompileAction.PrerequisiteItems);
+					PrepassAction.PrerequisiteItems.Remove(CompilerResponseFileItem);
+					PrepassAction.CommandDescription = "Preprocess Depends";
+					PrepassAction.StatusDescription = CompileAction.StatusDescription;
+					PrepassAction.bIsGCCCompiler = true;
+					PrepassAction.bCanExecuteRemotely = false;
+					PrepassAction.bShouldOutputStatusDescription = true;
+					PrepassAction.CommandPath = CompileAction.CommandPath;
+					PrepassAction.CommandVersion = CompileAction.CommandVersion;
+					PrepassAction.WorkingDirectory = CompileAction.WorkingDirectory;
+
+					string PreprocessArguments = Arguments.ToString();
+					string PreprocessFileArguments = FileArguments;
+					PreprocessArguments = PreprocessArguments.Replace(" -c ", " ");
+
+					FileItem DependencyListFile = FileItem.GetItemByFileReference(FileReference.Combine(OutputDir, Path.GetFileName(SourceFile.AbsolutePath) + ".d"));
+					PreprocessFileArguments += string.Format(" -M -MF\"{0}\"", NormalizeCommandLinePath(DependencyListFile));
+					PrepassAction.DependencyListFile = DependencyListFile;
+					PrepassAction.ProducedItems.Add(DependencyListFile);
+
+					PreprocessFileArguments = PreprocessFileArguments.Replace(" -ftime-trace", string.Empty);
+					PreprocessFileArguments = PreprocessFileArguments.Replace(string.Format(" -o \"{0}\"", NormalizeCommandLinePath(CompileAction.ProducedItems.First())), String.Empty);
+
+					PrepassAction.DeleteItems.AddRange(PrepassAction.ProducedItems);
+
+					string AllPreprocessArgs = PreprocessArguments + PreprocessFileArguments + CompileEnvironment.AdditionalArguments;
+
+					FileReference CompilerPrepassResponseFileName = PrepassAction.ProducedItems[0].Location + ".rsp";
+					FileItem CompilerPrepassResponseFileItem = Graph.CreateIntermediateTextFile(CompilerPrepassResponseFileName, AllPreprocessArgs);
+
+					PrepassAction.CommandArguments = string.Format(" @\"{0}\"", NormalizeCommandLinePath(CompilerPrepassResponseFileName));
+					PrepassAction.PrerequisiteItems.Add(CompilerResponseFileItem);
+
+					CompileAction.DependencyListFile = DependencyListFile;
+					CompileAction.PrerequisiteItems.Add(DependencyListFile);
+				}
 			}
 
 			return Result;
diff --git a/Engine/Source/Programs/UnrealBuildTool/Platform/Linux/UEBuildLinux.cs b/Engine/Source/Programs/UnrealBuildTool/Platform/Linux/UEBuildLinux.cs
--- a/Engine/Source/Programs/UnrealBuildTool/Platform/Linux/UEBuildLinux.cs
+++ b/Engine/Source/Programs/UnrealBuildTool/Platform/Linux/UEBuildLinux.cs
@@ -77,6 +77,13 @@ namespace UnrealBuildTool
 		[CommandLine("-EnableLLDB")]
 		[XmlConfigFile(Category = "BuildConfiguration", Name = "bTuneDebugInfoForLLDB")]
 		public bool bTuneDebugInfoForLLDB = false;
+
+		/// <summary>
+		/// Generate dependency files by preprocessing. This is only recommended when distributing builds as it adds additional overhead.
+		/// </summary>
+		[CommandLine("-PreprocessDepends")]
+		[XmlConfigFile(Category = "BuildConfiguration", Name = "bPreprocessDepends")]
+		public bool bPreprocessDepends = false;
 	}
 
 	/// <summary>
@@ -131,6 +138,11 @@ namespace UnrealBuildTool
 			get { return Inner.bTuneDebugInfoForLLDB; }
 		}
 
+		public bool bPreprocessDepends
+		{
+			get { return Inner.bPreprocessDepends; }
+		}
+
 		#pragma warning restore CS1591
 		#endregion
 	}
diff --git a/Engine/Source/Programs/UnrealBuildTool/Platform/Mac/MacToolChain.cs b/Engine/Source/Programs/UnrealBuildTool/Platform/Mac/MacToolChain.cs
--- a/Engine/Source/Programs/UnrealBuildTool/Platform/Mac/MacToolChain.cs
+++ b/Engine/Source/Programs/UnrealBuildTool/Platform/Mac/MacToolChain.cs
@@ -432,6 +432,22 @@ namespace UnrealBuildTool
 			return Result;
 		}
 
+		public static string NormalizeCommandLinePath(FileSystemReference Reference)
+		{
+			// Try to use a relative path to shorten command line length.
+			if (Reference.IsUnderDirectory(Unreal.RootDirectory))
+			{
+				return Reference.MakeRelativeTo(UnrealBuildTool.EngineSourceDirectory);
+			}
+
+			return Reference.FullName;
+		}
+
+		public static string NormalizeCommandLinePath(FileItem Item)
+		{
+			return NormalizeCommandLinePath(Item.Location);
+		}
+
 		public override CPPOutput CompileCPPFiles(CppCompileEnvironment CompileEnvironment, List<FileItem> InputFiles, DirectoryReference OutputDir, string ModuleName, IActionGraphBuilder Graph)
 		{
 			StringBuilder Arguments = new StringBuilder();
@@ -444,7 +460,7 @@ namespace UnrealBuildTool
 				// Add the precompiled header file's path to the include path so GCC can find it.
 				// This needs to be before the other include paths to ensure GCC uses it instead of the source header file.
 				PCHArguments.Append(" -include \"");
-				PCHArguments.Append(CompileEnvironment.PrecompiledHeaderIncludeFilename);
+				PCHArguments.Append(NormalizeCommandLinePath(CompileEnvironment.PrecompiledHeaderIncludeFilename!));
 				PCHArguments.Append("\"");
 			}
 
@@ -454,7 +470,7 @@ namespace UnrealBuildTool
 			foreach (DirectoryReference IncludePath in AllIncludes)
 			{
 				Arguments.Append(" -I\"");
-				Arguments.Append(IncludePath);
+				Arguments.Append(NormalizeCommandLinePath(IncludePath));
 				Arguments.Append("\"");
 			}
 
@@ -469,13 +485,13 @@ namespace UnrealBuildTool
 			List<string> FrameworksSearchPaths = new List<string>();
 			foreach (UEBuildFramework Framework in CompileEnvironment.AdditionalFrameworks)
 			{
-				string FrameworkPath = Path.GetDirectoryName(Path.GetFullPath(Framework.Name))!;
-				if (!FrameworksSearchPaths.Contains(FrameworkPath))
+				FileReference FrameworkPath = new FileReference(Path.GetFullPath(Framework.Name));
+				if (!FrameworksSearchPaths.Contains(FrameworkPath.Directory.FullName))
 				{
 					Arguments.Append(" -F \"");
-					Arguments.Append(FrameworkPath);
+					Arguments.Append(NormalizeCommandLinePath(FrameworkPath.Directory));
 					Arguments.Append("\"");
-					FrameworksSearchPaths.Add(FrameworkPath);
+					FrameworksSearchPaths.Add(FrameworkPath.Directory.FullName);
 				}
 			}
 
@@ -524,7 +540,7 @@ namespace UnrealBuildTool
 
 				foreach (FileItem ForceIncludeFile in CompileEnvironment.ForceIncludeFiles)
 				{
-					FileArguments += String.Format(" -include \"{0}\"", ForceIncludeFile.Location);
+					FileArguments += String.Format(" -include \"{0}\"", NormalizeCommandLinePath(ForceIncludeFile));
 				}
 
 				// Add the C++ source file and its included files to the prerequisite item list.
@@ -539,25 +555,26 @@ namespace UnrealBuildTool
 					Result.PrecompiledHeaderFile = PrecompiledHeaderFile;
 
 					// Add the parameters needed to compile the precompiled header file to the command-line.
-					FileArguments += string.Format(" -o \"{0}\"", PrecompiledHeaderFile.AbsolutePath);
+					FileArguments += string.Format(" -o \"{0}\"", NormalizeCommandLinePath(PrecompiledHeaderFile));
 				}
 				else
 				{
 					if (CompileEnvironment.PrecompiledHeaderAction == PrecompiledHeaderAction.Include)
 					{
 						CompileAction.PrerequisiteItems.Add(CompileEnvironment.PrecompiledHeaderFile!);
+						CompileAction.PrerequisiteItems.Add(FileItem.GetItemByPath(CompileEnvironment.PrecompiledHeaderFile!.FullName.Replace(".h.gch", ".h")));
 					}
 					// Add the object file to the produced item list.
 					FileItem ObjectFile = FileItem.GetItemByFileReference(FileReference.Combine(OutputDir, Path.GetFileName(SourceFile.AbsolutePath) + ".o"));
 
 					CompileAction.ProducedItems.Add(ObjectFile);
 					Result.ObjectFiles.Add(ObjectFile);
-					FileArguments += string.Format(" -o \"{0}\"", ObjectFile.AbsolutePath);
+					FileArguments += string.Format(" -o \"{0}\"", NormalizeCommandLinePath(ObjectFile));
 					OutputFilePath = ObjectFile.AbsolutePath;
 				}
 
 				// Add the source file path to the command-line.
-				FileArguments += string.Format(" \"{0}\"", SourceFile.AbsolutePath);
+				FileArguments += string.Format(" \"{0}\"", NormalizeCommandLinePath(SourceFile));
 
 				// Generate the timing info
 				if (CompileEnvironment.bPrintTimingInfo)
@@ -571,7 +588,7 @@ namespace UnrealBuildTool
 				if(!bPreprocessDepends && CompileEnvironment.bGenerateDependenciesFile)
 				{
 					FileItem DependencyListFile = FileItem.GetItemByFileReference(FileReference.Combine(OutputDir, Path.GetFileName(SourceFile.AbsolutePath) + ".d"));
-					FileArguments += string.Format(" -MD -MF\"{0}\"", DependencyListFile.AbsolutePath.Replace('\\', '/'));
+					FileArguments += string.Format(" -MD -MF\"{0}\"", NormalizeCommandLinePath(DependencyListFile));
 					CompileAction.DependencyListFile = DependencyListFile;
 					CompileAction.ProducedItems.Add(DependencyListFile);
 				}
@@ -608,7 +625,7 @@ namespace UnrealBuildTool
 
 				CompileAction.WorkingDirectory = GetMacDevSrcRoot();
 
-				if (MacExports.IsRunningUnderRosetta)
+				if (!MacExports.IsRunningUnderRosetta)
 				{
 					string ArchPath = "/usr/bin/arch";
 					CompileAction.CommandPath = new FileReference(ArchPath);
@@ -626,7 +643,7 @@ namespace UnrealBuildTool
 				CompileAction.StatusDescription = Path.GetFileName(SourceFile.AbsolutePath);
 				CompileAction.bIsGCCCompiler = true;
 				// We're already distributing the command by execution on Mac.
-				CompileAction.bCanExecuteRemotely = Extension != ".C";
+				CompileAction.bCanExecuteRemotely = CompileEnvironment.PrecompiledHeaderAction != PrecompiledHeaderAction.Create;
 				CompileAction.bShouldOutputStatusDescription = true;
 				CompileAction.CommandVersion = GetFullClangVersion();
 
@@ -647,19 +664,17 @@ namespace UnrealBuildTool
 					PreprocessArguments = PreprocessArguments.Replace(" -c ", " ");
 
 					FileItem DependencyListFile = FileItem.GetItemByFileReference(FileReference.Combine(OutputDir, Path.GetFileName(SourceFile.AbsolutePath) + ".d"));
-					PreprocessFileArguments += string.Format(" -M -MF\"{0}\"", DependencyListFile.AbsolutePath.Replace('\\', '/'));
+					PreprocessFileArguments += string.Format(" -M -MF\"{0}\"", NormalizeCommandLinePath(DependencyListFile));
 					PrepassAction.DependencyListFile = DependencyListFile;
 					PrepassAction.ProducedItems.Add(DependencyListFile);
 
 					PreprocessFileArguments = PreprocessFileArguments.Replace(" -ftime-trace", string.Empty);
-					PreprocessFileArguments = PreprocessFileArguments.Replace(string.Format(" -o \"{0}\"", CompileAction.ProducedItems.First().AbsolutePath), String.Empty);
+					PreprocessFileArguments = PreprocessFileArguments.Replace(string.Format(" -o \"{0}\"", NormalizeCommandLinePath(CompileAction.ProducedItems.First())), String.Empty);
 
 					PrepassAction.DeleteItems.AddRange(PrepassAction.ProducedItems);
 
 					string AllPreprocessArgs = PreprocessArguments + PreprocessFileArguments + EscapedAdditionalArgs;
 
-					PrepassAction.WorkingDirectory = GetMacDevSrcRoot();
-
 					if (MacExports.IsRunningUnderRosetta)
 					{
 						string ArchPath = "/usr/bin/arch";
