UnrealBuildTool: Clean up ParallelExecutor to not require waiting manually on prereq actions by building a continuation tree of async tasks
Split logic into multiple static functions and make protected so it can be shared with an inherited class

#rnx

diff --git a/Engine/Source/Programs/UnrealBuildTool/Executors/ParallelExecutor.cs b/Engine/Source/Programs/UnrealBuildTool/Executors/ParallelExecutor.cs
--- a/Engine/Source/Programs/UnrealBuildTool/Executors/ParallelExecutor.cs
+++ b/Engine/Source/Programs/UnrealBuildTool/Executors/ParallelExecutor.cs
@@ -124,7 +124,7 @@ namespace UnrealBuildTool
 			return true;
 		}
 
-		private class ExecuteResults
+		protected class ExecuteResults
 		{
 			public List<string> LogLines { get; private set; }
 			public int ExitCode { get; private set; }
@@ -145,7 +145,6 @@ namespace UnrealBuildTool
 			}
 		}
 
-
 		/// <summary>
 		/// Executes the specified actions locally.
 		/// </summary>
@@ -162,8 +161,8 @@ namespace UnrealBuildTool
 
 			Log.TraceInformation($"Building {TotalActions} {(TotalActions == 1 ? "action" : "actions")} with {ActualNumParallelProcesses} {(ActualNumParallelProcesses == 1 ? "process" : "processes")}...");
 
-			Dictionary<LinkedAction, Task<ExecuteResults>> Tasks = new Dictionary<LinkedAction, Task<ExecuteResults>>();
-			List<Task> AllTasks = new List<Task>();
+			Dictionary<LinkedAction, Task<ExecuteResults>> ExecuteTasks = new Dictionary<LinkedAction, Task<ExecuteResults>>();
+			List<Task> LogTasks = new List<Task>();
 
 			using LogIndentScope Indent = new LogIndentScope("  ");
 
@@ -173,96 +172,84 @@ namespace UnrealBuildTool
 			// Create a task for every action
 			foreach (LinkedAction Action in InputActions)
 			{
-				Task<ExecuteResults> ExecuteTask = ExecuteAction(Action, Tasks, ProcessGroup, MaxProcessSemaphore, CancellationToken);
+				if (ExecuteTasks.ContainsKey(Action))
+				{
+					continue;
+				}
+
+				Task<ExecuteResults> ExecuteTask = CreateExecuteTask(Action, ExecuteTasks, ProcessGroup, MaxProcessSemaphore, CancellationToken);
 				Task LogTask = ExecuteTask.ContinueWith(antecedent => LogCompletedAction(Action, antecedent, CancellationTokenSource, ProgressWriter, TotalActions, ref NumCompletedActions), CancellationToken);
 
-				Tasks.Add(Action, ExecuteTask);
-				AllTasks.Add(ExecuteTask);
-				AllTasks.Add(LogTask);
+				ExecuteTasks.Add(Action, ExecuteTask);
+				LogTasks.Add(LogTask);
 			}
 
-			// Wait for all tasks to complete
-			Task.WaitAll(AllTasks.ToArray(), CancellationToken);
+			Task SummaryTask = Task.Factory.ContinueWhenAll(LogTasks.ToArray(), (AntecedentTasks) => TraceSummary(ExecuteTasks, ProcessGroup), CancellationToken);
+			SummaryTask.Wait();
 
-			if (bShowCompilationTimes)
+			// Return if all tasks succeeded
+			return ExecuteTasks.Values.All(x => x.Result.ExitCode == 0);
+		}
+
+		protected static Task<ExecuteResults> CreateExecuteTask(LinkedAction Action, Dictionary<LinkedAction, Task<ExecuteResults>> ExecuteTasks, ManagedProcessGroup ProcessGroup, SemaphoreSlim MaxProcessSemaphore, CancellationToken CancellationToken)
+		{
+			if (Action.PrerequisiteActions.Count == 0)
 			{
-				Log.TraceInformation("");
-				if (ProcessGroup.TotalProcessorTime.Ticks > 0)
-				{
-					Log.TraceInformation($"Total CPU Time: {ProcessGroup.TotalProcessorTime.TotalSeconds} s");
-					Log.TraceInformation("");
-				}
+				return Task.Factory.StartNew(
+					() => ExecuteAction(new Task<ExecuteResults>[0], Action, ProcessGroup, MaxProcessSemaphore, CancellationToken),
+					CancellationToken,
+					TaskCreationOptions.LongRunning,
+					TaskScheduler.Current
+				).Unwrap();
+			}
 
-				if (Tasks.Count > 0)
+			// Create tasks for any preresquite actions if they don't exist already
+			List<Task<ExecuteResults>> PrerequisiteTasks = new List<Task<ExecuteResults>>();
+			foreach (var PrerequisiteAction in Action.PrerequisiteActions)
+			{
+				if (!ExecuteTasks.ContainsKey(PrerequisiteAction))
 				{
-					Log.TraceInformation($"Compilation Time Top {Math.Min(20, Tasks.Count)}");
-					Log.TraceInformation("");
-					foreach (var Pair in Tasks.OrderByDescending(x => x.Value.Result.ExecutionTime).Take(20))
-					{
-						string Description = $"{(Pair.Key.Inner.CommandDescription ?? Pair.Key.Inner.CommandPath.GetFileNameWithoutExtension())} {Pair.Key.Inner.StatusDescription}".Trim();
-						if (Pair.Value.Result.ProcessorTime.Ticks > 0)
-						{
-							Log.TraceInformation($"{Description} [ Wall Time {Pair.Value.Result.ExecutionTime.TotalSeconds:0.00} s / CPU Time {Pair.Value.Result.ProcessorTime.TotalSeconds:0.00} s ]");
-						}
-						else
-						{
-							Log.TraceInformation($"{Description} [ Time {Pair.Value.Result.ExecutionTime.TotalSeconds:0.00} s ]");
-						}
-
-					}
-					Log.TraceInformation("");
+					ExecuteTasks.Add(PrerequisiteAction, CreateExecuteTask(PrerequisiteAction, ExecuteTasks, ProcessGroup, MaxProcessSemaphore, CancellationToken));
 				}
+				PrerequisiteTasks.Add(ExecuteTasks[PrerequisiteAction]);
 			}
 
-			// Return if all tasks succeeded
-			return Tasks.Values.All(x => x.Result.ExitCode == 0);
+			return Task.Factory.ContinueWhenAll(
+				PrerequisiteTasks.ToArray(),
+				(AntecedentTasks) => ExecuteAction(AntecedentTasks, Action, ProcessGroup, MaxProcessSemaphore, CancellationToken),
+				CancellationToken,
+				TaskContinuationOptions.LongRunning,
+				TaskScheduler.Current
+			).Unwrap();
 		}
 
-		private async Task<ExecuteResults> ExecuteAction(LinkedAction Action, Dictionary<LinkedAction, Task<ExecuteResults>> Tasks, ManagedProcessGroup ProcessGroup, SemaphoreSlim MaxProcessSemaphore, CancellationToken CancellationToken)
+		protected static async Task WaitPrerequsiteActions(LinkedAction Action, Dictionary<LinkedAction, Task<ExecuteResults>> Tasks)
+		{
+			// Wait for all PrerequisiteActions to complete
+			ExecuteResults[] Results = await Task.WhenAll(Action.PrerequisiteActions.Select(x => Tasks[x]));
+
+			// Cancel tasks if any PrerequisiteActions fail, unless a PostBuildStep
+			if (Action.ActionType != ActionType.PostBuildStep && Results.Any(x => x.ExitCode != 0))
+			{
+				throw new OperationCanceledException();
+			}
+		}
+
+		protected static async Task<ExecuteResults> ExecuteAction(Task<ExecuteResults>[] AntecedentTasks, LinkedAction Action, ManagedProcessGroup ProcessGroup, SemaphoreSlim MaxProcessSemaphore, CancellationToken CancellationToken)
 		{
 			Task? SemaphoreTask = null;
 			try
 			{
-				// Wait for Tasks list to be populated with any PrerequisiteActions
-				while (!Action.PrerequisiteActions.All(x => Tasks.ContainsKey(x)))
-				{
-					await Task.Delay(100, CancellationToken);
-					CancellationToken.ThrowIfCancellationRequested();
-				}
-
-				// Wait for all PrerequisiteActions to complete
-				ExecuteResults[] Results = await Task.WhenAll(Action.PrerequisiteActions.Select(x => Tasks[x]).ToArray());
-
 				// Cancel tasks if any PrerequisiteActions fail, unless a PostBuildStep
-				if (Action.ActionType != ActionType.PostBuildStep && Results.Any(x => x.ExitCode != 0))
+				if (Action.ActionType != ActionType.PostBuildStep && AntecedentTasks.Any(x => x.Result.ExitCode != 0))
 				{
 					throw new OperationCanceledException();
 				}
 
-				CancellationToken.ThrowIfCancellationRequested();
-
 				// Limit the number of concurrent processes that will run in parallel
 				SemaphoreTask = MaxProcessSemaphore.WaitAsync(CancellationToken);
 				await SemaphoreTask;
-
-				CancellationToken.ThrowIfCancellationRequested();
-
-				using ManagedProcess Process = new ManagedProcess(ProcessGroup, Action.CommandPath.FullName, Action.CommandArguments, Action.WorkingDirectory.FullName, null, null, ProcessPriorityClass.BelowNormal);
-
-				MemoryStream StdOutStream = new MemoryStream();
-				await Process.CopyToAsync(StdOutStream, CancellationToken);
-
-				CancellationToken.ThrowIfCancellationRequested();
-
-				Process.WaitForExit();
-
-				CancellationToken.ThrowIfCancellationRequested();
-
-				List<string> LogLines = Console.OutputEncoding.GetString(StdOutStream.GetBuffer(), 0, Convert.ToInt32(StdOutStream.Length)).Split(LineEndingSplit, StringSplitOptions.RemoveEmptyEntries).ToList();
-				int ExitCode = Process.ExitCode;
-				TimeSpan ProcessorTime = Process.TotalProcessorTime;
-				TimeSpan ExecutionTime = Process.ExitTime - Process.StartTime;
-				return new ExecuteResults(LogLines, ExitCode, ExecutionTime, ProcessorTime);
+				return await RunAction(Action, ProcessGroup, CancellationToken);
 			}
 			catch (OperationCanceledException)
 			{
@@ -275,14 +262,34 @@ namespace UnrealBuildTool
 			}
 			finally
 			{
-				if (SemaphoreTask != null && SemaphoreTask.Status == TaskStatus.RanToCompletion)
+				if (SemaphoreTask?.Status == TaskStatus.RanToCompletion)
 				{
 					MaxProcessSemaphore.Release();
 				}
 			}
 		}
 
-		private void LogCompletedAction(LinkedAction Action, Task<ExecuteResults> ExecuteTask, CancellationTokenSource CancellationTokenSource, ProgressWriter ProgressWriter, int TotalActions, ref int NumCompletedActions)
+		protected static async Task<ExecuteResults> RunAction(LinkedAction Action, ManagedProcessGroup ProcessGroup, CancellationToken CancellationToken)
+		{
+			CancellationToken.ThrowIfCancellationRequested();
+
+			using ManagedProcess Process = new ManagedProcess(ProcessGroup, Action.CommandPath.FullName, Action.CommandArguments, Action.WorkingDirectory.FullName, null, null, ProcessPriorityClass.BelowNormal);
+
+			using MemoryStream StdOutStream = new MemoryStream();
+			await Process.CopyToAsync(StdOutStream, CancellationToken);
+
+			CancellationToken.ThrowIfCancellationRequested();
+
+			Process.WaitForExit();
+
+			List<string> LogLines = Console.OutputEncoding.GetString(StdOutStream.GetBuffer(), 0, Convert.ToInt32(StdOutStream.Length)).Split(LineEndingSplit, StringSplitOptions.RemoveEmptyEntries).ToList();
+			int ExitCode = Process.ExitCode;
+			TimeSpan ProcessorTime = Process.TotalProcessorTime;
+			TimeSpan ExecutionTime = Process.ExitTime - Process.StartTime;
+			return new ExecuteResults(LogLines, ExitCode, ExecutionTime, ProcessorTime);
+		}
+
+		protected static void LogCompletedAction(LinkedAction Action, Task<ExecuteResults> ExecuteTask, CancellationTokenSource CancellationTokenSource, ProgressWriter ProgressWriter, int TotalActions, ref int NumCompletedActions)
 		{
 			List<string> LogLines = new List<string>();
 			int ExitCode = int.MaxValue;
@@ -364,14 +371,51 @@ namespace UnrealBuildTool
 					}
 					// END TEMPORARY
 
-						// Cancel all other pending tasks
-						if (bStopCompilationAfterErrors)
+					// Cancel all other pending tasks
+					if (bStopCompilationAfterErrors)
 					{
 						CancellationTokenSource.Cancel();
 					}
 				}
 			}
 		}
+
+		protected static void TraceSummary(Dictionary<LinkedAction, Task<ExecuteResults>> Tasks, ManagedProcessGroup ProcessGroup)
+		{
+			if (!bShowCompilationTimes)
+			{
+				return;
+			}
+
+			Log.TraceInformation("");
+			if (ProcessGroup.TotalProcessorTime.Ticks > 0)
+			{
+				Log.TraceInformation($"Total CPU Time: {ProcessGroup.TotalProcessorTime.TotalSeconds} s");
+				Log.TraceInformation("");
+			}
+
+			var CompletedTasks = Tasks.Where(x => x.Value.Status == TaskStatus.RanToCompletion).OrderByDescending(x => x.Value.Result.ExecutionTime).Take(20);
+
+			if (CompletedTasks.Any())
+			{
+				Log.TraceInformation($"Compilation Time Top {CompletedTasks.Count()}");
+				Log.TraceInformation("");
+				foreach (var Pair in CompletedTasks)
+				{
+					string Description = $"{(Pair.Key.Inner.CommandDescription ?? Pair.Key.Inner.CommandPath.GetFileNameWithoutExtension())} {Pair.Key.Inner.StatusDescription}".Trim();
+					if (Pair.Value.Result.ProcessorTime.Ticks > 0)
+					{
+						Log.TraceInformation($"{Description} [ Wall Time {Pair.Value.Result.ExecutionTime.TotalSeconds:0.00} s / CPU Time {Pair.Value.Result.ProcessorTime.TotalSeconds:0.00} s ]");
+					}
+					else
+					{
+						Log.TraceInformation($"{Description} [ Time {Pair.Value.Result.ExecutionTime.TotalSeconds:0.00} s ]");
+					}
+
+				}
+				Log.TraceInformation("");
+			}
+		}
 	}
 
 	/// <summary>
