- Moved most of the remaining SDK discovery/versioning from UEBuildWindows.cs to MicrosoftPlatformSDK.cs

diff --git a/Engine/Source/Programs/UnrealBuildTool/Platform/Windows/MicrosoftPlatformSDK.cs b/Engine/Source/Programs/UnrealBuildTool/Platform/Windows/MicrosoftPlatformSDK.cs
--- a/Engine/Source/Programs/UnrealBuildTool/Platform/Windows/MicrosoftPlatformSDK.cs
+++ b/Engine/Source/Programs/UnrealBuildTool/Platform/Windows/MicrosoftPlatformSDK.cs
@@ -1,12 +1,17 @@
 // Copyright Epic Games, Inc. All Rights Reserved.
 
 using System;
+using System.IO;
 using System.Collections.Generic;
 using EpicGames.Core;
 using System.Linq;
 using System.Text.RegularExpressions;
 using Microsoft.Win32;
+using System.Diagnostics;
 using System.Diagnostics.CodeAnalysis;
+using Microsoft.VisualStudio.Setup.Configuration;
+using System.Runtime.InteropServices;
+using System.Buffers.Binary;
 using UnrealBuildBase;
 
 namespace UnrealBuildTool
@@ -21,6 +26,48 @@
 			VersionNumber.Parse("10.0.18362.0")
 		};
 
+		/// <summary>
+		/// The default compiler version to be used, if installed. 
+		/// </summary>
+		static readonly VersionNumberRange[] PreferredClangVersions =
+		{
+			VersionNumberRange.Parse("13.0.0", "13.999"),
+			VersionNumberRange.Parse("12.0.0", "12.999"),
+			VersionNumberRange.Parse("11.0.0", "11.999"),
+			VersionNumberRange.Parse("10.0.0", "10.999")
+		};
+
+		static readonly VersionNumber MinimumClangVersion = new VersionNumber(10, 0, 0);
+
+		/// <summary>
+		/// Ranges of tested compiler toolchains to be used, in order of preference. If multiple toolchains in a range are present, the latest version will be preferred.
+		/// Note that the numbers here correspond to the installation *folders* rather than precise executable versions. 
+		/// </summary>
+		static readonly VersionNumberRange[] PreferredVisualCppVersions = new VersionNumberRange[]
+		{
+			VersionNumberRange.Parse("14.29.30133", "14.29.30136"), // VS2019 16.11.5
+		};
+
+		/// <summary>
+		/// Tested compiler toolchanges that should not be allowed.
+		/// </summary>
+		static readonly VersionNumberRange[] BannedVisualCppVersions = new VersionNumberRange[]
+		{
+		};
+
+		static readonly VersionNumber MinimumVisualCppVersion = new VersionNumber(14, 29, 30133);
+
+		/// <summary>
+		/// The default compiler version to be used, if installed. 
+		/// </summary>
+		static readonly VersionNumberRange[] PreferredIntelOneApiVersions =
+		{
+			VersionNumberRange.Parse("2022.0.0", "2022.9999"),
+		};
+
+		static readonly VersionNumber MinimumIntelOneApiVersion = new VersionNumber(2021, 0, 0);
+
+
 		public override string GetMainVersion()
 		{
 			// preferred/main version is the top of the Preferred list - 
@@ -80,6 +127,13 @@
 
 		#region Windows Specific SDK discovery
 
+		#region Cached Locations
+
 		/// <summary>
+		/// Cache of Visual C++ installation directories
+		/// </summary>
+		private static Dictionary<WindowsCompiler, List<ToolChainInstallation>> CachedToolChainInstallations = new Dictionary<WindowsCompiler, List<ToolChainInstallation>>();
+
+		/// <summary>
 		/// Cache of Windows SDK installation directories
 		/// </summary>
@@ -90,62 +144,26 @@
 		/// </summary>
 		private static IReadOnlyDictionary<VersionNumber, DirectoryReference>? CachedUniversalCrtDirs;
 
-
-
+		/// <summary>
+		/// Visual Studio installations
+		/// </summary>
+		public static IReadOnlyList<VisualStudioInstallation> VisualStudioInstallations => CachedVisualStudioInstallations.Value;
 
 		/// <summary>
-		/// Updates the CachedWindowsSdkDirs and CachedUniversalCrtDirs variables
+		/// Cache of Visual Studio installation directories
 		/// </summary>
-		private static void UpdateCachedWindowsSdks()
-		{
-			Dictionary<VersionNumber, DirectoryReference> WindowsSdkDirs = new Dictionary<VersionNumber, DirectoryReference>();
-			Dictionary<VersionNumber, DirectoryReference> UniversalCrtDirs = new Dictionary<VersionNumber, DirectoryReference>();
-
-			// Enumerate the Windows 8.1 SDK, if present
-			DirectoryReference? InstallDir_8_1;
-			if (TryReadInstallDirRegistryKey32("Microsoft\\Microsoft SDKs\\Windows\\v8.1", "InstallationFolder", out InstallDir_8_1))
-			{
-				if (FileReference.Exists(FileReference.Combine(InstallDir_8_1, "Include", "um", "windows.h")))
-				{
-					Log.TraceLog("Found Windows 8.1 SDK at {0}", InstallDir_8_1);
-					VersionNumber Version_8_1 = new VersionNumber(8, 1);
-					WindowsSdkDirs[Version_8_1] = InstallDir_8_1;
-				}
-			}
-
-			// Find all the root directories for Windows 10 SDKs
-			List<DirectoryReference> InstallDirs_10 = new List<DirectoryReference>();
-			EnumerateSdkRootDirs(InstallDirs_10);
-
-			// Enumerate all the Windows 10 SDKs
-			foreach (DirectoryReference InstallDir_10 in InstallDirs_10.Distinct())
-			{
-				DirectoryReference IncludeRootDir = DirectoryReference.Combine(InstallDir_10, "Include");
-				if (DirectoryReference.Exists(IncludeRootDir))
-				{
-					foreach (DirectoryReference IncludeDir in DirectoryReference.EnumerateDirectories(IncludeRootDir))
-					{
-						VersionNumber? IncludeVersion;
-						if (VersionNumber.TryParse(IncludeDir.GetDirectoryName(), out IncludeVersion))
-						{
-							if (FileReference.Exists(FileReference.Combine(IncludeDir, "um", "windows.h")))
-							{
-								Log.TraceLog("Found Windows 10 SDK version {0} at {1}", IncludeVersion, InstallDir_10);
-								WindowsSdkDirs[IncludeVersion] = InstallDir_10;
-							}
-							if (FileReference.Exists(FileReference.Combine(IncludeDir, "ucrt", "corecrt.h")))
-							{
-								Log.TraceLog("Found Universal CRT version {0} at {1}", IncludeVersion, InstallDir_10);
-								UniversalCrtDirs[IncludeVersion] = InstallDir_10;
-							}
-						}
-					}
-				}
-			}
-
-			CachedWindowsSdkDirs = WindowsSdkDirs;
-			CachedUniversalCrtDirs = UniversalCrtDirs;
-		}
+		private static Lazy<IReadOnlyList<VisualStudioInstallation>> CachedVisualStudioInstallations = new Lazy<IReadOnlyList<VisualStudioInstallation>>(FindVisualStudioInstallations);
 
 		/// <summary>
+		/// Cache of DIA SDK installation directories
+		/// </summary>
+		private static Dictionary<WindowsCompiler, List<DirectoryReference>> CachedDiaSdkDirs = new Dictionary<WindowsCompiler, List<DirectoryReference>>();
+
+		#endregion
+
+		#region WindowsSdk 
+
+		#region Public Interface
+
+		/// <summary>
 		/// Finds all the installed Windows SDK versions
@@ -234,6 +252,99 @@
 			}
 		}
 
+		#endregion
+
+		#region Private implementation
+
 		/// <summary>
+		/// Updates the CachedWindowsSdkDirs and CachedUniversalCrtDirs variables
+		/// </summary>
+		private static void UpdateCachedWindowsSdks()
+		{
+			Dictionary<VersionNumber, DirectoryReference> WindowsSdkDirs = new Dictionary<VersionNumber, DirectoryReference>();
+			Dictionary<VersionNumber, DirectoryReference> UniversalCrtDirs = new Dictionary<VersionNumber, DirectoryReference>();
+
+			// Enumerate the Windows 8.1 SDK, if present
+			DirectoryReference? InstallDir_8_1;
+			if (TryReadInstallDirRegistryKey32("Microsoft\\Microsoft SDKs\\Windows\\v8.1", "InstallationFolder", out InstallDir_8_1))
+			{
+				if (FileReference.Exists(FileReference.Combine(InstallDir_8_1, "Include", "um", "windows.h")))
+				{
+					Log.TraceLog("Found Windows 8.1 SDK at {0}", InstallDir_8_1);
+					VersionNumber Version_8_1 = new VersionNumber(8, 1);
+					WindowsSdkDirs[Version_8_1] = InstallDir_8_1;
+				}
+			}
+
+			// Find all the root directories for Windows 10 SDKs
+			List<DirectoryReference> InstallDirs_10 = new List<DirectoryReference>();
+			EnumerateSdkRootDirs(InstallDirs_10);
+
+			// Enumerate all the Windows 10 SDKs
+			foreach (DirectoryReference InstallDir_10 in InstallDirs_10.Distinct())
+			{
+				DirectoryReference IncludeRootDir = DirectoryReference.Combine(InstallDir_10, "Include");
+				if (DirectoryReference.Exists(IncludeRootDir))
+				{
+					foreach (DirectoryReference IncludeDir in DirectoryReference.EnumerateDirectories(IncludeRootDir))
+					{
+						VersionNumber? IncludeVersion;
+						if (VersionNumber.TryParse(IncludeDir.GetDirectoryName(), out IncludeVersion))
+						{
+							if (FileReference.Exists(FileReference.Combine(IncludeDir, "um", "windows.h")))
+							{
+								Log.TraceLog("Found Windows 10 SDK version {0} at {1}", IncludeVersion, InstallDir_10);
+								WindowsSdkDirs[IncludeVersion] = InstallDir_10;
+							}
+							if (FileReference.Exists(FileReference.Combine(IncludeDir, "ucrt", "corecrt.h")))
+							{
+								Log.TraceLog("Found Universal CRT version {0} at {1}", IncludeVersion, InstallDir_10);
+								UniversalCrtDirs[IncludeVersion] = InstallDir_10;
+							}
+						}
+					}
+				}
+			}
+
+			CachedWindowsSdkDirs = WindowsSdkDirs;
+			CachedUniversalCrtDirs = UniversalCrtDirs;
+		}
+
+		/// <summary>
+		/// Enumerates all the Windows 10 SDK root directories
+		/// </summary>
+		/// <param name="RootDirs">Receives all the Windows 10 sdk root directories</param>
+		private static void EnumerateSdkRootDirs(List<DirectoryReference> RootDirs)
+		{
+			DirectoryReference? RootDir;
+			if (TryReadInstallDirRegistryKey32("Microsoft\\Windows Kits\\Installed Roots", "KitsRoot10", out RootDir))
+			{
+				Log.TraceLog("Found Windows 10 SDK root at {0} (1)", RootDir);
+				RootDirs.Add(RootDir);
+			}
+			if (TryReadInstallDirRegistryKey32("Microsoft\\Microsoft SDKs\\Windows\\v10.0", "InstallationFolder", out RootDir))
+			{
+				Log.TraceLog("Found Windows 10 SDK root at {0} (2)", RootDir);
+				RootDirs.Add(RootDir);
+			}
+
+			DirectoryReference? HostAutoSdkDir;
+			if (UEBuildPlatformSDK.TryGetHostPlatformAutoSDKDir(out HostAutoSdkDir))
+			{
+				DirectoryReference RootDirAutoSdk = DirectoryReference.Combine(HostAutoSdkDir, "Win64", "Windows Kits", "10");
+				if (DirectoryReference.Exists(RootDirAutoSdk))
+				{
+					Log.TraceLog("Found Windows 10 AutoSDK root at {0}", RootDirAutoSdk);
+					RootDirs.Add(RootDirAutoSdk);
+				}
+			}
+		}
+		#endregion
+
+		#endregion
+
+		#region NetFx
+
+		/// <summary>
 		/// Gets the installation directory for the NETFXSDK
 		/// </summary>
@@ -295,6 +406,7 @@
 			return false;
 		}
 
+		#region Private implementation
 
 		static readonly KeyValuePair<RegistryKey, string>[] InstallDirRoots = {
 			new KeyValuePair<RegistryKey, string>(Registry.CurrentUser, "SOFTWARE\\"),
@@ -375,36 +487,872 @@
 			}
 		}
 
+		#endregion
+
+		#endregion
+
+		#region Toolchain
+
+		#region Public Interface
+
 		/// <summary>
-		/// Enumerates all the Windows 10 SDK root directories
+		/// Gets the MSBuild path, and throws an exception on failure.
 		/// </summary>
-		/// <param name="RootDirs">Receives all the Windows 10 sdk root directories</param>
-		public static void EnumerateSdkRootDirs(List<DirectoryReference> RootDirs)
+		/// <returns>Path to MSBuild</returns>
+		public static FileReference GetMsBuildToolPath()
 		{
-			DirectoryReference? RootDir;
-			if (TryReadInstallDirRegistryKey32("Microsoft\\Windows Kits\\Installed Roots", "KitsRoot10", out RootDir))
+			FileReference? Location;
+			if (!TryGetMsBuildPath(out Location))
 			{
-				Log.TraceLog("Found Windows 10 SDK root at {0} (1)", RootDir);
-				RootDirs.Add(RootDir);
+				throw new BuildException("Unable to find installation of MSBuild.");
 			}
-			if (TryReadInstallDirRegistryKey32("Microsoft\\Microsoft SDKs\\Windows\\v10.0", "InstallationFolder", out RootDir))
+			return Location;
+		}
+
+		/// <summary>
+		/// Determines if a given compiler is installed
+		/// </summary>
+		/// <param name="Compiler">Compiler to check for</param>
+		/// <param name="Architecture">Architecture the compiler must support</param>
+		/// <returns>True if the given compiler is installed</returns>
+		public static bool HasCompiler(WindowsCompiler Compiler, WindowsArchitecture Architecture)
+		{
+			return FindToolChainInstallations(Compiler).Where(x => (x.Architecture == Architecture)).Count() > 0;
+		}
+
+		public static bool HasValidCompiler(WindowsCompiler Compiler, WindowsArchitecture Architecture)
+		{
+			// since this is static, we get the instance for it
+			MicrosoftPlatformSDK SDK = (MicrosoftPlatformSDK)GetSDKForPlatform("Win64")!;
+			return FindToolChainInstallations(Compiler).Where(x => (x.Version >= MinimumVisualCppVersion && x.Architecture == Architecture)).Count() > 0;
+		}
+
+		/// <summary>
+		/// Determines the directory containing the MSVC toolchain
+		/// </summary>
+		/// <param name="Compiler">Major version of the compiler to use</param>
+		/// <param name="CompilerVersion">The minimum compiler version to use</param>
+		/// <param name="Architecture">Architecture that is required</param>
+		/// <param name="OutToolChainVersion">Receives the chosen toolchain version</param>
+		/// <param name="OutToolChainDir">Receives the directory containing the toolchain</param>
+		/// <param name="OutRedistDir">Receives the optional directory containing redistributable components</param>
+		/// <returns>True if the toolchain directory was found correctly</returns>
+		public static bool TryGetToolChainDir(WindowsCompiler Compiler, string? CompilerVersion, WindowsArchitecture Architecture, [NotNullWhen(true)] out VersionNumber? OutToolChainVersion, [NotNullWhen(true)] out DirectoryReference? OutToolChainDir, out DirectoryReference? OutRedistDir)
+		{
+			// Find all the installed toolchains
+			List<ToolChainInstallation> ToolChains = FindToolChainInstallations(Compiler);
+
+			// Figure out the actual version number that we want
+			ToolChainInstallation? ToolChain = null;
+			if (CompilerVersion == null)
 			{
-				Log.TraceLog("Found Windows 10 SDK root at {0} (2)", RootDir);
-				RootDirs.Add(RootDir);
+				ToolChain = SelectToolChain(ToolChains, x => x, Architecture);
+				if (ToolChain == null)
+				{
+					OutToolChainVersion = null;
+					OutToolChainDir = null;
+					OutRedistDir = null;
+					return false;
+				}
+			}
+			else if (String.Compare(CompilerVersion, "Latest", StringComparison.InvariantCultureIgnoreCase) == 0)
+			{
+				ToolChain = SelectToolChain(ToolChains, x => x.ThenByDescending(x => x.Version), Architecture);
+				if (ToolChain == null)
+				{
+					throw new BuildException("Unable to find C++ toolchain for {0} {1}", Compiler, Architecture.ToString());
+				}
+			}
+			else if (String.Compare(CompilerVersion, "Preview", StringComparison.InvariantCultureIgnoreCase) == 0)
+			{
+				ToolChain = SelectToolChain(ToolChains, x => x.ThenByDescending(x => x.IsPreview), Architecture);
+				if (ToolChain == null || !ToolChain.IsPreview)
+				{
+					throw new BuildException("Unable to find preview toolchain for {0} {1}", Compiler, Architecture.ToString());
+				}
+			}
+			else if (VersionNumber.TryParse(CompilerVersion, out VersionNumber? ToolChainVersion))
+			{
+				ToolChain = SelectToolChain(ToolChains, x => x.ThenByDescending(x => x.Version == ToolChainVersion).ThenByDescending(x => x.Family == ToolChainVersion), Architecture);
+				if (ToolChain == null)
+				{
+					throw new BuildException("Unable to find {0} toolchain for {1} {2}", ToolChainVersion, Compiler, Architecture.ToString());
+				}
+			}
+			else
+			{
+				throw new BuildException("Unable to find {0} toolchain; '{1}' is an invalid version", GetCompilerName(Compiler), CompilerVersion);
 			}
 
-			DirectoryReference? HostAutoSdkDir;
-			if (UEBuildPlatformSDK.TryGetHostPlatformAutoSDKDir(out HostAutoSdkDir))
+			// Get the actual directory for this version
+			OutToolChainVersion = ToolChain.Version;
+			OutToolChainDir = ToolChain.BaseDir;
+			OutRedistDir = ToolChain.RedistDir;
+			return true;
+		}
+
+		/// <summary>
+		/// Returns the human-readable name of the given compiler
+		/// </summary>
+		/// <param name="Compiler">The compiler value</param>
+		/// <returns>Name of the compiler</returns>
+		public static string GetCompilerName(WindowsCompiler Compiler)
+		{
+			switch (Compiler)
 			{
-				DirectoryReference RootDirAutoSdk = DirectoryReference.Combine(HostAutoSdkDir, "Win64", "Windows Kits", "10");
-				if (DirectoryReference.Exists(RootDirAutoSdk))
+				case WindowsCompiler.VisualStudio2019:
+					return "Visual Studio 2019";
+				case WindowsCompiler.VisualStudio2022:
+					return "Visual Studio 2022";
+				default:
+					return Compiler.ToString();
+			}
+		}
+
+		#endregion
+
+		#region Private implementation
+
+		/// <summary>
+		/// Select which toolchain to use, combining a custom preference with a default sort order
+		/// </summary>
+		/// <param name="ToolChains"></param>
+		/// <param name="Preference">Ordering function</param>
+		/// <param name="Architecture">Architecture that must be supported</param>
+		/// <returns></returns>
+		static ToolChainInstallation? SelectToolChain(IEnumerable<ToolChainInstallation> ToolChains, Func<IOrderedEnumerable<ToolChainInstallation>, IOrderedEnumerable<ToolChainInstallation>> Preference, WindowsArchitecture Architecture)
+		{
+			ToolChainInstallation? ToolChain = Preference(ToolChains.Where(x => x.Architecture == Architecture)
+				.OrderByDescending(x => x.Error == null))
+				.ThenByDescending(x => x.Is64Bit)
+				.ThenBy(x => x.IsPreview)
+				.ThenBy(x => x.FamilyRank)
+				.ThenByDescending(x => x.Version)
+				.FirstOrDefault();
+
+			if (ToolChain?.Error != null)
+			{
+				throw new BuildException(ToolChain.Error);
+			}
+
+			return ToolChain;
+		}
+
+		static List<ToolChainInstallation> FindToolChainInstallations(WindowsCompiler Compiler)
+		{
+			List<ToolChainInstallation>? ToolChains;
+			if (!CachedToolChainInstallations.TryGetValue(Compiler, out ToolChains))
+			{
+				ToolChains = new List<ToolChainInstallation>();
+				if (BuildHostPlatform.Current.Platform == UnrealTargetPlatform.Win64)
 				{
-					Log.TraceLog("Found Windows 10 AutoSDK root at {0}", RootDirAutoSdk);
-					RootDirs.Add(RootDirAutoSdk);
+					if (Compiler == WindowsCompiler.Clang)
+					{
+						// Check for a manual installation to the default directory
+						DirectoryReference ManualInstallDir = DirectoryReference.Combine(DirectoryReference.GetSpecialFolder(Environment.SpecialFolder.ProgramFiles)!, "LLVM");
+						AddClangToolChain(ManualInstallDir, ToolChains);
+
+						// Check for a manual installation to a custom directory
+						string? LlvmPath = Environment.GetEnvironmentVariable("LLVM_PATH");
+						if (!String.IsNullOrEmpty(LlvmPath))
+						{
+							AddClangToolChain(new DirectoryReference(LlvmPath), ToolChains);
+						}
+
+						// Check for installations bundled with Visual Studio 2019
+						foreach (VisualStudioInstallation Installation in FindVisualStudioInstallations(WindowsCompiler.VisualStudio2019))
+						{
+							AddClangToolChain(DirectoryReference.Combine(Installation.BaseDir, "VC", "Tools", "Llvm"), ToolChains);
+						}
+
+						// Check for installations bundled with Visual Studio 2022
+						foreach (VisualStudioInstallation Installation in FindVisualStudioInstallations(WindowsCompiler.VisualStudio2022))
+						{
+							AddClangToolChain(DirectoryReference.Combine(Installation.BaseDir, "VC", "Tools", "Llvm"), ToolChains);
+						}
+
+						// Check for AutoSDK paths
+						DirectoryReference? AutoSdkDir;
+						if (UEBuildPlatformSDK.TryGetHostPlatformAutoSDKDir(out AutoSdkDir))
+						{
+							DirectoryReference ClangBaseDir = DirectoryReference.Combine(AutoSdkDir, "Win64", "LLVM");
+							if (DirectoryReference.Exists(ClangBaseDir))
+							{
+								foreach (DirectoryReference ToolChainDir in DirectoryReference.EnumerateDirectories(ClangBaseDir))
+								{
+									AddClangToolChain(ToolChainDir, ToolChains);
+								}
+							}
+						}
+					}
+					else if (Compiler == WindowsCompiler.Intel)
+					{
+						DirectoryReference InstallDir = DirectoryReference.Combine(DirectoryReference.GetSpecialFolder(Environment.SpecialFolder.ProgramFiles)!, "Intel", "oneAPI", "compiler");
+						FindIntelOneApiToolChains(InstallDir, ToolChains);
+
+						InstallDir = DirectoryReference.Combine(DirectoryReference.GetSpecialFolder(Environment.SpecialFolder.ProgramFilesX86)!, "Intel", "oneAPI", "compiler");
+						FindIntelOneApiToolChains(InstallDir, ToolChains);
+
+						// Check for AutoSDK paths
+						DirectoryReference? AutoSdkDir;
+						if (UEBuildPlatformSDK.TryGetHostPlatformAutoSDKDir(out AutoSdkDir))
+						{
+							DirectoryReference IntelBaseDir = DirectoryReference.Combine(AutoSdkDir, "Win64", "Intel");
+							if (DirectoryReference.Exists(IntelBaseDir))
+							{
+								FindIntelOneApiToolChains(IntelBaseDir, ToolChains);
+							}
+						}
+					}
+					else if (Compiler.IsMSVC())
+					{
+						// Enumerate all the manually installed toolchains
+						List<VisualStudioInstallation> Installations = FindVisualStudioInstallations(Compiler);
+						foreach (VisualStudioInstallation Installation in Installations)
+						{
+							DirectoryReference ToolChainBaseDir = DirectoryReference.Combine(Installation.BaseDir, "VC", "Tools", "MSVC");
+							DirectoryReference RedistBaseDir = DirectoryReference.Combine(Installation.BaseDir, "VC", "Redist", "MSVC");
+							FindVisualStudioToolChains(ToolChainBaseDir, RedistBaseDir, Installation.bPreview, ToolChains);
+						}
+
+						// Enumerate all the AutoSDK toolchains
+						DirectoryReference? PlatformDir;
+						if (UEBuildPlatformSDK.TryGetHostPlatformAutoSDKDir(out PlatformDir))
+						{
+							string VSDir = string.Empty;
+							switch (Compiler)
+							{
+								case WindowsCompiler.VisualStudio2019: VSDir = "VS2019"; break;
+								case WindowsCompiler.VisualStudio2022: VSDir = "VS2022"; break;
+							}
+
+							if (!string.IsNullOrEmpty(VSDir))
+							{
+								DirectoryReference ReleaseBaseDir = DirectoryReference.Combine(PlatformDir, "Win64", VSDir);
+								FindVisualStudioToolChains(ReleaseBaseDir, null, false, ToolChains);
+
+								DirectoryReference PreviewBaseDir = DirectoryReference.Combine(PlatformDir, "Win64", $"{VSDir}-Preview");
+								FindVisualStudioToolChains(PreviewBaseDir, null, true, ToolChains);
+							}
+						}
+					}
+					else
+					{
+						throw new BuildException("Unsupported compiler version ({0})", Compiler);
+					}
+				}
+				CachedToolChainInstallations.Add(Compiler, ToolChains);
+			}
+			return ToolChains;
+		}
+
+		/// <summary>
+		/// Read the Visual Studio install directory for the given compiler version. Note that it is possible for the compiler toolchain to be installed without
+		/// Visual Studio, and vice versa.
+		/// </summary>
+		/// <returns>List of directories containing Visual Studio installations</returns>
+		private static IReadOnlyList<VisualStudioInstallation> FindVisualStudioInstallations()
+		{
+			List<VisualStudioInstallation> Installations = new List<VisualStudioInstallation>();
+			if (BuildHostPlatform.Current.Platform == UnrealTargetPlatform.Win64)
+			{
+				try
+				{
+					SetupConfiguration Setup = new SetupConfiguration();
+					IEnumSetupInstances Enumerator = Setup.EnumAllInstances();
+
+					ISetupInstance[] Instances = new ISetupInstance[1];
+					for (; ; )
+					{
+						int NumFetched;
+						Enumerator.Next(1, Instances, out NumFetched);
+
+						if (NumFetched == 0)
+						{
+							break;
+						}
+
+						ISetupInstance2 Instance = (ISetupInstance2)Instances[0];
+						if ((Instance.GetState() & InstanceState.Local) != InstanceState.Local)
+						{
+							continue;
+						}
+
+						VersionNumber? Version;
+						if (!VersionNumber.TryParse(Instance.GetInstallationVersion(), out Version))
+						{
+							continue;
+						}
+
+						int MajorVersion = Version.GetComponent(0);
+
+						WindowsCompiler Compiler;
+						if (MajorVersion >= 17) // Treat any newer versions as 2022, until we have an explicit enum for them
+						{
+							Compiler = WindowsCompiler.VisualStudio2022;
+						}
+						else if (MajorVersion == 16)
+						{
+							Compiler = WindowsCompiler.VisualStudio2019;
+						}
+						else
+						{
+							continue;
+						}
+
+						ISetupInstanceCatalog? Catalog = Instance as ISetupInstanceCatalog;
+						bool bPreview = Catalog != null && Catalog.IsPrerelease();
+
+						string ProductId = Instance.GetProduct().GetId();
+						bool bExpress = ProductId.Equals("Microsoft.VisualStudio.Product.WDExpress", StringComparison.Ordinal);
+
+						DirectoryReference BaseDir = new DirectoryReference(Instance.GetInstallationPath());
+						Installations.Add(new VisualStudioInstallation(Compiler, Version, BaseDir, bExpress, bPreview));
+
+						Log.TraceLog("Found Visual Studio installation: {0} (Product={1}, Version={2})", BaseDir, ProductId, Version);
+					}
+
+					Installations = Installations.OrderBy(x => x.bExpress).ThenBy(x => x.bPreview).ThenByDescending(x => x.Version).ToList();
+				}
+				catch (Exception Ex)
+				{
+					throw new BuildException(Ex, "Error while enumerating Visual Studio toolchains");
+				}
+			}
+			return Installations;
+		}
+
+		/// <summary>
+		/// Read the Visual Studio install directory for the given compiler version. Note that it is possible for the compiler toolchain to be installed without
+		/// Visual Studio, and vice versa.
+		/// </summary>
+		/// <returns>List of directories containing Visual Studio installations</returns>
+		public static List<VisualStudioInstallation> FindVisualStudioInstallations(WindowsCompiler Compiler)
+		{
+			return VisualStudioInstallations.Where(x => x.Compiler == Compiler).ToList();
+		}
+
+		/// <summary>
+		/// Finds all the valid Visual Studio toolchains under the given base directory
+		/// </summary>
+		/// <param name="BaseDir">Base directory to search</param>
+		/// <param name="OptionalRedistDir">Optional directory for redistributable components (DLLs etc)</param>
+		/// <param name="bPreview">Whether this is a preview installation</param>
+		/// <param name="ToolChains">Map of tool chain version to installation info</param>
+		static void FindVisualStudioToolChains(DirectoryReference BaseDir, DirectoryReference? OptionalRedistDir, bool bPreview, List<ToolChainInstallation> ToolChains)
+		{
+			if (DirectoryReference.Exists(BaseDir))
+			{
+				foreach (DirectoryReference ToolChainDir in DirectoryReference.EnumerateDirectories(BaseDir))
+				{
+					VersionNumber? Version;
+					if (IsValidToolChainDirMSVC(ToolChainDir, out Version))
+					{
+						DirectoryReference? RedistDir = FindVisualStudioRedistForToolChain(ToolChainDir, OptionalRedistDir, Version);
+
+						AddVisualCppToolChain(Version, bPreview, ToolChainDir, RedistDir, ToolChains);
+					}
 				}
 			}
 		}
 
-		#endregion
+
+		/// <summary>
+		/// Finds the most appropriate redist directory for the given toolchain version
+		/// </summary>
+		/// <param name="ToolChainDir"></param>
+		/// <param name="OptionalRedistDir"></param>
+		/// <param name="Version"></param>
+		/// <returns></returns>
+		static DirectoryReference? FindVisualStudioRedistForToolChain(DirectoryReference ToolChainDir, DirectoryReference? OptionalRedistDir, VersionNumber Version)
+		{
+			DirectoryReference? RedistDir;
+			if (OptionalRedistDir == null)
+			{
+				RedistDir = DirectoryReference.Combine(ToolChainDir, "redist"); // AutoSDK keeps redist under the toolchain
+			}
+			else
+			{
+				RedistDir = DirectoryReference.Combine(OptionalRedistDir, ToolChainDir.GetDirectoryName());
+
+				// exact redist not found - find highest version (they are backwards compatible)
+				if (!DirectoryReference.Exists(RedistDir) && DirectoryReference.Exists(OptionalRedistDir))
+				{
+					RedistDir = DirectoryReference.EnumerateDirectories(OptionalRedistDir)
+						.Where(X => VersionNumber.TryParse(X.GetDirectoryName(), out VersionNumber? DirVersion))
+						.OrderByDescending(X => VersionNumber.Parse(X.GetDirectoryName()))
+						.FirstOrDefault();
+				}
+			}
+
+			if (RedistDir != null && DirectoryReference.Exists(RedistDir))
+			{
+				return RedistDir;
+			}
+
+			return null;
+		}
+
+		/// <summary>
+		/// Adds a Visual C++ toolchain to a list of installations
+		/// </summary>
+		/// <param name="Version"></param>
+		/// <param name="bPreview"></param>
+		/// <param name="ToolChainDir"></param>
+		/// <param name="RedistDir"></param>
+		/// <param name="ToolChains"></param>
+		static void AddVisualCppToolChain(VersionNumber Version, bool bPreview, DirectoryReference ToolChainDir, DirectoryReference? RedistDir, List<ToolChainInstallation> ToolChains)
+		{
+			bool Is64Bit = Has64BitToolChain(ToolChainDir);
+
+			VersionNumber? Family;
+			if (!VersionNumber.TryParse(ToolChainDir.GetDirectoryName(), out Family))
+			{
+				Family = Version;
+			}
+
+			int FamilyRank = PreferredVisualCppVersions.TakeWhile(x => !x.Contains(Family)).Count();
+
+			string? Error = null;
+			if (Version < MinimumVisualCppVersion)
+			{
+				Error = $"UnrealBuildTool requires at minimum the MSVC {MinimumVisualCppVersion} toolchain. Please install a later toolchain from the Visual Studio installer.";
+			}
+
+			VersionNumberRange? Banned = BannedVisualCppVersions.FirstOrDefault(x => x.Contains(Version));
+			if (Banned != null)
+			{
+				Error = $"UnrealBuildTool has banned the MSVC {Banned} toolchains due to compiler issues. Please install a different toolchain from the Visual Studio installer.";
+			}
+
+			Log.TraceLog("Found Visual Studio toolchain: {0} (Family={1}, FamilyRank={2}, Version={3}, Is64Bit={4}, Preview={5}, Architecture={6}, Error={7}, Redist={8})", ToolChainDir, Family, FamilyRank, Version, Is64Bit, bPreview, WindowsArchitecture.x64.ToString(), Error != null, RedistDir);
+			ToolChains.Add(new ToolChainInstallation(Family, FamilyRank, Version, Is64Bit, bPreview, WindowsArchitecture.x64, Error, ToolChainDir, RedistDir));
+
+			bool HasArm64 = HasArm64ToolChain(ToolChainDir);
+			if (HasArm64)
+			{
+				Log.TraceLog("Found Visual Studio toolchain: {0} (Family={1}, FamilyRank={2}, Version={3}, Is64Bit={4}, Preview={5}, Architecture={6}, Error={7}, Redist={8})", ToolChainDir, Family, FamilyRank, Version, Is64Bit, bPreview, WindowsArchitecture.ARM64.ToString(), Error != null, RedistDir);
+				ToolChains.Add(new ToolChainInstallation(Family, FamilyRank, Version, Is64Bit, bPreview, WindowsArchitecture.ARM64, Error, ToolChainDir, RedistDir));
+			}
+		}
+
+		/// <summary>
+		/// Add a Clang toolchain
+		/// </summary>
+		/// <param name="ToolChainDir"></param>
+		/// <param name="ToolChains"></param>
+		static void AddClangToolChain(DirectoryReference ToolChainDir, List<ToolChainInstallation> ToolChains)
+		{
+			FileReference CompilerFile = FileReference.Combine(ToolChainDir, "bin", "clang-cl.exe");
+			if (FileReference.Exists(CompilerFile))
+			{
+				FileVersionInfo VersionInfo = FileVersionInfo.GetVersionInfo(CompilerFile.FullName);
+				VersionNumber Version = new VersionNumber(VersionInfo.FileMajorPart, VersionInfo.FileMinorPart, VersionInfo.FileBuildPart);
+
+				int Rank = PreferredClangVersions.TakeWhile(x => !x.Contains(Version)).Count();
+				bool Is64Bit = Is64BitExecutable(CompilerFile);
+
+				string? Error = null;
+				if (Version < MinimumClangVersion)
+				{
+					Error = $"UnrealBuildTool requires at minimum the Clang {MinimumClangVersion} toolchain. Please install a later toolchain from LLVM.";
+				}
+
+				Log.TraceLog("Found Clang toolchain: {0} (Version={1}, Is64Bit={2}, Rank={3})", ToolChainDir, Version, Is64Bit, Rank);
+				ToolChains.Add(new ToolChainInstallation(Version, Rank, Version, Is64Bit, false, WindowsArchitecture.x64, null, ToolChainDir, null));
+			}
+		}
+
+
+		/// <summary>
+		/// Add an Intel OneAPI toolchain
+		/// </summary>
+		/// <param name="ToolChainDir"></param>
+		/// <param name="ToolChains"></param>
+		static void AddIntelOneApiToolChain(DirectoryReference ToolChainDir, List<ToolChainInstallation> ToolChains)
+		{
+			FileReference CompilerFile = FileReference.Combine(ToolChainDir, "windows", "bin", "icx.exe");
+			if (FileReference.Exists(CompilerFile))
+			{
+				FileVersionInfo VersionInfo = FileVersionInfo.GetVersionInfo(CompilerFile.FullName);
+				VersionNumber Version = new VersionNumber(VersionInfo.FileMajorPart, VersionInfo.FileMinorPart, VersionInfo.FileBuildPart);
+
+				int Rank = PreferredIntelOneApiVersions.TakeWhile(x => !x.Contains(Version)).Count();
+				bool Is64Bit = Is64BitExecutable(CompilerFile);
+
+				string? Error = null;
+				if (Version < MinimumIntelOneApiVersion)
+				{
+					Error = $"UnrealBuildTool requires at minimum the Intel OneAPI {MinimumIntelOneApiVersion} toolchain. Please install a later toolchain from Intel.";
+				}
+
+				Log.TraceLog("Found Intel OneAPI toolchain: {0} (Version={1}, Is64Bit={2}, Rank={3})", ToolChainDir, Version, Is64Bit, Rank);
+				ToolChains.Add(new ToolChainInstallation(Version, Rank, Version, Is64Bit, false, WindowsArchitecture.x64, Error, ToolChainDir, null));
+			}
+		}
+
+		/// <summary>
+		/// Finds all the valid Intel oneAPI toolchains under the given base directory
+		/// </summary>
+		/// <param name="BaseDir">Base directory to search</param>
+		/// <param name="ToolChains">Map of tool chain version to installation info</param>
+		static void FindIntelOneApiToolChains(DirectoryReference BaseDir, List<ToolChainInstallation> ToolChains)
+		{
+			if (DirectoryReference.Exists(BaseDir))
+			{
+				foreach (DirectoryReference ToolChainDir in DirectoryReference.EnumerateDirectories(BaseDir))
+				{
+					AddIntelOneApiToolChain(ToolChainDir, ToolChains);
+				}
+			}
+		}
+
+
+		/// <summary>
+		/// Test whether an executable is 64-bit
+		/// </summary>
+		/// <param name="File">Executable to test</param>
+		/// <returns></returns>
+		static bool Is64BitExecutable(FileReference File)
+		{
+			using (FileStream Stream = new FileStream(File.FullName, FileMode.Open, FileAccess.Read, FileShare.Read | FileShare.Delete))
+			{
+				byte[] Header = new byte[64];
+				if (Stream.Read(Header, 0, Header.Length) != Header.Length)
+				{
+					return false;
+				}
+				if (Header[0] != (byte)'M' || Header[1] != (byte)'Z')
+				{
+					return false;
+				}
+
+				int Offset = BinaryPrimitives.ReadInt32LittleEndian(Header.AsSpan(0x3c));
+				if (Stream.Seek(Offset, SeekOrigin.Begin) != Offset)
+				{
+					return false;
+				}
+
+				byte[] PeHeader = new byte[6];
+				if (Stream.Read(PeHeader, 0, PeHeader.Length) != PeHeader.Length)
+				{
+					return false;
+				}
+				if (BinaryPrimitives.ReadUInt32BigEndian(PeHeader.AsSpan()) != 0x50450000)
+				{
+					return false;
+				}
+
+				ushort MachineType = BinaryPrimitives.ReadUInt16LittleEndian(PeHeader.AsSpan(4));
+
+				const ushort IMAGE_FILE_MACHINE_AMD64 = 0x8664;
+				return MachineType == IMAGE_FILE_MACHINE_AMD64;
+			}
+		}
+
+
+		/// <summary>
+		/// Checks if the given directory contains a valid Clang toolchain
+		/// </summary>
+		/// <param name="ToolChainDir">Directory to check</param>
+		/// <returns>True if the given directory is valid</returns>
+		static bool IsValidToolChainDirClang(DirectoryReference ToolChainDir)
+		{
+			return FileReference.Exists(FileReference.Combine(ToolChainDir, "bin", "clang-cl.exe"));
+		}
+
+		/// <summary>
+		/// Determines if the given path is a valid Visual C++ version number
+		/// </summary>
+		/// <param name="ToolChainDir">The toolchain directory</param>
+		/// <param name="Version">The version number for the toolchain</param>
+		/// <returns>True if the path is a valid version</returns>
+		static bool IsValidToolChainDirMSVC(DirectoryReference ToolChainDir, [NotNullWhen(true)] out VersionNumber? Version)
+		{
+			FileReference CompilerExe = FileReference.Combine(ToolChainDir, "bin", "Hostx86", "x64", "cl.exe");
+			if (!FileReference.Exists(CompilerExe))
+			{
+				CompilerExe = FileReference.Combine(ToolChainDir, "bin", "Hostx64", "x64", "cl.exe");
+				if (!FileReference.Exists(CompilerExe))
+				{
+					Version = null;
+					return false;
+				}
+			}
+
+			FileVersionInfo VersionInfo = FileVersionInfo.GetVersionInfo(CompilerExe.FullName);
+			if (VersionInfo.ProductMajorPart != 0)
+			{
+				Version = new VersionNumber(VersionInfo.ProductMajorPart, VersionInfo.ProductMinorPart, VersionInfo.ProductBuildPart);
+				return true;
+			}
+
+			return VersionNumber.TryParse(ToolChainDir.GetDirectoryName(), out Version);
+		}
+
+		/// <summary>
+		/// Checks if the given directory contains a 64-bit toolchain. Used to prefer regular Visual Studio versions over express editions.
+		/// </summary>
+		/// <param name="ToolChainDir">Directory to check</param>
+		/// <returns>True if the given directory contains a 64-bit toolchain</returns>
+		static bool Has64BitToolChain(DirectoryReference ToolChainDir)
+		{
+			return FileReference.Exists(FileReference.Combine(ToolChainDir, "bin", "amd64", "cl.exe")) || FileReference.Exists(FileReference.Combine(ToolChainDir, "bin", "Hostx64", "x64", "cl.exe"));
+		}
+
+		/// <summary>
+		/// Checks if the given directory contains a arm64 toolchain.  Used to require arm64, which is an optional install item, when that is our target architecture.
+		/// </summary>
+		/// <param name="ToolChainDir">Directory to check</param>
+		/// <returns>True if the given directory contains the arm64 toolchain</returns>
+		static bool HasArm64ToolChain(DirectoryReference ToolChainDir)
+		{
+			return FileReference.Exists(FileReference.Combine(ToolChainDir, "bin", "Hostx64", "arm64", "cl.exe"));
+		}
+
+		/// <summary>
+		/// Determines if an IDE for the given compiler is installed.
+		/// </summary>
+		/// <param name="Compiler">Compiler to check for</param>
+		/// <returns>True if the given compiler is installed</returns>
+		public static bool HasIDE(WindowsCompiler Compiler)
+		{
+			return FindVisualStudioInstallations(Compiler).Count > 0;
+		}
+
+		/// <summary>
+		/// Gets the path to MSBuild. This mirrors the logic in GetMSBuildPath.bat.
+		/// </summary>
+		/// <param name="OutLocation">On success, receives the path to the MSBuild executable.</param>
+		/// <returns>True on success.</returns>
+		public static bool TryGetMsBuildPath([NotNullWhen(true)] out FileReference? OutLocation)
+		{
+			// Get the Visual Studio 2019 install directory
+			List<DirectoryReference> InstallDirs2019 = MicrosoftPlatformSDK.FindVisualStudioInstallations(WindowsCompiler.VisualStudio2019).ConvertAll(x => x.BaseDir);
+			foreach (DirectoryReference InstallDir in InstallDirs2019)
+			{
+				FileReference MsBuildLocation = FileReference.Combine(InstallDir, "MSBuild", "Current", "Bin", "MSBuild.exe");
+				if (FileReference.Exists(MsBuildLocation))
+				{
+					OutLocation = MsBuildLocation;
+					return true;
+				}
+			}
+
+			// Get the Visual Studio 2022 install directory
+			List<DirectoryReference> InstallDirs2022 = MicrosoftPlatformSDK.FindVisualStudioInstallations(WindowsCompiler.VisualStudio2022).ConvertAll(x => x.BaseDir);
+			foreach (DirectoryReference InstallDir in InstallDirs2022)
+			{
+				FileReference MsBuildLocation = FileReference.Combine(InstallDir, "MSBuild", "Current", "Bin", "MSBuild.exe");
+				if (FileReference.Exists(MsBuildLocation))
+				{
+					OutLocation = MsBuildLocation;
+					return true;
+				}
+			}
+
+			// Try to get the MSBuild 14.0 path directly (see https://msdn.microsoft.com/en-us/library/hh162058(v=vs.120).aspx)
+			FileReference? ToolPath = FileReference.Combine(DirectoryReference.GetSpecialFolder(Environment.SpecialFolder.ProgramFilesX86)!, "MSBuild", "14.0", "bin", "MSBuild.exe");
+			if (FileReference.Exists(ToolPath))
+			{
+				OutLocation = ToolPath;
+				return true;
+			}
+
+			// Check for older versions of MSBuild. These are registered as separate versions in the registry.
+			if (TryReadMsBuildInstallPath("Microsoft\\MSBuild\\ToolsVersions\\14.0", "MSBuildToolsPath", "MSBuild.exe", out ToolPath))
+			{
+				OutLocation = ToolPath;
+				return true;
+			}
+			if (TryReadMsBuildInstallPath("Microsoft\\MSBuild\\ToolsVersions\\12.0", "MSBuildToolsPath", "MSBuild.exe", out ToolPath))
+			{
+				OutLocation = ToolPath;
+				return true;
+			}
+			if (TryReadMsBuildInstallPath("Microsoft\\MSBuild\\ToolsVersions\\4.0", "MSBuildToolsPath", "MSBuild.exe", out ToolPath))
+			{
+				OutLocation = ToolPath;
+				return true;
+			}
+
+			OutLocation = null;
+			return false;
+		}
+
+		/// <summary>
+		/// Function to query the registry under HKCU/HKLM Win32/Wow64 software registry keys for a certain install directory.
+		/// This mirrors the logic in GetMSBuildPath.bat.
+		/// </summary>
+		/// <returns></returns>
+		static bool TryReadMsBuildInstallPath(string KeyRelativePath, string KeyName, string MsBuildRelativePath, [NotNullWhen(true)] out FileReference? OutMsBuildPath)
+		{
+			string[] KeyBasePaths =
+			{
+				@"HKEY_CURRENT_USER\SOFTWARE\",
+				@"HKEY_LOCAL_MACHINE\SOFTWARE\",
+				@"HKEY_CURRENT_USER\SOFTWARE\Wow6432Node\",
+				@"HKEY_LOCAL_MACHINE\SOFTWARE\Wow6432Node\"
+			};
+
+			foreach (string KeyBasePath in KeyBasePaths)
+			{
+				string? Value = Registry.GetValue(KeyBasePath + KeyRelativePath, KeyName, null) as string;
+				if (Value != null)
+				{
+					FileReference MsBuildPath = FileReference.Combine(new DirectoryReference(Value), MsBuildRelativePath);
+					if (FileReference.Exists(MsBuildPath))
+					{
+						OutMsBuildPath = MsBuildPath;
+						return true;
+					}
+				}
+			}
+
+			OutMsBuildPath = null;
+			return false;
+		}
+
+		#endregion // Toolchain Private
+
+		#endregion // Toolchain
+
+		#region Dia Sdk
+
+		/// <summary>
+		/// Determines the directory containing the MSVC toolchain
+		/// </summary>
+		/// <param name="Compiler">Major version of the compiler to use</param>
+		/// <returns>Map of version number to directories</returns>
+		public static List<DirectoryReference> FindDiaSdkDirs(WindowsCompiler Compiler)
+		{
+			List<DirectoryReference>? DiaSdkDirs;
+			if (!CachedDiaSdkDirs.TryGetValue(Compiler, out DiaSdkDirs))
+			{
+				DiaSdkDirs = new List<DirectoryReference>();
+
+				DirectoryReference? PlatformDir;
+				if (UEBuildPlatformSDK.TryGetHostPlatformAutoSDKDir(out PlatformDir))
+				{
+					string VSDir = string.Empty;
+					switch (Compiler)
+					{
+						case WindowsCompiler.VisualStudio2019: VSDir = "VS2019"; break;
+						case WindowsCompiler.VisualStudio2022: VSDir = "VS2022"; break;
+					}
+
+					if (!string.IsNullOrEmpty(VSDir))
+					{
+						DirectoryReference DiaSdkDir = DirectoryReference.Combine(PlatformDir, "Win64", "DIA SDK", VSDir);
+						if (IsValidDiaSdkDir(DiaSdkDir))
+						{
+							DiaSdkDirs.Add(DiaSdkDir);
+						}
+					}
+				}
+
+				List<DirectoryReference> VisualStudioDirs = MicrosoftPlatformSDK.FindVisualStudioInstallations(Compiler).ConvertAll(x => x.BaseDir);
+				foreach (DirectoryReference VisualStudioDir in VisualStudioDirs)
+				{
+					DirectoryReference DiaSdkDir = DirectoryReference.Combine(VisualStudioDir, "DIA SDK");
+					if (IsValidDiaSdkDir(DiaSdkDir))
+					{
+						DiaSdkDirs.Add(DiaSdkDir);
+					}
+				}
+			}
+			return DiaSdkDirs;
+		}
+
+		/// <summary>
+		/// Determines if a directory contains a valid DIA SDK
+		/// </summary>
+		/// <param name="DiaSdkDir">The directory to check</param>
+		/// <returns>True if it contains a valid DIA SDK</returns>
+		static bool IsValidDiaSdkDir(DirectoryReference DiaSdkDir)
+		{
+			return FileReference.Exists(FileReference.Combine(DiaSdkDir, "bin", "amd64", "msdia140.dll"));
+		}
+
+		#endregion // Dia
+
+		#endregion // Windows Specific SDK 
 	}
+
+
+	/// <summary>
+	/// Information about a particular toolchain installation
+	/// </summary>
+	internal class ToolChainInstallation
+	{
+		/// <summary>
+		/// The version "family" (ie. the nominal version number of the directory this toolchain is installed to)
+		/// </summary>
+		public VersionNumber Family { get; }
+
+		/// <summary>
+		/// Index into the preferred version range
+		/// </summary>
+		public int FamilyRank { get; }
+
+		/// <summary>
+		/// The actual version number of this toolchain
+		/// </summary>
+		public VersionNumber Version { get; }
+
+		/// <summary>
+		/// Whether this is a 64-bit toolchain
+		/// </summary>
+		public bool Is64Bit { get; }
+
+		/// <summary>
+		/// Whether it's a pre-release version of the toolchain.
+		/// </summary>
+		public bool IsPreview { get; }
+
+		/// <summary>
+		/// The architecture of this ToolChainInstallation (multiple ToolChainInstallation instances may be created, one per architecture).
+		/// </summary>
+		public WindowsArchitecture Architecture { get; }
+
+		/// <summary>
+		/// Reason for this toolchain not being compatible
+		/// </summary>
+		public string? Error { get; }
+
+		/// <summary>
+		/// Base directory for the toolchain
+		/// </summary>
+		public DirectoryReference BaseDir { get; }
+
+		/// <summary>
+		/// Base directory for the redistributable components
+		/// </summary>
+		public DirectoryReference? RedistDir { get; }
+
+		/// <summary>
+		/// Constructor
+		/// </summary>
+		/// <param name="Family"></param>
+		/// <param name="FamilyRank"></param>
+		/// <param name="Version"></param>
+		/// <param name="Is64Bit"></param>
+		/// <param name="IsPreview">Whether it's a pre-release version of the toolchain</param>
+		/// <param name="Architecture"></param>
+		/// <param name="Error"></param>
+		/// <param name="BaseDir">Base directory for the toolchain</param>
+		/// <param name="RedistDir">Optional directory for redistributable components (DLLs etc)</param>
+		public ToolChainInstallation(VersionNumber Family, int FamilyRank, VersionNumber Version, bool Is64Bit, bool IsPreview, WindowsArchitecture Architecture, string? Error, DirectoryReference BaseDir, DirectoryReference? RedistDir)
+		{
+			this.Family = Family;
+			this.FamilyRank = FamilyRank;
+			this.Version = Version;
+			this.Is64Bit = Is64Bit;
+			this.IsPreview = IsPreview;
+			this.Architecture = Architecture;
+			this.Error = Error;
+			this.BaseDir = BaseDir;
+			this.RedistDir = RedistDir;
+		}
+	}
 }
diff --git a/Engine/Source/Programs/UnrealBuildTool/Platform/Windows/UEBuildWindows.cs b/Engine/Source/Programs/UnrealBuildTool/Platform/Windows/UEBuildWindows.cs
--- a/Engine/Source/Programs/UnrealBuildTool/Platform/Windows/UEBuildWindows.cs
+++ b/Engine/Source/Programs/UnrealBuildTool/Platform/Windows/UEBuildWindows.cs
@@ -5,13 +5,9 @@ using System.Collections.Generic;
 using System.Text;
 using System.Diagnostics;
 using System.IO;
-using Microsoft.Win32;
 using System.Linq;
 using EpicGames.Core;
-using Microsoft.VisualStudio.Setup.Configuration;
-using System.Runtime.InteropServices;
 using System.Diagnostics.CodeAnalysis;
-using System.Buffers.Binary;
 using UnrealBuildBase;
 
 namespace UnrealBuildTool
@@ -418,7 +414,7 @@ namespace UnrealBuildTool
 		/// </summary>
 		public string DiaSdkDir
 		{
-			get { return WindowsPlatform.FindDiaSdkDirs(Environment!.Compiler).Select(x => x.FullName).FirstOrDefault(); }
+			get { return MicrosoftPlatformSDK.FindDiaSdkDirs(Environment!.Compiler).Select(x => x.FullName).FirstOrDefault(); }
 		}
 
 		/// <summary>
@@ -430,7 +426,7 @@ namespace UnrealBuildTool
 			{
 				try
 				{
-					return WindowsPlatform.FindVisualStudioInstallations(Environment!.Compiler).Select(x => x.BaseDir.FullName).FirstOrDefault();
+					return MicrosoftPlatformSDK.FindVisualStudioInstallations(Environment!.Compiler).Select(x => x.BaseDir.FullName).FirstOrDefault();
 				}
 				catch(Exception) // Find function will throw if there is no visual studio installed! This can happen w/ clang builds
 				{
@@ -751,143 +747,6 @@ namespace UnrealBuildTool
 
 	class WindowsPlatform : UEBuildPlatform
 	{
-		/// <summary>
-		/// Information about a particular toolchain installation
-		/// </summary>
-		class ToolChainInstallation
-		{
-			/// <summary>
-			/// The version "family" (ie. the nominal version number of the directory this toolchain is installed to)
-			/// </summary>
-			public VersionNumber Family { get; }
-
-			/// <summary>
-			/// Index into the preferred version range
-			/// </summary>
-			public int FamilyRank { get; }
-
-			/// <summary>
-			/// The actual version number of this toolchain
-			/// </summary>
-			public VersionNumber Version { get; }
-
-			/// <summary>
-			/// Whether this is a 64-bit toolchain
-			/// </summary>
-			public bool Is64Bit { get; }
-
-			/// <summary>
-			/// Whether it's a pre-release version of the toolchain.
-			/// </summary>
-			public bool IsPreview { get; }
-
-			/// <summary>
-			/// The architecture of this ToolChainInstallation (multiple ToolChainInstallation instances may be created, one per architecture).
-			/// </summary>
-			public WindowsArchitecture Architecture { get; }
-
-			/// <summary>
-			/// Reason for this toolchain not being compatible
-			/// </summary>
-			public string? Error { get; }
-
-			/// <summary>
-			/// Base directory for the toolchain
-			/// </summary>
-			public DirectoryReference BaseDir { get; }
-
-			/// <summary>
-			/// Base directory for the redistributable components
-			/// </summary>
-			public DirectoryReference? RedistDir { get; }
-
-			/// <summary>
-			/// Constructor
-			/// </summary>
-			/// <param name="Family"></param>
-			/// <param name="FamilyRank"></param>
-			/// <param name="Version"></param>
-			/// <param name="Is64Bit"></param>
-			/// <param name="IsPreview">Whether it's a pre-release version of the toolchain</param>
-			/// <param name="Architecture"></param>
-			/// <param name="Error"></param>
-			/// <param name="BaseDir">Base directory for the toolchain</param>
-			/// <param name="RedistDir">Optional directory for redistributable components (DLLs etc)</param>
-			public ToolChainInstallation(VersionNumber Family, int FamilyRank, VersionNumber Version, bool Is64Bit, bool IsPreview, WindowsArchitecture Architecture, string? Error, DirectoryReference BaseDir, DirectoryReference? RedistDir)
-			{
-				this.Family = Family;
-				this.FamilyRank = FamilyRank;
-				this.Version = Version;
-				this.Is64Bit = Is64Bit;
-				this.IsPreview = IsPreview;
-				this.Architecture = Architecture;
-				this.Error = Error;
-				this.BaseDir = BaseDir;
-				this.RedistDir = RedistDir;
-			}
-		}
-
-		/// <summary>
-		/// The default compiler version to be used, if installed. 
-		/// </summary>
-		static readonly VersionNumberRange[] PreferredClangVersions =
-		{
-			VersionNumberRange.Parse("13.0.0", "13.999"),
-			VersionNumberRange.Parse("12.0.0", "12.999"),
-			VersionNumberRange.Parse("11.0.0", "11.999"),
-			VersionNumberRange.Parse("10.0.0", "10.999")
-		};
-		
-		static readonly VersionNumber MinimumClangVersion = new VersionNumber(10, 0, 0);
-
-		/// <summary>
-		/// Ranges of tested compiler toolchains to be used, in order of preference. If multiple toolchains in a range are present, the latest version will be preferred.
-		/// Note that the numbers here correspond to the installation *folders* rather than precise executable versions. 
-		/// </summary>
-		static readonly VersionNumberRange[] PreferredVisualCppVersions = new VersionNumberRange[]
-		{
-			VersionNumberRange.Parse("14.29.30133", "14.29.30136"), // VS2019 16.11.5
-		};
-
-		/// <summary>
-		/// Tested compiler toolchanges that should not be allowed.
-		/// </summary>
-		static readonly VersionNumberRange[] BannedVisualCppVersions = new VersionNumberRange[]
-		{
-		};
-
-		static readonly VersionNumber MinimumVisualCppVersion = new VersionNumber(14, 29, 30133);
-
-		/// <summary>
-		/// The default compiler version to be used, if installed. 
-		/// </summary>
-		static readonly VersionNumberRange[] PreferredIntelOneApiVersions =
-		{
-			VersionNumberRange.Parse("2022.0.0", "2022.9999"),
-		};
-		
-		static readonly VersionNumber MinimumIntelOneApiVersion = new VersionNumber(2021, 0, 0);
-
-		/// <summary>
-		/// Visual Studio installations
-		/// </summary>
-		public static IReadOnlyList<VisualStudioInstallation> VisualStudioInstallations => CachedVisualStudioInstallations.Value;
-
-		/// <summary>
-		/// Cache of Visual Studio installation directories
-		/// </summary>
-		private static Lazy<IReadOnlyList<VisualStudioInstallation>> CachedVisualStudioInstallations = new Lazy<IReadOnlyList<VisualStudioInstallation>>(FindVisualStudioInstallations);
-
-		/// <summary>
-		/// Cache of Visual C++ installation directories
-		/// </summary>
-		private static Dictionary<WindowsCompiler, List<ToolChainInstallation>> CachedToolChainInstallations = new Dictionary<WindowsCompiler, List<ToolChainInstallation>>();
-
-		/// <summary>
-		/// Cache of DIA SDK installation directories
-		/// </summary>
-		private static Dictionary<WindowsCompiler, List<DirectoryReference>> CachedDiaSdkDirs = new Dictionary<WindowsCompiler, List<DirectoryReference>>();
-
 		/// <summary>
 		/// True if we should use the Clang linker (LLD) when we are compiling with Clang, otherwise we use the MSVC linker
 		/// </summary>
@@ -1105,11 +964,11 @@ namespace UnrealBuildTool
 			}
 
 			// Second, default based on what's installed, test for 2019 first
-			if (FindToolChainInstallations(WindowsCompiler.VisualStudio2019).Where(x => x.Version >= MinimumVisualCppVersion && x.Architecture == Architecture).Count() > 0)
+			if (MicrosoftPlatformSDK.HasValidCompiler(WindowsCompiler.VisualStudio2019, Architecture))
 			{
 				return WindowsCompiler.VisualStudio2019;
 			}
-			else if (FindToolChainInstallations(WindowsCompiler.VisualStudio2022).Where(x => x.Version >= MinimumVisualCppVersion && x.Architecture == Architecture).Count() > 0)
+			else if (MicrosoftPlatformSDK.HasValidCompiler(WindowsCompiler.VisualStudio2022, Architecture))
 			{
 				return WindowsCompiler.VisualStudio2022;
 			}
@@ -1145,15 +1004,7 @@ namespace UnrealBuildTool
 		/// <returns>Name of the compiler</returns>
 		public static string GetCompilerName(WindowsCompiler Compiler)
 		{
-			switch (Compiler)
-			{
-				case WindowsCompiler.VisualStudio2019:
-					return "Visual Studio 2019";
-				case WindowsCompiler.VisualStudio2022:
-					return "Visual Studio 2022";
-				default:
-					return Compiler.ToString();
-			}
+			return MicrosoftPlatformSDK.GetCompilerName(Compiler);
 		}
 
 		/// <summary>
@@ -1164,7 +1015,7 @@ namespace UnrealBuildTool
 		/// <returns>True if the directory was found, false otherwise.</returns>
 		public static IEnumerable<DirectoryReference>? TryGetVSInstallDirs(WindowsCompiler Compiler)
 		{
-			List<VisualStudioInstallation> Installations = FindVisualStudioInstallations(Compiler);
+			List<VisualStudioInstallation> Installations = MicrosoftPlatformSDK.FindVisualStudioInstallations(Compiler);
 			if(Installations.Count == 0)
 			{
 				return null;
@@ -1173,489 +1024,6 @@ namespace UnrealBuildTool
 			return Installations.Select(x => x.BaseDir);
 		}
 
-		/// <summary>
-		/// Read the Visual Studio install directory for the given compiler version. Note that it is possible for the compiler toolchain to be installed without
-		/// Visual Studio, and vice versa.
-		/// </summary>
-		/// <returns>List of directories containing Visual Studio installations</returns>
-		private static IReadOnlyList<VisualStudioInstallation> FindVisualStudioInstallations()
-		{
-			List<VisualStudioInstallation> Installations = new List<VisualStudioInstallation>();
-			if (BuildHostPlatform.Current.Platform == UnrealTargetPlatform.Win64)
-			{
-				try
-				{
-					SetupConfiguration Setup = new SetupConfiguration();
-					IEnumSetupInstances Enumerator = Setup.EnumAllInstances();
-
-					ISetupInstance[] Instances = new ISetupInstance[1];
-					for (; ; )
-					{
-						int NumFetched;
-						Enumerator.Next(1, Instances, out NumFetched);
-
-						if (NumFetched == 0)
-						{
-							break;
-						}
-
-						ISetupInstance2 Instance = (ISetupInstance2)Instances[0];
-						if ((Instance.GetState() & InstanceState.Local) != InstanceState.Local)
-						{
-							continue;
-						}
-
-						VersionNumber? Version;
-						if (!VersionNumber.TryParse(Instance.GetInstallationVersion(), out Version))
-						{
-							continue;
-						}
-
-						int MajorVersion = Version.GetComponent(0);
-
-						WindowsCompiler Compiler;
-						if (MajorVersion >= 17) // Treat any newer versions as 2022, until we have an explicit enum for them
-						{
-							Compiler = WindowsCompiler.VisualStudio2022;
-						}
-						else if (MajorVersion == 16)
-						{
-							Compiler = WindowsCompiler.VisualStudio2019;
-						}
-						else
-						{
-							continue;
-						}
-
-						ISetupInstanceCatalog? Catalog = Instance as ISetupInstanceCatalog;
-						bool bPreview = Catalog != null && Catalog.IsPrerelease();
-
-						string ProductId = Instance.GetProduct().GetId();
-						bool bExpress = ProductId.Equals("Microsoft.VisualStudio.Product.WDExpress", StringComparison.Ordinal);
-
-						DirectoryReference BaseDir = new DirectoryReference(Instance.GetInstallationPath());
-						Installations.Add(new VisualStudioInstallation(Compiler, Version, BaseDir, bExpress, bPreview));
-
-						Log.TraceLog("Found Visual Studio installation: {0} (Product={1}, Version={2})", BaseDir, ProductId, Version);
-					}
-
-					Installations = Installations.OrderBy(x => x.bExpress).ThenBy(x => x.bPreview).ThenByDescending(x => x.Version).ToList();
-				}
-				catch (Exception Ex)
-				{
-					throw new BuildException(Ex, "Error while enumerating Visual Studio toolchains");
-				}
-			}
-			return Installations;
-		}
-
-		/// <summary>
-		/// Read the Visual Studio install directory for the given compiler version. Note that it is possible for the compiler toolchain to be installed without
-		/// Visual Studio, and vice versa.
-		/// </summary>
-		/// <returns>List of directories containing Visual Studio installations</returns>
-		public static List<VisualStudioInstallation> FindVisualStudioInstallations(WindowsCompiler Compiler)
-		{
-			return VisualStudioInstallations.Where(x => x.Compiler == Compiler).ToList();
-		}
-
-		/// <summary>
-		/// Determines the directory containing the MSVC toolchain
-		/// </summary>
-		/// <param name="Compiler">Major version of the compiler to use</param>
-		/// <returns>Map of version number to directories</returns>
-		static List<ToolChainInstallation> FindToolChainInstallations(WindowsCompiler Compiler)
-		{
-			List<ToolChainInstallation>? ToolChains;
-			if(!CachedToolChainInstallations.TryGetValue(Compiler, out ToolChains))
-			{
-				ToolChains = new List<ToolChainInstallation>();
-			    if (BuildHostPlatform.Current.Platform == UnrealTargetPlatform.Win64)
-			    {
-					if(Compiler == WindowsCompiler.Clang)
-					{
-						// Check for a manual installation to the default directory
-						DirectoryReference ManualInstallDir = DirectoryReference.Combine(DirectoryReference.GetSpecialFolder(Environment.SpecialFolder.ProgramFiles)!, "LLVM");
-						AddClangToolChain(ManualInstallDir, ToolChains);
-
-						// Check for a manual installation to a custom directory
-						string? LlvmPath = Environment.GetEnvironmentVariable("LLVM_PATH");
-						if (!String.IsNullOrEmpty(LlvmPath))
-						{
-							AddClangToolChain(new DirectoryReference(LlvmPath), ToolChains);
-						}
-
-						// Check for installations bundled with Visual Studio 2019
-						foreach (VisualStudioInstallation Installation in FindVisualStudioInstallations(WindowsCompiler.VisualStudio2019))
-						{
-							AddClangToolChain(DirectoryReference.Combine(Installation.BaseDir, "VC", "Tools", "Llvm"), ToolChains);
-						}
-
-						// Check for installations bundled with Visual Studio 2022
-						foreach (VisualStudioInstallation Installation in FindVisualStudioInstallations(WindowsCompiler.VisualStudio2022))
-						{
-							AddClangToolChain(DirectoryReference.Combine(Installation.BaseDir, "VC", "Tools", "Llvm"), ToolChains);
-						}
-
-						// Check for AutoSDK paths
-						DirectoryReference? AutoSdkDir;
-						if(UEBuildPlatformSDK.TryGetHostPlatformAutoSDKDir(out AutoSdkDir))
-						{
-							DirectoryReference ClangBaseDir = DirectoryReference.Combine(AutoSdkDir, "Win64", "LLVM");
-							if(DirectoryReference.Exists(ClangBaseDir))
-							{
-								foreach(DirectoryReference ToolChainDir in DirectoryReference.EnumerateDirectories(ClangBaseDir))
-								{
-									AddClangToolChain(ToolChainDir, ToolChains);
-								}
-							}
-						}
-					}
-					else if(Compiler == WindowsCompiler.Intel)
-					{
-						DirectoryReference InstallDir = DirectoryReference.Combine(DirectoryReference.GetSpecialFolder(Environment.SpecialFolder.ProgramFiles)!, "Intel", "oneAPI", "compiler");
-						FindIntelOneApiToolChains(InstallDir, ToolChains);
-
-						InstallDir = DirectoryReference.Combine(DirectoryReference.GetSpecialFolder(Environment.SpecialFolder.ProgramFilesX86)!, "Intel", "oneAPI", "compiler");
-						FindIntelOneApiToolChains(InstallDir, ToolChains);
-
-						// Check for AutoSDK paths
-						DirectoryReference? AutoSdkDir;
-						if (UEBuildPlatformSDK.TryGetHostPlatformAutoSDKDir(out AutoSdkDir))
-						{
-							DirectoryReference IntelBaseDir = DirectoryReference.Combine(AutoSdkDir, "Win64", "Intel");
-							if (DirectoryReference.Exists(IntelBaseDir))
-							{
-								FindIntelOneApiToolChains(IntelBaseDir, ToolChains);
-							}
-						}
-					}
-				    else if(Compiler.IsMSVC())
-				    {
-						// Enumerate all the manually installed toolchains
-						List<VisualStudioInstallation> Installations = FindVisualStudioInstallations(Compiler);
-					    foreach(VisualStudioInstallation Installation in Installations)
-					    {
-						    DirectoryReference ToolChainBaseDir = DirectoryReference.Combine(Installation.BaseDir, "VC", "Tools", "MSVC");
-						    DirectoryReference RedistBaseDir = DirectoryReference.Combine(Installation.BaseDir, "VC", "Redist", "MSVC");
-							FindVisualStudioToolChains(ToolChainBaseDir, RedistBaseDir, Installation.bPreview, ToolChains);
-					    }
-
-						// Enumerate all the AutoSDK toolchains
-						DirectoryReference? PlatformDir;
-						if (UEBuildPlatformSDK.TryGetHostPlatformAutoSDKDir(out PlatformDir))
-						{
-							string VSDir = string.Empty;
-							switch (Compiler)
-							{
-								case WindowsCompiler.VisualStudio2019: VSDir = "VS2019"; break;
-								case WindowsCompiler.VisualStudio2022: VSDir = "VS2022"; break;
-							}
-
-							if (!string.IsNullOrEmpty(VSDir))
-							{
-								DirectoryReference ReleaseBaseDir = DirectoryReference.Combine(PlatformDir, "Win64", VSDir);
-								FindVisualStudioToolChains(ReleaseBaseDir, null, false, ToolChains);
-
-								DirectoryReference PreviewBaseDir = DirectoryReference.Combine(PlatformDir, "Win64", $"{VSDir}-Preview");
-								FindVisualStudioToolChains(PreviewBaseDir, null, true, ToolChains);
-							}
-						}
-					}
-					else
-				    {
-					    throw new BuildException("Unsupported compiler version ({0})", Compiler);
-				    }
-				}
-				CachedToolChainInstallations.Add(Compiler, ToolChains);
-			}
-			return ToolChains;
-		}
-
-		/// <summary>
-		/// Finds all the valid Visual Studio toolchains under the given base directory
-		/// </summary>
-		/// <param name="BaseDir">Base directory to search</param>
-		/// <param name="OptionalRedistDir">Optional directory for redistributable components (DLLs etc)</param>
-		/// <param name="bPreview">Whether this is a preview installation</param>
-		/// <param name="ToolChains">Map of tool chain version to installation info</param>
-		static void FindVisualStudioToolChains(DirectoryReference BaseDir, DirectoryReference? OptionalRedistDir, bool bPreview, List<ToolChainInstallation> ToolChains)
-		{
-			if (DirectoryReference.Exists(BaseDir))
-			{
-				foreach (DirectoryReference ToolChainDir in DirectoryReference.EnumerateDirectories(BaseDir))
-				{
-					VersionNumber? Version;
-					if (IsValidToolChainDirMSVC(ToolChainDir, out Version))
-					{
-						DirectoryReference? RedistDir = FindVisualStudioRedistForToolChain(ToolChainDir, OptionalRedistDir, Version );
-
-						AddVisualCppToolChain(Version, bPreview, ToolChainDir, RedistDir, ToolChains);
-					}
-				}
-			}
-		}
-
-		/// <summary>
-		/// Finds all the valid Intel oneAPI toolchains under the given base directory
-		/// </summary>
-		/// <param name="BaseDir">Base directory to search</param>
-		/// <param name="ToolChains">Map of tool chain version to installation info</param>
-		static void FindIntelOneApiToolChains(DirectoryReference BaseDir, List<ToolChainInstallation> ToolChains)
-		{
-			if (DirectoryReference.Exists(BaseDir))
-			{
-				foreach (DirectoryReference ToolChainDir in DirectoryReference.EnumerateDirectories(BaseDir))
-				{
-					AddIntelOneApiToolChain(ToolChainDir, ToolChains);
-				}
-			}
-		}
-
-		/// <summary>
-		/// Finds the most appropriate redist directory for the given toolchain version
-		/// </summary>
-		/// <param name="ToolChainDir"></param>
-		/// <param name="OptionalRedistDir"></param>
-		/// <param name="Version"></param>
-		/// <returns></returns>
-		static DirectoryReference? FindVisualStudioRedistForToolChain( DirectoryReference ToolChainDir, DirectoryReference? OptionalRedistDir, VersionNumber Version )
-		{
-			DirectoryReference? RedistDir;
-			if (OptionalRedistDir == null)
-			{
-				RedistDir = DirectoryReference.Combine(ToolChainDir, "redist"); // AutoSDK keeps redist under the toolchain
-			}
-			else
-			{
-				RedistDir = DirectoryReference.Combine(OptionalRedistDir, ToolChainDir.GetDirectoryName() );
-
-				// exact redist not found - find highest version (they are backwards compatible)
-				if (!DirectoryReference.Exists(RedistDir) && DirectoryReference.Exists(OptionalRedistDir))
-				{
-					RedistDir = DirectoryReference.EnumerateDirectories(OptionalRedistDir)
-						.Where( X => VersionNumber.TryParse(X.GetDirectoryName(), out VersionNumber? DirVersion) )
-						.OrderByDescending( X => VersionNumber.Parse(X.GetDirectoryName()) )
-						.FirstOrDefault();
-				}
-			}
-			
-			if (RedistDir != null && DirectoryReference.Exists(RedistDir) )
-			{
-				return RedistDir;
-			}
-
-			return null;
-		}
-
-		/// <summary>
-		/// Add a Clang toolchain
-		/// </summary>
-		/// <param name="ToolChainDir"></param>
-		/// <param name="ToolChains"></param>
-		static void AddClangToolChain(DirectoryReference ToolChainDir, List<ToolChainInstallation> ToolChains)
-		{
-			FileReference CompilerFile = FileReference.Combine(ToolChainDir, "bin", "clang-cl.exe");
-			if (FileReference.Exists(CompilerFile))
-			{
-				FileVersionInfo VersionInfo = FileVersionInfo.GetVersionInfo(CompilerFile.FullName);
-				VersionNumber Version = new VersionNumber(VersionInfo.FileMajorPart, VersionInfo.FileMinorPart, VersionInfo.FileBuildPart);
-
-				int Rank = PreferredClangVersions.TakeWhile(x => !x.Contains(Version)).Count();
-				bool Is64Bit = Is64BitExecutable(CompilerFile);
-
-				string? Error = null;
-				if (Version < MinimumClangVersion)
-				{
-					Error = $"UnrealBuildTool requires at minimum the Clang {MinimumClangVersion} toolchain. Please install a later toolchain from LLVM.";
-				}
-
-				Log.TraceLog("Found Clang toolchain: {0} (Version={1}, Is64Bit={2}, Rank={3})", ToolChainDir, Version, Is64Bit, Rank);
-				ToolChains.Add(new ToolChainInstallation(Version, Rank, Version, Is64Bit, false, WindowsArchitecture.x64, null, ToolChainDir, null));
-			}
-		}
-
-		/// <summary>
-		/// Add an Intel OneAPI toolchain
-		/// </summary>
-		/// <param name="ToolChainDir"></param>
-		/// <param name="ToolChains"></param>
-		static void AddIntelOneApiToolChain(DirectoryReference ToolChainDir, List<ToolChainInstallation> ToolChains)
-		{
-			FileReference CompilerFile = FileReference.Combine(ToolChainDir, "windows", "bin", "icx.exe");
-			if (FileReference.Exists(CompilerFile))
-			{
-				FileVersionInfo VersionInfo = FileVersionInfo.GetVersionInfo(CompilerFile.FullName);
-				VersionNumber Version = new VersionNumber(VersionInfo.FileMajorPart, VersionInfo.FileMinorPart, VersionInfo.FileBuildPart);
-
-				int Rank = PreferredIntelOneApiVersions.TakeWhile(x => !x.Contains(Version)).Count();
-				bool Is64Bit = Is64BitExecutable(CompilerFile);
-
-				string? Error = null;
-				if (Version < MinimumIntelOneApiVersion)
-				{
-					Error = $"UnrealBuildTool requires at minimum the Intel OneAPI {MinimumIntelOneApiVersion} toolchain. Please install a later toolchain from Intel.";
-				}
-
-				Log.TraceLog("Found Intel OneAPI toolchain: {0} (Version={1}, Is64Bit={2}, Rank={3})", ToolChainDir, Version, Is64Bit, Rank);
-				ToolChains.Add(new ToolChainInstallation(Version, Rank, Version, Is64Bit, false, WindowsArchitecture.x64, Error, ToolChainDir, null));
-			}
-		}
-
-		/// <summary>
-		/// Test whether an executable is 64-bit
-		/// </summary>
-		/// <param name="File">Executable to test</param>
-		/// <returns></returns>
-		static bool Is64BitExecutable(FileReference File)
-		{
-			using (FileStream Stream = new FileStream(File.FullName, FileMode.Open, FileAccess.Read, FileShare.Read | FileShare.Delete))
-			{
-				byte[] Header = new byte[64];
-				if (Stream.Read(Header, 0, Header.Length) != Header.Length)
-				{
-					return false;
-				}
-				if (Header[0] != (byte)'M' || Header[1] != (byte)'Z')
-				{
-					return false;
-				}
-
-				int Offset = BinaryPrimitives.ReadInt32LittleEndian(Header.AsSpan(0x3c));
-				if (Stream.Seek(Offset, SeekOrigin.Begin) != Offset)
-				{
-					return false;
-				}
-
-				byte[] PeHeader = new byte[6];
-				if(Stream.Read(PeHeader, 0, PeHeader.Length) != PeHeader.Length)
-				{
-					return false;
-				}
-				if (BinaryPrimitives.ReadUInt32BigEndian(PeHeader.AsSpan()) != 0x50450000)
-				{
-					return false;
-				}
-
-				ushort MachineType = BinaryPrimitives.ReadUInt16LittleEndian(PeHeader.AsSpan(4));
-
-				const ushort IMAGE_FILE_MACHINE_AMD64 = 0x8664;
-				return MachineType == IMAGE_FILE_MACHINE_AMD64;
-			}
-		}
-
-		/// <summary>
-		/// Adds a Visual C++ toolchain to a list of installations
-		/// </summary>
-		/// <param name="Version"></param>
-		/// <param name="bPreview"></param>
-		/// <param name="ToolChainDir"></param>
-		/// <param name="RedistDir"></param>
-		/// <param name="ToolChains"></param>
-		static void AddVisualCppToolChain(VersionNumber Version, bool bPreview, DirectoryReference ToolChainDir, DirectoryReference? RedistDir, List<ToolChainInstallation> ToolChains)
-		{
-			bool Is64Bit = Has64BitToolChain(ToolChainDir);
-
-			VersionNumber? Family;
-			if (!VersionNumber.TryParse(ToolChainDir.GetDirectoryName(), out Family))
-			{
-				Family = Version;
-			}
-
-			int FamilyRank = PreferredVisualCppVersions.TakeWhile(x => !x.Contains(Family)).Count();
-
-			string? Error = null;
-			if (Version < MinimumVisualCppVersion)
-			{
-				Error = $"UnrealBuildTool requires at minimum the MSVC {MinimumVisualCppVersion} toolchain. Please install a later toolchain from the Visual Studio installer.";
-			}
-
-			VersionNumberRange? Banned = BannedVisualCppVersions.FirstOrDefault(x => x.Contains(Version));
-			if (Banned != null)
-			{
-				Error = $"UnrealBuildTool has banned the MSVC {Banned} toolchains due to compiler issues. Please install a different toolchain from the Visual Studio installer.";
-			}
-
-			Log.TraceLog("Found Visual Studio toolchain: {0} (Family={1}, FamilyRank={2}, Version={3}, Is64Bit={4}, Preview={5}, Architecture={6}, Error={7}, Redist={8})", ToolChainDir, Family, FamilyRank, Version, Is64Bit, bPreview, WindowsArchitecture.x64.ToString(), Error != null, RedistDir);
-			ToolChains.Add(new ToolChainInstallation(Family, FamilyRank, Version, Is64Bit, bPreview, WindowsArchitecture.x64, Error, ToolChainDir, RedistDir));
-
-			bool HasArm64 = HasArm64ToolChain(ToolChainDir);
-			if (HasArm64)
-			{
-				Log.TraceLog("Found Visual Studio toolchain: {0} (Family={1}, FamilyRank={2}, Version={3}, Is64Bit={4}, Preview={5}, Architecture={6}, Error={7}, Redist={8})", ToolChainDir, Family, FamilyRank, Version, Is64Bit, bPreview, WindowsArchitecture.ARM64.ToString(), Error != null, RedistDir);
-				ToolChains.Add(new ToolChainInstallation(Family, FamilyRank, Version, Is64Bit, bPreview, WindowsArchitecture.ARM64, Error, ToolChainDir, RedistDir));
-			}
-		}
-
-		/// <summary>
-		/// Checks if the given directory contains a valid Clang toolchain
-		/// </summary>
-		/// <param name="ToolChainDir">Directory to check</param>
-		/// <returns>True if the given directory is valid</returns>
-		static bool IsValidToolChainDirClang(DirectoryReference ToolChainDir)
-		{
-			return FileReference.Exists(FileReference.Combine(ToolChainDir, "bin", "clang-cl.exe"));
-		}
-
-		/// <summary>
-		/// Determines if the given path is a valid Visual C++ version number
-		/// </summary>
-		/// <param name="ToolChainDir">The toolchain directory</param>
-		/// <param name="Version">The version number for the toolchain</param>
-		/// <returns>True if the path is a valid version</returns>
-		static bool IsValidToolChainDirMSVC(DirectoryReference ToolChainDir, [NotNullWhen(true)] out VersionNumber? Version)
-		{
-			FileReference CompilerExe = FileReference.Combine(ToolChainDir, "bin", "Hostx86", "x64", "cl.exe");
-			if(!FileReference.Exists(CompilerExe))
-			{
-				CompilerExe = FileReference.Combine(ToolChainDir, "bin", "Hostx64", "x64", "cl.exe");
-				if(!FileReference.Exists(CompilerExe))
-				{
-					Version = null;
-					return false;
-				}
-			}
-
-			FileVersionInfo VersionInfo = FileVersionInfo.GetVersionInfo(CompilerExe.FullName);
-			if (VersionInfo.ProductMajorPart != 0)
-			{
-				Version = new VersionNumber(VersionInfo.ProductMajorPart, VersionInfo.ProductMinorPart, VersionInfo.ProductBuildPart);
-				return true;
-			}
-
-			return VersionNumber.TryParse(ToolChainDir.GetDirectoryName(), out Version);
-		}
-
-		/// <summary>
-		/// Checks if the given directory contains a 64-bit toolchain. Used to prefer regular Visual Studio versions over express editions.
-		/// </summary>
-		/// <param name="ToolChainDir">Directory to check</param>
-		/// <returns>True if the given directory contains a 64-bit toolchain</returns>
-		static bool Has64BitToolChain(DirectoryReference ToolChainDir)
-		{
-			return FileReference.Exists(FileReference.Combine(ToolChainDir, "bin", "amd64", "cl.exe")) || FileReference.Exists(FileReference.Combine(ToolChainDir, "bin", "Hostx64", "x64", "cl.exe"));
-		}
-
-		/// <summary>
-		/// Checks if the given directory contains a arm64 toolchain.  Used to require arm64, which is an optional install item, when that is our target architecture.
-		/// </summary>
-		/// <param name="ToolChainDir">Directory to check</param>
-		/// <returns>True if the given directory contains the arm64 toolchain</returns>
-		static bool HasArm64ToolChain(DirectoryReference ToolChainDir)
-		{
-			return FileReference.Exists(FileReference.Combine(ToolChainDir, "bin", "Hostx64", "arm64", "cl.exe"));
-		}
-
-		/// <summary>
-		/// Determines if an IDE for the given compiler is installed.
-		/// </summary>
-		/// <param name="Compiler">Compiler to check for</param>
-		/// <returns>True if the given compiler is installed</returns>
-		public static bool HasIDE(WindowsCompiler Compiler)
-		{
-			return FindVisualStudioInstallations(Compiler).Count > 0;
-		}
-
 		/// <summary>
 		/// Determines if a given compiler is installed
 		/// </summary>
@@ -1664,32 +1032,7 @@ namespace UnrealBuildTool
 		/// <returns>True if the given compiler is installed</returns>
 		public static bool HasCompiler(WindowsCompiler Compiler, WindowsArchitecture Architecture)
 		{
-			return FindToolChainInstallations(Compiler).Where(x => (x.Architecture == Architecture)).Count() > 0;
-		}
-
-		/// <summary>
-		/// Select which toolchain to use, combining a custom preference with a default sort order
-		/// </summary>
-		/// <param name="ToolChains"></param>
-		/// <param name="Preference">Ordering function</param>
-		/// <param name="Architecture">Architecture that must be supported</param>
-		/// <returns></returns>
-		static ToolChainInstallation? SelectToolChain(IEnumerable<ToolChainInstallation> ToolChains, Func<IOrderedEnumerable<ToolChainInstallation>, IOrderedEnumerable<ToolChainInstallation>> Preference, WindowsArchitecture Architecture)
-		{
-			ToolChainInstallation? ToolChain = Preference(ToolChains.Where(x => x.Architecture == Architecture)
-				.OrderByDescending(x => x.Error == null))
-				.ThenByDescending(x => x.Is64Bit)
-				.ThenBy(x => x.IsPreview)
-				.ThenBy(x => x.FamilyRank)
-				.ThenByDescending(x => x.Version)
-				.FirstOrDefault();
-
-			if (ToolChain?.Error != null)
-			{
-				throw new BuildException(ToolChain.Error);
-			}
-
-			return ToolChain;
+			return MicrosoftPlatformSDK.HasCompiler(Compiler, Architecture);
 		}
 
 		/// <summary>
@@ -1704,130 +1047,7 @@ namespace UnrealBuildTool
 		/// <returns>True if the toolchain directory was found correctly</returns>
 		public static bool TryGetToolChainDir(WindowsCompiler Compiler, string? CompilerVersion, WindowsArchitecture Architecture, [NotNullWhen(true)] out VersionNumber? OutToolChainVersion, [NotNullWhen(true)] out DirectoryReference? OutToolChainDir, out DirectoryReference? OutRedistDir)
 		{
-			// Find all the installed toolchains
-			List<ToolChainInstallation> ToolChains = FindToolChainInstallations(Compiler);
-
-			// Figure out the actual version number that we want
-			ToolChainInstallation? ToolChain = null;
-			if (CompilerVersion == null)
-			{
-				ToolChain = SelectToolChain(ToolChains, x => x, Architecture);
-				if (ToolChain == null)
-				{
-					OutToolChainVersion = null;
-					OutToolChainDir = null;
-					OutRedistDir = null;
-					return false;
-				}
-			}
-			else if (String.Compare(CompilerVersion, "Latest", StringComparison.InvariantCultureIgnoreCase) == 0)
-			{
-				ToolChain = SelectToolChain(ToolChains, x => x.ThenByDescending(x => x.Version), Architecture);
-				if (ToolChain == null)
-				{
-					throw new BuildException("Unable to find C++ toolchain for {0} {1}", Compiler, Architecture.ToString());
-				}
-			}
-			else if (String.Compare(CompilerVersion, "Preview", StringComparison.InvariantCultureIgnoreCase) == 0)
-			{
-				ToolChain = SelectToolChain(ToolChains, x => x.ThenByDescending(x => x.IsPreview), Architecture);
-				if (ToolChain == null || !ToolChain.IsPreview)
-				{
-					throw new BuildException("Unable to find preview toolchain for {0} {1}", Compiler, Architecture.ToString());
-				}
-			}
-			else if (VersionNumber.TryParse(CompilerVersion, out VersionNumber? ToolChainVersion))
-			{
-				ToolChain = SelectToolChain(ToolChains, x => x.ThenByDescending(x => x.Version == ToolChainVersion).ThenByDescending(x => x.Family == ToolChainVersion), Architecture);
-				if (ToolChain == null)
-				{
-					throw new BuildException("Unable to find {0} toolchain for {1} {2}", ToolChainVersion, Compiler, Architecture.ToString());
-				}
-			}
-			else
-			{
-				throw new BuildException("Unable to find {0} toolchain; '{1}' is an invalid version", GetCompilerName(Compiler), CompilerVersion);
-			}
-
-			// Get the actual directory for this version
-			OutToolChainVersion = ToolChain.Version;
-			OutToolChainDir = ToolChain.BaseDir;
-			OutRedistDir = ToolChain.RedistDir;
-			return true;
-		}
-
-		/// <summary>
-		/// Gets the path to MSBuild. This mirrors the logic in GetMSBuildPath.bat.
-		/// </summary>
-		/// <param name="OutLocation">On success, receives the path to the MSBuild executable.</param>
-		/// <returns>True on success.</returns>
-		public static bool TryGetMsBuildPath([NotNullWhen(true)] out FileReference? OutLocation)
-		{
-			// Get the Visual Studio 2019 install directory
-			List<DirectoryReference> InstallDirs2019 = WindowsPlatform.FindVisualStudioInstallations(WindowsCompiler.VisualStudio2019).ConvertAll(x => x.BaseDir);
-			foreach (DirectoryReference InstallDir in InstallDirs2019)
-			{
-				FileReference MsBuildLocation = FileReference.Combine(InstallDir, "MSBuild", "Current", "Bin", "MSBuild.exe");
-				if (FileReference.Exists(MsBuildLocation))
-				{
-					OutLocation = MsBuildLocation;
-					return true;
-				}
-			}
-
-			// Get the Visual Studio 2022 install directory
-			List<DirectoryReference> InstallDirs2022 = WindowsPlatform.FindVisualStudioInstallations(WindowsCompiler.VisualStudio2022).ConvertAll(x => x.BaseDir);
-			foreach (DirectoryReference InstallDir in InstallDirs2022)
-			{
-				FileReference MsBuildLocation = FileReference.Combine(InstallDir, "MSBuild", "Current", "Bin", "MSBuild.exe");
-				if(FileReference.Exists(MsBuildLocation))
-				{
-					OutLocation = MsBuildLocation;
-					return true;
-				}
-			}
-
-			// Try to get the MSBuild 14.0 path directly (see https://msdn.microsoft.com/en-us/library/hh162058(v=vs.120).aspx)
-			FileReference? ToolPath = FileReference.Combine(DirectoryReference.GetSpecialFolder(Environment.SpecialFolder.ProgramFilesX86)!, "MSBuild", "14.0", "bin", "MSBuild.exe");
-			if(FileReference.Exists(ToolPath))
-			{
-				OutLocation = ToolPath;
-				return true;
-			} 
-
-			// Check for older versions of MSBuild. These are registered as separate versions in the registry.
-			if (TryReadMsBuildInstallPath("Microsoft\\MSBuild\\ToolsVersions\\14.0", "MSBuildToolsPath", "MSBuild.exe", out ToolPath))
-			{
-				OutLocation = ToolPath;
-				return true;
-			}
-			if (TryReadMsBuildInstallPath("Microsoft\\MSBuild\\ToolsVersions\\12.0", "MSBuildToolsPath", "MSBuild.exe", out ToolPath))
-			{
-				OutLocation = ToolPath;
-				return true;
-			}
-			if (TryReadMsBuildInstallPath("Microsoft\\MSBuild\\ToolsVersions\\4.0", "MSBuildToolsPath", "MSBuild.exe", out ToolPath))
-			{
-				OutLocation = ToolPath;
-				return true;
-			}
-
-			OutLocation = null;
-			return false;
-		}
-
-		/// <summary>
-		/// Gets the MSBuild path, and throws an exception on failure.
-		/// </summary>
-		/// <returns>Path to MSBuild</returns>
-		public static FileReference GetMsBuildToolPath()
-		{
-			FileReference? Location;
-			if(!TryGetMsBuildPath(out Location))
-			{
-				throw new BuildException("Unable to find installation of MSBuild.");
-			}
- 			return Location;
+			return MicrosoftPlatformSDK.TryGetToolChainDir(Compiler, CompilerVersion, Architecture, out OutToolChainVersion, out OutToolChainDir, out OutRedistDir);
 		}
 
 		public static string GetArchitectureSubpath(WindowsArchitecture arch)
@@ -1844,83 +1064,6 @@ namespace UnrealBuildTool
 			return archPath;
 		}
 
-		/// <summary>
-		/// Function to query the registry under HKCU/HKLM Win32/Wow64 software registry keys for a certain install directory.
-		/// This mirrors the logic in GetMSBuildPath.bat.
-		/// </summary>
-		/// <returns></returns>
-		static bool TryReadMsBuildInstallPath(string KeyRelativePath, string KeyName, string MsBuildRelativePath, [NotNullWhen(true)] out FileReference? OutMsBuildPath)
-		{
-			string[] KeyBasePaths =
-			{
-				@"HKEY_CURRENT_USER\SOFTWARE\",
-				@"HKEY_LOCAL_MACHINE\SOFTWARE\",
-				@"HKEY_CURRENT_USER\SOFTWARE\Wow6432Node\",
-				@"HKEY_LOCAL_MACHINE\SOFTWARE\Wow6432Node\"
-			};
-
-			foreach (string KeyBasePath in KeyBasePaths)
-			{
-				string? Value = Registry.GetValue(KeyBasePath + KeyRelativePath, KeyName, null) as string;
-				if (Value != null)
-				{
-					FileReference MsBuildPath = FileReference.Combine(new DirectoryReference(Value), MsBuildRelativePath);
-					if (FileReference.Exists(MsBuildPath))
-					{
-						OutMsBuildPath = MsBuildPath;
-						return true;
-					}
-				}
-			}
-
-			OutMsBuildPath = null;
-			return false;
-		}
-
-		/// <summary>
-		/// Determines the directory containing the MSVC toolchain
-		/// </summary>
-		/// <param name="Compiler">Major version of the compiler to use</param>
-		/// <returns>Map of version number to directories</returns>
-		public static List<DirectoryReference> FindDiaSdkDirs(WindowsCompiler Compiler)
-		{
-			List<DirectoryReference>? DiaSdkDirs;
-			if (!CachedDiaSdkDirs.TryGetValue(Compiler, out DiaSdkDirs))
-			{
-				DiaSdkDirs = new List<DirectoryReference>();
-
-				DirectoryReference? PlatformDir;
-				if (UEBuildPlatformSDK.TryGetHostPlatformAutoSDKDir(out PlatformDir))
-				{
-					string VSDir = string.Empty;
-					switch (Compiler)
-					{
-						case WindowsCompiler.VisualStudio2019: VSDir = "VS2019"; break;
-						case WindowsCompiler.VisualStudio2022: VSDir = "VS2022"; break;
-					}
-
-					if (!string.IsNullOrEmpty(VSDir))
-					{
-						DirectoryReference DiaSdkDir = DirectoryReference.Combine(PlatformDir, "Win64", "DIA SDK", VSDir);
-						if (IsValidDiaSdkDir(DiaSdkDir))
-						{
-							DiaSdkDirs.Add(DiaSdkDir);
-						}
-					}
-				}
-
-				List<DirectoryReference> VisualStudioDirs = VisualStudioInstallations.ConvertAll(x => x.BaseDir);
-				foreach (DirectoryReference VisualStudioDir in VisualStudioDirs)
-				{
-					DirectoryReference DiaSdkDir = DirectoryReference.Combine(VisualStudioDir, "DIA SDK");
-					if (IsValidDiaSdkDir(DiaSdkDir))
-					{
-						DiaSdkDirs.Add(DiaSdkDir);
-					}
-				}
-			}
-			return DiaSdkDirs;
-		}
 
 		/// <summary>
 		/// Determines if a directory contains a valid DIA SDK
diff --git a/Engine/Source/Programs/UnrealBuildTool/Platform/Windows/WindowsExports.cs b/Engine/Source/Programs/UnrealBuildTool/Platform/Windows/WindowsExports.cs
--- a/Engine/Source/Programs/UnrealBuildTool/Platform/Windows/WindowsExports.cs
+++ b/Engine/Source/Programs/UnrealBuildTool/Platform/Windows/WindowsExports.cs
@@ -31,7 +31,7 @@ namespace UnrealBuildTool
 		/// <returns>Path to MSBuild.exe</returns>
 		public static string GetMSBuildToolPath()
 		{
-			return WindowsPlatform.GetMsBuildToolPath().FullName;
+			return MicrosoftPlatformSDK.GetMsBuildToolPath().FullName;
 		}
 
 		/// <summary>
diff --git a/Engine/Source/Programs/UnrealBuildTool/ProjectFiles/VisualStudio/VCProjectFileGenerator.cs b/Engine/Source/Programs/UnrealBuildTool/ProjectFiles/VisualStudio/VCProjectFileGenerator.cs
--- a/Engine/Source/Programs/UnrealBuildTool/ProjectFiles/VisualStudio/VCProjectFileGenerator.cs
+++ b/Engine/Source/Programs/UnrealBuildTool/ProjectFiles/VisualStudio/VCProjectFileGenerator.cs
@@ -279,7 +279,7 @@ namespace UnrealBuildTool
 			if (Settings.ProjectFileFormat == VCProjectFileFormat.Default)
 			{
 				// Enumerate all the valid installations. This list is already sorted by preference.
-				List<VisualStudioInstallation> Installations = WindowsPlatform.VisualStudioInstallations.Where(x => WindowsPlatform.HasCompiler(x.Compiler, WindowsArchitecture.x64)).ToList();
+				List<VisualStudioInstallation> Installations = MicrosoftPlatformSDK.VisualStudioInstallations.Where(x => WindowsPlatform.HasCompiler(x.Compiler, WindowsArchitecture.x64)).ToList();
 
 				// Get the corresponding project file format
 				VCProjectFileFormat Format = VCProjectFileFormat.VisualStudio2019;
