UnrealBuildTool: Return correct IntelOneAPI verson when icx.exe isn't updated

#rnx

diff --git a/Engine/Source/Programs/UnrealBuildTool/Platform/Windows/MicrosoftPlatformSDK.cs b/Engine/Source/Programs/UnrealBuildTool/Platform/Windows/MicrosoftPlatformSDK.cs
--- a/Engine/Source/Programs/UnrealBuildTool/Platform/Windows/MicrosoftPlatformSDK.cs
+++ b/Engine/Source/Programs/UnrealBuildTool/Platform/Windows/MicrosoftPlatformSDK.cs
@@ -977,6 +977,12 @@ namespace UnrealBuildTool
 				FileVersionInfo VersionInfo = FileVersionInfo.GetVersionInfo(CompilerFile.FullName);
 				VersionNumber Version = new VersionNumber(VersionInfo.FileMajorPart, VersionInfo.FileMinorPart, VersionInfo.FileBuildPart);
 
+				// The icx.exe version may be lower than the toolchain folder version, so use that instead if available
+				if (VersionNumber.TryParse(ToolChainDir.GetDirectoryName(), out VersionNumber? FolderVersion) && FolderVersion > Version)
+				{
+					Version = FolderVersion;
+				}
+
 				int Rank = PreferredIntelOneApiVersions.TakeWhile(x => !x.Contains(Version)).Count();
 				bool Is64Bit = Is64BitExecutable(CompilerFile);
 
