UnrealBuildTool: Print warnings of any normalized resolved paths in VCToolchain are longer that MAX_PATH as those consistently cause unexpected build errors

#rnx

diff --git a/Engine/Source/Programs/UnrealBuildTool/Platform/Windows/VCToolChain.cs b/Engine/Source/Programs/UnrealBuildTool/Platform/Windows/VCToolChain.cs
--- a/Engine/Source/Programs/UnrealBuildTool/Platform/Windows/VCToolChain.cs
+++ b/Engine/Source/Programs/UnrealBuildTool/Platform/Windows/VCToolChain.cs
@@ -25,6 +25,11 @@ namespace UnrealBuildTool
 		/// </summary>
 		protected VCEnvironment EnvVars;
 
+		/// <summary>
+		/// Length of a path string that will trigger a build warning, as long paths may cause unexpected errors with the MSVC toolchain.
+		/// </summary>
+		private static int MaxPathWarningLength = 260;
+
 		public VCToolChain(ReadOnlyTargetRules Target)
 		{
 			this.Target = Target;
@@ -114,14 +119,33 @@ namespace UnrealBuildTool
 			}
 		}
 
+		private static void CheckCommandLinePathLength(string PathString)
+		{
+			if (!Path.IsPathRooted(PathString))
+			{
+				string ResolvedPath = Path.Combine(UnrealBuildTool.EngineSourceDirectory.FullName, PathString);
+				if (ResolvedPath.Length > MaxPathWarningLength)
+				{
+					Log.TraceWarningOnce($"Relative path '{PathString}' when resolved will have length '{ResolvedPath.Length}' which is greater than MAX_PATH (260) and may cause unexpected errors with the MSVC toolchain.");
+				}
+			}
+			else if (PathString.Length > MaxPathWarningLength)
+			{
+				Log.TraceWarningOnce($"Absolute path '{PathString}' has length '{PathString.Length}' which is greater than MAX_PATH (260) and may cause unexpected errors with the MSVC toolchain.");
+			}
+		}
+
 		public static string NormalizeCommandLinePath(FileSystemReference Reference)
 		{
-			// Try to use a relative path to shorten command line length.
+			// Try to use a relative path to shorten command line length and to enabled remote distribution where absolute paths are not desired
 			if (Reference.IsUnderDirectory(Unreal.EngineDirectory))
 			{
-				return Reference.MakeRelativeTo(UnrealBuildTool.EngineSourceDirectory);
+				string RelativePath = Reference.MakeRelativeTo(UnrealBuildTool.EngineSourceDirectory);
+				CheckCommandLinePathLength(RelativePath);
+				return RelativePath;
 			}
 
+			CheckCommandLinePathLength(Reference.FullName);
 			return Reference.FullName;
 		}
 
