UnrealBuildTool: Use dotnet UnrealBuildTool.dll (with correct path) for Make and QMake generated files

#rnx

diff --git a/Engine/Source/Programs/UnrealBuildTool/ProjectFiles/Make/MakefileGenerator.cs b/Engine/Source/Programs/UnrealBuildTool/ProjectFiles/Make/MakefileGenerator.cs
--- a/Engine/Source/Programs/UnrealBuildTool/ProjectFiles/Make/MakefileGenerator.cs
+++ b/Engine/Source/Programs/UnrealBuildTool/ProjectFiles/Make/MakefileGenerator.cs
@@ -68,7 +68,7 @@ namespace UnrealBuildTool
 			{
 				GameProjectFile = OnlyGameProject!.FullName;
 				MakeGameProjectFile = "GAMEPROJECTFILE =" + GameProjectFile + "\n";
-				ProjectBuildCommand = $"PROJECTBUILD = bash \"$(UNREALROOTPATH)/Engine/{Unreal.RelativeDotnetDirectory}/dotnet\" \"$(UNREALROOTPATH)/Engine/Binaries/DotNET/UnrealBuildTool.dll\"\n";
+				ProjectBuildCommand = $"PROJECTBUILD = bash \"$(UNREALROOTPATH)/Engine/{Unreal.RelativeDotnetDirectory}/dotnet\" \"$(UNREALROOTPATH)/Engine/Binaries/DotNET/UnrealBuildTool/UnrealBuildTool.dll\"\n";
 			}
 
 			BuildCommand = "BUILD = bash \"$(UNREALROOTPATH)/Engine/Build/BatchFiles/Linux/Build.sh\"\n";
diff --git a/Engine/Source/Programs/UnrealBuildTool/ProjectFiles/QMake/QMakefileGenerator.cs b/Engine/Source/Programs/UnrealBuildTool/ProjectFiles/QMake/QMakefileGenerator.cs
--- a/Engine/Source/Programs/UnrealBuildTool/ProjectFiles/QMake/QMakefileGenerator.cs
+++ b/Engine/Source/Programs/UnrealBuildTool/ProjectFiles/QMake/QMakefileGenerator.cs
@@ -210,7 +210,7 @@ namespace UnrealBuildTool
 				GameProjectPath = OnlyGameProject!.Directory.FullName;
 				GameProjectFile = OnlyGameProject.FullName;
 				QMakeGameProjectFile = "gameProjectFile=" + GameProjectFile + "\n";
-				BuildCommand = $"build=bash $$unrealRootPath/Engine/{Unreal.RelativeDotnetDirectory}/dotnet $$unrealRootPath/Engine/Binaries/DotNET/UnrealBuildTool.exe\n\n";
+				BuildCommand = $"build=bash $$unrealRootPath/Engine/{Unreal.RelativeDotnetDirectory}/dotnet $$unrealRootPath/Engine/Binaries/DotNET/UnrealBuildTool/UnrealBuildTool.dll\n\n";
 			}
 			else
 			{
