UBT detects script source directories more robustly
- besides existence also checks if any script source files are present within

#rnx

diff --git a/Engine/Source/Programs/UnrealBuildTool/Configuration/UEBuildModuleCPP.cs b/Engine/Source/Programs/UnrealBuildTool/Configuration/UEBuildModuleCPP.cs
--- a/Engine/Source/Programs/UnrealBuildTool/Configuration/UEBuildModuleCPP.cs
+++ b/Engine/Source/Programs/UnrealBuildTool/Configuration/UEBuildModuleCPP.cs
@@ -210,7 +210,7 @@ namespace UnrealBuildTool
 
 			// Check for a Verse directory next to the rules file
 			DirectoryReference MaybeVerseDirectory = DirectoryReference.Combine(Rules.File.Directory, "Verse");
-			if (DirectoryReference.Exists(MaybeVerseDirectory))
+			if (IsValidVerseDirectory(MaybeVerseDirectory))
 			{
 				this.AssociatedVerseDirectory = MaybeVerseDirectory;
 				this.bDependsOnVerse = true;
@@ -1710,5 +1710,26 @@ namespace UnrealBuildTool
 		{
 			return FindInputFilesFromDirectory(Directory, new InputFileCollection());
 		}
+
+		/// <summary>
+		/// Checks a given directory path whether it exists and if it contains any Verse source files
+		/// </summary>
+		public static bool IsValidVerseDirectory(DirectoryReference MaybeVerseDirectory)
+		{
+			if (!DirectoryReference.Exists(MaybeVerseDirectory))
+			{
+				return false;
+			}
+
+			foreach (string FilePath in Directory.EnumerateFiles(MaybeVerseDirectory.FullName, "*.v*", SearchOption.AllDirectories))
+			{
+				if (FilePath.EndsWith(".verse") || FilePath.EndsWith(".vmodule"))
+				{
+					return true;
+				}
+			}
+
+			return false;
+		}
 	}
 }
