Manifest generation optimizations to improve iteration times when nothing has changed:

 - resource binary files are not copied to an Intermediate location first.
 - resources.pri is also not created in the Intermediate location first.
 - file diffs are done on generated files only; resource binary files use timestamps.
 - only copy the resource binaries for the cultures that we are staging.
 - final file removal & copying is left to the end of the process rather than being ad-hoc.

Additionally, the output from makepri is only logged if there was an error, to reduce signal/noise.

#rnx

diff --git a/Engine/Source/Programs/UnrealBuildTool/Platform/Windows/VCManifestGenerator.cs b/Engine/Source/Programs/UnrealBuildTool/Platform/Windows/VCManifestGenerator.cs
--- a/Engine/Source/Programs/UnrealBuildTool/Platform/Windows/VCManifestGenerator.cs
+++ b/Engine/Source/Programs/UnrealBuildTool/Platform/Windows/VCManifestGenerator.cs
@@ -74,8 +74,8 @@ namespace UnrealBuildTool
 		/// intermediate path - used for temporary files
 		protected string? IntermediatePath;
 
-		/// files that have been updated
-		protected List<string>? UpdatedFilePaths;
+		/// files that should be included in the manifest output folder and where to source them from
+		protected Dictionary<string,string>? ManifestFiles; // Dst, Src
 
 		/// <summary>
 		/// Create a manifest generator for the given platform variant.
@@ -124,47 +124,6 @@ namespace UnrealBuildTool
 			return true;
 		}
 
-		/// <summary>
-		/// Delete the entire directory tree
-		/// </summary>
-		protected static void RecursivelyForceDeleteDirectory(string InDirectoryToDelete)
-		{
-			if (Directory.Exists(InDirectoryToDelete))
-			{
-				try
-				{
-					List<string> SubDirectories = new List<string>(Directory.GetDirectories(InDirectoryToDelete, "*.*", SearchOption.AllDirectories));
-					foreach (string DirectoryToRemove in SubDirectories)
-					{
-						RecursivelyForceDeleteDirectory(DirectoryToRemove);
-					}
-					List<string> FilesInDirectory = new List<string>(Directory.GetFiles(InDirectoryToDelete));
-					foreach (string FileToRemove in FilesInDirectory)
-					{
-						try
-						{
-							FileAttributes Attributes = File.GetAttributes(FileToRemove);
-							if ((Attributes & FileAttributes.ReadOnly) == FileAttributes.ReadOnly)
-							{
-								Attributes &= ~FileAttributes.ReadOnly;
-								File.SetAttributes(FileToRemove, Attributes);
-							}
-							File.Delete(FileToRemove);
-						}
-						catch (Exception)
-						{
-							Log.TraceWarning("Could not remove file {0} to remove directory {1}.", FileToRemove, InDirectoryToDelete);
-						}
-					}
-					Directory.Delete(InDirectoryToDelete, true);
-				}
-				catch (Exception)
-				{
-					Log.TraceWarning("Could not remove directory {0}.", InDirectoryToDelete);
-				}
-			}
-		}
-
 
         /// <summary>
         /// Runs Makepri. Blocking
@@ -180,19 +139,35 @@ namespace UnrealBuildTool
 				throw new BuildException("BUILD FAILED: Couldn't find the makepri executable: {0}", PriExecutable);
 			}
 
+			StringBuilder ProcessOutput = new StringBuilder();
+			void LocalProcessOutput(DataReceivedEventArgs Args)
+			{
+				if (Args != null && Args.Data != null)
+				{
+					ProcessOutput.AppendLine(Args.Data.TrimEnd());
+				}
+			}
+
 			ProcessStartInfo StartInfo = new ProcessStartInfo(PriExecutable, CommandLine);			
 			StartInfo.UseShellExecute = false;
 			StartInfo.CreateNoWindow = true;
 			StartInfo.StandardOutputEncoding = Encoding.Unicode;
 			StartInfo.StandardErrorEncoding = Encoding.Unicode;
-			int ExitCode = Utils.RunLocalProcessAndLogOutput(StartInfo);
+
+			Process LocalProcess = new Process();
+			LocalProcess.StartInfo = StartInfo;
+			LocalProcess.OutputDataReceived += (Sender, Args) => { LocalProcessOutput(Args); };
+			LocalProcess.ErrorDataReceived += (Sender, Args) => { LocalProcessOutput(Args); };
+			int ExitCode = Utils.RunLocalProcess(LocalProcess);
 
 			if (ExitCode == 0)
 			{
+				Log.TraceVerbose(ProcessOutput.ToString());
 				return true;
 			}
 			else
 			{
+				Log.TraceInformation(ProcessOutput.ToString());
 				Log.TraceError(Path.GetFileName(PriExecutable) + " returned an error.\nExit code:" + ExitCode );
 				return false;
 			}
@@ -338,6 +313,42 @@ namespace UnrealBuildTool
 			return DefaultValue;
 		}
 
+
+
+
+		private bool RemoveStaleResourceFiles()
+		{
+			// remove all resource files that should not be included
+			string TargetResourceDir = Path.Combine(OutputPath!, BuildResourceSubPath);
+			var TargetResourceInstances = Directory.EnumerateFiles(TargetResourceDir, "*.*", SearchOption.AllDirectories);
+
+			var StaleResourceFiles = TargetResourceInstances.Where(X => !ManifestFiles!.ContainsKey(X)).ToList();
+			if (StaleResourceFiles.Any())
+			{
+				Log.TraceVerbose("Removing stale manifest resource files...");
+				foreach (string StaleResourceFile in StaleResourceFiles)
+				{
+					// try to delete the file & the directory that contains it
+					try
+					{
+						Log.TraceVerbose($"\tremoving {Utils.MakePathRelativeTo(StaleResourceFile, OutputPath!)}");
+						FileUtils.ForceDeleteFile(StaleResourceFile);
+						if (!Directory.EnumerateFileSystemEntries(Path.GetDirectoryName(StaleResourceFile)).Any())
+						{
+							Directory.Delete(Path.GetDirectoryName(StaleResourceFile), false);
+						}
+					}
+					catch (Exception E)
+					{
+						Log.TraceError($"\tCould not remove {StaleResourceFile} - {E.Message}.");
+					}
+				}
+			}
+
+			return StaleResourceFiles.Any();
+		}
+
+
 		/// <summary>
 		/// Attempts to locate the given resource binary file in several known folder locations
 		/// </summary>
@@ -380,14 +391,15 @@ namespace UnrealBuildTool
 			return FindResourceBinaryFile( out SourcePath, ResourceFileName, AllowEngineFallback );
 		}
 
+
 		/// <summary>
-		/// Updates the given resource binary file
+		/// Adds the given resource binary file(s) to the manifest files
 		/// </summary>
-		protected bool CopyAndReplaceBinaryIntermediate(string ResourceFileName, bool AllowEngineFallback = true)
+		protected bool AddResourceBinaryFileReference(string ResourceFileName, bool AllowEngineFallback = true)
 		{
-			string TargetPath = Path.Combine(IntermediatePath!, BuildResourceSubPath);
+			string TargetPath = Path.Combine(OutputPath!, BuildResourceSubPath);
 			string SourcePath;
-			bool bFileExists = FindResourceBinaryFile( out SourcePath, ResourceFileName, AllowEngineFallback );
+			bool bFileExists = FindResourceBinaryFile(out SourcePath, ResourceFileName, AllowEngineFallback);
 
 			// At least the default culture entry for any resource binary must always exist
 			if (!bFileExists)
@@ -402,155 +414,134 @@ namespace UnrealBuildTool
 			}
 
 			// Find all copies of the resource file in the source directory (could be up to one for each culture and the default).
-			IEnumerable<string> SourceResourceInstances = Directory.EnumerateFiles(SourcePath, ResourceFileName, SearchOption.AllDirectories);
+			List<string> SourceResourceInstances = new List<string>();
+			SourceResourceInstances.Add( Path.Combine( SourcePath, ResourceFileName) );
+			foreach( string CultureId in CulturesToStage!)
+			{
+				string CultureResourceFile = Path.Combine(SourcePath, CultureId, ResourceFileName);
+				if (File.Exists(CultureResourceFile))
+				{
+					SourceResourceInstances.Add( CultureResourceFile );
+				}
+			}
 
 			// Copy new resource files
 			foreach (string SourceResourceFile in SourceResourceInstances)
 			{
-				//@todo only copy files for cultures we are staging
 				string TargetResourcePath = Path.Combine(TargetPath, SourceResourceFile.Substring(SourcePath.Length + 1));
 				if (!CreateCheckDirectory(Path.GetDirectoryName(TargetResourcePath)!))
 				{
 					Log.TraceError("Unable to create intermediate directory {0}.", Path.GetDirectoryName(TargetResourcePath));
 					continue;
 				}
-				if (!File.Exists(TargetResourcePath))
-				{
-					try
-					{
-						File.Copy(SourceResourceFile, TargetResourcePath);
-
-						// File.Copy also copies the attributes, so make sure the new file isn't read only
-						FileAttributes Attrs = File.GetAttributes(TargetResourcePath);
-						if (Attrs.HasFlag(FileAttributes.ReadOnly))
-						{
-							File.SetAttributes(TargetResourcePath, Attrs & ~FileAttributes.ReadOnly);
-						}
-					}
-					catch (Exception)
-					{
-						Log.TraceError("Unable to copy file {0} to {1}.", SourceResourceFile, TargetResourcePath);
-						return false;
-					}
-				}
+				AddFileReference(SourceResourceFile, TargetResourcePath, bIsGeneratedFile:false);
 			}
 
 			return true;
 		}
 
+
 		/// <summary>
-		/// Updates the given file if necessary
+		/// Adds the given file to the manifest files
 		/// </summary>
-		protected void CompareAndReplaceModifiedTarget(string IntermediatePath, string TargetPath)
+		protected void AddFileReference(string SourcePath, string TargetPath, bool bIsGeneratedFile)
 		{
-			if (!File.Exists(IntermediatePath))
+			// check if the file is unchanged
+			if (File.Exists(TargetPath) && !string.IsNullOrEmpty(SourcePath) )
 			{
-				Log.TraceError("Tried to copy non-existant intermediate file {0}.", IntermediatePath);
-				return;
-			}
+				FileInfo SrcFileInfo = new FileInfo(SourcePath);
+				FileInfo DstFileInfo = new FileInfo(TargetPath);
 
-			CreateCheckDirectory(Path.GetDirectoryName(TargetPath)!);
-
-			// Check for differences in file contents
-			if (File.Exists(TargetPath))
-			{
-				byte[] OriginalContents = File.ReadAllBytes(TargetPath);
-				byte[] NewContents = File.ReadAllBytes(IntermediatePath);
-				if (!OriginalContents.Equals(NewContents))
+				bool bFileIsUnchanged = (SrcFileInfo.Length == DstFileInfo.Length);
+				if (bFileIsUnchanged)
 				{
-					try
+					if (bIsGeneratedFile)
 					{
-						FileAttributes Attrs = File.GetAttributes(TargetPath);
-						if ((Attrs & FileAttributes.ReadOnly) == FileAttributes.ReadOnly)
-						{
-							Attrs &= ~FileAttributes.ReadOnly;
-							File.SetAttributes(TargetPath, Attrs);
-						}
-						File.Delete(TargetPath);
+						// this file is auto-generated - we need to compare the contents
+						byte[] OriginalContents = File.ReadAllBytes(TargetPath);
+						byte[] NewContents = File.ReadAllBytes(SourcePath);
+						bFileIsUnchanged = Enumerable.SequenceEqual(OriginalContents, NewContents);
 					}
-					catch (Exception)
+					else
 					{
-						Log.TraceError("Could not replace file {0}.", TargetPath);
-						return;
+						// this file is not generated, just copied from somewhere else - can check the time to confirm if they're different
+						bFileIsUnchanged = (SrcFileInfo.CreationTime == DstFileInfo.CreationTime);
 					}
 				}
-			}
 
-			// If the file is present it is unmodified and should not be overwritten
-			if (!File.Exists(TargetPath))
-			{
-				try
+				if (bFileIsUnchanged)
 				{
-					File.Copy(IntermediatePath, TargetPath);
+					// use an empty source string if the file doesn't need changing
+					SourcePath = "";
 				}
-				catch (Exception)
-				{
-					Log.TraceError("Unable to copy file {0}.", TargetPath);
-					return;
-				}
-				UpdatedFilePaths!.Add(TargetPath);
 			}
+
+			ManifestFiles!.Add(TargetPath, SourcePath);
 		}
 
+
 		/// <summary>
 		/// Copies all of the generated files to the output folder
 		/// </summary>
-		protected void CopyResourcesToTargetDir()
+		private List<string> CopyFilesToOutput()
 		{
-			string TargetPath = Path.Combine(OutputPath!, BuildResourceSubPath);
-			string SourcePath = Path.Combine(IntermediatePath!, BuildResourceSubPath);
+			List<string> UpdatedFiles = new List<string>();
 
-			// If the target resource folder doesn't exist yet, create it
-			if (!CreateCheckDirectory(TargetPath))
+			// early out if there's nothing to copy
+			if (!ManifestFiles.Any(X => !string.IsNullOrEmpty(X.Value)))
 			{
-				return;
+				return UpdatedFiles;
 			}
 
-			// Find all copies of the resource file in both target and source directories (could be up to one for each culture and the default, but must have at least the default).
-			var TargetResourceInstances = Directory.EnumerateFiles(TargetPath, "*.*", SearchOption.AllDirectories);
-			var SourceResourceInstances = Directory.EnumerateFiles(SourcePath, "*.*", SearchOption.AllDirectories);
+			Log.TraceVerbose("Updating manifest resource files...");
 
-			// Remove any target files that aren't part of the source file list
-			foreach (string TargetResourceFile in TargetResourceInstances)
+			// copy over any new or updated files
+			foreach ( var ManifestFilePair in ManifestFiles! )
 			{
-				// Ignore string tables (the only non-binary resources that will be present)
-				if (!TargetResourceFile.Contains(".resw"))
+				string TargetPath = ManifestFilePair.Key;
+				string SourcePath = ManifestFilePair.Value;
+				if (string.IsNullOrEmpty(SourcePath))
 				{
-					//@todo always delete for cultures we aren't staging
-					bool bRelativeSourceFileFound = false;
-					foreach (string SourceResourceFile in SourceResourceInstances)
+					// if source is an empty string then the file is up-to-date
+					continue;
+				}
+				
+				// remove old version, if any
+				bool bFileExists = File.Exists(TargetPath);
+				if (bFileExists)
+				{
+					try
 					{
-						string SourceRelativeFile = SourceResourceFile.Substring(SourcePath.Length + 1);
-						string TargetRelativeFile = TargetResourceFile.Substring(TargetPath.Length + 1);
-						if (SourceRelativeFile.Equals(TargetRelativeFile))
-						{
-							bRelativeSourceFileFound = true;
-							break;
-						}
+						FileUtils.ForceDeleteFile(TargetPath);
 					}
-					if (!bRelativeSourceFileFound)
+					catch (Exception E)
 					{
-						try
-						{
-							File.Delete(TargetResourceFile);
-						}
-						catch (Exception E)
-						{
-							Log.TraceError("Could not remove stale resource file {0} - {1}.", TargetResourceFile, E.Message);
-						}
+						Log.TraceError($"\tCould not replace file {TargetPath} - {E.Message}");
 					}
 				}
-			}
 
-			// Copy new resource files only if they differ from the destination
-			foreach (string SourceResourceFile in SourceResourceInstances)
-			{
-				//@todo only copy files for cultures we are staging
-				string TargetResourcePath = Path.Combine(TargetPath, SourceResourceFile.Substring(SourcePath.Length + 1));
-				CompareAndReplaceModifiedTarget(SourceResourceFile, TargetResourcePath);
+				// copy new version
+				try
+				{
+					Log.TraceVerbose($"\t{(bFileExists ? "updating" : "adding")} {Utils.MakePathRelativeTo(TargetPath, OutputPath!)}");
+
+					Directory.CreateDirectory(Path.GetDirectoryName(TargetPath));
+					File.Copy(SourcePath, TargetPath);
+					File.SetAttributes(TargetPath, FileAttributes.Normal);
+					File.SetCreationTime(TargetPath, File.GetCreationTime(SourcePath));
+
+					UpdatedFiles.Add(TargetPath);
+				}
+				catch (Exception E)
+				{
+					Log.TraceError($"\tUnable to copy file {TargetPath} - {E.Message}");
+				}
 			}
+
+			return UpdatedFiles;
 		}
 
+
 		/// <summary>
 		/// Adds the given string to the culture string writers
 		/// </summary>
@@ -828,7 +819,7 @@ namespace UnrealBuildTool
 			IntermediatePath = InIntermediatePath;
 			ProjectFile = InProjectFile;
 			ProjectPath = InProjectDirectory;
-			UpdatedFilePaths = new List<string>();
+			ManifestFiles = new Dictionary<string, string>();
 
 			// Load up INI settings. We'll use engine settings to retrieve the manifest configuration, but these may reference
 			// values in either game or engine settings, so we'll keep both.
@@ -891,7 +882,7 @@ namespace UnrealBuildTool
 
 			// Clean out the resources intermediate path so that we know there are no stale binary files.
 			string IntermediateResourceDirectory = Path.Combine(IntermediatePath, BuildResourceSubPath);
-			RecursivelyForceDeleteDirectory(IntermediateResourceDirectory);
+			FileUtils.ForceDeleteDirectory(IntermediateResourceDirectory);
 			if (!CreateCheckDirectory(IntermediateResourceDirectory))
 			{
 				Log.TraceError("Could not create directory {0}.", IntermediateResourceDirectory);
@@ -933,27 +924,13 @@ namespace UnrealBuildTool
 			string ManifestIntermediatePath = Path.Combine(IntermediatePath, InManifestName);
 			string ManifestTargetPath = Path.Combine(OutputPath, InManifestName);
 			ManifestXmlDocument.Save(ManifestIntermediatePath);
-			CompareAndReplaceModifiedTarget(ManifestIntermediatePath, ManifestTargetPath);
+			AddFileReference(ManifestIntermediatePath, ManifestTargetPath, bIsGeneratedFile: true);
 			ProcessManifest(InTargetConfigs, InExecutables, InManifestName, ManifestTargetPath, ManifestIntermediatePath);
 
-			// Clean out any resource directories that we aren't staging
-			string TargetResourcePath = Path.Combine(OutputPath, BuildResourceSubPath);
-			if (Directory.Exists(TargetResourcePath))
-			{
-				List<string> TargetResourceDirectories = new List<string>(Directory.GetDirectories(TargetResourcePath, "*.*", SearchOption.AllDirectories));
-				foreach (string ResourceDirectory in TargetResourceDirectories)
-				{
-					if (!CulturesToStage.Contains(Path.GetFileName(ResourceDirectory)))
-					{
-						RecursivelyForceDeleteDirectory(ResourceDirectory);
-					}
-				}
-			}
-
 			// Export the resource tables starting with the default culture
 			string DefaultResourceTargetPath = Path.Combine(OutputPath, BuildResourceSubPath, "resources.resw");
 			DefaultResourceWriter.Close();
-			CompareAndReplaceModifiedTarget(DefaultResourceIntermediatePath, DefaultResourceTargetPath);
+			AddFileReference(DefaultResourceIntermediatePath, DefaultResourceTargetPath, bIsGeneratedFile: true);
 
 			foreach (var Writer in PerCultureResourceWriters)
 			{
@@ -962,15 +939,20 @@ namespace UnrealBuildTool
 				string IntermediateStringResourceFile = Path.Combine(IntermediateResourceDirectory, Writer.Key, "resources.resw");
 				string TargetStringResourceFile = Path.Combine(OutputPath, BuildResourceSubPath, Writer.Key, "resources.resw");
 
-				CompareAndReplaceModifiedTarget(IntermediateStringResourceFile, TargetStringResourceFile);
+				AddFileReference(IntermediateStringResourceFile, TargetStringResourceFile, bIsGeneratedFile: true);
 			}
 
-			// Copy all the binary resources into the target directory.
-			CopyResourcesToTargetDir();
+
+			// include a reference to the Package Resource Index so it isn't removed by RemoveStaleResourceFiles
+			string TargetResourceIndexFile = Path.Combine(OutputPath, "resources.pri");
+			AddFileReference("", TargetResourceIndexFile, bIsGeneratedFile: true);
 
 			// The resource database is dependent on everything else calculated here (manifest, resource string tables, binary resources).
 			// So if any file has been updated we'll need to run the config.
-			if (UpdatedFilePaths.Count > 0)
+			bool bHadStaleResources = RemoveStaleResourceFiles();
+			List<string> UpdatedFilePaths = CopyFilesToOutput();
+
+			if (bHadStaleResources || UpdatedFilePaths.Any())
 			{
 				// Create resource index configuration
 				string ResourceConfigFile = Path.Combine(IntermediatePath, "priconfig.xml");
@@ -1003,61 +985,37 @@ namespace UnrealBuildTool
 				System.Text.StringBuilder ResourcesList = new System.Text.StringBuilder();
 				foreach (string Resource in Resources)
 				{
-					ResourcesList.AppendLine(Resource.Replace(OutputPath, "").TrimStart('\\'));
+					ResourcesList.AppendLine( Utils.MakePathRelativeTo( Resource, OutputPath ) );
 				}
 				File.WriteAllText(ResourcesResFile, ResourcesList.ToString());
 
-				// Remove previous pri files so we can enumerate which ones are new since the resource generator could produce a file for each staged language.
-				IEnumerable<string> OldPriFiles = Directory.EnumerateFiles(IntermediatePath, "*.pri");
-				foreach (string OldPri in OldPriFiles)
+				// remove old Package Resource Index
+				try
 				{
-					try
-					{
-						File.Delete(OldPri);
-					}
-					catch (Exception)
-					{
-						Log.TraceError("Could not delete file {0}.", OldPri);
-					}
+					FileUtils.ForceDeleteFile(TargetResourceIndexFile);
+				}
+				catch (Exception E)
+				{
+					Log.TraceError($"cannot remove old pri file: {TargetResourceIndexFile} - {E.Message}");
 				}
 
-				// Generate the resource index
+				// Generate the Package Resource Index
 				string ResourceLogFile = Path.Combine(IntermediatePath, "ResIndexLog.xml");
-				string ResourceIndexFile = Path.Combine(IntermediatePath, "resources.pri");
-
-				string MakePriCommandLine = "new /pr \"" + OutputPath + "\" /cf \"" + ResourceConfigFile + "\" /mn \"" + ManifestTargetPath + "\" /il \"" + ResourceLogFile + "\" /of \"" + ResourceIndexFile + "\" /o";
-
+				string MakePriCommandLine = "new /pr \"" + OutputPath + "\" /cf \"" + ResourceConfigFile + "\" /mn \"" + ManifestTargetPath + "\" /il \"" + ResourceLogFile + "\" /of \"" + TargetResourceIndexFile + "\" /o";
 				if (IdentityName != null)
 				{
 					MakePriCommandLine += " /indexName \"" + IdentityName + "\"";
 				}
-				RunMakePri(MakePriCommandLine);
 
-				// Remove any existing pri target files that were not generated by this latest update
-				IEnumerable<string> NewPriFiles = Directory.EnumerateFiles(IntermediatePath, "*.pri");
-				IEnumerable<string> TargetPriFiles = Directory.EnumerateFiles(OutputPath, "*.pri");
-				foreach (string TargetPri in TargetPriFiles)
-				{
-					if (!NewPriFiles.Contains(TargetPri))
-					{
-						try
-						{
-							File.Delete(TargetPri);
-						}
-						catch (Exception)
-						{
-							Log.TraceError("Could not remove stale file {0}.", TargetPri);
-						}
-					}
-				}
+				Log.TraceVerbose($"\tgenerating {Utils.MakePathRelativeTo(TargetResourceIndexFile, OutputPath!)}");
+				RunMakePri(MakePriCommandLine);
+				UpdatedFilePaths.Add(TargetResourceIndexFile);
+			}
 
-				// Stage all the modified pri files to the output directory
-				foreach (string NewPri in NewPriFiles)
-				{
-					string NewResourceIndexFile = Path.Combine(IntermediatePath, Path.GetFileName(NewPri));
-					string FinalResourceIndexFile = Path.Combine(OutputPath, Path.GetFileName(NewPri));
-					CompareAndReplaceModifiedTarget(NewResourceIndexFile, FinalResourceIndexFile);
-				}
+			// Report if nothing was changed
+			if (!bHadStaleResources && !UpdatedFilePaths.Any())
+			{
+				Log.TraceVerbose($"Manifest resource files are up to date");
 			}
 
 			return UpdatedFilePaths;
