UBT: Added VSTestRunSettingsFile to TargetRules

diff --git a/Engine/Source/Programs/LowLevelTests/LowLevelTests.Target.cs b/Engine/Source/Programs/LowLevelTests/LowLevelTests.Target.cs
--- a/Engine/Source/Programs/LowLevelTests/LowLevelTests.Target.cs
+++ b/Engine/Source/Programs/LowLevelTests/LowLevelTests.Target.cs
@@ -1,6 +1,8 @@
 // Copyright Epic Games, Inc. All Rights Reserved.
 
 using System.IO;
+using EpicGames.Core;
+using UnrealBuildBase;
 using UnrealBuildTool;
 
 [SupportedPlatforms(UnrealPlatformClass.All)]
@@ -55,6 +57,8 @@ public class LowLevelTestsTarget : TargetRules
 			bDebugBuildsActuallyUseDebugCRT = true;
 		}
 
+		VSTestRunSettingsFile = FileReference.Combine(Unreal.EngineDirectory, "Source", "Programs", "LowLevelTests", "vstest.runsettings");
+
 		SetupPreprocessorDefinitions(Target);
 	}
 
diff --git a/Engine/Source/Programs/UnrealBuildTool/Configuration/TargetRules.cs b/Engine/Source/Programs/UnrealBuildTool/Configuration/TargetRules.cs
--- a/Engine/Source/Programs/UnrealBuildTool/Configuration/TargetRules.cs
+++ b/Engine/Source/Programs/UnrealBuildTool/Configuration/TargetRules.cs
@@ -442,6 +442,11 @@ namespace UnrealBuildTool
 		/// </summary>
 		protected bool? bIncludeAllTestsOverride;
 
+		/// <summary>
+		/// Set this to reference a VSTest run settings file from generated projects.
+		/// </summary>
+		public FileReference? VSTestRunSettingsFile;
+
 		/// <summary>
 		/// Additional plugins that are built for this target type but not enabled.
 		/// </summary>
diff --git a/Engine/Source/Programs/UnrealBuildTool/ProjectFiles/VisualStudio/VCProject.cs b/Engine/Source/Programs/UnrealBuildTool/ProjectFiles/VisualStudio/VCProject.cs
--- a/Engine/Source/Programs/UnrealBuildTool/ProjectFiles/VisualStudio/VCProject.cs
+++ b/Engine/Source/Programs/UnrealBuildTool/ProjectFiles/VisualStudio/VCProject.cs
@@ -1345,6 +1345,22 @@ namespace UnrealBuildTool
 			VCProjectFileContent.AppendLine("    <CleanDependsOn> $(CleanDependsOn); </CleanDependsOn>");
 			VCProjectFileContent.AppendLine("    <CppCleanDependsOn></CppCleanDependsOn>");
 			VCProjectFileContent.AppendLine("  </PropertyGroup>");
+
+			FileReference[] VSTestRunSettingsFiles = ProjectTargets.Select(Target => Target.TargetRules?.VSTestRunSettingsFile)
+				.Where(File => File != null).Select(File => File!).Distinct().ToArray();
+			if (VSTestRunSettingsFiles.Length == 1)
+			{
+				string RunSettingsRelativePath = VSTestRunSettingsFiles[0].MakeRelativeTo(ProjectFilePath.Directory);
+				VCProjectFileContent.AppendLine("  <PropertyGroup>");
+				VCProjectFileContent.AppendLine($"    <RunSettingsFilePath>$(ProjectDir){RunSettingsRelativePath}</RunSettingsFilePath>");
+				VCProjectFileContent.AppendLine("  </PropertyGroup>");
+			}
+			else if (VSTestRunSettingsFiles.Length > 1)
+			{
+				string Files = String.Join(", ", VSTestRunSettingsFiles.Select(File => File!.FullName));
+				Log.TraceWarning($"Inconsistent VSTest run settings files for project '{ProjectFilePath}': {Files}");
+			}
+
 			if (!IsStubProject)
 			{
 				foreach (UnrealTargetPlatform Platform in ProjectPlatforms)
