- Split SDK versions into their own Versions.cs files

diff --git a/Engine/Source/Programs/UnrealBuildTool/Platform/Android/AndroidPlatformSDK.Versions.cs b/Engine/Source/Programs/UnrealBuildTool/Platform/Android/AndroidPlatformSDK.Versions.cs
new file mode 100644
--- /dev/null
+++ b/Engine/Source/Programs/UnrealBuildTool/Platform/Android/AndroidPlatformSDK.Versions.cs
@@ -0,0 +1,43 @@
+// Copyright Epic Games, Inc. All Rights Reserved.
+
+using EpicGames.Core;
+
+namespace UnrealBuildTool
+{
+	partial class AndroidPlatformSDK : UEBuildPlatformSDK
+	{
+		public override string GetMainVersion()
+		{
+			return "r21e";
+		}
+		
+		public override string GetAutoSDKDirectoryForMainVersion()
+		{
+			return "-24";
+		}
+
+		protected override void GetValidVersionRange(out string MinVersion, out string MaxVersion)
+		{
+			MinVersion = "r21a";
+			MaxVersion = "r23a";
+		}
+
+		protected override void GetValidSoftwareVersionRange(out string? MinVersion, out string? MaxVersion)
+		{
+			MinVersion = MaxVersion = null;
+		}
+
+		public override string GetPlatformSpecificVersion(string VersionType)
+		{
+			switch (VersionType.ToLower())
+			{
+				case "platforms": return "android-28";
+				case "build-tools": return "28.0.3";
+				case "cmake": return "3.10.2.4988404";
+				case "ndk": return "21.4.7075529";
+			}
+
+			return "";
+		}
+	}
+}
diff --git a/Engine/Source/Programs/UnrealBuildTool/Platform/Android/AndroidPlatformSDK.cs b/Engine/Source/Programs/UnrealBuildTool/Platform/Android/AndroidPlatformSDK.cs
--- a/Engine/Source/Programs/UnrealBuildTool/Platform/Android/AndroidPlatformSDK.cs
+++ b/Engine/Source/Programs/UnrealBuildTool/Platform/Android/AndroidPlatformSDK.cs
@@ -9,42 +9,8 @@ using System.Linq;
 
 namespace UnrealBuildTool
 {
-	class AndroidPlatformSDK : UEBuildPlatformSDK
+	partial class AndroidPlatformSDK : UEBuildPlatformSDK
 	{
-		public override string GetMainVersion()
-		{
-			return "r21e";
-		}
-		public override string GetAutoSDKDirectoryForMainVersion()
-		{
-			return "-24";
-		}
-
-		protected override void GetValidVersionRange(out string MinVersion, out string MaxVersion)
-		{
-			MinVersion = "r21a";
-			MaxVersion = "r23a";
-		}
-
-		protected override void GetValidSoftwareVersionRange(out string? MinVersion, out string? MaxVersion)
-		{
-			MinVersion = MaxVersion = null;
-		}
-
-		public override string GetPlatformSpecificVersion(string VersionType)
-		{
-			switch (VersionType.ToLower())
-			{
-				case "platforms": return "android-28";
-				case "build-tools": return "28.0.3";
-				case "cmake": return "3.10.2.4988404";
-				case "ndk": return "21.4.7075529";
-			}
-
-			return "";
-		}
-
-
 		protected override string? GetInstalledSDKVersion()
 		{
 			string? NDKPath = Environment.GetEnvironmentVariable("NDKROOT");
diff --git a/Engine/Source/Programs/UnrealBuildTool/Platform/IOS/IOSPlatformSDK.Versions.cs b/Engine/Source/Programs/UnrealBuildTool/Platform/IOS/IOSPlatformSDK.Versions.cs
new file mode 100644
--- /dev/null
+++ b/Engine/Source/Programs/UnrealBuildTool/Platform/IOS/IOSPlatformSDK.Versions.cs
@@ -0,0 +1,14 @@
+// Copyright Epic Games, Inc. All Rights Reserved.
+
+namespace UnrealBuildTool
+{
+	internal class IOSPlatformSDK : ApplePlatformSDK
+	{
+		protected override void GetValidSoftwareVersionRange(out string MinVersion, out string? MaxVersion)
+		{
+			// what is our min IOS version?
+			MinVersion = "12.0";
+			MaxVersion = null;
+		}
+	}
+}
diff --git a/Engine/Source/Programs/UnrealBuildTool/Platform/IOS/IOSPlatformSDK.cs b/Engine/Source/Programs/UnrealBuildTool/Platform/IOS/IOSPlatformSDK.cs
deleted file mode 100644
diff --git a/Engine/Source/Programs/UnrealBuildTool/Platform/Linux/LinuxPlatformSDK.Versions.cs b/Engine/Source/Programs/UnrealBuildTool/Platform/Linux/LinuxPlatformSDK.Versions.cs
new file mode 100644
--- /dev/null
+++ b/Engine/Source/Programs/UnrealBuildTool/Platform/Linux/LinuxPlatformSDK.Versions.cs
@@ -0,0 +1,26 @@
+// Copyright Epic Games, Inc. All Rights Reserved.
+
+using EpicGames.Core;
+
+namespace UnrealBuildTool
+{
+	partial class LinuxPlatformSDK : UEBuildPlatformSDK
+	{
+		public override string GetMainVersion()
+		{
+			return "v19_clang-11.0.1-centos7";
+		}
+
+		protected override void GetValidVersionRange(out string MinVersion, out string MaxVersion)
+		{
+			// all that matters is the number after the v, according to TryConvertVersionToInt()
+			MinVersion = "v10_clang-5.0.0-centos7";
+			MaxVersion = "v19_clang-11.0.1-centos7";
+		}
+
+		protected override void GetValidSoftwareVersionRange(out string? MinVersion, out string? MaxVersion)
+		{
+			MinVersion = MaxVersion = null;
+		}
+	}
+}
diff --git a/Engine/Source/Programs/UnrealBuildTool/Platform/Linux/LinuxPlatformSDK.cs b/Engine/Source/Programs/UnrealBuildTool/Platform/Linux/LinuxPlatformSDK.cs
--- a/Engine/Source/Programs/UnrealBuildTool/Platform/Linux/LinuxPlatformSDK.cs
+++ b/Engine/Source/Programs/UnrealBuildTool/Platform/Linux/LinuxPlatformSDK.cs
@@ -9,25 +9,8 @@ using UnrealBuildBase;
 
 namespace UnrealBuildTool
 {
-	class LinuxPlatformSDK : UEBuildPlatformSDK
+	partial class LinuxPlatformSDK : UEBuildPlatformSDK
 	{
-		public override string GetMainVersion()
-		{
-			return "v19_clang-11.0.1-centos7";
-		}
-
-		protected override void GetValidVersionRange(out string MinVersion, out string MaxVersion)
-		{
-			// all that matters is the number after the v, according to TryConvertVersionToInt()
-			MinVersion = "v10_clang-5.0.0-centos7";
-			MaxVersion = "v19_clang-11.0.1-centos7";
-		}
-
-		protected override void GetValidSoftwareVersionRange(out string? MinVersion, out string? MaxVersion)
-		{
-			MinVersion = MaxVersion = null;
-		}
-
 		public override string GetAutoSDKPlatformName()
 		{
 			return TargetPlatformName;
diff --git a/Engine/Source/Programs/UnrealBuildTool/Platform/Mac/MacPlatformSDK.Versions.cs b/Engine/Source/Programs/UnrealBuildTool/Platform/Mac/MacPlatformSDK.Versions.cs
new file mode 100644
--- /dev/null
+++ b/Engine/Source/Programs/UnrealBuildTool/Platform/Mac/MacPlatformSDK.Versions.cs
@@ -0,0 +1,34 @@
+// Copyright Epic Games, Inc. All Rights Reserved.
+
+using EpicGames.Core;
+
+namespace UnrealBuildTool
+{
+	internal partial class MacPlatformSDK : ApplePlatformSDK
+	{
+		public override string GetMainVersion()
+		{
+			return "11.5";
+		}
+
+		protected override void GetValidVersionRange(out string MinVersion, out string MaxVersion)
+		{
+			if (RuntimePlatform.IsMac)
+			{
+				MinVersion = "11.0.0";
+				MaxVersion = "13.9.9";
+			}
+			else
+			{
+				// @todo turnkey: these are MobileDevice .dll versions in Windows - to get the iTunes app version (12.3.4.1 etc) would need to hunt down the .exe
+				MinVersion = "1100.0.0.0";
+				MaxVersion = "1399.0";
+			}
+		}
+
+		protected override void GetValidSoftwareVersionRange(out string? MinVersion, out string? MaxVersion)
+		{
+			MinVersion = MaxVersion = null;
+		}
+	}
+}
\ No newline at end of file
diff --git a/Engine/Source/Programs/UnrealBuildTool/Platform/Mac/UEBuildMac.cs b/Engine/Source/Programs/UnrealBuildTool/Platform/Mac/UEBuildMac.cs
--- a/Engine/Source/Programs/UnrealBuildTool/Platform/Mac/UEBuildMac.cs
+++ b/Engine/Source/Programs/UnrealBuildTool/Platform/Mac/UEBuildMac.cs
@@ -424,7 +424,7 @@ namespace UnrealBuildTool
 		/// </summary>
 		public override void RegisterBuildPlatforms()
 		{
-			ApplePlatformSDK SDK = new ApplePlatformSDK();
+			MacPlatformSDK SDK = new MacPlatformSDK();
 
 			// Register this build platform for Mac
 			UEBuildPlatform.RegisterBuildPlatform(new MacPlatform(SDK));
diff --git a/Engine/Source/Programs/UnrealBuildTool/Platform/Windows/MicrosoftPlatformSDK.Versions.cs b/Engine/Source/Programs/UnrealBuildTool/Platform/Windows/MicrosoftPlatformSDK.Versions.cs
new file mode 100644
--- /dev/null
+++ b/Engine/Source/Programs/UnrealBuildTool/Platform/Windows/MicrosoftPlatformSDK.Versions.cs
@@ -0,0 +1,79 @@
+// Copyright Epic Games, Inc. All Rights Reserved.
+
+using EpicGames.Core;
+using System.Linq;
+
+namespace UnrealBuildTool
+{
+	partial class MicrosoftPlatformSDK : UEBuildPlatformSDK
+	{
+		/// <summary>
+		/// The default Windows SDK version to be used, if installed.
+		/// </summary>
+		static readonly VersionNumber[] PreferredWindowsSdkVersions = new VersionNumber[]
+		{
+			VersionNumber.Parse("10.0.18362.0")
+		};
+
+		/// <summary>
+		/// The default compiler version to be used, if installed. 
+		/// </summary>
+		static readonly VersionNumberRange[] PreferredClangVersions =
+		{
+			VersionNumberRange.Parse("13.0.0", "13.999"),
+			VersionNumberRange.Parse("12.0.0", "12.999"),
+			VersionNumberRange.Parse("11.0.0", "11.999"),
+			VersionNumberRange.Parse("10.0.0", "10.999")
+		};
+
+		static readonly VersionNumber MinimumClangVersion = new VersionNumber(10, 0, 0);
+
+		/// <summary>
+		/// Ranges of tested compiler toolchains to be used, in order of preference. If multiple toolchains in a range are present, the latest version will be preferred.
+		/// Note that the numbers here correspond to the installation *folders* rather than precise executable versions. 
+		/// </summary>
+		static readonly VersionNumberRange[] PreferredVisualCppVersions = new VersionNumberRange[]
+		{
+			VersionNumberRange.Parse("14.29.30133", "14.29.30136"), // VS2019 16.11.5
+		};
+
+		/// <summary>
+		/// Tested compiler toolchanges that should not be allowed.
+		/// </summary>
+		static readonly VersionNumberRange[] BannedVisualCppVersions = new VersionNumberRange[]
+		{
+		};
+
+		static readonly VersionNumber MinimumVisualCppVersion = new VersionNumber(14, 29, 30133);
+
+		/// <summary>
+		/// The default compiler version to be used, if installed. 
+		/// </summary>
+		static readonly VersionNumberRange[] PreferredIntelOneApiVersions =
+		{
+			VersionNumberRange.Parse("2022.0.0", "2022.9999"),
+		};
+
+		static readonly VersionNumber MinimumIntelOneApiVersion = new VersionNumber(2021, 0, 0);
+
+
+		public override string GetMainVersion()
+		{
+			// preferred/main version is the top of the Preferred list - 
+			return PreferredWindowsSdkVersions.First().ToString();
+		}
+
+		protected override void GetValidVersionRange(out string MinVersion, out string MaxVersion)
+		{
+			MinVersion = "10.0.00000.0";
+			MaxVersion = "10.9.99999.0";
+		}
+
+		protected override void GetValidSoftwareVersionRange(out string? MinVersion, out string? MaxVersion)
+		{
+			// minimum version is the oldest version in the Preferred list -
+			MinVersion = PreferredWindowsSdkVersions.Min()?.ToString();
+			MaxVersion = null;
+		}
+	}
+}
diff --git a/Engine/Source/Programs/UnrealBuildTool/Platform/Windows/MicrosoftPlatformSDK.cs b/Engine/Source/Programs/UnrealBuildTool/Platform/Windows/MicrosoftPlatformSDK.cs
--- a/Engine/Source/Programs/UnrealBuildTool/Platform/Windows/MicrosoftPlatformSDK.cs
+++ b/Engine/Source/Programs/UnrealBuildTool/Platform/Windows/MicrosoftPlatformSDK.cs
@@ -17,77 +17,8 @@ using System.Runtime.Versioning;
 
 namespace UnrealBuildTool
 {
-	internal class MicrosoftPlatformSDK : UEBuildPlatformSDK
+	internal partial class MicrosoftPlatformSDK : UEBuildPlatformSDK
 	{
-		/// <summary>
-		/// The default Windows SDK version to be used, if installed.
-		/// </summary>
-		static readonly VersionNumber[] PreferredWindowsSdkVersions = new VersionNumber[]
-		{
-			VersionNumber.Parse("10.0.18362.0")
-		};
-
-		/// <summary>
-		/// The default compiler version to be used, if installed. 
-		/// </summary>
-		static readonly VersionNumberRange[] PreferredClangVersions =
-		{
-			VersionNumberRange.Parse("13.0.0", "13.999"),
-			VersionNumberRange.Parse("12.0.0", "12.999"),
-			VersionNumberRange.Parse("11.0.0", "11.999"),
-			VersionNumberRange.Parse("10.0.0", "10.999")
-		};
-
-		static readonly VersionNumber MinimumClangVersion = new VersionNumber(10, 0, 0);
-
-		/// <summary>
-		/// Ranges of tested compiler toolchains to be used, in order of preference. If multiple toolchains in a range are present, the latest version will be preferred.
-		/// Note that the numbers here correspond to the installation *folders* rather than precise executable versions. 
-		/// </summary>
-		static readonly VersionNumberRange[] PreferredVisualCppVersions = new VersionNumberRange[]
-		{
-			VersionNumberRange.Parse("14.29.30133", "14.29.30136"), // VS2019 16.11.5
-		};
-
-		/// <summary>
-		/// Tested compiler toolchanges that should not be allowed.
-		/// </summary>
-		static readonly VersionNumberRange[] BannedVisualCppVersions = new VersionNumberRange[]
-		{
-		};
-
-		static readonly VersionNumber MinimumVisualCppVersion = new VersionNumber(14, 29, 30133);
-
-		/// <summary>
-		/// The default compiler version to be used, if installed. 
-		/// </summary>
-		static readonly VersionNumberRange[] PreferredIntelOneApiVersions =
-		{
-			VersionNumberRange.Parse("2022.0.0", "2022.9999"),
-		};
-
-		static readonly VersionNumber MinimumIntelOneApiVersion = new VersionNumber(2021, 0, 0);
-
-
-		public override string GetMainVersion()
-		{
-			// preferred/main version is the top of the Preferred list - 
-			return PreferredWindowsSdkVersions.First().ToString();
-		}
-
-		protected override void GetValidVersionRange(out string MinVersion, out string MaxVersion)
-		{
-			MinVersion = "10.0.00000.0";
-			MaxVersion = "10.9.99999.0";
-		}
-
-		protected override void GetValidSoftwareVersionRange(out string? MinVersion, out string? MaxVersion)
-		{
-			// minimum version is the oldest version in the Preferred list -
-			MinVersion = PreferredWindowsSdkVersions.Min()?.ToString();
-			MaxVersion = null;
-		}
-
 		[SupportedOSPlatform("windows")]
 		protected override string? GetInstalledSDKVersion()
 		{
