UBT will now find "Shared" C# project files in the Restricted directories.

#rnx

diff --git a/Engine/Source/Programs/UnrealBuildTool/ProjectFiles/ProjectFileGenerator.cs b/Engine/Source/Programs/UnrealBuildTool/ProjectFiles/ProjectFileGenerator.cs
--- a/Engine/Source/Programs/UnrealBuildTool/ProjectFiles/ProjectFileGenerator.cs
+++ b/Engine/Source/Programs/UnrealBuildTool/ProjectFiles/ProjectFileGenerator.cs
@@ -491,31 +491,43 @@ namespace UnrealBuildTool
 			}
 		}
 
+		/// <summary>
+		/// Return a collection of all the directories to search for shared C# projects
+		/// </summary>
+		IEnumerable<DirectoryInfo> FindSharedDirectories()
+		{
+			yield return new DirectoryInfo(DirectoryReference.Combine(UnrealBuildTool.EngineSourceDirectory, "Programs", "Shared").FullName);
+			yield return new DirectoryInfo(DirectoryReference.Combine(Unreal.EngineDirectory, "Restricted", "NoRedist", "Source", "Programs", "Shared").FullName);
+			yield return new DirectoryInfo(DirectoryReference.Combine(Unreal.EngineDirectory, "Restricted", "NotForLicensees", "Source", "Programs", "Shared").FullName);
+		}
+
 		/// <summary>
 		/// Adds all the DotNet folder to the solution.
 		/// </summary>
 		void AddSharedDotNetModules(PrimaryProjectFolder ProgramsFolder)
 		{
-			DirectoryInfo DotNetDir = new DirectoryInfo(DirectoryReference.Combine(UnrealBuildTool.EngineSourceDirectory, "Programs", "Shared").FullName);
-			if (DotNetDir.Exists)
+			List<FileInfo> ProjectFiles = new List<FileInfo>();
+			foreach (DirectoryInfo DotNetDir in FindSharedDirectories())
 			{
-				List<FileInfo> ProjectFiles = new List<FileInfo>();
-				foreach (DirectoryInfo ProjectDir in DotNetDir.EnumerateDirectories())
-				{
-					ProjectFiles.AddRange(ProjectDir.EnumerateFiles("*.csproj"));
-				}
-				if (ProjectFiles.Count > 0)
+				if (DotNetDir.Exists)
 				{
-					PrimaryProjectFolder Folder = ProgramsFolder.AddSubFolder("Shared");
-					foreach (FileInfo ProjectFile in ProjectFiles)
+					foreach (DirectoryInfo ProjectDir in DotNetDir.EnumerateDirectories())
 					{
-						VCSharpProjectFile Project = new VCSharpProjectFile(new FileReference(ProjectFile));
-						Project.ShouldBuildForAllSolutionTargets = false;
-						AddExistingProjectFile(Project, bForceDevelopmentConfiguration: true);
-						Folder.ChildProjects.Add(Project);
+						ProjectFiles.AddRange(ProjectDir.EnumerateFiles("*.csproj"));
 					}
 				}
 			}
+			if (ProjectFiles.Count > 0)
+			{
+				PrimaryProjectFolder Folder = ProgramsFolder.AddSubFolder("Shared");
+				foreach (FileInfo ProjectFile in ProjectFiles)
+				{
+					VCSharpProjectFile Project = new VCSharpProjectFile(new FileReference(ProjectFile));
+					Project.ShouldBuildForAllSolutionTargets = false;
+					AddExistingProjectFile(Project, bForceDevelopmentConfiguration: true);
+					Folder.ChildProjects.Add(Project);
+				}
+			}
 		}
 
 		/// <summary>
