- VSMac project generation now takes into account projects that are windows only, and doesn't attempt to put them into the solution
- Don't make solution folders if there are no projects anywhere in them
- Fixed MemProfiler2.csproj to open on VSMac (although now with the above it is skipped anyway because it requires windows - but if we make it not require windows, we'd need this fix to match what UBT does)

diff --git a/Engine/Source/Programs/MemoryProfiler2/MemoryProfiler2.csproj b/Engine/Source/Programs/MemoryProfiler2/MemoryProfiler2.csproj
--- a/Engine/Source/Programs/MemoryProfiler2/MemoryProfiler2.csproj
+++ b/Engine/Source/Programs/MemoryProfiler2/MemoryProfiler2.csproj
@@ -211,10 +211,8 @@
     <Compile Include="StreamToken.cs" />
     <Compile Include="TimeLineView.cs" />
     <Compile Include="../../../Platforms/*/Source/Programs/MemoryProfiler2/**/*.cs">
-      <!-- RecursiveDir metadata expands to [PlatformName]/Source/Programs/[etc]. The Replace() will
-      replace everything from the first slash to the end leaving just the platform name. Should the
-      Replace() fail then RecursiveDir is still an agreeable value -->
-      <Link>$([System.Text.RegularExpressions.Regex]::Replace(%(Compile.RecursiveDir), [\\/].+$, ``))/_</Link>
+      <!-- Since RecursiveDir is bugged on VS Mac, we use the FullPath, and use Regex to pull out portions of the path to get Platform name -->
+      <Link>Platform/$([System.Text.RegularExpressions.Regex]::Replace(%(FullPath), '^(.*[\\/]Platforms[\\/])(.+?)[\\/](.*)$', '$2'))/%(Filename)%(Extension)</Link>
     </Compile>
   </ItemGroup>
   <ItemGroup>
diff --git a/Engine/Source/Programs/UnrealBuildTool/ProjectFiles/VisualStudio/VCProjectFileGenerator.cs b/Engine/Source/Programs/UnrealBuildTool/ProjectFiles/VisualStudio/VCProjectFileGenerator.cs
--- a/Engine/Source/Programs/UnrealBuildTool/ProjectFiles/VisualStudio/VCProjectFileGenerator.cs
+++ b/Engine/Source/Programs/UnrealBuildTool/ProjectFiles/VisualStudio/VCProjectFileGenerator.cs
@@ -502,6 +502,28 @@ namespace UnrealBuildTool
 
 			IDictionary<PrimaryProjectFolder, Guid> ProjectFolderGuids = GenerateProjectFolderGuids(RootFolder);
 
+			// check if the give folder has any projects that will be put into the solution
+			System.Func<PrimaryProjectFolder, bool>? HasProjectFunc = null;
+			HasProjectFunc = (ProjectFolder) =>
+			{
+				foreach (MSBuildProjectFile ChildProject in ProjectFolder.ChildProjects)
+				{
+					if (AllProjectFiles.Contains(ChildProject))
+					{
+						return true;
+					}
+				}
+
+				foreach (PrimaryProjectFolder SubFolder in ProjectFolder.SubFolders)
+				{
+					if (HasProjectFunc!(SubFolder))
+					{
+						return true;
+					}
+				}
+				return false;
+			};
+
 			// Solution folders, files and project entries
 			{
 				// This the GUID that Visual Studio uses to identify a solution folder
@@ -512,6 +534,11 @@ namespace UnrealBuildTool
 					IEnumerable<PrimaryProjectFolder> AllSolutionFolders = ProjectFolderGuids.Keys.OrderBy(Folder => Folder.FolderName).ThenBy(Folder => ProjectFolderGuids[Folder]);
 					foreach (PrimaryProjectFolder CurFolder in AllSolutionFolders)
 					{
+						// if this folder has no projects anywhere under it then skip it
+						if (!HasProjectFunc(CurFolder))
+						{
+							continue;
+						}
 						string FolderGUIDString = ProjectFolderGuids[CurFolder].ToString("B").ToUpperInvariant();
 						VCSolutionFileContent.AppendLine("Project(\"" + SolutionFolderEntryGUID + "\") = \"" + CurFolder.FolderName + "\", \"" + CurFolder.FolderName + "\", \"" + FolderGUIDString + "\"");
 
@@ -763,14 +790,21 @@ namespace UnrealBuildTool
 
 									foreach (MSBuildProjectFile ChildProject in CurFolder.ChildProjects)
 									{
-										//	e.g. "{BF6FB09F-A2A6-468F-BE6F-DEBE07EAD3EA} = {C43B6BB5-3EF0-4784-B896-4099753BCDA9}"
-										LocalVCSolutionFileContent.AppendLine("		" + ChildProject.ProjectGUID.ToString("B").ToUpperInvariant() + " = " + CurFolderGUIDString);
+										if (AllProjectFiles.Contains(ChildProject))
+										{
+											//	e.g. "{BF6FB09F-A2A6-468F-BE6F-DEBE07EAD3EA} = {C43B6BB5-3EF0-4784-B896-4099753BCDA9}"
+											LocalVCSolutionFileContent.AppendLine("		" + ChildProject.ProjectGUID.ToString("B").ToUpperInvariant() + " = " + CurFolderGUIDString);
+										}
 									}
 
 									foreach (PrimaryProjectFolder SubFolder in CurFolder.SubFolders)
 									{
-										//	e.g. "{BF6FB09F-A2A6-468F-BE6F-DEBE07EAD3EA} = {C43B6BB5-3EF0-4784-B896-4099753BCDA9}"
-										LocalVCSolutionFileContent.AppendLine("		" + ProjectFolderGuids[SubFolder].ToString("B").ToUpperInvariant() + " = " + CurFolderGUIDString);
+										// if this folder has no projects anywhere under it then skip it
+										if (HasProjectFunc(SubFolder))
+										{
+											//	e.g. "{BF6FB09F-A2A6-468F-BE6F-DEBE07EAD3EA} = {C43B6BB5-3EF0-4784-B896-4099753BCDA9}"
+											LocalVCSolutionFileContent.AppendLine("		" + ProjectFolderGuids[SubFolder].ToString("B").ToUpperInvariant() + " = " + CurFolderGUIDString);
+										}
 									}
 
 									// Recurse into subfolders
diff --git a/Engine/Source/Programs/UnrealBuildTool/ProjectFiles/VisualStudioMac/VisualStudioMacProjectFileGenerator.cs b/Engine/Source/Programs/UnrealBuildTool/ProjectFiles/VisualStudioMac/VisualStudioMacProjectFileGenerator.cs
--- a/Engine/Source/Programs/UnrealBuildTool/ProjectFiles/VisualStudioMac/VisualStudioMacProjectFileGenerator.cs
+++ b/Engine/Source/Programs/UnrealBuildTool/ProjectFiles/VisualStudioMac/VisualStudioMacProjectFileGenerator.cs
@@ -5,9 +5,11 @@ using System.Collections.Generic;
 using System.Text;
 using System.IO;
 using System.Linq;
+using System.Xml;
 using System.Xml.Linq;
 using EpicGames.Core;
 
+
 namespace UnrealBuildTool
 {
     
@@ -35,6 +37,31 @@ namespace UnrealBuildTool
             return false;
         }
 
+        protected bool IsValidProject(ProjectFile ProjectFile)
+        {
+            XmlDocument Doc = new XmlDocument();
+            Doc.Load(ProjectFile.ProjectFilePath.FullName);
+
+            XmlNodeList Elements = Doc.GetElementsByTagName("TargetFramework");
+            foreach (XmlElement? Element in Elements)
+            {
+                if (Element == null)
+                {
+                    continue;
+                }
+
+                // some projects can have TargetFramework's nested in conditionals for IsWindows, etc. Rather than try to handle that,
+                // we look if _any_ TargetFramework doesn't contain Windows and assume that means the conditionals are set up as expected
+                if (!Element.InnerText.Contains("windows"))
+				{
+                    return true;
+				}
+            }
+
+            // if there was no framework detected without windows, then this is not valid
+            return false;
+        }
+
         /// <summary>
         /// Writes the project files to disk
         /// </summary>
@@ -49,6 +76,10 @@ namespace UnrealBuildTool
             // write out OtherProjectFiles and AutomationProjectFiles
             GeneratedProjectFiles.Clear();
 
+            // remove C# projects that require windows (see BuildAllScriptPlugins() looking at TargetFramework)
+            OtherProjectFiles.RemoveAll(x => !IsValidProject(x));
+            AutomationProjectFiles.RemoveAll(x => !IsValidProject(x));
+
             if (!base.WriteProjectFiles(PlatformProjectGenerators))
             {
                 return false;
