- Look in Restricted directories for Shared projects

diff --git a/Engine/Source/Programs/UnrealBuildTool/ProjectFiles/ProjectFileGenerator.cs b/Engine/Source/Programs/UnrealBuildTool/ProjectFiles/ProjectFileGenerator.cs
--- a/Engine/Source/Programs/UnrealBuildTool/ProjectFiles/ProjectFileGenerator.cs
+++ b/Engine/Source/Programs/UnrealBuildTool/ProjectFiles/ProjectFileGenerator.cs
@@ -491,41 +491,32 @@ namespace UnrealBuildTool
 			}
 		}
 
-		/// <summary>
-		/// Return a collection of all the directories to search for shared C# projects
-		/// </summary>
-		IEnumerable<DirectoryInfo> FindSharedDirectories()
-		{
-			yield return new DirectoryInfo(DirectoryReference.Combine(UnrealBuildTool.EngineSourceDirectory, "Programs", "Shared").FullName);
-			yield return new DirectoryInfo(DirectoryReference.Combine(Unreal.EngineDirectory, "Restricted", "NoRedist", "Source", "Programs", "Shared").FullName);
-			yield return new DirectoryInfo(DirectoryReference.Combine(Unreal.EngineDirectory, "Restricted", "NotForLicensees", "Source", "Programs", "Shared").FullName);
-		}
-
 		/// <summary>
 		/// Adds all the DotNet folder to the solution.
 		/// </summary>
-		void AddSharedDotNetModules(PrimaryProjectFolder ProgramsFolder)
+		void AddSharedDotNetModules(List<DirectoryReference> AllEngineDirectories, PrimaryProjectFolder ProgramsFolder)
 		{
-			List<FileInfo> ProjectFiles = new List<FileInfo>();
-			foreach (DirectoryInfo DotNetDir in FindSharedDirectories())
+			foreach (DirectoryReference EngineDir in AllEngineDirectories)
 			{
+				DirectoryInfo DotNetDir = new DirectoryInfo(DirectoryReference.Combine(EngineDir, "Shared").FullName);
 				if (DotNetDir.Exists)
 				{
+					List<FileInfo> ProjectFiles = new List<FileInfo>();
 					foreach (DirectoryInfo ProjectDir in DotNetDir.EnumerateDirectories())
 					{
 						ProjectFiles.AddRange(ProjectDir.EnumerateFiles("*.csproj"));
 					}
-				}
-			}
-			if (ProjectFiles.Count > 0)
-			{
-				PrimaryProjectFolder Folder = ProgramsFolder.AddSubFolder("Shared");
-				foreach (FileInfo ProjectFile in ProjectFiles)
-				{
-					VCSharpProjectFile Project = new VCSharpProjectFile(new FileReference(ProjectFile));
-					Project.ShouldBuildForAllSolutionTargets = false;
-					AddExistingProjectFile(Project, bForceDevelopmentConfiguration: true);
-					Folder.ChildProjects.Add(Project);
+					if (ProjectFiles.Count > 0)
+					{
+						PrimaryProjectFolder Folder = ProgramsFolder.AddSubFolder("Shared");
+						foreach (FileInfo ProjectFile in ProjectFiles)
+						{
+							VCSharpProjectFile Project = new VCSharpProjectFile(new FileReference(ProjectFile));
+							Project.ShouldBuildForAllSolutionTargets = false;
+							AddExistingProjectFile(Project, bForceDevelopmentConfiguration: true);
+							Folder.ChildProjects.Add(Project);
+						}
+					}
 				}
 			}
 		}
@@ -553,27 +544,14 @@ namespace UnrealBuildTool
 		/// <summary>
 		/// Finds all csproj within Engine/Source/Programs, and add them if their .uprogram file exists.
 		/// </summary>
-		void DiscoverCSharpProgramProjects(List<FileReference> AllGameProjects, PrimaryProjectFolder ProgramsFolder)
+		void DiscoverCSharpProgramProjects(List<DirectoryReference> AllEngineDirectories, List<FileReference> AllGameProjects, PrimaryProjectFolder ProgramsFolder)
 		{
-			string[] UnsupportedPlatformNames = Utils.MakeListOfUnsupportedPlatforms(SupportedPlatforms, bIncludeUnbuildablePlatforms: true).ToArray();
-
 			List<FileReference> FoundProjects = new List<FileReference>();
 			List<DirectoryReference> ProjectDirs = new List<DirectoryReference>();
 
 			DirectoryReference EngineExtras = DirectoryReference.Combine(Unreal.EngineDirectory, "Extras");
 			ProjectDirs.Add(EngineExtras);
 			DiscoverCSharpProgramProjectsRecursively(EngineExtras, FoundProjects);
-
-			List<DirectoryReference> AllEngineDirectories =
-				Unreal.GetExtensionDirs(Unreal.EngineDirectory, "Source/Programs").Where(x =>
-				{
-					if (x.ContainsAnyNames(UnsupportedPlatformNames, Unreal.EngineDirectory))
-					{
-						Log.TraceLog($"Skipping any C# project files in \"{x}\" due to unsupported platform");
-						return false;
-					}
-					return true;
-				}).ToList();
 			
 			ProjectDirs = ProjectDirs.Union(AllEngineDirectories).ToList();
 			foreach (DirectoryReference EngineDir in AllEngineDirectories)
@@ -592,6 +570,7 @@ namespace UnrealBuildTool
 				}
 			}
 
+			string[] UnsupportedPlatformNames = Utils.MakeListOfUnsupportedPlatforms(SupportedPlatforms, bIncludeUnbuildablePlatforms: true).ToArray();
 			foreach (FileReference FoundProject in FoundProjects)
 			{
 				foreach (DirectoryReference ProjectDir in ProjectDirs)
@@ -1085,6 +1064,19 @@ namespace UnrealBuildTool
 			{
 				if (bIncludeDotNetPrograms)
 				{
+					// gather engine directories across extension dirs
+					string[] UnsupportedPlatformNames = Utils.MakeListOfUnsupportedPlatforms(SupportedPlatforms, bIncludeUnbuildablePlatforms: true).ToArray();
+					List<DirectoryReference> AllEngineDirectories =
+						Unreal.GetExtensionDirs(Unreal.EngineDirectory, "Source/Programs").Where(x =>
+						{
+							if (x.ContainsAnyNames(UnsupportedPlatformNames, Unreal.EngineDirectory))
+							{
+								Log.TraceLog($"Skipping any C# project files in \"{x}\" due to unsupported platform");
+								return false;
+							}
+							return true;
+						}).ToList();
+
 					PrimaryProjectFolder ProgramsFolder = RootFolder.AddSubFolder("Programs");
 
 					// Add UnrealBuildTool to the primary project
@@ -1099,10 +1091,10 @@ namespace UnrealBuildTool
 					AddAutomationModules(AllGameProjects, RootFolder, ProgramsFolder);
 
 					// Add shared projects
-					AddSharedDotNetModules(ProgramsFolder);
+					AddSharedDotNetModules(AllEngineDirectories, ProgramsFolder);
 
 					// Discover C# programs which should additionally be included in the solution.
-					DiscoverCSharpProgramProjects(AllGameProjects, ProgramsFolder);
+					DiscoverCSharpProgramProjects(AllEngineDirectories, AllGameProjects, ProgramsFolder);
 
 					foreach (KeyValuePair<RulesAssembly, DirectoryReference> RulesAssemblyEntry in RulesAssemblies)
 					{
