UnrealBuildTool: Stop running mac toolchain binaries through /usr/bin/arch, it is no longer necessary with .NET6 as UBT is being run through the correct architecture version of dotnet
#rnx

diff --git a/Engine/Source/Programs/UnrealBuildTool/Platform/Mac/MacToolChain.cs b/Engine/Source/Programs/UnrealBuildTool/Platform/Mac/MacToolChain.cs
--- a/Engine/Source/Programs/UnrealBuildTool/Platform/Mac/MacToolChain.cs
+++ b/Engine/Source/Programs/UnrealBuildTool/Platform/Mac/MacToolChain.cs
@@ -624,18 +624,8 @@ namespace UnrealBuildTool
 				}
 
 				CompileAction.WorkingDirectory = GetMacDevSrcRoot();
-
-				if (!MacExports.IsRunningUnderRosetta)
-				{
-					string ArchPath = "/usr/bin/arch";
-					CompileAction.CommandPath = new FileReference(ArchPath);
-					CompileAction.CommandArguments = string.Format("-{0} {1} {2}", MacExports.HostArchitecture, CompilerPath, AllArgs);
-				}
-				else
-				{
-					CompileAction.CommandPath = new FileReference(CompilerPath);
-					CompileAction.CommandArguments = AllArgs;
-				}
+				CompileAction.CommandPath = new FileReference(CompilerPath);
+				CompileAction.CommandArguments = AllArgs;
 
 				// For compilation we delete everything we produce
 				CompileAction.DeleteItems.AddRange(CompileAction.ProducedItems);
@@ -675,18 +665,8 @@ namespace UnrealBuildTool
 
 					string AllPreprocessArgs = PreprocessArguments + PreprocessFileArguments + EscapedAdditionalArgs;
 
-					if (MacExports.IsRunningUnderRosetta)
-					{
-						string ArchPath = "/usr/bin/arch";
-						PrepassAction.CommandPath = new FileReference(ArchPath);
-						PrepassAction.CommandArguments = string.Format("-{0} {1} {2}", MacExports.HostArchitecture, CompilerPath, AllPreprocessArgs);
-					}
-					else
-					{
-						PrepassAction.CommandPath = new FileReference(CompilerPath);
-						PrepassAction.CommandArguments = AllPreprocessArgs;
-					}
-
+					PrepassAction.CommandPath = new FileReference(CompilerPath);
+					PrepassAction.CommandArguments = AllPreprocessArgs;
 					CompileAction.DependencyListFile = DependencyListFile;
 					CompileAction.PrerequisiteItems.Add(DependencyListFile);
 				}
@@ -860,11 +840,6 @@ namespace UnrealBuildTool
 			string Linker = bIsBuildingLibrary ? MacArchiver : MacLinker;
 			string LinkCommand = Settings.ToolchainDir + Linker + VersionArg + " " + (bIsBuildingLibrary ? GetArchiveArguments_Global(LinkEnvironment) : GetLinkArguments_Global(LinkEnvironment));
 
-			if (MacExports.IsRunningUnderRosetta)
-			{
-				LinkCommand = string.Format("/usr/bin/arch -{0} {1}", MacExports.HostArchitecture, LinkCommand);
-			}
-
 			// Tell the action that we're building an import library here and it should conditionally be
 			// ignored as a prerequisite for other actions
 			LinkAction.bProducesImportLibrary = RuntimePlatform.IsWindows && (bBuildImportLibraryOnly || LinkEnvironment.bIsBuildingDLL);
