HordeExecutor Prototype

#rnx

diff --git a/Engine/Source/Programs/UnrealBuildTool/Configuration/BuildConfiguration.cs b/Engine/Source/Programs/UnrealBuildTool/Configuration/BuildConfiguration.cs
--- a/Engine/Source/Programs/UnrealBuildTool/Configuration/BuildConfiguration.cs
+++ b/Engine/Source/Programs/UnrealBuildTool/Configuration/BuildConfiguration.cs
@@ -42,6 +42,13 @@ namespace UnrealBuildTool
 		[XmlConfigFile]
 		public bool bAllowHybridExecutor = false;
 
+		/// <summary>
+		/// Whether Horde remote compute may be used. Highly experimental, disabled by default.
+		/// </summary>
+		[XmlConfigFile]
+		[CommandLine("-HordeCompute", Value = "true")]
+		public bool bAllowHordeCompute = false;
+
 		/// <summary>
 		/// Whether XGE may be used.
 		/// </summary>
diff --git a/Engine/Source/Programs/UnrealBuildTool/Executors/Experimental/HordeExecutor.cs b/Engine/Source/Programs/UnrealBuildTool/Executors/Experimental/HordeExecutor.cs
new file mode 100644
--- /dev/null
+++ b/Engine/Source/Programs/UnrealBuildTool/Executors/Experimental/HordeExecutor.cs
@@ -0,0 +1,689 @@
+// Copyright Epic Games, Inc. All Rights Reserved.
+
+using EpicGames.Core;
+using EpicGames.Serialization;
+using System;
+using System.Collections.Generic;
+using System.Diagnostics;
+using System.IO;
+using System.Linq;
+using System.Net.Http;
+using System.Text;
+using System.Threading;
+using System.Threading.Tasks;
+using UnrealBuildBase;
+
+namespace UnrealBuildTool
+{
+	/// <summary>
+	/// Highly experimental, a work in progress and not currently in active use.
+	/// * This executor uses Horde Compute to build actions remotely. Requests are made to Horde via a local Zen server.
+	/// * Currently only supported for Linux and Mac toolchains.
+	/// 
+	/// Current Requirements:
+	/// * -PreprocessDepends: For Linux and Mac, run a prepass of clang to generate the dependency list. Needed to get all input files required to compile.
+	///						  This will add additional overhead to the build and is not recommended unless using Horde compute remote distribution.
+	/// * -NoPCH:			  Cannot build with pch, as absolute paths are embedded into the .gch file for included headers, and these paths must exist exactly.
+	/// * -HordeCompute:	  Flag to enable this executor
+	/// * ZenServer:	      Must be running with the correct credentials to connect to both Horde Storage to upload files and Horde to make execution requests.
+	///
+	/// TODO:
+	/// * Investigate overhead when connecting from Zen to Horde
+	/// * Local process racing when remote tasks are blocking more tasks from starting, or at the end of a build.
+	/// * Print remote status as tasks complete, print more concise and useful summary at the end of the build
+	/// * Better error handling when remote tasks fail, locally retry failed actions
+	/// * Additional toolchain platform support
+	/// * Dynamic control over the number of remote processes to run at once. Needs knowledge about availibility of remote compute resources, if the pending remote queue is overloaded and by how much, etc.
+	/// </summary>
+	class HordeExecutor : ParallelExecutor
+	{
+		/// <summary>
+		/// How many processes that will be executed in parallel remotely
+		/// </summary>
+		[XmlConfigFile]
+		private static int NumRemoteParallelProcesses = 24;
+
+		/// <summary>
+		/// Only run work remotely, for testing
+		/// </summary>
+		[XmlConfigFile]
+		private static bool RemoteProcessOnly = false;
+
+		private static readonly Lazy<HttpClient> LazyZenHttpClient = new Lazy<HttpClient>();
+		private static HttpClient ZenHttpClient { get { return LazyZenHttpClient.Value; } }
+
+		// Manager to prevent rehashing and reuploading duplicate data to Zen.
+		private static readonly Lazy<ZenUploader> LazyHashLoader = new Lazy<ZenUploader>();
+		private static ZenUploader HashLoader { get { return LazyHashLoader.Value; } }
+
+		internal class ZenWorker
+		{
+			internal class ZenFile : IComparable<ZenFile>
+			{
+				[CbField("name")]
+				public string? Name { get; set; }
+				[CbField("hash")]
+				public CbBinaryAttachment? Hash { get; set; }
+				[CbField("size")]
+				public long? Size { get; set; }
+
+				public int CompareTo(ZenFile? Other)
+				{
+					return string.Compare(Name, Other?.Name);
+				}
+			}
+
+			[CbField("path")]
+			public string? Path { get; set; }
+			[CbField("workdir")]
+			public string? WorkingDirectory { get; set; }
+			[CbField("arguments")]
+			public List<string>? Arguments { get; set; } = new List<string>();
+			[CbField("executables")]
+			public List<ZenFile>? Executables { get; set; } = new List<ZenFile>();
+			[CbField("files")]
+			public List<ZenFile>? Files { get; set; } = new List<ZenFile>();
+			[CbField("directories")]
+			public List<string>? Directories { get; set; } = new List<string>();
+			[CbField("environment")]
+			public List<string>? Environment { get; set; } = new List<string>();
+			[CbField("outputs")]
+			public List<string>? Outputs { get; set; } = new List<string>();
+			[CbField("host")]
+			public string? HostPlatform { get; set; }
+			[CbField("cores")]
+			public int? Cores { get; set; }
+			[CbField("memory")]
+			public long? Memory { get; set; }
+			[CbField("exclusive")]
+			public bool? Excusive { get; set; }
+			[CbField("timeout")]
+			public int? TimeoutSeconds { get; set; }
+		}
+
+		internal class ZenWorkerPostResult
+		{
+			[CbField("need")]
+			public List<IoHash>? Need { get; set; } = new List<IoHash>();
+		}
+
+		internal class ZenWorkerResult
+		{
+			internal class ZenWorkerResultFile
+			{
+				[CbField("name")]
+				public string? Name { get; set; }
+				[CbField("data")]
+				public byte[]? Data { get; set; }
+			}
+
+			[DebuggerDisplay("{Name}: {Time}")]
+			internal class ZenWorkerResultStat : IComparable<ZenWorkerResultStat>
+			{
+				[CbField("name")]
+				public string? Name { get; set; }
+				[CbField("time")]
+				public DateTime? Time { get; set; }
+
+				public int CompareTo(ZenWorkerResultStat? Other)
+				{
+					return DateTime.Compare(Time ?? DateTime.MaxValue, Other?.Time ?? DateTime.MaxValue);
+				}
+			}
+
+			[CbField("agent")]
+			public string? Agent { get; set; }
+			[CbField("detail")]
+			public string? Detail { get; set; }
+			[CbField("stdout")]
+			public string? StdOut { get; set; }
+			[CbField("stderr")]
+			public string? StdErr { get; set; }
+			[CbField("exitcode")]
+			public int? ExitCode { get; set; }
+			[CbField("files")]
+			public List<ZenWorkerResultFile>? Files { get; set; } = new List<ZenWorkerResultFile>();
+
+			[CbField("stats")]
+			public List<ZenWorkerResultStat>? Stats { get; set; } = new List<ZenWorkerResultStat>();
+		}
+
+		internal class ZenUploader
+		{
+			Dictionary<string, IoHash> FileHashes = new Dictionary<string, IoHash>();
+
+			Dictionary<IoHash, byte[]> Datas = new Dictionary<IoHash, byte[]>();
+			HashSet<IoHash> Uploaded = new HashSet<IoHash>();
+
+			public IoHash HashFile(FileReference File)
+			{
+				if (!FileHashes.ContainsKey(File.FullName))
+				{
+					byte[] Data = FileReference.ReadAllBytes(File);
+					IoHash Hash = HashAndStoreData(Data);
+					lock (FileHashes)
+					{
+						FileHashes.TryAdd(File.FullName, Hash);
+					}
+				}
+				return FileHashes[File.FullName];
+			}
+
+			// HashString is not cached
+			public Tuple<IoHash, long> HashString(string String)
+			{
+				byte[] Data = Encoding.ASCII.GetBytes(String);
+				return Tuple.Create(HashAndStoreData(Data), Data.LongLength);
+			}
+
+			private IoHash HashAndStoreData(byte[] Data)
+			{
+				IoHash Hash = IoHash.Compute(Data);
+				lock (Datas)
+				{
+					Datas.TryAdd(Hash, Data);
+				}
+				return Hash;
+			}
+
+			public bool GetData(IoHash Hash, [System.Diagnostics.CodeAnalysis.NotNullWhen(true)] out ReadOnlyMemory<byte> Memory)
+			{
+				Memory = null;
+				if (Datas.TryGetValue(Hash, out byte[]? Data))
+				{
+					Memory = Data;
+					return true;
+				}
+				return false;
+			}
+
+			public async Task UploadData(IoHash Hash, CancellationToken CancellationToken)
+			{
+				if (Uploaded.Contains(Hash))
+				{
+					return;
+				}
+				if (!GetData(Hash, out ReadOnlyMemory<byte> Data))
+				{
+					throw new BuildException($"Missing data {Hash}");
+				}
+				var Uri = $"http://localhost:1337/cas/{Hash}";
+				var HttpPutResponse = await ZenHttpClient.PutAsync(Uri, new ReadOnlyMemoryContent(Data), CancellationToken);
+				HttpPutResponse.EnsureSuccessStatusCode();
+				lock (Uploaded)
+				{
+					Uploaded.Add(Hash);
+				}
+				lock (Datas)
+				{
+					Datas.Remove(Hash);
+				}
+			}
+
+			public void MarkAsUploaded(IoHash Hash)
+			{
+				if (Uploaded.Contains(Hash))
+				{
+					return;
+				}
+				lock (Uploaded)
+				{
+					Uploaded.Add(Hash);
+				}
+				lock (Datas)
+				{
+					Datas.Remove(Hash);
+				}
+			}
+
+			public void MarkAsUploaded(IEnumerable<IoHash> Hashes)
+			{
+				foreach (IoHash Hash in Hashes)
+				{
+					MarkAsUploaded(Hash);
+				}
+			}
+		}
+
+		protected class RemoteExecuteResults : ExecuteResults
+		{
+			public Dictionary<string, DateTime> Timepoints { get; private set; }
+
+			public RemoteExecuteResults(List<string> LogLines, int ExitCode, TimeSpan ExecutionTime, TimeSpan ProcessorTime, Dictionary<string, DateTime> Timepoints) : base(LogLines, ExitCode, ExecutionTime, ProcessorTime)
+			{
+				this.Timepoints = Timepoints;
+			}
+		}
+
+		/// <summary>
+		/// Constructor
+		/// </summary>
+		/// <param name="MaxLocalActions">How many local actions to execute in parallel</param>
+		public HordeExecutor(int MaxLocalActions) : base(MaxLocalActions)
+		{
+			XmlConfig.ApplyTo(this);
+		}
+
+		/// <summary>
+		/// Returns the name of this executor
+		/// </summary>
+		public override string Name
+		{
+			get { return "Horde"; }
+		}
+
+		/// <summary>
+		/// Checks whether the task executor can be used
+		/// </summary>
+		/// <returns>True if the task executor can be used</returns>
+		public static new bool IsAvailable()
+		{
+			try
+			{
+				var HttpPostResult = ZenHttpClient.GetAsync("http://localhost:1337/apply/ready");
+				HttpPostResult.Wait();
+				return HttpPostResult.Result.IsSuccessStatusCode;
+			}
+			catch
+			{
+				return false;
+			}
+		}
+
+		/// <summary>
+		/// Executes the specified actions locally.
+		/// </summary>
+		/// <returns>True if all the tasks successfully executed, or false if any of them failed.</returns>
+		public override bool ExecuteActions(List<LinkedAction> InputActions)
+		{
+			int NumCompletedActions = 0;
+			int TotalActions = InputActions.Count;
+			int ActualNumParallelProcesses = Math.Min(TotalActions, NumParallelProcesses);
+
+			using ManagedProcessGroup ProcessGroup = new ManagedProcessGroup();
+			using SemaphoreSlim MaxProcessSemaphore = new SemaphoreSlim(ActualNumParallelProcesses, ActualNumParallelProcesses);
+			using SemaphoreSlim MaxRemoteProcessSemaphore = new SemaphoreSlim(NumRemoteParallelProcesses, NumRemoteParallelProcesses);
+			using ProgressWriter ProgressWriter = new ProgressWriter("Compiling C++ source code...", false);
+
+			Log.TraceInformation($"Building {TotalActions} {(TotalActions == 1 ? "action" : "actions")} with {ActualNumParallelProcesses} local {(ActualNumParallelProcesses == 1 ? "process" : "processes")} and {NumRemoteParallelProcesses} remote {(NumRemoteParallelProcesses == 1 ? "process" : "processes")} ...");
+			if (RemoteProcessOnly)
+			{
+				Log.TraceWarning("Processing all distributable actions remotely. This is not recommended unless testing remote compute.");
+			}
+
+			Dictionary<LinkedAction, Task<ExecuteResults>> ExecuteTasks = new Dictionary<LinkedAction, Task<ExecuteResults>>();
+			List<Task> LogTasks = new List<Task>();
+
+			using LogIndentScope Indent = new LogIndentScope("  ");
+
+			CancellationTokenSource CancellationTokenSource = new CancellationTokenSource();
+			CancellationToken CancellationToken = CancellationTokenSource.Token;
+
+			// Create a task for every action
+			foreach (LinkedAction Action in InputActions)
+			{
+				if (ExecuteTasks.ContainsKey(Action))
+				{
+					continue;
+				}
+
+				Task<ExecuteResults> ExecuteTask = CreateRemoteExecuteTask(Action, ExecuteTasks, ProcessGroup, MaxProcessSemaphore, MaxRemoteProcessSemaphore, CancellationToken);
+				Task LogTask = ExecuteTask.ContinueWith(antecedent => LogCompletedAction(Action, antecedent, CancellationTokenSource, ProgressWriter, TotalActions, ref NumCompletedActions), CancellationToken);
+
+				ExecuteTasks.Add(Action, ExecuteTask);
+				LogTasks.Add(LogTask);
+			}
+
+			Task SummaryTask = Task.Factory.ContinueWhenAll(LogTasks.ToArray(), (AntecedentTasks) => TraceSummary(ExecuteTasks, ProcessGroup), CancellationToken);
+			SummaryTask = SummaryTask.ContinueWith(antecedent => TraceRemoteSummary(ExecuteTasks));
+			while (!SummaryTask.IsCompleted)
+			{
+				SummaryTask.Wait(1000);
+				Log.TraceInformation($"Actions Running Local:{ActualNumParallelProcesses - MaxProcessSemaphore.CurrentCount} Remote:{NumRemoteParallelProcesses - MaxRemoteProcessSemaphore.CurrentCount}");
+			}
+
+			// Return if all tasks succeeded
+			return ExecuteTasks.Values.All(x => x.Result.ExitCode == 0);
+		}
+
+		private static void TraceStats(List<RemoteExecuteResults> RemoteExecuteResults, string Description, string StartTime, string EndTime)
+		{
+			var Results = RemoteExecuteResults.Where(x => x.Timepoints.ContainsKey(StartTime) && x.Timepoints.ContainsKey(EndTime)).ToArray();
+			if (Results.Length == 0)
+			{
+				return;
+			}
+
+			Log.TraceInformation($"  {Description}:");
+			Log.TraceInformation($"    Total {Results.Sum(x => (x.Timepoints[EndTime] - x.Timepoints[StartTime]).TotalSeconds):0.00}");
+			Log.TraceInformation($"    Avg   {Results.Average(x => (x.Timepoints[EndTime] - x.Timepoints[StartTime]).TotalSeconds):0.00}");
+			Log.TraceInformation($"    Min   {Results.Min(x => (x.Timepoints[EndTime] - x.Timepoints[StartTime]).TotalSeconds):0.00}");
+			Log.TraceInformation($"    Max   {Results.Max(x => (x.Timepoints[EndTime] - x.Timepoints[StartTime]).TotalSeconds):0.00}");
+		}
+
+		private static void TraceRemoteSummary(Dictionary<LinkedAction, Task<ExecuteResults>> Tasks)
+		{
+			// TODO: Beautify summary
+			List<ExecuteResults> LocalExecuteResults = Tasks.Values.Where(x => x.IsCompletedSuccessfully && x.Result.ExitCode == 0 && !(x.Result is RemoteExecuteResults)).Select(x => x.Result).ToList();
+			List<RemoteExecuteResults> RemoteExecuteResults = Tasks.Values.Where(x => x.IsCompletedSuccessfully && x.Result.ExitCode == 0 && x.Result is RemoteExecuteResults).Select(x => (RemoteExecuteResults)x.Result).ToList();
+			Log.TraceInformation("");
+			Log.TraceInformation($"Successful Tasks Run: Local {LocalExecuteResults.Count} Remote {RemoteExecuteResults.Count}");
+			TraceStats(RemoteExecuteResults, "ubt iohash", "ubt-queue-dispatched", "ubt-hash-files");
+			TraceStats(RemoteExecuteResults, "ubt put files", "ubt-hash-files", "ubt-put-files");
+			TraceStats(RemoteExecuteResults, "zen queue wait", "zen-queue-added", "zen-queue-dispatched");
+			TraceStats(RemoteExecuteResults, "zen put horde", "zen-storage-build-ref", "zen-storage-put-ref");
+			TraceStats(RemoteExecuteResults, "horde queue wait", "horde-queue-added", "horde-queue-dispatched");
+			TraceStats(RemoteExecuteResults, "horde prepare sandbox", "horde-execution-start", "horde-execution-download-input");
+			TraceStats(RemoteExecuteResults, "horde execute", "horde-execution-download-input", "horde-execution-execute");
+			TraceStats(RemoteExecuteResults, "horde upload results", "horde-execution-execute", "horde-execution-upload-ref");
+			TraceStats(RemoteExecuteResults, "zen complete", "horde-queue-complete", "zen-queue-complete");
+			Log.TraceInformation("");
+		}
+
+		private static Task<ExecuteResults> CreateRemoteExecuteTask(LinkedAction Action, Dictionary<LinkedAction, Task<ExecuteResults>> ExecuteTasks, ManagedProcessGroup ProcessGroup, SemaphoreSlim MaxProcessSemaphore, SemaphoreSlim MaxRemoteProcessSemaphore, CancellationToken CancellationToken)
+		{
+			if (Action.PrerequisiteActions.Count == 0)
+			{
+				return Task.Factory.StartNew(
+					() => ExecuteRemoteAction(Array.Empty<Task<ExecuteResults>>(), Action, ProcessGroup, MaxProcessSemaphore, MaxRemoteProcessSemaphore, CancellationToken),
+					CancellationToken,
+					TaskCreationOptions.LongRunning,
+					TaskScheduler.Current
+				).Unwrap();
+			}
+
+			// Create tasks for any preresquite actions if they don't exist already
+			List<Task<ExecuteResults>> PrerequisiteTasks = new List<Task<ExecuteResults>>();
+			foreach (var PrerequisiteAction in Action.PrerequisiteActions)
+			{
+				if (!ExecuteTasks.ContainsKey(PrerequisiteAction))
+				{
+					ExecuteTasks.Add(PrerequisiteAction, CreateRemoteExecuteTask(PrerequisiteAction, ExecuteTasks, ProcessGroup, MaxProcessSemaphore, MaxRemoteProcessSemaphore, CancellationToken));
+				}
+				PrerequisiteTasks.Add(ExecuteTasks[PrerequisiteAction]);
+			}
+
+			return Task.Factory.ContinueWhenAll(
+				PrerequisiteTasks.ToArray(),
+				(AntecedentTasks) => ExecuteRemoteAction(AntecedentTasks, Action, ProcessGroup, MaxProcessSemaphore, MaxRemoteProcessSemaphore, CancellationToken),
+				CancellationToken,
+				TaskContinuationOptions.LongRunning,
+				TaskScheduler.Current
+			).Unwrap();
+		}
+
+		private static HashSet<FileReference>? GatherDependencies(FileItem DependencyListFile, IEnumerable<FileItem> PrerequisiteItems)
+		{
+			DependencyListFile.ResetCachedInfo();
+			// TODO: Share with CppDependencyCache in BuildMode.Build and remove TryGetDependenciesUncached
+			CppDependencyCache.TryGetDependenciesUncached(DependencyListFile, out List<FileItem>? DependencyFileItems);
+			if (DependencyFileItems == null || DependencyFileItems.Count == 0)
+			{
+				return null;
+			}
+
+			HashSet<FileReference> DependencyItems = DependencyFileItems.Distinct().Select(x => new FileReference(x.FullName)).ToHashSet();
+			DependencyItems.UnionWith(PrerequisiteItems.Where(x => x != DependencyListFile).Select(x => new FileReference(x.FullName)));
+			return DependencyItems;
+		}
+
+		private static async Task<ExecuteResults> ExecuteRemoteAction(Task<ExecuteResults>[] AntecedentTasks, LinkedAction Action, ManagedProcessGroup ProcessGroup, SemaphoreSlim MaxProcessSemaphore, SemaphoreSlim MaxRemoteProcessSemaphore, CancellationToken CancellationToken)
+		{
+			Task<bool>? SemaphoreTask = null;
+			Task<bool>? RemoteSemaphoreTask = null;
+			try
+			{
+				// Cancel tasks if any PrerequisiteActions fail, unless a PostBuildStep
+				if (Action.ActionType != ActionType.PostBuildStep && AntecedentTasks.Any(x => x.Result.ExitCode != 0))
+				{
+					throw new OperationCanceledException();
+				}
+
+				// Run locally if available, skipping any dependency checking or file hashing
+				if (!RemoteProcessOnly)
+				{
+					SemaphoreTask = MaxProcessSemaphore.WaitAsync(0, CancellationToken);
+					await SemaphoreTask;
+
+					if (SemaphoreTask.Status == TaskStatus.RanToCompletion && SemaphoreTask.Result)
+					{
+						return await RunAction(Action, ProcessGroup, CancellationToken);
+					}
+					CancellationToken.ThrowIfCancellationRequested();
+				}
+
+				// Gather dependencies from Dependency file if can execute remotely
+				HashSet<FileReference>? DependencyItems = null;
+				if (Action.bCanExecuteRemotely && Action.DependencyListFile != null && Action.PrerequisiteItems.Any(x => x == Action.DependencyListFile))
+				{
+					DependencyItems = GatherDependencies(Action.DependencyListFile!, Action.PrerequisiteItems);
+				}
+
+				// Unable to gather dependencies, must build local
+				if (DependencyItems == null || DependencyItems.Count == 0)
+				{
+					// Run ParallelExecutor.ExecuteAction
+					return await ExecuteAction(AntecedentTasks, Action, ProcessGroup, MaxProcessSemaphore, CancellationToken);
+				}
+
+				while (true)
+				{
+					if (!RemoteProcessOnly)
+					{
+						SemaphoreTask = MaxProcessSemaphore.WaitAsync(100, CancellationToken);
+						await SemaphoreTask;
+
+						if (SemaphoreTask.Result)
+						{
+							break;
+						}
+					}
+
+					RemoteSemaphoreTask = MaxRemoteProcessSemaphore.WaitAsync(0, CancellationToken);
+					await RemoteSemaphoreTask;
+
+					if (RemoteSemaphoreTask.Result)
+					{
+						break;
+					}
+					CancellationToken.ThrowIfCancellationRequested();
+				}
+
+				if (SemaphoreTask != null && SemaphoreTask.Status == TaskStatus.RanToCompletion && SemaphoreTask.Result)
+				{
+					// Local Execution
+					return await RunAction(Action, ProcessGroup, CancellationToken);
+				}
+				else if (RemoteSemaphoreTask != null && RemoteSemaphoreTask.Status == TaskStatus.RanToCompletion && RemoteSemaphoreTask.Result)
+				{
+					// Remote Execution
+					return await RunRemoteAction(Action, DependencyItems, CancellationToken);
+				}
+				else
+				{
+					// Unexpected
+					throw new OperationCanceledException();
+				}
+			}
+			catch (OperationCanceledException)
+			{
+				return new ExecuteResults(new List<string>(), int.MaxValue);
+			}
+			catch (Exception Ex)
+			{
+				Log.WriteException(Ex, null);
+				return new ExecuteResults(new List<string>(), int.MaxValue);
+			}
+			finally
+			{
+				if (SemaphoreTask?.Status == TaskStatus.RanToCompletion && SemaphoreTask?.Result == true)
+				{
+					MaxProcessSemaphore.Release();
+				}
+
+				if (RemoteSemaphoreTask?.Status == TaskStatus.RanToCompletion && RemoteSemaphoreTask?.Result == true)
+				{
+					MaxRemoteProcessSemaphore.Release();
+				}
+			}
+		}
+
+		private static async Task<ExecuteResults> RunRemoteAction(LinkedAction Action, HashSet<FileReference> DependencyItems, CancellationToken CancellationToken)
+		{
+			Dictionary<string, DateTime> Timepoints = new Dictionary<string, DateTime>();
+			Timepoints["ubt-queue-dispatched"] = DateTime.UtcNow;
+
+			ZenWorker? Worker = new ZenWorker
+			{
+				Path = Action.CommandPath.IsUnderDirectory(Unreal.RootDirectory) ? Action.CommandPath.MakeRelativeTo(UnrealBuildTool.EngineSourceDirectory) : Action.CommandPath.FullName,
+				WorkingDirectory = UnrealBuildTool.EngineSourceDirectory.MakeRelativeTo(Unreal.RootDirectory),
+				Arguments = new List<string>() { Action.CommandArguments.Trim() },
+				Directories = new List<string>() { UnrealBuildTool.EngineSourceDirectory.MakeRelativeTo(Unreal.RootDirectory) },
+				HostPlatform = BuildHostPlatform.Current.Platform.ToString(),
+				Cores = 1,
+				Memory = MemoryPerActionBytesOverride > 0 ? (long)MemoryPerActionBytesOverride : 1 * 1024 * 1024 * 1024, // 1 GB default
+				TimeoutSeconds = 300,
+				Executables = new List<ZenWorker.ZenFile>(),
+				Files = new List<ZenWorker.ZenFile>(),
+				Outputs = new List<string>(Action.ProducedItems.Distinct().Select(x => FileReference.FromString(x.FullName).MakeRelativeTo(Unreal.RootDirectory)))
+			};
+
+			foreach (var SourceFile in DependencyItems!.Where(SourceFile => SourceFile.IsUnderDirectory(Unreal.RootDirectory)))
+			{
+				Worker.Files.Add(new ZenWorker.ZenFile { Name = SourceFile.MakeRelativeTo(Unreal.RootDirectory), Hash = HashLoader.HashFile(SourceFile), Size = SourceFile.ToFileInfo().Length });
+			}
+			Timepoints["ubt-hash-files"] = DateTime.UtcNow;
+
+			// TODO: Running command through a script file shouldn't be necessary
+			if (BuildHostPlatform.Current.Platform == UnrealTargetPlatform.Win64)
+			{
+				StringBuilder ScriptContents = new StringBuilder();
+				ScriptContents.AppendLine("@echo off");
+				ScriptContents.AppendLine($"\"{Action.CommandPath.FullName}\" {Action.CommandArguments.Trim()}");
+
+				var Result = HashLoader.HashString(ScriptContents.ToString());
+				Worker.Executables.Add(new ZenWorker.ZenFile { Name = Path.Combine(Worker.WorkingDirectory, "Run.bat"), Hash = Result.Item1, Size = Result.Item2 });
+
+				Worker.Path = "cmd.exe";
+				Worker.Arguments.Clear();
+				Worker.Arguments.Add("/C");
+				Worker.Arguments.Add("Run.bat");
+			}
+			else if (BuildHostPlatform.Current.Platform == UnrealTargetPlatform.Mac || BuildHostPlatform.Current.Platform == UnrealTargetPlatform.Linux)
+			{
+				StringBuilder ScriptContents = new StringBuilder();
+				ScriptContents.AppendLine("#!/bin/sh");
+				ScriptContents.AppendLine($"'{Action.CommandPath}' {Action.CommandArguments.Trim()}");
+
+				var Result = HashLoader.HashString(ScriptContents.ToString());
+				Worker.Executables.Add(new ZenWorker.ZenFile { Name = Path.Combine(Worker.WorkingDirectory, "Run.sh"), Hash = Result.Item1, Size = Result.Item2 });
+
+				Worker.Path = "/bin/sh";
+				Worker.Arguments.Clear();
+				Worker.Arguments.Add("Run.sh");
+			}
+
+			Worker.Executables?.Sort();
+			Worker.Files?.Sort();
+			Worker.Directories?.Sort();
+			Worker.Environment?.Sort();
+			Worker.Outputs?.Sort();
+
+			IoHash WorkerHash;
+			{
+				CbObject Object = CbSerializer.Serialize(Worker);
+				WorkerHash = Object.GetHash();
+
+				ReadOnlyMemoryContent Content = new ReadOnlyMemoryContent(Object.GetView());
+				Content.Headers.Add("Content-Type", "application/x-ue-cb");
+				var HttpPostResult = await ZenHttpClient.PostAsync($"http://localhost:1337/apply/workers/{WorkerHash}", Content, CancellationToken);
+				if (HttpPostResult.StatusCode == System.Net.HttpStatusCode.NotFound)
+				{
+					CancellationToken.ThrowIfCancellationRequested();
+					ZenWorkerPostResult NeedResult = CbSerializer.Deserialize<ZenWorkerPostResult>(await HttpPostResult.Content.ReadAsByteArrayAsync(CancellationToken));
+					if (NeedResult.Need?.Count > 0)
+					{
+						foreach (var Need in NeedResult.Need)
+						{
+							await HashLoader.UploadData(Need, CancellationToken);
+						}
+
+						Content = new ReadOnlyMemoryContent(Object.GetView());
+						Content.Headers.Add("Content-Type", "application/x-ue-cb");
+						HttpPostResult = await ZenHttpClient.PostAsync($"http://localhost:1337/apply/workers/{WorkerHash}", Content, CancellationToken);
+					}
+				}
+				HttpPostResult.EnsureSuccessStatusCode();
+				CancellationToken.ThrowIfCancellationRequested();
+				HashLoader.MarkAsUploaded(Worker.Executables!.Where(x => x.Hash.HasValue).Select(x => x.Hash!.Value.Hash));
+				HashLoader.MarkAsUploaded(Worker.Files!.Where(x => x.Hash.HasValue).Select(x => x.Hash!.Value.Hash));
+			}
+			Timepoints["ubt-put-files"] = DateTime.UtcNow;
+			Worker = null;
+
+			{
+				using var Request = new HttpRequestMessage(HttpMethod.Post, $"http://localhost:1337/apply/simple/{WorkerHash}");
+				Request.Headers.Add("Accept", "application/x-ue-cb");
+				var HttpPostResult = await ZenHttpClient.SendAsync(Request, CancellationToken);
+				HttpPostResult.EnsureSuccessStatusCode();
+				CancellationToken.ThrowIfCancellationRequested();
+			}
+
+			byte[]? Buffer = null;
+			{
+				DateTime ResultsStartTime = DateTime.Now;
+				while ((DateTime.UtcNow - Timepoints["ubt-queue-dispatched"]).TotalSeconds < 600)
+				{
+					using var Request = new HttpRequestMessage(HttpMethod.Get, $"http://localhost:1337/apply/simple/{WorkerHash}");
+					Request.Headers.Add("Accept", "application/x-ue-cb");
+					var HttpGetResult = await ZenHttpClient.SendAsync(Request, CancellationToken);
+					HttpGetResult.EnsureSuccessStatusCode();
+
+					if (HttpGetResult.StatusCode == System.Net.HttpStatusCode.Accepted)
+					{
+						await Task.Delay(100, CancellationToken);
+						continue;
+					}
+
+					// TODO: Streaming?
+					Buffer = await HttpGetResult.Content.ReadAsByteArrayAsync(CancellationToken);
+					CancellationToken.ThrowIfCancellationRequested();
+					break;
+				}
+			}
+
+			if (Buffer == null)
+			{
+				throw new BuildException("Failed remote execution");
+			}
+
+			ZenWorkerResult OutputData = CbSerializer.Deserialize<ZenWorkerResult>(Buffer);
+			List<string> LogLines = new List<string>();
+			if (!string.IsNullOrEmpty(OutputData.StdOut))
+			{
+				LogLines.AddRange(OutputData.StdOut.Split('\n'));
+			}
+			if (!string.IsNullOrEmpty(OutputData.StdErr))
+			{
+				LogLines.AddRange(OutputData.StdErr.Split('\n'));
+			}
+			int ExitCode = OutputData.ExitCode ?? 0;
+
+			if (ExitCode == 0)
+			{
+				foreach (var ResultFile in OutputData.Files!)
+				{
+					FileReference File = FileReference.Combine(Unreal.RootDirectory, ResultFile.Name!);
+					await System.IO.File.WriteAllBytesAsync(File.FullName, ResultFile.Data!, CancellationToken);
+					CancellationToken.ThrowIfCancellationRequested();
+				}
+			}
+
+			OutputData.Stats?.Sort();
+			OutputData.Stats?.ForEach(x => { if (x.Name != null && x.Time.HasValue) Timepoints[x.Name] = x.Time.Value; });
+			Timepoints["ubt-queue-complete"] = DateTime.UtcNow;
+
+			return new RemoteExecuteResults(LogLines, ExitCode, Timepoints["ubt-queue-complete"] - Timepoints["ubt-queue-dispatched"], TimeSpan.Zero, Timepoints);
+		}
+	}
+}
diff --git a/Engine/Source/Programs/UnrealBuildTool/System/ActionGraph.cs b/Engine/Source/Programs/UnrealBuildTool/System/ActionGraph.cs
--- a/Engine/Source/Programs/UnrealBuildTool/System/ActionGraph.cs
+++ b/Engine/Source/Programs/UnrealBuildTool/System/ActionGraph.cs
@@ -350,6 +350,10 @@ namespace UnrealBuildTool
 				{
 					return new SNDBS(TargetDescriptors);
 				}
+				else if (BuildConfiguration.bAllowHordeCompute && HordeExecutor.IsAvailable())
+				{
+					return new HordeExecutor(BuildConfiguration.MaxParallelActions);
+				}
 			}
 
 			return new ParallelExecutor(BuildConfiguration.MaxParallelActions);
diff --git a/Engine/Source/Programs/UnrealBuildTool/System/CppDependencyCache.cs b/Engine/Source/Programs/UnrealBuildTool/System/CppDependencyCache.cs
--- a/Engine/Source/Programs/UnrealBuildTool/System/CppDependencyCache.cs
+++ b/Engine/Source/Programs/UnrealBuildTool/System/CppDependencyCache.cs
@@ -299,6 +299,24 @@ namespace UnrealBuildTool
 			return false;
 		}
 
+		/// <summary>
+		/// Attempts to read the dependencies from the given input file. This static version will not use any caches.
+		/// </summary>
+		/// <param name="InputFile">File to be read</param>
+		/// <param name="OutDependencyItems">Receives a list of output items</param>
+		/// <returns>True if the input file exists and the dependencies were read</returns>
+		public static bool TryGetDependenciesUncached(FileItem InputFile, [NotNullWhen(true)] out List<FileItem>? OutDependencyItems)
+		{
+			if (!InputFile.Exists)
+			{
+				OutDependencyItems = null;
+				return false;
+			}
+
+			OutDependencyItems = ReadDependencyInfo(InputFile).Files;
+			return true;
+		}
+
 		/// <summary>
 		/// Creates a cache hierarchy for a particular target
 		/// </summary>
diff --git a/Engine/Source/Programs/UnrealBuildTool/UnrealBuildTool.csproj b/Engine/Source/Programs/UnrealBuildTool/UnrealBuildTool.csproj
--- a/Engine/Source/Programs/UnrealBuildTool/UnrealBuildTool.csproj
+++ b/Engine/Source/Programs/UnrealBuildTool/UnrealBuildTool.csproj
@@ -97,6 +97,8 @@
     While similar checks are in RunUBT.bat/.sh and GenerateProjectFiles.bat, they don't currently check for these projects. -->
     <ProjectReference Include="..\Shared\EpicGames.Build\EpicGames.Build.csproj" />
     <ProjectReference Include="..\Shared\EpicGames.Core\EpicGames.Core.csproj" />
+    <ProjectReference Include="..\Shared\EpicGames.IoHash\EpicGames.IoHash.csproj" />
+    <ProjectReference Include="..\Shared\EpicGames.Serialization\EpicGames.Serialization.csproj" />
     <ProjectReference Include="..\Shared\EpicGames.UHT\EpicGames.UHT.csproj" />
   </ItemGroup>
   <ItemGroup>
