UnrealBuildTool: Move ISPCToolChain to it's own file

#rnx

diff --git a/Engine/Source/Programs/UnrealBuildTool/ToolChain/ISPCToolChain.cs b/Engine/Source/Programs/UnrealBuildTool/ToolChain/ISPCToolChain.cs
new file mode 100644
--- /dev/null
+++ b/Engine/Source/Programs/UnrealBuildTool/ToolChain/ISPCToolChain.cs
@@ -0,0 +1,646 @@
+// Copyright Epic Games, Inc. All Rights Reserved.
+
+using System;
+using System.Collections.Generic;
+using System.IO;
+using System.Linq;
+using System.Text;
+using EpicGames.Core;
+using UnrealBuildBase;
+
+namespace UnrealBuildTool
+{
+	abstract class ISPCToolChain : UEToolChain
+	{
+		/// <summary>
+		/// Get CPU Instruction set targets for ISPC.
+		/// </summary>
+		/// <param name="Platform">Which OS platform to target.</param>
+		/// <param name="Arch">Which architecture inside an OS platform to target. Only used for Android currently.</param>
+		/// <returns>List of instruction set targets passed to ISPC compiler</returns>
+		public virtual List<string> GetISPCCompileTargets(UnrealTargetPlatform Platform, string? Arch)
+		{
+			List<string> ISPCTargets = new List<string>();
+
+			if (UEBuildPlatform.IsPlatformInGroup(Platform, UnrealPlatformGroup.Windows) ||
+				(UEBuildPlatform.IsPlatformInGroup(Platform, UnrealPlatformGroup.Unix) && Platform != UnrealTargetPlatform.LinuxArm64) ||
+				Platform == UnrealTargetPlatform.Mac)
+			{
+				ISPCTargets.AddRange(new string[] { "avx512skx-i32x8", "avx2", "avx", "sse4", "sse2" });
+			}
+			else if (Platform == UnrealTargetPlatform.LinuxArm64)
+			{
+				ISPCTargets.AddRange(new string[] { "neon" });
+			}
+			else if (Platform == UnrealTargetPlatform.Android)
+			{
+				switch (Arch)
+				{
+					case "-armv7": ISPCTargets.Add("neon"); break; // Assumes NEON is in use
+					case "-arm64": ISPCTargets.Add("neon"); break;
+					case "-x86": ISPCTargets.AddRange(new string[] { "sse4", "sse2" }); break;
+					case "-x64": ISPCTargets.AddRange(new string[] { "sse4", "sse2" }); break;
+					default: Log.TraceWarning("Invalid Android architecture for ISPC. At least one architecture (armv7, x86, etc) needs to be selected in the project settings to build"); break;
+				}
+			}
+			else
+			{
+				Log.TraceWarning("Unsupported ISPC platform target!");
+			}
+
+			return ISPCTargets;
+		}
+
+		/// <summary>
+		/// Get OS target for ISPC.
+		/// </summary>
+		/// <param name="Platform">Which OS platform to target.</param>
+		/// <returns>OS string passed to ISPC compiler</returns>
+		public virtual string GetISPCOSTarget(UnrealTargetPlatform Platform)
+		{
+			string ISPCOS = "";
+
+			if (UEBuildPlatform.IsPlatformInGroup(Platform, UnrealPlatformGroup.Windows))
+			{
+				ISPCOS += "windows";
+			}
+			else if (UEBuildPlatform.IsPlatformInGroup(Platform, UnrealPlatformGroup.Unix))
+			{
+				ISPCOS += "linux";
+			}
+			else if (Platform == UnrealTargetPlatform.Android)
+			{
+				ISPCOS += "android";
+			}
+			else if (Platform == UnrealTargetPlatform.Mac)
+			{
+				ISPCOS += "macos";
+			}
+			else
+			{
+				Log.TraceWarning("Unsupported ISPC platform target!");
+			}
+
+			return ISPCOS;
+		}
+
+		/// <summary>
+		/// Get CPU architecture target for ISPC.
+		/// </summary>
+		/// <param name="Platform">Which OS platform to target.</param>
+		/// <param name="Arch">Which architecture inside an OS platform to target. Only used for Android currently.</param>
+		/// <returns>Arch string passed to ISPC compiler</returns>
+		public virtual string GetISPCArchTarget(UnrealTargetPlatform Platform, string? Arch)
+		{
+			string ISPCArch = "";
+
+			if (UEBuildPlatform.IsPlatformInGroup(Platform, UnrealPlatformGroup.Windows) ||
+				(UEBuildPlatform.IsPlatformInGroup(Platform, UnrealPlatformGroup.Unix) && Platform != UnrealTargetPlatform.LinuxArm64) ||
+				Platform == UnrealTargetPlatform.Mac)
+			{
+				ISPCArch += "x86-64";
+			}
+			else if (Platform == UnrealTargetPlatform.LinuxArm64)
+			{
+				ISPCArch += "aarch64";
+			}
+			else if (Platform == UnrealTargetPlatform.Android)
+			{
+				switch (Arch)
+				{
+					case "-armv7": ISPCArch += "arm"; break; // Assumes NEON is in use
+					case "-arm64": ISPCArch += "aarch64"; break;
+					case "-x86": ISPCArch += "x86"; break;
+					case "-x64": ISPCArch += "x86-64"; break;
+					default: Log.TraceWarning("Invalid Android architecture for ISPC. At least one architecture (armv7, x86, etc) needs to be selected in the project settings to build"); break;
+				}
+			}
+			else
+			{
+				Log.TraceWarning("Unsupported ISPC platform target!");
+			}
+
+			return ISPCArch;
+		}
+
+		/// <summary>
+		/// Get CPU target for ISPC.
+		/// </summary>
+		/// <param name="Platform">Which OS platform to target.</param>
+		/// <returns>CPU string passed to ISPC compiler</returns>
+		public virtual string? GetISPCCpuTarget(UnrealTargetPlatform Platform)
+		{
+			return null;  // no specific CPU selected
+		}
+
+		/// <summary>
+		/// Get host compiler path for ISPC.
+		/// </summary>
+		/// <param name="Platform">Which OS build platform is running on.</param>
+		/// <returns>Path to ISPC compiler</returns>
+		public virtual string GetISPCHostCompilerPath(UnrealTargetPlatform Platform)
+		{
+			string ISPCCompilerPathCommon = Path.Combine(UnrealBuildTool.EngineSourceDirectory.FullName, "ThirdParty", "Intel", "ISPC", "bin");
+			string ISPCArchitecturePath = "";
+			string ExeExtension = ".exe";
+
+			if (UEBuildPlatform.IsPlatformInGroup(Platform, UnrealPlatformGroup.Windows))
+			{
+				ISPCArchitecturePath = "Windows";
+			}
+			else if (Platform == UnrealTargetPlatform.Linux)
+			{
+				ISPCArchitecturePath = "Linux";
+				ExeExtension = "";
+			}
+			else if (Platform == UnrealTargetPlatform.Mac)
+			{
+				ISPCArchitecturePath = "Mac";
+				ExeExtension = "";
+			}
+			else
+			{
+				Log.TraceWarning("Unsupported ISPC host!");
+			}
+
+			return Path.Combine(ISPCCompilerPathCommon, ISPCArchitecturePath, "ispc" + ExeExtension);
+		}
+
+		static Dictionary<UnrealTargetPlatform, string> ISPCCompilerVersions = new Dictionary<UnrealTargetPlatform, string>();
+
+		/// <summary>
+		/// Returns the version of the ISPC compiler for the specified platform. If GetISPCHostCompilerPath() doesn't return a valid path
+		/// this will return a -1 version.
+		/// </summary>
+		/// <param name="Platform">Which OS build platform is running on.</param>
+		/// <returns>Version reported by the ISPC compilerr</returns>
+		public virtual string GetISPCHostCompilerVersion(UnrealTargetPlatform Platform)
+		{
+			if (!ISPCCompilerVersions.ContainsKey(Platform))
+			{
+				Version? CompilerVersion = null;
+				string CompilerPath = GetISPCHostCompilerPath(Platform);
+
+				if (!File.Exists(CompilerPath))
+				{
+					Log.TraceWarning("No ISPC compiler at {0}", CompilerPath);
+					CompilerVersion = new Version(-1, -1);
+				}
+
+				ISPCCompilerVersions[Platform] = RunToolAndCaptureOutput(new FileReference(CompilerPath), "--version", "(.*)")!;
+			}
+
+			return ISPCCompilerVersions[Platform];
+		}
+
+		/// <summary>
+		/// Get object file format for ISPC.
+		/// </summary>
+		/// <param name="Platform">Which OS build platform is running on.</param>
+		/// <returns>Object file suffix</returns>
+		public virtual string GetISPCObjectFileFormat(UnrealTargetPlatform Platform)
+		{
+			string Format = "";
+
+			if (UEBuildPlatform.IsPlatformInGroup(Platform, UnrealPlatformGroup.Windows))
+			{
+				Format += "obj";
+			}
+			else if (UEBuildPlatform.IsPlatformInGroup(Platform, UnrealPlatformGroup.Unix) ||
+					Platform == UnrealTargetPlatform.Mac ||
+					Platform == UnrealTargetPlatform.Android)
+			{
+				Format += "obj";
+			}
+			else
+			{
+				Log.TraceWarning("Unsupported ISPC platform target!");
+			}
+
+			return Format;
+		}
+
+		/// <summary>
+		/// Get object file suffix for ISPC.
+		/// </summary>
+		/// <param name="Platform">Which OS build platform is running on.</param>
+		/// <returns>Object file suffix</returns>
+		public virtual string GetISPCObjectFileSuffix(UnrealTargetPlatform Platform)
+		{
+			string Suffix = "";
+
+			if (UEBuildPlatform.IsPlatformInGroup(Platform, UnrealPlatformGroup.Windows))
+			{
+				Suffix += ".obj";
+			}
+			else if (UEBuildPlatform.IsPlatformInGroup(Platform, UnrealPlatformGroup.Unix) ||
+					Platform == UnrealTargetPlatform.Mac ||
+					Platform == UnrealTargetPlatform.Android)
+			{
+				Suffix += ".o";
+			}
+			else
+			{
+				Log.TraceWarning("Unsupported ISPC platform target!");
+			}
+
+			return Suffix;
+		}
+
+		private string EscapeDefinitionForISPC(string Definition)
+		{
+			// See: [REDACTED_PROFILE_URL]/ispc/blob/4ee767560cd752eaf464c124eb7ef1b0fd37f1df/src/main.cpp#L264 for ispc's argument parsing code, which does the following (and does not support escaping):
+			// Argument      Parses as 
+			// "abc""def"    One agrument:  abcdef
+			// "'abc'"       One argument:  'abc'
+			// -D"X="Y Z""   Two arguments: -DX=Y and Z
+			// -D'X="Y Z"'   One argument:  -DX="Y Z"  (i.e. with quotes in value)
+			// -DX="Y Z"     One argument:  -DX=Y Z    (this is what we want on the command line)
+
+			// Assumes that quotes at the start and end of the value string mean that everything between them should be passed on unchanged.
+
+			int DoubleQuoteCount = Definition.Count(c => c == '"');
+			bool bHasSingleQuote = Definition.Contains('\'');
+			bool bHasSpace = Definition.Contains(' ');
+
+			string Escaped = Definition;
+
+			if (DoubleQuoteCount > 0 || bHasSingleQuote || bHasSpace)
+			{
+				int EqualsIndex = Definition.IndexOf('=');
+				string Name = Definition[0..EqualsIndex];
+				string Value = Definition[(EqualsIndex + 1)..];
+
+				string UnquotedValue = Value;
+
+				// remove one layer of quoting, if present
+				if (Value.StartsWith('"') && Value.EndsWith('"') && Value.Length != 1)
+				{
+					UnquotedValue = Value[1..^1];
+					DoubleQuoteCount -= 2;
+				}
+
+				if (DoubleQuoteCount == 0 && (bHasSingleQuote || bHasSpace))
+				{
+					Escaped = $"{Name}=\"{UnquotedValue}\"";
+				}
+				else if (!bHasSingleQuote && (bHasSpace || DoubleQuoteCount > 0))
+				{
+					// If there are no single quotes, we can use them to quote the value string
+					Escaped = $"{Name}='{UnquotedValue}'";
+				}
+				else
+				{
+					// Treat all special chars in the value string as needing explicit extra quoting. Thoroughly clumsy.
+					StringBuilder Requoted = new StringBuilder();
+					foreach (char c in UnquotedValue)
+					{
+						if (c == '"')
+						{
+							Requoted.Append("'\"'");
+						}
+						else if (c == '\'')
+						{
+							Requoted.Append("\"'\"");
+						}
+						else if (c == ' ')
+						{
+							Requoted.Append("\" \"");
+						}
+						else
+						{
+							Requoted.Append(c);
+						}
+					}
+					Escaped = $"{Name}={Requoted}";
+				}
+			}
+
+			return Escaped;
+		}
+
+		public override CPPOutput GenerateISPCHeaders(CppCompileEnvironment CompileEnvironment, List<FileItem> InputFiles, DirectoryReference OutputDir, IActionGraphBuilder Graph)
+		{
+			CPPOutput Result = new CPPOutput();
+
+			if (!CompileEnvironment.bCompileISPC)
+			{
+				return Result;
+			}
+
+			List<string> CompileTargets = GetISPCCompileTargets(CompileEnvironment.Platform, null);
+
+			foreach (FileItem ISPCFile in InputFiles)
+			{
+				Action CompileAction = Graph.CreateAction(ActionType.Compile);
+				CompileAction.CommandDescription = "Compile";
+				CompileAction.WorkingDirectory = UnrealBuildTool.EngineSourceDirectory;
+				CompileAction.CommandPath = new FileReference(GetISPCHostCompilerPath(BuildHostPlatform.Current.Platform));
+				CompileAction.StatusDescription = Path.GetFileName(ISPCFile.AbsolutePath);
+
+				// Disable remote execution to workaround mismatched case on XGE
+				CompileAction.bCanExecuteRemotely = false;
+
+				List<string> Arguments = new List<string>();
+
+				// Add the ISPC obj file as a prerequisite of the action.
+				CompileAction.CommandArguments = String.Format("\"{0}\" ", ISPCFile.AbsolutePath);
+
+				// Add the ISPC h file to the produced item list.
+				FileItem ISPCIncludeHeaderFile = FileItem.GetItemByFileReference(
+					FileReference.Combine(
+						OutputDir,
+						Path.GetFileName(ISPCFile.AbsolutePath) + ".generated.dummy.h"
+						)
+					);
+
+				// Add the ISPC file to be compiled.
+				Arguments.Add(String.Format("-h \"{0}\"", ISPCIncludeHeaderFile));
+
+				// Build target string. No comma on last
+				string TargetString = "";
+				foreach (string Target in CompileTargets)
+				{
+					if (Target == CompileTargets[CompileTargets.Count - 1]) // .Last()
+					{
+						TargetString += Target;
+					}
+					else
+					{
+						TargetString += Target + ",";
+					}
+				}
+
+				// Build target triplet
+				Arguments.Add(String.Format("--target-os={0}", GetISPCOSTarget(CompileEnvironment.Platform)));
+				Arguments.Add(String.Format("--arch={0}", GetISPCArchTarget(CompileEnvironment.Platform, null)));
+				Arguments.Add(String.Format("--target={0}", TargetString));
+				Arguments.Add(String.Format("--emit-{0}", GetISPCObjectFileFormat(CompileEnvironment.Platform)));
+
+				string? CpuTarget = GetISPCCpuTarget(CompileEnvironment.Platform);
+				if (!String.IsNullOrEmpty(CpuTarget))
+				{
+					Arguments.Add(String.Format("--cpu={0}", CpuTarget));
+				}
+
+				// PIC is needed for modular builds except on Windows
+				if ((CompileEnvironment.bIsBuildingDLL ||
+					CompileEnvironment.bIsBuildingLibrary) &&
+					!UEBuildPlatform.IsPlatformInGroup(CompileEnvironment.Platform, UnrealPlatformGroup.Windows))
+				{
+					Arguments.Add("--pic");
+				}
+
+				// Include paths. Don't use AddIncludePath() here, since it uses the full path and exceeds the max command line length.
+				// Because ISPC response files don't support white space in arguments, paths with white space need to be passed to the command line directly.
+				foreach (DirectoryReference IncludePath in CompileEnvironment.UserIncludePaths)
+				{
+					Arguments.Add(String.Format("-I\"{0}\"", IncludePath));
+				}
+
+				// System include paths.
+				foreach (DirectoryReference SystemIncludePath in CompileEnvironment.SystemIncludePaths)
+				{
+					Arguments.Add(String.Format("-I\"{0}\"", SystemIncludePath));
+				}
+
+				// Preprocessor definitions.
+				foreach (string Definition in CompileEnvironment.Definitions)
+				{
+					// TODO: Causes ISPC compiler to generate a spurious warning about the universal character set
+					if (!Definition.Contains("\\\\U") && !Definition.Contains("\\\\u"))
+					{
+						Arguments.Add($"-D{EscapeDefinitionForISPC(Definition)}");
+					}
+				}
+
+				// Generate the included header dependency list
+				if (CompileEnvironment.bGenerateDependenciesFile)
+				{
+					FileItem DependencyListFile = FileItem.GetItemByFileReference(FileReference.Combine(OutputDir, Path.GetFileName(ISPCFile.AbsolutePath) + ".txt"));
+					Arguments.Add(String.Format("-MMM \"{0}\"", DependencyListFile.AbsolutePath.Replace('\\', '/')));
+					CompileAction.DependencyListFile = DependencyListFile;
+					CompileAction.ProducedItems.Add(DependencyListFile);
+				}
+
+				CompileAction.ProducedItems.Add(ISPCIncludeHeaderFile);
+
+				FileReference ResponseFileName = new FileReference(ISPCIncludeHeaderFile.AbsolutePath + ".response");
+				FileItem ResponseFileItem = Graph.CreateIntermediateTextFile(ResponseFileName, Arguments.Select(x => Utils.ExpandVariables(x)));
+				CompileAction.CommandArguments += String.Format("@\"{0}\"", ResponseFileName);
+				CompileAction.PrerequisiteItems.Add(ResponseFileItem);
+
+				// Add the source file and its included files to the prerequisite item list.
+				CompileAction.PrerequisiteItems.Add(ISPCFile);
+
+				FileItem ISPCFinalHeaderFile = FileItem.GetItemByFileReference(
+					FileReference.Combine(
+						OutputDir,
+						Path.GetFileName(ISPCFile.AbsolutePath) + ".generated.h"
+						)
+					);
+
+				// Fix interrupted build issue by copying header after generation completes
+				FileReference SourceFile = ISPCIncludeHeaderFile.Location;
+				FileReference TargetFile = ISPCFinalHeaderFile.Location;
+
+				FileItem SourceFileItem = FileItem.GetItemByFileReference(SourceFile);
+				FileItem TargetFileItem = FileItem.GetItemByFileReference(TargetFile);
+
+				Action CopyAction = Graph.CreateAction(ActionType.BuildProject);
+				CopyAction.CommandDescription = "Copy";
+				CopyAction.CommandPath = BuildHostPlatform.Current.Shell;
+				if (BuildHostPlatform.Current.ShellType == ShellType.Cmd)
+				{
+					CopyAction.CommandArguments = String.Format("/C \"copy /Y \"{0}\" \"{1}\" 1>nul\"", SourceFile, TargetFile);
+				}
+				else
+				{
+					CopyAction.CommandArguments = String.Format("-c 'cp -f \"\"{0}\"\" \"\"{1}\"'", SourceFile.FullName, TargetFile.FullName);
+				}
+				CopyAction.WorkingDirectory = UnrealBuildTool.EngineSourceDirectory;
+				CopyAction.PrerequisiteItems.Add(SourceFileItem);
+				CopyAction.ProducedItems.Add(TargetFileItem);
+				CopyAction.StatusDescription = TargetFileItem.Location.GetFileName();
+				CopyAction.bCanExecuteRemotely = false;
+				CopyAction.bShouldOutputStatusDescription = false;
+
+				Result.GeneratedHeaderFiles.Add(TargetFileItem);
+
+				Log.TraceVerbose("   ISPC Generating Header " + CompileAction.StatusDescription + ": \"" + CompileAction.CommandPath + "\"" + CompileAction.CommandArguments);
+			}
+
+			return Result;
+		}
+
+		public override CPPOutput CompileISPCFiles(CppCompileEnvironment CompileEnvironment, List<FileItem> InputFiles, DirectoryReference OutputDir, IActionGraphBuilder Graph)
+		{
+			CPPOutput Result = new CPPOutput();
+
+			if (!CompileEnvironment.bCompileISPC)
+			{
+				return Result;
+			}
+
+			List<string> CompileTargets = GetISPCCompileTargets(CompileEnvironment.Platform, null);
+
+			foreach (FileItem ISPCFile in InputFiles)
+			{
+				Action CompileAction = Graph.CreateAction(ActionType.Compile);
+				CompileAction.CommandDescription = "Compile";
+				CompileAction.WorkingDirectory = UnrealBuildTool.EngineSourceDirectory;
+				CompileAction.CommandPath = new FileReference(GetISPCHostCompilerPath(BuildHostPlatform.Current.Platform));
+				CompileAction.StatusDescription = Path.GetFileName(ISPCFile.AbsolutePath);
+
+				// Disable remote execution to workaround mismatched case on XGE
+				CompileAction.bCanExecuteRemotely = false;
+
+				List<string> Arguments = new List<string>();
+
+				// Add the ISPC file to be compiled.
+				Arguments.Add(String.Format(" \"{0}\"", ISPCFile.AbsolutePath));
+
+				List<FileItem> CompiledISPCObjFiles = new List<FileItem>();
+				string TargetString = "";
+
+				foreach (string Target in CompileTargets)
+				{
+					string ObjTarget = Target;
+
+					if (Target.Contains("-"))
+					{
+						// Remove lane width and gang size from obj file name
+						ObjTarget = Target.Split('-')[0];
+					}
+
+					FileItem CompiledISPCObjFile;
+
+					if (CompileTargets.Count > 1)
+					{
+						CompiledISPCObjFile = FileItem.GetItemByFileReference(
+						FileReference.Combine(
+							OutputDir,
+							Path.GetFileName(ISPCFile.AbsolutePath) + "_" + ObjTarget + GetISPCObjectFileSuffix(CompileEnvironment.Platform)
+							)
+						);
+					}
+					else
+					{
+						CompiledISPCObjFile = FileItem.GetItemByFileReference(
+						FileReference.Combine(
+							OutputDir,
+							Path.GetFileName(ISPCFile.AbsolutePath) + GetISPCObjectFileSuffix(CompileEnvironment.Platform)
+							)
+						);
+					}
+
+					// Add the ISA specific ISPC obj files to the produced item list.
+					CompiledISPCObjFiles.Add(CompiledISPCObjFile);
+
+					// Build target string. No comma on last
+					if (Target == CompileTargets[CompileTargets.Count - 1]) // .Last()
+					{
+						TargetString += Target;
+					}
+					else
+					{
+						TargetString += Target + ",";
+					}
+				}
+
+				// Add the common ISPC obj file to the produced item list.
+				FileItem CompiledISPCObjFileNoISA = FileItem.GetItemByFileReference(
+					FileReference.Combine(
+						OutputDir,
+						Path.GetFileName(ISPCFile.AbsolutePath) + GetISPCObjectFileSuffix(CompileEnvironment.Platform)
+						)
+					);
+
+				CompiledISPCObjFiles.Add(CompiledISPCObjFileNoISA);
+
+				// Add the output ISPC obj file
+				Arguments.Add(String.Format("-o \"{0}\"", CompiledISPCObjFileNoISA));
+
+				// Build target triplet
+				Arguments.Add(String.Format("--target-os=\"{0}\"", GetISPCOSTarget(CompileEnvironment.Platform)));
+				Arguments.Add(String.Format("--arch=\"{0}\"", GetISPCArchTarget(CompileEnvironment.Platform, null)));
+				Arguments.Add(String.Format("--target=\"{0}\"", TargetString));
+				Arguments.Add(String.Format("--emit-{0}", GetISPCObjectFileFormat(CompileEnvironment.Platform)));
+
+				if (CompileEnvironment.Configuration == CppConfiguration.Debug)
+				{
+					if (CompileEnvironment.Platform == UnrealTargetPlatform.Mac)
+					{
+						// Turn off debug symbols on Mac due to dsym generation issue
+						Arguments.Add("-O0");
+						// Ideally we would be able to turn on symbols and specify the dwarf version, but that does
+						// does not seem to be working currently, ie:
+						//    Arguments.Add("-g -O0 --dwarf-version=2");
+
+					}
+					else
+					{
+						Arguments.Add("-g -O0");
+					}
+				}
+				else
+				{
+					Arguments.Add("-O2");
+				}
+
+				// PIC is needed for modular builds except on Windows
+				if ((CompileEnvironment.bIsBuildingDLL ||
+					CompileEnvironment.bIsBuildingLibrary) &&
+					!UEBuildPlatform.IsPlatformInGroup(CompileEnvironment.Platform, UnrealPlatformGroup.Windows))
+				{
+					Arguments.Add("--pic");
+				}
+
+				// Include paths. Don't use AddIncludePath() here, since it uses the full path and exceeds the max command line length.
+				foreach (DirectoryReference IncludePath in CompileEnvironment.UserIncludePaths)
+				{
+					Arguments.Add(String.Format("-I\"{0}\"", IncludePath));
+				}
+
+				// System include paths.
+				foreach (DirectoryReference SystemIncludePath in CompileEnvironment.SystemIncludePaths)
+				{
+					Arguments.Add(String.Format("-I\"{0}\"", SystemIncludePath));
+				}
+
+				// Preprocessor definitions.
+				foreach (string Definition in CompileEnvironment.Definitions)
+				{
+					// TODO: Causes ISPC compiler to generate a spurious warning about the universal character set
+					if (!Definition.Contains("\\\\U") && !Definition.Contains("\\\\u"))
+					{
+						Arguments.Add($"-D{EscapeDefinitionForISPC(Definition)}");
+					}
+				}
+
+				// Consume the included header dependency list
+				if (CompileEnvironment.bGenerateDependenciesFile)
+				{
+					FileItem DependencyListFile = FileItem.GetItemByFileReference(FileReference.Combine(OutputDir, Path.GetFileName(ISPCFile.AbsolutePath) + ".txt"));
+					CompileAction.DependencyListFile = DependencyListFile;
+					CompileAction.PrerequisiteItems.Add(DependencyListFile);
+				}
+
+				CompileAction.ProducedItems.AddRange(CompiledISPCObjFiles);
+				Result.ObjectFiles.AddRange(CompiledISPCObjFiles);
+
+				FileReference ResponseFileName = new FileReference(CompiledISPCObjFileNoISA.AbsolutePath + ".response");
+				FileItem ResponseFileItem = Graph.CreateIntermediateTextFile(ResponseFileName, Arguments.Select(x => Utils.ExpandVariables(x)));
+				CompileAction.CommandArguments = " @\"" + ResponseFileName + "\"";
+				CompileAction.PrerequisiteItems.Add(ResponseFileItem);
+
+				// Add the source file and its included files to the prerequisite item list.
+				CompileAction.PrerequisiteItems.Add(ISPCFile);
+
+				Log.TraceVerbose("   ISPC Compiling " + CompileAction.StatusDescription + ": \"" + CompileAction.CommandPath + "\"" + CompileAction.CommandArguments);
+			}
+
+			return Result;
+		}
+	}
+}
diff --git a/Engine/Source/Programs/UnrealBuildTool/ToolChain/UEToolChain.cs b/Engine/Source/Programs/UnrealBuildTool/ToolChain/UEToolChain.cs
--- a/Engine/Source/Programs/UnrealBuildTool/ToolChain/UEToolChain.cs
+++ b/Engine/Source/Programs/UnrealBuildTool/ToolChain/UEToolChain.cs
@@ -2,9 +2,6 @@
 
 using System;
 using System.Collections.Generic;
-using System.IO;
-using System.Linq;
-using System.Text;
 using System.Text.RegularExpressions;
 using EpicGames.Core;
 using UnrealBuildBase;
@@ -47,6 +44,7 @@ namespace UnrealBuildTool
 			CPPOutput Result = new CPPOutput();
 			return Result;
 		}
+
 		public virtual CPPOutput GenerateISPCHeaders(CppCompileEnvironment Environment, List<FileItem> InputFiles, DirectoryReference OutputDir, IActionGraphBuilder Graph)
 		{
 			CPPOutput Result = new CPPOutput();
@@ -65,7 +63,6 @@ namespace UnrealBuildTool
 			return LinkFile != null ? new FileItem[] { LinkFile } : new FileItem[] { };
 		}
 
-
 		/// <summary>
 		/// Get the name of the response file for the current linker environment and output file
 		/// </summary>
@@ -99,9 +96,8 @@ namespace UnrealBuildTool
 		{
 		}
 
-		public virtual void PrepareRuntimeDependencies( List<RuntimeDependency> RuntimeDependencies, Dictionary<FileReference, FileReference> TargetFileToSourceFile, DirectoryReference ExeDir )
+		public virtual void PrepareRuntimeDependencies(List<RuntimeDependency> RuntimeDependencies, Dictionary<FileReference, FileReference> TargetFileToSourceFile, DirectoryReference ExeDir)
 		{
-			
 		}
 
 		/// <summary>
@@ -113,23 +109,21 @@ namespace UnrealBuildTool
 		{
 			return true;
 		}
-		
+
 		public virtual FileReference GetDebugFile(FileReference OutputFile, string DebugExtension)
 		{
 			//  by default, just change the extension to the debug extension
 			return OutputFile.ChangeExtension(DebugExtension);
 		}
 
-
 		public virtual void SetupBundleDependencies(List<UEBuildBinary> Binaries, string GameName)
 		{
-
 		}
 
-        public virtual string GetSDKVersion()
-        {
-            return "Not Applicable";
-        }
+		public virtual string GetSDKVersion()
+		{
+			return "Not Applicable";
+		}
 
 		/// <summary>
 		/// Runs the provided tool and argument. Returns the output, using a rexex capture if one is provided
@@ -146,9 +140,9 @@ namespace UnrealBuildTool
 			{
 				return ProcessOutput;
 			}
-			
+
 			Match M = Regex.Match(ProcessOutput, Expression);
-			return M.Success ? M.Groups[1].ToString() : null;	
+			return M.Success ? M.Groups[1].ToString() : null;
 		}
 
 		/// <summary>
@@ -174,638 +168,4 @@ namespace UnrealBuildTool
 			return new Version(0, 0);
 		}
 	};
-
-	abstract class ISPCToolChain : UEToolChain
-	{
-		/// <summary>
-		/// Get CPU Instruction set targets for ISPC.
-		/// </summary>
-		/// <param name="Platform">Which OS platform to target.</param>
-		/// <param name="Arch">Which architecture inside an OS platform to target. Only used for Android currently.</param>
-		/// <returns>List of instruction set targets passed to ISPC compiler</returns>
-		public virtual List<string> GetISPCCompileTargets(UnrealTargetPlatform Platform, string? Arch)
-		{
-			List<string> ISPCTargets = new List<string>();
-
-			if (UEBuildPlatform.IsPlatformInGroup(Platform, UnrealPlatformGroup.Windows) ||
-				(UEBuildPlatform.IsPlatformInGroup(Platform, UnrealPlatformGroup.Unix) && Platform != UnrealTargetPlatform.LinuxArm64) ||
-				Platform == UnrealTargetPlatform.Mac)
-			{
-				ISPCTargets.AddRange(new string[] { "avx512skx-i32x8", "avx2", "avx", "sse4", "sse2" });
-			}
-			else if (Platform == UnrealTargetPlatform.LinuxArm64)
-			{
-				ISPCTargets.AddRange(new string[] { "neon" });
-			}
-			else if (Platform == UnrealTargetPlatform.Android)
-			{
-				switch (Arch)
-				{
-					case "-armv7": ISPCTargets.Add("neon"); break; // Assumes NEON is in use
-					case "-arm64": ISPCTargets.Add("neon"); break;
-					case "-x86": ISPCTargets.AddRange(new string[] { "sse4", "sse2" }); break;
-					case "-x64": ISPCTargets.AddRange(new string[] { "sse4", "sse2" }); break;
-					default: Log.TraceWarning("Invalid Android architecture for ISPC. At least one architecture (armv7, x86, etc) needs to be selected in the project settings to build"); break;
-				}
-			}
-			else
-			{
-				Log.TraceWarning("Unsupported ISPC platform target!");
-			}
-
-			return ISPCTargets;
-		}
-
-		/// <summary>
-		/// Get OS target for ISPC.
-		/// </summary>
-		/// <param name="Platform">Which OS platform to target.</param>
-		/// <returns>OS string passed to ISPC compiler</returns>
-		public virtual string GetISPCOSTarget(UnrealTargetPlatform Platform)
-		{
-			string ISPCOS = "";
-
-			if (UEBuildPlatform.IsPlatformInGroup(Platform, UnrealPlatformGroup.Windows))
-			{
-				ISPCOS += "windows";
-			}
-			else if (UEBuildPlatform.IsPlatformInGroup(Platform, UnrealPlatformGroup.Unix))
-			{
-				ISPCOS += "linux";
-			}
-			else if (Platform == UnrealTargetPlatform.Android)
-			{
-				ISPCOS += "android";
-			}
-			else if (Platform == UnrealTargetPlatform.Mac)
-			{
-				ISPCOS += "macos";
-			}
-			else
-			{
-				Log.TraceWarning("Unsupported ISPC platform target!");
-			}
-
-			return ISPCOS;
-		}
-
-		/// <summary>
-		/// Get CPU architecture target for ISPC.
-		/// </summary>
-		/// <param name="Platform">Which OS platform to target.</param>
-		/// <param name="Arch">Which architecture inside an OS platform to target. Only used for Android currently.</param>
-		/// <returns>Arch string passed to ISPC compiler</returns>
-		public virtual string GetISPCArchTarget(UnrealTargetPlatform Platform, string? Arch)
-		{
-			string ISPCArch = "";
-
-			if (UEBuildPlatform.IsPlatformInGroup(Platform, UnrealPlatformGroup.Windows) ||
-				(UEBuildPlatform.IsPlatformInGroup(Platform, UnrealPlatformGroup.Unix) && Platform != UnrealTargetPlatform.LinuxArm64) ||
-				Platform == UnrealTargetPlatform.Mac)
-			{
-				ISPCArch += "x86-64";
-			}
-			else if (Platform == UnrealTargetPlatform.LinuxArm64)
-			{
-				ISPCArch += "aarch64";
-			}
-			else if (Platform == UnrealTargetPlatform.Android)
-			{
-				switch (Arch)
-				{
-					case "-armv7": ISPCArch += "arm"; break; // Assumes NEON is in use
-					case "-arm64": ISPCArch += "aarch64"; break;
-					case "-x86": ISPCArch += "x86"; break;
-					case "-x64": ISPCArch += "x86-64"; break;
-					default: Log.TraceWarning("Invalid Android architecture for ISPC. At least one architecture (armv7, x86, etc) needs to be selected in the project settings to build"); break;
-				}
-			}
-			else
-			{
-				Log.TraceWarning("Unsupported ISPC platform target!");
-			}
-
-			return ISPCArch;
-		}
-
-		/// <summary>
-		/// Get CPU target for ISPC.
-		/// </summary>
-		/// <param name="Platform">Which OS platform to target.</param>
-		/// <returns>CPU string passed to ISPC compiler</returns>
-		public virtual string? GetISPCCpuTarget(UnrealTargetPlatform Platform)
-		{
-			return null;  // no specific CPU selected
-		}
-
-		/// <summary>
-		/// Get host compiler path for ISPC.
-		/// </summary>
-		/// <param name="Platform">Which OS build platform is running on.</param>
-		/// <returns>Path to ISPC compiler</returns>
-		public virtual string GetISPCHostCompilerPath(UnrealTargetPlatform Platform)
-		{
-			string ISPCCompilerPathCommon = Path.Combine(UnrealBuildTool.EngineSourceDirectory.FullName, "ThirdParty", "Intel", "ISPC", "bin");
-			string ISPCArchitecturePath = "";
-			string ExeExtension = ".exe";
-
-			if (UEBuildPlatform.IsPlatformInGroup(Platform, UnrealPlatformGroup.Windows))
-			{
-				ISPCArchitecturePath = "Windows";
-			}
-			else if (Platform == UnrealTargetPlatform.Linux)
-			{
-				ISPCArchitecturePath = "Linux";
-				ExeExtension = "";
-			}
-			else if (Platform == UnrealTargetPlatform.Mac)
-			{
-				ISPCArchitecturePath = "Mac";
-				ExeExtension = "";
-			}
-			else
-			{
-				Log.TraceWarning("Unsupported ISPC host!");
-			}
-
-			return Path.Combine(ISPCCompilerPathCommon, ISPCArchitecturePath, "ispc" + ExeExtension);
-		}
-
-		static Dictionary<UnrealTargetPlatform, string> ISPCCompilerVersions = new Dictionary<UnrealTargetPlatform, string>();
-
-		/// <summary>
-		/// Returns the version of the ISPC compiler for the specified platform. If GetISPCHostCompilerPath() doesn't return a valid path
-		/// this will return a -1 version.
-		/// </summary>
-		/// <param name="Platform">Which OS build platform is running on.</param>
-		/// <returns>Version reported by the ISPC compilerr</returns>
-		public virtual string GetISPCHostCompilerVersion(UnrealTargetPlatform Platform)
-		{
-			if (!ISPCCompilerVersions.ContainsKey(Platform))
-			{
-				Version? CompilerVersion = null;
-				string CompilerPath = GetISPCHostCompilerPath(Platform);
-
-				if (!File.Exists(CompilerPath))
-				{
-					Log.TraceWarning("No ISPC compiler at {0}", CompilerPath);
-					CompilerVersion = new Version(-1, -1);
-				}
-
-				ISPCCompilerVersions[Platform] = RunToolAndCaptureOutput(new FileReference(CompilerPath), "--version", "(.*)")!;
-			}
-
-			return ISPCCompilerVersions[Platform];			
-		}
-
-		/// <summary>
-		/// Get object file format for ISPC.
-		/// </summary>
-		/// <param name="Platform">Which OS build platform is running on.</param>
-		/// <returns>Object file suffix</returns>
-		public virtual string GetISPCObjectFileFormat(UnrealTargetPlatform Platform)
-		{
-			string Format = "";
-
-			if (UEBuildPlatform.IsPlatformInGroup(Platform, UnrealPlatformGroup.Windows))
-			{
-				Format += "obj";
-			}
-			else if (UEBuildPlatform.IsPlatformInGroup(Platform, UnrealPlatformGroup.Unix) ||
-					Platform == UnrealTargetPlatform.Mac ||
-					Platform == UnrealTargetPlatform.Android)
-			{
-				Format += "obj";
-			}
-			else
-			{
-				Log.TraceWarning("Unsupported ISPC platform target!");
-			}
-
-			return Format;
-		}
-
-		/// <summary>
-		/// Get object file suffix for ISPC.
-		/// </summary>
-		/// <param name="Platform">Which OS build platform is running on.</param>
-		/// <returns>Object file suffix</returns>
-		public virtual string GetISPCObjectFileSuffix(UnrealTargetPlatform Platform)
-		{
-			string Suffix = "";
-
-			if (UEBuildPlatform.IsPlatformInGroup(Platform, UnrealPlatformGroup.Windows))
-			{
-				Suffix += ".obj";
-			}
-			else if (UEBuildPlatform.IsPlatformInGroup(Platform, UnrealPlatformGroup.Unix) ||
-					Platform == UnrealTargetPlatform.Mac ||
-					Platform == UnrealTargetPlatform.Android)
-			{
-				Suffix += ".o";
-			}
-			else
-			{
-				Log.TraceWarning("Unsupported ISPC platform target!");
-			}
-
-			return Suffix;
-		}
-
-		private string EscapeDefinitionForISPC(string Definition)
-		{
-			// See: [REDACTED_PROFILE_URL]/ispc/blob/4ee767560cd752eaf464c124eb7ef1b0fd37f1df/src/main.cpp#L264 for ispc's argument parsing code, which does the following (and does not support escaping):
-			// Argument      Parses as 
-			// "abc""def"    One agrument:  abcdef
-			// "'abc'"       One argument:  'abc'
-			// -D"X="Y Z""   Two arguments: -DX=Y and Z
-			// -D'X="Y Z"'   One argument:  -DX="Y Z"  (i.e. with quotes in value)
-			// -DX="Y Z"     One argument:  -DX=Y Z    (this is what we want on the command line)
-
-			// Assumes that quotes at the start and end of the value string mean that everything between them should be passed on unchanged.
-
-			int DoubleQuoteCount = Definition.Count(c => c == '"');
-			bool bHasSingleQuote = Definition.Contains('\'');
-			bool bHasSpace = Definition.Contains(' ');
-
-			string Escaped = Definition;
-
-			if (DoubleQuoteCount > 0 || bHasSingleQuote || bHasSpace)
-			{
-				int EqualsIndex = Definition.IndexOf('=');
-				string Name = Definition[0..EqualsIndex];
-				string Value = Definition[(EqualsIndex + 1)..];
-
-				string UnquotedValue = Value;
-
-				// remove one layer of quoting, if present
-				if (Value.StartsWith('"') && Value.EndsWith('"') && Value.Length != 1)
-				{
-					UnquotedValue = Value[1..^1];
-					DoubleQuoteCount -= 2;
-				}
-
-				if (DoubleQuoteCount == 0 && (bHasSingleQuote || bHasSpace))
-				{
-					Escaped = $"{Name}=\"{UnquotedValue}\"";
-				}
-				else if (!bHasSingleQuote && (bHasSpace || DoubleQuoteCount > 0))
-				{
-					// If there are no single quotes, we can use them to quote the value string
-					Escaped = $"{Name}='{UnquotedValue}'";
-				}
-				else
-				{
-					// Treat all special chars in the value string as needing explicit extra quoting. Thoroughly clumsy.
-					StringBuilder Requoted = new StringBuilder();
-					foreach (char c in UnquotedValue)
-					{
-						if (c == '"')
-						{
-							Requoted.Append("'\"'");
-						}
-						else if (c == '\'')
-						{
-							Requoted.Append("\"'\"");
-						}
-						else if (c == ' ')
-						{
-							Requoted.Append("\" \"");
-						}
-						else
-						{
-							Requoted.Append(c);
-						}
-					}
-					Escaped = $"{Name}={Requoted}";
-				}
-			}
-
-			return Escaped;
-		}
-
-		public override CPPOutput GenerateISPCHeaders(CppCompileEnvironment CompileEnvironment, List<FileItem> InputFiles, DirectoryReference OutputDir, IActionGraphBuilder Graph)
-		{
-			CPPOutput Result = new CPPOutput();
-
-			if(!CompileEnvironment.bCompileISPC)
-			{
-				return Result;
-			}
-
-			List<string> CompileTargets = GetISPCCompileTargets(CompileEnvironment.Platform, null);
-
-			foreach (FileItem ISPCFile in InputFiles)
-			{
-				Action CompileAction = Graph.CreateAction(ActionType.Compile);
-				CompileAction.CommandDescription = "Compile";
-				CompileAction.WorkingDirectory = UnrealBuildTool.EngineSourceDirectory;
-				CompileAction.CommandPath = new FileReference(GetISPCHostCompilerPath(BuildHostPlatform.Current.Platform));
-				CompileAction.StatusDescription = Path.GetFileName(ISPCFile.AbsolutePath);
-
-				// Disable remote execution to workaround mismatched case on XGE
-				CompileAction.bCanExecuteRemotely = false;
-
-				List<string> Arguments = new List<string>();
-
-				// Add the ISPC obj file as a prerequisite of the action.
-				CompileAction.CommandArguments = String.Format("\"{0}\" ", ISPCFile.AbsolutePath);
-
-				// Add the ISPC h file to the produced item list.
-				FileItem ISPCIncludeHeaderFile = FileItem.GetItemByFileReference(
-					FileReference.Combine(
-						OutputDir,
-						Path.GetFileName(ISPCFile.AbsolutePath) + ".generated.dummy.h"
-						)
-					);
-
-				// Add the ISPC file to be compiled.
-				Arguments.Add(String.Format("-h \"{0}\"", ISPCIncludeHeaderFile));
-
-				// Build target string. No comma on last
-				string TargetString = "";
-				foreach (string Target in CompileTargets)
-				{
-					if (Target == CompileTargets[CompileTargets.Count-1]) // .Last()
-					{
-						TargetString += Target;
-					}
-					else
-					{
-						TargetString += Target + ",";
-					}
-				}
-
-				// Build target triplet
-				Arguments.Add(String.Format("--target-os={0}", GetISPCOSTarget(CompileEnvironment.Platform)));
-				Arguments.Add(String.Format("--arch={0}", GetISPCArchTarget(CompileEnvironment.Platform, null)));
-				Arguments.Add(String.Format("--target={0}", TargetString));
-				Arguments.Add(String.Format("--emit-{0}", GetISPCObjectFileFormat(CompileEnvironment.Platform)));
-				
-				string? CpuTarget = GetISPCCpuTarget(CompileEnvironment.Platform);
-				if (!String.IsNullOrEmpty(CpuTarget))
-                {
-					Arguments.Add(String.Format("--cpu={0}", CpuTarget));
-				}
-
-				// PIC is needed for modular builds except on Windows
-				if ((CompileEnvironment.bIsBuildingDLL ||
-					CompileEnvironment.bIsBuildingLibrary) &&
-					!UEBuildPlatform.IsPlatformInGroup(CompileEnvironment.Platform, UnrealPlatformGroup.Windows))
-				{
-					Arguments.Add("--pic");
-				}
-
-				// Include paths. Don't use AddIncludePath() here, since it uses the full path and exceeds the max command line length.
-				// Because ISPC response files don't support white space in arguments, paths with white space need to be passed to the command line directly.
-				foreach (DirectoryReference IncludePath in CompileEnvironment.UserIncludePaths)
-				{
-					Arguments.Add(String.Format("-I\"{0}\"", IncludePath));
-				}
-
-				// System include paths.
-				foreach (DirectoryReference SystemIncludePath in CompileEnvironment.SystemIncludePaths)
-				{
-					Arguments.Add(String.Format("-I\"{0}\"", SystemIncludePath));
-				}
-
-				// Preprocessor definitions.
-				foreach (string Definition in CompileEnvironment.Definitions)
-				{
-					// TODO: Causes ISPC compiler to generate a spurious warning about the universal character set
-					if (!Definition.Contains("\\\\U") && !Definition.Contains("\\\\u"))
-					{
-						Arguments.Add($"-D{EscapeDefinitionForISPC(Definition)}");
-					}
-				}
-
-				// Generate the included header dependency list
-				if (CompileEnvironment.bGenerateDependenciesFile)
-				{
-					FileItem DependencyListFile = FileItem.GetItemByFileReference(FileReference.Combine(OutputDir, Path.GetFileName(ISPCFile.AbsolutePath) + ".txt"));
-					Arguments.Add(String.Format("-MMM \"{0}\"", DependencyListFile.AbsolutePath.Replace('\\', '/')));
-					CompileAction.DependencyListFile = DependencyListFile;
-					CompileAction.ProducedItems.Add(DependencyListFile);
-				}
-
-				CompileAction.ProducedItems.Add(ISPCIncludeHeaderFile);
-
-				FileReference ResponseFileName = new FileReference(ISPCIncludeHeaderFile.AbsolutePath + ".response");
-				FileItem ResponseFileItem = Graph.CreateIntermediateTextFile(ResponseFileName, Arguments.Select(x => Utils.ExpandVariables(x)));
-				CompileAction.CommandArguments += String.Format("@\"{0}\"", ResponseFileName);
-				CompileAction.PrerequisiteItems.Add(ResponseFileItem);
-
-				// Add the source file and its included files to the prerequisite item list.
-				CompileAction.PrerequisiteItems.Add(ISPCFile);
-
-				FileItem ISPCFinalHeaderFile = FileItem.GetItemByFileReference(
-					FileReference.Combine(
-						OutputDir,
-						Path.GetFileName(ISPCFile.AbsolutePath) + ".generated.h"
-						)
-					);
-
-				// Fix interrupted build issue by copying header after generation completes
-				FileReference SourceFile = ISPCIncludeHeaderFile.Location;
-				FileReference TargetFile = ISPCFinalHeaderFile.Location;
-
-				FileItem SourceFileItem = FileItem.GetItemByFileReference(SourceFile);
-				FileItem TargetFileItem = FileItem.GetItemByFileReference(TargetFile);
-
-				Action CopyAction = Graph.CreateAction(ActionType.BuildProject);
-				CopyAction.CommandDescription = "Copy";
-				CopyAction.CommandPath = BuildHostPlatform.Current.Shell;
-				if (BuildHostPlatform.Current.ShellType == ShellType.Cmd)
-				{
-					CopyAction.CommandArguments = String.Format("/C \"copy /Y \"{0}\" \"{1}\" 1>nul\"", SourceFile, TargetFile);
-				}
-				else
-				{
-					CopyAction.CommandArguments = String.Format("-c 'cp -f \"\"{0}\"\" \"\"{1}\"'", SourceFile.FullName, TargetFile.FullName);
-				}
-				CopyAction.WorkingDirectory = UnrealBuildTool.EngineSourceDirectory;
-				CopyAction.PrerequisiteItems.Add(SourceFileItem);
-				CopyAction.ProducedItems.Add(TargetFileItem);
-				CopyAction.StatusDescription = TargetFileItem.Location.GetFileName();
-				CopyAction.bCanExecuteRemotely = false;
-				CopyAction.bShouldOutputStatusDescription = false;
-
-				Result.GeneratedHeaderFiles.Add(TargetFileItem);
-
-				Log.TraceVerbose("   ISPC Generating Header " + CompileAction.StatusDescription + ": \"" + CompileAction.CommandPath + "\"" + CompileAction.CommandArguments);
-			}
-
-			return Result;
-		}
-
-		public override CPPOutput CompileISPCFiles(CppCompileEnvironment CompileEnvironment, List<FileItem> InputFiles, DirectoryReference OutputDir, IActionGraphBuilder Graph)
-		{
-			CPPOutput Result = new CPPOutput();
-
-			if (!CompileEnvironment.bCompileISPC)
-			{
-				return Result;
-			}
-
-			List<string> CompileTargets = GetISPCCompileTargets(CompileEnvironment.Platform, null);
-
-			foreach (FileItem ISPCFile in InputFiles)
-			{
-				Action CompileAction = Graph.CreateAction(ActionType.Compile);
-				CompileAction.CommandDescription = "Compile";
-				CompileAction.WorkingDirectory = UnrealBuildTool.EngineSourceDirectory;
-				CompileAction.CommandPath = new FileReference(GetISPCHostCompilerPath(BuildHostPlatform.Current.Platform));
-				CompileAction.StatusDescription = Path.GetFileName(ISPCFile.AbsolutePath);
-
-				// Disable remote execution to workaround mismatched case on XGE
-				CompileAction.bCanExecuteRemotely = false;
-
-				List<string> Arguments = new List<string>();
-
-				// Add the ISPC file to be compiled.
-				Arguments.Add(String.Format(" \"{0}\"", ISPCFile.AbsolutePath));
-
-				List<FileItem> CompiledISPCObjFiles = new List<FileItem>();
-				string TargetString = "";
-
-				foreach (string Target in CompileTargets)
-				{
-					string ObjTarget = Target;
-
-					if (Target.Contains("-"))
-					{
-						// Remove lane width and gang size from obj file name
-						ObjTarget = Target.Split('-')[0];
-					}
-
-					FileItem CompiledISPCObjFile;
-
-					if (CompileTargets.Count > 1)
-					{
-						CompiledISPCObjFile = FileItem.GetItemByFileReference(
-						FileReference.Combine(
-							OutputDir,
-							Path.GetFileName(ISPCFile.AbsolutePath) + "_" + ObjTarget + GetISPCObjectFileSuffix(CompileEnvironment.Platform)
-							)
-						);
-					}
-					else
-					{
-						CompiledISPCObjFile = FileItem.GetItemByFileReference(
-						FileReference.Combine(
-							OutputDir,
-							Path.GetFileName(ISPCFile.AbsolutePath) + GetISPCObjectFileSuffix(CompileEnvironment.Platform)
-							)
-						);
-					}
-
-					// Add the ISA specific ISPC obj files to the produced item list.
-					CompiledISPCObjFiles.Add(CompiledISPCObjFile);
-
-					// Build target string. No comma on last
-					if (Target == CompileTargets[CompileTargets.Count-1]) // .Last()
-					{
-						TargetString += Target;
-					}
-					else
-					{
-						TargetString += Target + ",";
-					}
-				}
-
-				// Add the common ISPC obj file to the produced item list.
-				FileItem CompiledISPCObjFileNoISA = FileItem.GetItemByFileReference(
-					FileReference.Combine(
-						OutputDir,
-						Path.GetFileName(ISPCFile.AbsolutePath) + GetISPCObjectFileSuffix(CompileEnvironment.Platform)
-						)
-					);
-
-				CompiledISPCObjFiles.Add(CompiledISPCObjFileNoISA);
-
-				// Add the output ISPC obj file
-				Arguments.Add(String.Format("-o \"{0}\"", CompiledISPCObjFileNoISA));
-
-				// Build target triplet
-				Arguments.Add(String.Format("--target-os=\"{0}\"", GetISPCOSTarget(CompileEnvironment.Platform)));
-				Arguments.Add(String.Format("--arch=\"{0}\"", GetISPCArchTarget(CompileEnvironment.Platform, null)));
-				Arguments.Add(String.Format("--target=\"{0}\"", TargetString));
-				Arguments.Add(String.Format("--emit-{0}", GetISPCObjectFileFormat(CompileEnvironment.Platform)));
-
-				if (CompileEnvironment.Configuration == CppConfiguration.Debug)
-				{
-					if (CompileEnvironment.Platform == UnrealTargetPlatform.Mac)
-					{
-						// Turn off debug symbols on Mac due to dsym generation issue
-						Arguments.Add("-O0");
-						// Ideally we would be able to turn on symbols and specify the dwarf version, but that does
-						// does not seem to be working currently, ie:
-						//    Arguments.Add("-g -O0 --dwarf-version=2");
-
-					}
-					else
-					{
-						Arguments.Add("-g -O0");
-					}
-				}
-				else
-				{
-					Arguments.Add("-O2");
-				}
-
-				// PIC is needed for modular builds except on Windows
-				if ((CompileEnvironment.bIsBuildingDLL || 
-					CompileEnvironment.bIsBuildingLibrary) &&
-					!UEBuildPlatform.IsPlatformInGroup(CompileEnvironment.Platform, UnrealPlatformGroup.Windows))
-				{
-					Arguments.Add("--pic");
-				}
-
-				// Include paths. Don't use AddIncludePath() here, since it uses the full path and exceeds the max command line length.
-				foreach (DirectoryReference IncludePath in CompileEnvironment.UserIncludePaths)
-				{
-					Arguments.Add(String.Format("-I\"{0}\"", IncludePath));
-				}
-
-				// System include paths.
-				foreach (DirectoryReference SystemIncludePath in CompileEnvironment.SystemIncludePaths)
-				{
-					Arguments.Add(String.Format("-I\"{0}\"", SystemIncludePath));
-				}
-
-				// Preprocessor definitions.
-				foreach (string Definition in CompileEnvironment.Definitions)
-				{
-					// TODO: Causes ISPC compiler to generate a spurious warning about the universal character set
-					if (!Definition.Contains("\\\\U") && !Definition.Contains("\\\\u"))
-					{
-						Arguments.Add($"-D{EscapeDefinitionForISPC(Definition)}");
-					}
-				}
-
-				// Consume the included header dependency list
-				if (CompileEnvironment.bGenerateDependenciesFile)
-				{
-					FileItem DependencyListFile = FileItem.GetItemByFileReference(FileReference.Combine(OutputDir, Path.GetFileName(ISPCFile.AbsolutePath) + ".txt"));
-					CompileAction.DependencyListFile = DependencyListFile;
-					CompileAction.PrerequisiteItems.Add(DependencyListFile);
-				}
-
-				CompileAction.ProducedItems.AddRange(CompiledISPCObjFiles);
-				Result.ObjectFiles.AddRange(CompiledISPCObjFiles);
-
-				FileReference ResponseFileName = new FileReference(CompiledISPCObjFileNoISA.AbsolutePath + ".response");
-				FileItem ResponseFileItem = Graph.CreateIntermediateTextFile(ResponseFileName, Arguments.Select(x => Utils.ExpandVariables(x)));
-				CompileAction.CommandArguments = " @\"" + ResponseFileName + "\"";
-				CompileAction.PrerequisiteItems.Add(ResponseFileItem);
-
-				// Add the source file and its included files to the prerequisite item list.
-				CompileAction.PrerequisiteItems.Add(ISPCFile);
-
-				Log.TraceVerbose("   ISPC Compiling " + CompileAction.StatusDescription + ": \"" + CompileAction.CommandPath + "\"" + CompileAction.CommandArguments);
-			}
-
-			return Result;
-		}
-	}
 }
