UnrealBuildTool: Set PreservePSYM in LinuxToolChainOptions rather than passing into LinuxToolChain constructor

#rnx

diff --git a/Engine/Source/Programs/UnrealBuildTool/Platform/Linux/LinuxToolChain.cs b/Engine/Source/Programs/UnrealBuildTool/Platform/Linux/LinuxToolChain.cs
--- a/Engine/Source/Programs/UnrealBuildTool/Platform/Linux/LinuxToolChain.cs
+++ b/Engine/Source/Programs/UnrealBuildTool/Platform/Linux/LinuxToolChain.cs
@@ -63,6 +63,11 @@ namespace UnrealBuildTool
 		/// Enable tuning of debug info for LLDB
 		/// </summary>
 		TuneDebugInfoForLLDB = 0x80,
+
+		/// <summary>
+		/// Whether or not to preserve the portable symbol file produced by dump_syms
+		/// </summary>
+		PreservePSYM = 0x100,
 	}
 
 	class LinuxToolChain : ISPCToolChain
@@ -73,9 +78,6 @@ namespace UnrealBuildTool
 		/** Whether the compiler is set up to produce PIE executables by default */
 		bool bSuppressPIE = false;
 
-		/** Whether or not to preserve the portable symbol file produced by dump_syms */
-		bool bPreservePSYM = false;
-
 		/** Pass --gdb-index option to linker to generate .gdb_index section. */
 		protected bool bGdbIndexSection = true;
 
@@ -96,8 +98,8 @@ namespace UnrealBuildTool
 		/// </summary>
 		LinuxToolChainOptions Options;
 
-		public LinuxToolChain(string InArchitecture, LinuxPlatformSDK InSDK, bool InPreservePSYM = false, LinuxToolChainOptions InOptions = LinuxToolChainOptions.None)
-			: this(UnrealTargetPlatform.Linux, InArchitecture, InSDK, InPreservePSYM, InOptions)
+		public LinuxToolChain(string InArchitecture, LinuxPlatformSDK InSDK, LinuxToolChainOptions InOptions = LinuxToolChainOptions.None)
+			: this(UnrealTargetPlatform.Linux, InArchitecture, InSDK, InOptions)
 		{
 			MultiArchRoot = PlatformSDK.GetSDKLocation();
 			BaseLinuxPath = PlatformSDK.GetBaseLinuxPathForArchitecture(InArchitecture);
@@ -188,13 +190,12 @@ namespace UnrealBuildTool
 			}
 		}
 
-		public LinuxToolChain(UnrealTargetPlatform InPlatform, string InArchitecture, LinuxPlatformSDK InSDK, bool InPreservePSYM = false, LinuxToolChainOptions InOptions = LinuxToolChainOptions.None)
+		public LinuxToolChain(UnrealTargetPlatform InPlatform, string InArchitecture, LinuxPlatformSDK InSDK, LinuxToolChainOptions InOptions = LinuxToolChainOptions.None)
 			: base()
 		{
 			Architecture = InArchitecture;
 			PlatformSDK = InSDK;
 			Options = InOptions;
-			bPreservePSYM = InPreservePSYM;
 		}
 
 		public override void SetUpGlobalEnvironment(ReadOnlyTargetRules Target)
@@ -249,7 +250,7 @@ namespace UnrealBuildTool
 			FileItem StrippedFile = FileItem.GetItemByPath(Path.Combine(LinkEnvironment.LocalShadowDirectory.FullName, OutputFile.Location.GetFileName() + "_nodebug"));
 			FileItem DebugFile = FileItem.GetItemByPath(Path.Combine(LinkEnvironment.OutputDirectory.FullName, OutputFile.Location.GetFileNameWithoutExtension() + ".debug"));
 
-			if (bPreservePSYM)
+			if (Options.HasFlag(LinuxToolChainOptions.PreservePSYM))
 			{
 				SymbolsFile = FileItem.GetItemByPath(Path.Combine(LinkEnvironment.OutputDirectory.FullName, OutputFile.Location.GetFileNameWithoutExtension() + ".psym"));
 			}
diff --git a/Engine/Source/Programs/UnrealBuildTool/Platform/Linux/UEBuildLinux.cs b/Engine/Source/Programs/UnrealBuildTool/Platform/Linux/UEBuildLinux.cs
--- a/Engine/Source/Programs/UnrealBuildTool/Platform/Linux/UEBuildLinux.cs
+++ b/Engine/Source/Programs/UnrealBuildTool/Platform/Linux/UEBuildLinux.cs
@@ -605,13 +605,18 @@ namespace UnrealBuildTool
 				Options |= LinuxToolChainOptions.TuneDebugInfoForLLDB;
 			}
 
+			if (Target.LinuxPlatform.bPreservePSYM)
+			{
+				Options |= LinuxToolChainOptions.PreservePSYM;
+			}
+
 			// Disable color logging if we are on a build machine
 			if (Environment.GetEnvironmentVariable("IsBuildMachine") == "1")
 			{
 				Log.ColorConsoleOutput = false;
 			}
 
-			return new LinuxToolChain(Target.Architecture, SDK, Target.LinuxPlatform.bPreservePSYM, Options);
+			return new LinuxToolChain(Target.Architecture, SDK, Options);
 		}
 
 		/// <summary>
