UnrealBuildTool: Remove /usr/bin/arch trimming logic from FASTBuild, it is no longer needed

#rnx

diff --git a/Engine/Source/Programs/UnrealBuildTool/Executors/Experimental/FASTBuild.cs b/Engine/Source/Programs/UnrealBuildTool/Executors/Experimental/FASTBuild.cs
--- a/Engine/Source/Programs/UnrealBuildTool/Executors/Experimental/FASTBuild.cs
+++ b/Engine/Source/Programs/UnrealBuildTool/Executors/Experimental/FASTBuild.cs
@@ -522,20 +522,6 @@ namespace UnrealBuildTool
 			// Make sure we substituted the known environment variables with corresponding BFF friendly imported vars
 			CompilerCommandLine = SubstituteEnvironmentVariables(CompilerCommandLine);
 
-			// Some tools are now executed via arch so they can run natively even if UBT is under mono/rosetta. If so we
-			// need to remove the architecture from the argument list
-			if (LocalToolName.ToLower() == "arch")
-			{
-				// Action would be /usr/bin/arch -<arch> -other -flags so remove the first
-				// argument and any trailing spaces
-				Match M = Regex.Match(CompilerCommandLine, @"^\s*-.+?\s+");
-
-				if (M.Success)
-				{
-					CompilerCommandLine = CompilerCommandLine.Substring(M.Length);
-				}
-			}
-
 			// Some tricky defines /DTROUBLE=\"\\\" abc  123\\\"\" aren't handled properly by either Unreal or FASTBuild, but we do our best.
 			char[] SpaceChar = { ' ' };
 			string[] RawTokens = CompilerCommandLine.Trim().Split(' ');
