Don't require SDK to build Android module

#android

diff --git a/Engine/Source/Programs/UnrealBuildTool/Platform/Android/UEBuildAndroid.cs b/Engine/Source/Programs/UnrealBuildTool/Platform/Android/UEBuildAndroid.cs
--- a/Engine/Source/Programs/UnrealBuildTool/Platform/Android/UEBuildAndroid.cs
+++ b/Engine/Source/Programs/UnrealBuildTool/Platform/Android/UEBuildAndroid.cs
@@ -264,8 +264,9 @@ namespace UnrealBuildTool
 		/// <param name="Target">The target being build</param>
 		public override void ModifyModuleRulesForOtherPlatform(string ModuleName, ModuleRules Rules, ReadOnlyTargetRules Target)
 		{
-			// don't do any target platform stuff if SDK is not available
-			if (!UEBuildPlatform.IsPlatformAvailableForTarget(Platform, Target))
+			// don't do any target platform stuff if not available for host and opted in
+			// do not require SDK to build it since we don't necessarily need it for editor building
+			if (!(TryGetBuildPlatform(Platform, out _) && Target.IsPlatformOptedIn(Platform)))
 			{
 				return;
 			}
