UnrealBuildTool: treat Marketplace plugins as read-only for the purpose of cleaning.

Changing the BaseDir for the marketplace plugin to Engine/Plugins/Marketplace prevents all of Engine from being cleaned when a Marketplace plugin is present.
Setting bReadOnly for the the Marketplace rules assembly prevents Marketplace plugins from being cleaned.

#rnx

diff --git a/Engine/Source/Programs/UnrealBuildTool/System/RulesCompiler.cs b/Engine/Source/Programs/UnrealBuildTool/System/RulesCompiler.cs
--- a/Engine/Source/Programs/UnrealBuildTool/System/RulesCompiler.cs
+++ b/Engine/Source/Programs/UnrealBuildTool/System/RulesCompiler.cs
@@ -205,7 +205,7 @@ namespace UnrealBuildTool
 			if (ModuleFileToContext.Count > 0)
 			{
 				FileReference AssemblyFileName = FileReference.Combine(UnrealBuildTool.WritableEngineDirectory, "Intermediate", "Build", "BuildRules", "MarketplaceRules.dll");
-				Result = new RulesAssembly(MarketplaceScope, new List<DirectoryReference> { Unreal.EngineDirectory }, Plugins, ModuleFileToContext, new List<FileReference>(), AssemblyFileName, bContainsEngineModules: true, DefaultBuildSettings: BuildSettingsVersion.Latest, bReadOnly: false, bSkipCompile: bSkipCompile, bForceCompile: bForceCompile, Parent: Parent);
+				Result = new RulesAssembly(MarketplaceScope, new List<DirectoryReference> { DirectoryReference.Combine(Unreal.EngineDirectory, "Plugins", "Marketplace") }, Plugins, ModuleFileToContext, new List<FileReference>(), AssemblyFileName, bContainsEngineModules: true, DefaultBuildSettings: BuildSettingsVersion.Latest, bReadOnly: true, bSkipCompile: bSkipCompile, bForceCompile: bForceCompile, Parent: Parent);
 			}
 			return Result;
 		}
