Adding UBT plugin support currently only used by UHT

diff --git a/Engine/Source/Programs/Shared/EpicGames.Build/System/EnumeratePlugins.cs b/Engine/Source/Programs/Shared/EpicGames.Build/System/EnumeratePlugins.cs
--- a/Engine/Source/Programs/Shared/EpicGames.Build/System/EnumeratePlugins.cs
+++ b/Engine/Source/Programs/Shared/EpicGames.Build/System/EnumeratePlugins.cs
@@ -99,5 +99,60 @@ namespace UnrealBuildBase
 				}
 			}
 		}
+
+		/// <summary>
+		/// Enumerates all the UBT plugin files available to the given project
+		/// </summary>
+		/// <param name="ProjectFile">Path to the project file</param>
+		/// <param name="Plugins">Collection of found plugins</param>
+		/// <returns>True if the plugins compiled</returns>
+		public static bool EnumerateUbtPlugins(FileReference? ProjectFile, out (FileReference ProjectFile, FileReference TargetAssembly)[]? Plugins)
+		{
+			bool bBuildSuccess = true;
+			Plugins = null;
+
+			List<DirectoryReference> PluginDirectories = new List<DirectoryReference>();
+			PluginDirectories.AddRange(Unreal.GetExtensionDirs(Unreal.EngineDirectory, "Source"));
+			PluginDirectories.AddRange(Unreal.GetExtensionDirs(Unreal.EngineDirectory, "Plugins"));
+			if (ProjectFile != null)
+			{
+				PluginDirectories.AddRange(Unreal.GetExtensionDirs(ProjectFile.Directory, "Source"));
+				PluginDirectories.AddRange(Unreal.GetExtensionDirs(ProjectFile.Directory, "Plugins"));
+				PluginDirectories.AddRange(Unreal.GetExtensionDirs(ProjectFile.Directory, "Mods"));
+			}
+
+			// Scan for UBT plugin projects
+			HashSet<FileReference> ScannedPlugins = new HashSet<FileReference>(UnrealBuildBase.Rules.FindAllRulesFiles(
+				PluginDirectories, UnrealBuildBase.Rules.RulesFileType.UbtPlugin));
+
+			// Build the located plugins
+			Dictionary<FileReference, (CsProjBuildRecord BuildRecord, FileReference BuildRecordFile)>? BuiltPlugins = null;
+			if (ScannedPlugins.Count > 0)
+			{
+				List<DirectoryReference> BaseDirectories = new List<DirectoryReference>(2);
+				BaseDirectories.Add(Unreal.EngineDirectory);
+				if (ProjectFile != null)
+				{
+					BaseDirectories.Add(ProjectFile.Directory);
+				}
+
+				BuiltPlugins = CompileScriptModule.Build(UnrealBuildBase.Rules.RulesFileType.UbtPlugin, ScannedPlugins, BaseDirectories, false, false, true, out bBuildSuccess,
+					Count =>
+					{
+						if (Log.OutputFile != null)
+						{
+							Log.TraceInformation($"Building {Count} plugins (see Log '{Log.OutputFile}' for more details)");
+						}
+						else
+						{
+							Log.TraceInformation($"Building {Count} plugins");
+						}
+					}
+				);
+
+				Plugins = BuiltPlugins.Select(P => (P.Key, CompileScriptModule.GetTargetPath(P))).OrderBy(X => X.Key.FullName).ToArray();
+			}
+			return bBuildSuccess;
+		}
 	}
 }
diff --git a/Engine/Source/Programs/UnrealBuildTool/Configuration/UEBuildTarget.cs b/Engine/Source/Programs/UnrealBuildTool/Configuration/UEBuildTarget.cs
--- a/Engine/Source/Programs/UnrealBuildTool/Configuration/UEBuildTarget.cs
+++ b/Engine/Source/Programs/UnrealBuildTool/Configuration/UEBuildTarget.cs
@@ -16,6 +16,7 @@ using EpicGames.Core;
 using System.Reflection;
 using OpenTracing.Util;
 using UnrealBuildBase;
+using EpicGames.UHT.Utils;
 
 namespace UnrealBuildTool
 {
@@ -1059,6 +1060,11 @@ namespace UnrealBuildTool
 		/// </summary>
 		public bool bHasProjectScriptPlugin;
 
+		/// <summary>
+		/// If true, then a script plugin was found that didn't have a matching UBT plugin
+		/// </summary>
+		public bool bHasRequiredProjectScriptPlugin;
+
 		/// <summary>
 		/// All plugins which are built for this target
 		/// </summary>
@@ -1074,6 +1080,16 @@ namespace UnrealBuildTool
 		/// </summary>
 		public FileReference? ForeignPlugin;
 
+		/// <summary>
+		/// Collection of all UBT plugin project files
+		/// </summary>
+		public List<FileReference>? UbtPlugins;
+
+		/// <summary>
+		/// Collection of all located UHT plugin assemblies
+		/// </summary>
+		public List<FileReference>? UhtPlugins;
+
 		/// <summary>
 		/// All application binaries; may include binaries not built by this target.
 		/// </summary>
@@ -1793,7 +1809,8 @@ namespace UnrealBuildTool
 
 			// Create the makefile
 			string ExternalMetadata = UEBuildPlatform.GetBuildPlatform(Platform).GetExternalBuildMetadata(ProjectFile);
-			TargetMakefile Makefile = new TargetMakefile(ExternalMetadata, Binaries[0].OutputFilePaths[0], ReceiptFileName, ProjectIntermediateDirectory, TargetType, Rules.ConfigValueTracker, bDeployAfterCompile, bHasProjectScriptPlugin);
+			TargetMakefile Makefile = new TargetMakefile(ExternalMetadata, Binaries[0].OutputFilePaths[0], ReceiptFileName, ProjectIntermediateDirectory, TargetType, 
+				Rules.ConfigValueTracker, bDeployAfterCompile, bHasProjectScriptPlugin, bHasRequiredProjectScriptPlugin, UbtPlugins?.ToArray(), UhtPlugins?.ToArray());
 
 			// Get diagnostic info to be printed before each build
 			TargetToolChain.GetVersionInfo(Makefile.Diagnostics);
@@ -3414,6 +3431,9 @@ namespace UnrealBuildTool
 
 			// Set the list of plugins that should be built
 			BuildPlugins = new List<UEBuildPlugin>(NameToInstance.Values);
+
+			// Setup the UHT plugins
+			SetupUbtPlugins();
 		}
 
 		/// <summary>
@@ -3506,12 +3526,6 @@ namespace UnrealBuildTool
 					return null;
 				}
 
-				// If this plugin supports UHT, flag that we need to pass the project file to UHT when generating headers
-				if (Info.Descriptor.SupportedPrograms != null && Info.Descriptor.SupportedPrograms.Contains("UnrealHeaderTool"))
-				{
-					bHasProjectScriptPlugin = true;
-				}
-
 				// Disable any plugins that aren't compatible with this program
 				if (Rules.Type == TargetType.Program && (Info.Descriptor.SupportedPrograms == null || !Info.Descriptor.SupportedPrograms.Contains(AppName)))
 				{
@@ -3599,6 +3613,69 @@ namespace UnrealBuildTool
 			}
 		}
 
+		/// <summary>
+		/// Scan for UBT plugins.  This also detects if the current target has a C++ style UHT that don't have
+		/// an associated UBT/UHT plugin.  When this happens, the old UHT must be used.
+		/// </summary>
+		private void SetupUbtPlugins()
+		{
+			if (BuildPlugins == null)
+			{
+				throw new BuildException("Build plugins has not be populated");
+			}
+
+			// Collect the plugins
+			(FileReference ProjectFile, FileReference TargetAssembly)[]? BuiltPlugins;
+			if (!PluginsBase.EnumerateUbtPlugins(ProjectFile, out BuiltPlugins))
+			{
+				throw new BuildException("One or more UBT plugins failed to compile.");
+			}
+
+			// Save the collection
+			if (BuiltPlugins != null)
+			{
+				UbtPlugins = BuiltPlugins.Select(P => P.ProjectFile).ToList();
+			}
+
+			// For all the plugins 
+			foreach (UEBuildPlugin Plugin in BuildPlugins)
+			{
+				PluginInfo Info = Plugin.Info;
+
+				// Check to see if this plugin has an associated UBT plugin
+				bool bUbtPluginFound = false;
+				if (BuiltPlugins != null)
+				{
+					foreach ((FileReference ProjectFile, FileReference TargetAssembly) BuiltPlugin in BuiltPlugins)
+					{
+						if (BuiltPlugin.ProjectFile.IsUnderDirectory(Info.Directory))
+						{
+							if (UhtTables.IsUhtPlugin(BuiltPlugin.TargetAssembly.FullName))
+							{
+								if (UhtPlugins == null)
+								{
+									UhtPlugins = new List<FileReference>();
+								}
+								UhtPlugins.Add(BuiltPlugin.TargetAssembly);
+								bUbtPluginFound = true;
+								break;
+							}
+						}
+					}
+				}
+
+				// If this plugin supports UHT, flag that we need to pass the project file to UHT when generating headers
+				if (Info.Descriptor.SupportedPrograms != null && Info.Descriptor.SupportedPrograms.Contains("UnrealHeaderTool"))
+				{
+					bHasProjectScriptPlugin = true;
+					if (!bUbtPluginFound)
+					{
+						bHasRequiredProjectScriptPlugin = true;
+					}
+				}
+			}
+		}
+
 		/// <summary>
 		/// Sets up the binaries for the target.
 		/// </summary>
diff --git a/Engine/Source/Programs/UnrealBuildTool/ProjectFiles/ProjectFileGenerator.cs b/Engine/Source/Programs/UnrealBuildTool/ProjectFiles/ProjectFileGenerator.cs
--- a/Engine/Source/Programs/UnrealBuildTool/ProjectFiles/ProjectFileGenerator.cs
+++ b/Engine/Source/Programs/UnrealBuildTool/ProjectFiles/ProjectFileGenerator.cs
@@ -430,9 +430,10 @@ namespace UnrealBuildTool
 		}
 
 		/// <summary>
-		/// Adds all .automation.csproj files to the solution.
+		/// Adds all rules project files to the solution.
 		/// </summary>
-		void AddAutomationModules(List<FileReference> UnrealProjectFiles, PrimaryProjectFolder RootFolder, PrimaryProjectFolder ProgramsFolder)
+		void AddRulesModules(Rules.RulesFileType RulesFileType, string ProgramSubDirectory, List<ProjectFile> AddedProjectFiles, 
+			List<FileReference> UnrealProjectFiles, PrimaryProjectFolder RootFolder, PrimaryProjectFolder ProgramsFolder)
 		{
 			List<DirectoryReference> BuildFolders = new List<DirectoryReference>();
 			foreach (FileReference UnrealProjectFile in UnrealProjectFiles)
@@ -444,11 +445,11 @@ namespace UnrealBuildTool
 				}
 			}
 
-			PrimaryProjectFolder AutomationFolder = ProgramsFolder.AddSubFolder("Automation");
+			PrimaryProjectFolder Folder = ProgramsFolder.AddSubFolder(ProgramSubDirectory);
 			DirectoryReference SamplesDirectory = DirectoryReference.Combine(Unreal.RootDirectory, "Samples");
 
-			// Find all the automation modules .csproj files to add
-			List<FileReference> ModuleFiles = Rules.FindAllRulesSourceFiles(Rules.RulesFileType.AutomationModule, null, ForeignPlugins: null, AdditionalSearchPaths: BuildFolders);
+			// Find all the modules .csproj files to add
+			List<FileReference> ModuleFiles = Rules.FindAllRulesSourceFiles(RulesFileType, null, ForeignPlugins: null, AdditionalSearchPaths: BuildFolders);
 			foreach (FileReference ProjectFile in ModuleFiles)
 			{
 				if (FileReference.Exists(ProjectFile))
@@ -457,12 +458,12 @@ namespace UnrealBuildTool
 
 					Project.ShouldBuildForAllSolutionTargets = false;//true;
 					AddExistingProjectFile(Project, bForceDevelopmentConfiguration: true);
-					AutomationProjectFiles.Add(Project);
+					AddedProjectFiles.Add(Project);
 
 					if (!ProjectFile.IsUnderDirectory(Unreal.EngineDirectory))
 					{
 						FileReference PropsFile = new FileReference(ProjectFile.FullName + ".props");
-						CreateAutomationProjectPropsFile(PropsFile);
+						CreateProjectPropsFile(PropsFile);
 
 						if (ProjectFile.IsUnderDirectory(SamplesDirectory))
 						{
@@ -475,7 +476,7 @@ namespace UnrealBuildTool
 					}
 					else
 					{
-						AutomationFolder.ChildProjects.Add(Project);
+						Folder.ChildProjects.Add(Project);
 					}
 				}
 			}
@@ -512,10 +513,10 @@ namespace UnrealBuildTool
 		}
 
 		/// <summary>
-		/// Creates a .props file next to each automation project which specifies the path to the engine directory
+		/// Creates a .props file next to each project which specifies the path to the engine directory
 		/// </summary>
 		/// <param name="PropsFile">The properties file path</param>
-		void CreateAutomationProjectPropsFile(FileReference PropsFile)
+		void CreateProjectPropsFile(FileReference PropsFile)
 		{
 			using (FileStream Stream = FileReference.Open(PropsFile, FileMode.Create, FileAccess.Write, FileShare.Read))
 			{
@@ -1078,7 +1079,10 @@ namespace UnrealBuildTool
 						ProgramsFolder.ChildProjects.Add(AutomationToolProject);
 
 					// Add automation.csproj files to the primary project
-					AddAutomationModules(AllGameProjects, RootFolder, ProgramsFolder);
+					AddRulesModules(Rules.RulesFileType.AutomationModule, "Automation", AutomationProjectFiles, AllGameProjects, RootFolder, ProgramsFolder);
+
+					// Add ubtplugin.csproj files to the primary project
+					AddRulesModules(Rules.RulesFileType.UbtPlugin, "UnrealBuildTool.Plugins", UbtPluginProjectFiles, AllGameProjects, RootFolder, ProgramsFolder);
 
 					// Add shared projects
 					AddSharedDotNetModules(AllEngineDirectories, ProgramsFolder);
@@ -3048,6 +3052,8 @@ namespace UnrealBuildTool
 
 		protected readonly List<ProjectFile> AutomationProjectFiles = new List<ProjectFile>();
 
+		protected readonly List<ProjectFile> UbtPluginProjectFiles = new List<ProjectFile>();
+
 		/// List of top-level folders in the primary project file
 		protected PrimaryProjectFolder RootFolder;
 	}
diff --git a/Engine/Source/Programs/UnrealBuildTool/System/ExternalExecution.cs b/Engine/Source/Programs/UnrealBuildTool/System/ExternalExecution.cs
--- a/Engine/Source/Programs/UnrealBuildTool/System/ExternalExecution.cs
+++ b/Engine/Source/Programs/UnrealBuildTool/System/ExternalExecution.cs
@@ -19,6 +19,7 @@ using System.Threading.Tasks;
 using OpenTracing;
 using OpenTracing.Util;
 using UnrealBuildTool.Modes;
+using EpicGames.UHT.Utils;
 
 namespace UnrealBuildTool
 {
@@ -147,6 +148,11 @@ namespace UnrealBuildTool
 		/// </summary>
 		public DirectoryItem GeneratedCodeDirectory;
 
+		/// <summary>
+		/// Collection of the module's public defines
+		/// </summary>
+		public List<string> PublicDefines;
+
 		/// <summary>
 		/// Base (i.e. extensionless) path+filename of the .gen files
 		/// </summary>
@@ -173,6 +179,7 @@ namespace UnrealBuildTool
 			this.PublicUObjectHeaders = new List<FileItem>();
 			this.InternalUObjectHeaders = new List<FileItem>();
 			this.PrivateUObjectHeaders = new List<FileItem>();
+			this.PublicDefines = new List<string>();
 			this.GeneratedCodeDirectory = GeneratedCodeDirectory;
 			this.GeneratedCodeVersion = GeneratedCodeVersion;
 			this.bIsReadOnly = bIsReadOnly;
@@ -193,6 +200,7 @@ namespace UnrealBuildTool
 			GeneratedCodeDirectory = Reader.ReadDirectoryItem()!;
 			GeneratedCodeVersion = (EGeneratedCodeVersion)Reader.ReadInt();
 			bIsReadOnly = Reader.ReadBool();
+			PublicDefines = Reader.ReadList(() => Reader.ReadString())!;
 		}
 
 		public void Write(BinaryArchiveWriter Writer)
@@ -210,6 +218,7 @@ namespace UnrealBuildTool
 			Writer.WriteDirectoryItem(GeneratedCodeDirectory);
 			Writer.WriteInt((int)GeneratedCodeVersion);
 			Writer.WriteBool(bIsReadOnly);
+			Writer.WriteList(PublicDefines, Item => Writer.WriteString(Item));
 		}
 
 		public override string ToString()
@@ -473,6 +482,7 @@ namespace UnrealBuildTool
 			{
 				UEBuildModuleCPP Module = ModulesSortedByType[Idx];
 				UHTModuleInfo Info = ModuleInfoArray[Idx];
+				Info.PublicDefines.AddRange(Module.PublicDefinitions);
 
 				// Delete any obsolete files at the root of the generated code directory
 				// Remove this code after engine version 5.2 (SOL-2551)
@@ -587,9 +597,9 @@ namespace UnrealBuildTool
 		/// Gets the path to the receipt for UHT
 		/// </summary>
 		/// <returns>Path to the UHT receipt</returns>
-		public static FileReference GetHeaderToolReceiptFile(FileReference? ProjectFile, UnrealTargetConfiguration Configuration, bool bHasProjectScriptPlugin)
+		public static FileReference GetHeaderToolReceiptFile(FileReference? ProjectFile, UnrealTargetConfiguration Configuration, bool bHasProjectScriptPlugin, FileReference[]? UhtPlugins)
 		{
-			if(bHasProjectScriptPlugin && ProjectFile != null)
+			if ((bHasProjectScriptPlugin || UhtPlugins != null) && ProjectFile != null)
 			{
 				return TargetReceipt.GetDefaultPath(ProjectFile.Directory, "UnrealHeaderTool", BuildHostPlatform.Current.Platform, Configuration, "");
 			}
@@ -918,13 +928,13 @@ namespace UnrealBuildTool
 		/// </summary>
 		public static void ExecuteHeaderToolIfNecessary(BuildConfiguration BuildConfiguration, FileReference? ProjectFile, TargetMakefile Makefile, string TargetName, ISourceFileWorkingSet WorkingSet)
 		{
-			if (!BuildConfiguration.bUseBuiltInUnrealHeaderTool || Makefile.bHasProjectScriptPlugin || TargetName.Equals("UnrealHeaderTool", StringComparison.InvariantCultureIgnoreCase))
+			if (!BuildConfiguration.bUseBuiltInUnrealHeaderTool || Makefile.bHasRequiredProjectScriptPlugin || TargetName.Equals("UnrealHeaderTool", StringComparison.InvariantCultureIgnoreCase))
 			{
 				ExecuteExternalHeaderToolIfNecessary(BuildConfiguration, ProjectFile, Makefile, TargetName, WorkingSet);
 			}
 			else
 			{
-				ExecuteInternalHeaderToolIfNecessary(BuildConfiguration, ProjectFile, Makefile, TargetName, WorkingSet);
+				ExecuteInternalHeaderToolIfNecessary(BuildConfiguration, ProjectFile, Makefile, TargetName, Makefile.UhtPlugins, WorkingSet);
 			}
 		}
 
@@ -974,11 +984,22 @@ namespace UnrealBuildTool
 						PublicHeaders = UObjectModule.PublicUObjectHeaders.Select((Header) => Header.AbsolutePath).ToList(),
 						InternalHeaders = UObjectModule.InternalUObjectHeaders.Select((Header) => Header.AbsolutePath).ToList(),
 						PrivateHeaders = UObjectModule.PrivateUObjectHeaders.Select((Header) => Header.AbsolutePath).ToList(),
+						PublicDefines = UObjectModule.PublicDefines,
 						GeneratedCPPFilenameBase = UObjectModule.GeneratedCPPFilenameBase,
 						SaveExportedHeaders = !UObjectModule.bIsReadOnly,
 						GeneratedCodeVersion = UObjectModule.GeneratedCodeVersion,
 					});
 			}
+
+			List<string> UhtPlugins = new List<string>();
+			if (Makefile.UhtPlugins != null)
+			{
+				foreach (FileReference UhtPlugin in Makefile.UhtPlugins)
+				{
+					UhtPlugins.Add(UhtPlugin.FullName);
+				}
+			}
+
 			UHTManifest Manifest = new UHTManifest
 			{
 				TargetName = TargetName,
@@ -986,6 +1007,7 @@ namespace UnrealBuildTool
 				RootLocalPath = Unreal.RootDirectory.FullName,
 				ExternalDependenciesFile = DepsFileName.FullName,
 				Modules = Modules,
+				UhtPlugins = UhtPlugins,
 			};
 
 			return Manifest;
@@ -1026,7 +1048,7 @@ namespace UnrealBuildTool
 				UnrealTargetConfiguration UHTConfig = BuildConfiguration.bForceDebugUnrealHeaderTool ? UnrealTargetConfiguration.Debug : UnrealTargetConfiguration.Development;
 
 				// Figure out the receipt path
-				FileReference HeaderToolReceipt = GetHeaderToolReceiptFile(ProjectFile, UHTConfig, Makefile.bHasProjectScriptPlugin);
+				FileReference HeaderToolReceipt = GetHeaderToolReceiptFile(ProjectFile, UHTConfig, Makefile.bHasProjectScriptPlugin, Makefile.UhtPlugins);
 
 				// check if UHT is out of date
 				DateTime HeaderToolTimestampUtc = DateTime.MaxValue;
@@ -1223,7 +1245,8 @@ namespace UnrealBuildTool
 		/// Builds and runs the header tool and touches the header directories.
 		/// Performs any early outs if headers need no changes, given the UObject modules, tool path, game name, and configuration
 		/// </summary>
-		private static void ExecuteInternalHeaderToolIfNecessary(BuildConfiguration BuildConfiguration, FileReference? ProjectFile, TargetMakefile Makefile, string TargetName, ISourceFileWorkingSet WorkingSet)
+		private static void ExecuteInternalHeaderToolIfNecessary(BuildConfiguration BuildConfiguration, FileReference? ProjectFile, 
+			TargetMakefile Makefile, string TargetName, FileReference[]? UhtPlugins, ISourceFileWorkingSet WorkingSet)
 		{
 			if (ProgressWriter.bWriteMarkup)
 			{
@@ -1236,10 +1259,26 @@ namespace UnrealBuildTool
 				UnrealTargetConfiguration UHTConfig = BuildConfiguration.bForceDebugUnrealHeaderTool ? UnrealTargetConfiguration.Debug : UnrealTargetConfiguration.Development;
 
 				// Figure out the receipt path
-				FileReference HeaderToolReceipt = GetHeaderToolReceiptFile(ProjectFile, UHTConfig, Makefile.bHasProjectScriptPlugin);
+				FileReference HeaderToolReceipt = GetHeaderToolReceiptFile(ProjectFile, UHTConfig, Makefile.bHasProjectScriptPlugin, Makefile.UhtPlugins);
 
 				// Get UHT assembly timestamp
-				DateTime UnrealBuildToolTimestamp = new FileInfo(Assembly.GetExecutingAssembly().Location).LastWriteTimeUtc;
+				DateTime CompositeTimestamp = new FileInfo(Assembly.GetExecutingAssembly().Location).LastWriteTimeUtc;
+				DateTime Timestamp = new FileInfo(typeof(UhtSession).Assembly.Location).LastWriteTimeUtc;
+				if (CompositeTimestamp < Timestamp)
+				{
+					CompositeTimestamp = Timestamp;
+				}
+				if (UhtPlugins != null)
+				{
+					foreach (FileReference Plugin in UhtPlugins)
+					{
+						Timestamp = new FileInfo(Plugin.FullName).LastWriteTimeUtc;
+						if (CompositeTimestamp < Timestamp)
+						{
+							CompositeTimestamp = Timestamp;
+						}
+					}
+				}
 
 				// ensure the headers are up to date
 				bool bUHTNeedsToRun = false;
@@ -1247,7 +1286,7 @@ namespace UnrealBuildTool
 				{
 					bUHTNeedsToRun = true;
 				}
-				else if (AreGeneratedCodeFilesOutOfDate(BuildConfiguration, Makefile.UObjectModules, UnrealBuildToolTimestamp))
+				else if (AreGeneratedCodeFilesOutOfDate(BuildConfiguration, Makefile.UObjectModules, CompositeTimestamp))
 				{
 					bUHTNeedsToRun = true;
 				}
@@ -1284,40 +1323,44 @@ namespace UnrealBuildTool
 					Log.TraceInformation("Parsing headers for {0}", ActualTargetName);
 
 					// Generate the command line
-					string CmdLine = (ProjectFile != null) ? "\"" + ProjectFile.FullName + "\"" : TargetName;
-					CmdLine += " \"" + ModuleInfoFileName + "\" -WarningsAsErrors";
-
-					if (Makefile.UHTAdditionalArguments != null)
-					{
-						foreach (string UHTArgument in Makefile.UHTAdditionalArguments)
-						{
-							CmdLine += " " + UHTArgument;
-						}
-					}
+					List<string> CmdArgs = new List<string>();
+					CmdArgs.Add((ProjectFile != null) ? ProjectFile.FullName : TargetName);
+					CmdArgs.Add(ModuleInfoFileName.FullName);
+					CmdArgs.Add("-WarningsAsErrors");
 
 					if (Unreal.IsEngineInstalled())
 					{
-						CmdLine += " -installed";
+						CmdArgs.Add("-installed");
 					}
 
 					if (BuildConfiguration.bFailIfGeneratedCodeChanges)
 					{
-						CmdLine += " -FailIfGeneratedCodeChanges";
+						CmdArgs.Add("-FailIfGeneratedCodeChanges");
 					}
 
-					Log.TraceInformation("  Running Internal UnrealHeaderTool {0}", CmdLine);
-
-					// Convert the command line to an argument list
-					string[] Args = CmdLine.Split();
-					for (int Index = 0; Index < Args.Length; ++Index)
+					if (Makefile.UHTAdditionalArguments != null)
 					{
-						string Arg = Args[Index];
-						if (Arg[0] == '"' && Arg[Arg.Length - 1] == '"')
+						foreach (string Arg in Makefile.UHTAdditionalArguments)
 						{
-							Args[Index] = Arg.Substring(1, Arg.Length - 2);
+							if (Arg[0] == '"' && Arg[Arg.Length - 1] == '"')
+							{
+								CmdArgs.Add(Arg.Substring(1, Arg.Length - 2));
+							}
+							else
+							{
+								CmdArgs.Add(Arg);
+							}
 						}
 					}
-					CommandLineArguments Arguments = new CommandLineArguments(Args);
+					CommandLineArguments Arguments = new CommandLineArguments(CmdArgs.ToArray());
+
+					StringBuilder CmdLine = new StringBuilder();
+					foreach (string Arg in CmdArgs)
+					{
+						CmdLine.AppendArgument(Arg);
+					}
+
+					Log.TraceInformation("  Running Internal UnrealHeaderTool {0}", CmdLine);
 
 					// Run UHT
 					Stopwatch s = new Stopwatch();
diff --git a/Engine/Source/Programs/UnrealBuildTool/System/TargetMakefile.cs b/Engine/Source/Programs/UnrealBuildTool/System/TargetMakefile.cs
--- a/Engine/Source/Programs/UnrealBuildTool/System/TargetMakefile.cs
+++ b/Engine/Source/Programs/UnrealBuildTool/System/TargetMakefile.cs
@@ -6,6 +6,7 @@ using System.IO;
 using System.Linq;
 using System.Reflection;
 using EpicGames.Core;
+using EpicGames.UHT.Utils;
 using OpenTracing.Util;
 using UnrealBuildBase;
 
@@ -19,7 +20,7 @@ namespace UnrealBuildTool
 		/// <summary>
 		/// The version number to write
 		/// </summary>
-		public const int CurrentVersion = 28;
+		public const int CurrentVersion = 29;
 
 		/// <summary>
 		/// The time at which the makefile was created
@@ -76,6 +77,21 @@ namespace UnrealBuildTool
 		/// </summary>
 		public bool bHasProjectScriptPlugin;
 
+		/// <summary>
+		/// If true, there exists script plugins that don't have a matching UBT plugin
+		/// </summary>
+		public bool bHasRequiredProjectScriptPlugin;
+
+		/// <summary>
+		/// Collection of UBT C# plugins (project file)
+		/// </summary>
+		public FileReference[]? UbtPlugins;
+
+		/// <summary>
+		/// Collection of UHT C# plugins (target assembly file)
+		/// </summary>
+		public FileReference[]? UhtPlugins;
+
 		/// <summary>
 		/// The array of command-line arguments. The makefile will be invalidated whenever these change.
 		/// </summary>
@@ -195,7 +211,12 @@ namespace UnrealBuildTool
 		/// <param name="ConfigValueTracker">Set of dependencies on config files</param>
 		/// <param name="bDeployAfterCompile">Whether to deploy the target after compiling</param>
 		/// <param name="bHasProjectScriptPlugin">Whether the target has a project script plugin</param>
-		public TargetMakefile(string ExternalMetadata, FileReference ExecutableFile, FileReference ReceiptFile, DirectoryReference ProjectIntermediateDirectory, TargetType TargetType, ConfigValueTracker ConfigValueTracker, bool bDeployAfterCompile, bool bHasProjectScriptPlugin)
+		/// <param name="bHasRequiredProjectScriptPlugin">If true, we have a script plugin without a matching UBT plugin</param>
+		/// <param name="UbtPlugins">Collection of UBT plugins</param>
+		/// <param name="UhtPlugins">Collection of UBT plugins for UHT</param>
+		public TargetMakefile(string ExternalMetadata, FileReference ExecutableFile, FileReference ReceiptFile, DirectoryReference ProjectIntermediateDirectory, 
+			TargetType TargetType, ConfigValueTracker ConfigValueTracker, bool bDeployAfterCompile, bool bHasProjectScriptPlugin,
+			bool bHasRequiredProjectScriptPlugin, FileReference[]? UbtPlugins, FileReference[]? UhtPlugins)
 		{
 			this.CreateTimeUtc = UnrealBuildTool.StartTimeUtc;
 			this.ModifiedTimeUtc = CreateTimeUtc;
@@ -208,6 +229,9 @@ namespace UnrealBuildTool
 			this.ConfigValueTracker = ConfigValueTracker;
 			this.bDeployAfterCompile = bDeployAfterCompile;
 			this.bHasProjectScriptPlugin = bHasProjectScriptPlugin;
+			this.bHasRequiredProjectScriptPlugin = bHasRequiredProjectScriptPlugin;
+			this.UbtPlugins = UbtPlugins;
+			this.UhtPlugins = UhtPlugins;
 			this.Actions = new List<IExternalAction>();
 			this.OutputItems = new List<FileItem>();
 			this.ModuleNameToOutputItems = new Dictionary<string, FileItem[]>(StringComparer.OrdinalIgnoreCase);
@@ -244,6 +268,9 @@ namespace UnrealBuildTool
 			ConfigValueTracker = new ConfigValueTracker(Reader);
 			bDeployAfterCompile = Reader.ReadBool();
 			bHasProjectScriptPlugin = Reader.ReadBool();
+			bHasRequiredProjectScriptPlugin = Reader.ReadBool();
+			UbtPlugins = Reader.ReadArray(() => Reader.ReadFileReference())!;
+			UhtPlugins = Reader.ReadArray(() => Reader.ReadFileReference())!;
 			AdditionalArguments = Reader.ReadArray(() => Reader.ReadString())!;
 			UHTAdditionalArguments = Reader.ReadArray(() => Reader.ReadString())!;
 			PreBuildScripts = Reader.ReadArray(() => Reader.ReadFileReference())!;
@@ -284,6 +311,9 @@ namespace UnrealBuildTool
 			ConfigValueTracker.Write(Writer);
 			Writer.WriteBool(bDeployAfterCompile);
 			Writer.WriteBool(bHasProjectScriptPlugin);
+			Writer.WriteBool(bHasRequiredProjectScriptPlugin);
+			Writer.WriteArray(UbtPlugins, Item => Writer.WriteFileReference(Item));
+			Writer.WriteArray(UhtPlugins, Item => Writer.WriteFileReference(Item));
 			Writer.WriteArray(AdditionalArguments, Item => Writer.WriteString(Item));
 			Writer.WriteArray(UHTAdditionalArguments, Item => Writer.WriteString(Item));
 			Writer.WriteArray(PreBuildScripts, Item => Writer.WriteFileReference(Item));
@@ -338,6 +368,7 @@ namespace UnrealBuildTool
 		/// <returns>The loaded makefile, or null if it failed for some reason.  On failure, the 'ReasonNotLoaded' variable will contain information about why</returns>
 		public static TargetMakefile? Load(FileReference MakefilePath, FileReference? ProjectFile, UnrealTargetPlatform Platform, string[] Arguments, out string? ReasonNotLoaded)
 		{
+			(FileReference ProjectFile, FileReference TargetAssembly)[]? Plugins;
 			FileInfo MakefileInfo;
 			using (GlobalTracer.Instance.BuildSpan("Checking dependent timestamps").StartActive())
 			{
@@ -383,6 +414,16 @@ namespace UnrealBuildTool
 					return null;
 				}
 
+				// Check to see if EpicGames.UHT was compiled more recently than the makefile
+				DateTime UhtTimestamp = new FileInfo(typeof(UhtSession).Assembly.Location).LastWriteTime;
+				if (MakefileInfo.LastWriteTime.CompareTo(UhtTimestamp) < 0)
+				{
+					// UnrealBuildTool was compiled more recently than the makefile
+					Log.TraceLog("Makefile is older than EpicGames.UHT assembly, ignoring it");
+					ReasonNotLoaded = "EpicGames.UHT assembly is newer";
+					return null;
+				}
+
 				// Check to see if any BuildConfiguration files have changed since the last build
 				foreach (XmlConfig.InputFile InputFile in XmlConfig.InputFiles)
 				{
@@ -394,6 +435,33 @@ namespace UnrealBuildTool
 						return null;
 					}
 				}
+
+				// Check for any out of date UBT plugins
+				{
+					if (PluginsBase.EnumerateUbtPlugins(ProjectFile, out Plugins))
+					{
+						if (Plugins != null)
+						{
+							foreach ((FileReference ProjectFile, FileReference TargetAssembly) Plugin in Plugins)
+							{
+								DateTime PluginTimestamp = new FileInfo(Plugin.TargetAssembly.FullName).LastWriteTime;
+								if (MakefileInfo.LastWriteTime.CompareTo(PluginTimestamp) < 0)
+								{
+									// UnrealBuildTool was compiled more recently than the makefile
+									Log.TraceLog($"Makefile is older than {Plugin.TargetAssembly.GetFileNameWithoutExtension()} UBT plugin assembly, ignoring it");
+									ReasonNotLoaded = "UBT plugin assembly is newer";
+									return null;
+								}
+							}
+						}
+					}
+					else
+					{
+						Log.TraceLog($"Error collecting plugins");
+						ReasonNotLoaded = "Error collecting plugins";
+						return null;
+					}
+				}
 			}
 
 			TargetMakefile Makefile;
@@ -456,6 +524,24 @@ namespace UnrealBuildTool
 					ReasonNotLoaded = "build metadata has changed";
 					return null;
 				}
+
+				// Check to see if the number of plugins changed
+				{
+					// Make sure the list of plugins hasn't changed
+					int NewCount = Plugins != null ? Plugins.Length : 0;
+					int OldCount = Makefile.UbtPlugins != null ? Makefile.UbtPlugins.Length : 0;
+					bool bSame = NewCount == OldCount;
+					if (bSame && Plugins != null)
+					{
+						bSame = Enumerable.SequenceEqual(Plugins.Select(P => P.ProjectFile), Makefile.UbtPlugins!);
+					}
+					if (!bSame)
+					{
+						Log.TraceLog("Available UBT plugins changed");
+						ReasonNotLoaded = "Available UBT plugins changed";
+						return null;
+					}
+				}
 			}
 
 			// The makefile is ok
diff --git a/Engine/Source/Programs/UnrealBuildTool/UnrealBuildTool.csproj b/Engine/Source/Programs/UnrealBuildTool/UnrealBuildTool.csproj
--- a/Engine/Source/Programs/UnrealBuildTool/UnrealBuildTool.csproj
+++ b/Engine/Source/Programs/UnrealBuildTool/UnrealBuildTool.csproj
@@ -116,6 +116,10 @@
     <!-- Required for Hololens -->
     <PackageReference Include="System.Reflection.MetadataLoadContext" Version="4.7.2" /> 
     
+    <!-- Required for EpicGames.MsBuild -->
+    <PackageReference Include="Microsoft.Build" Version="17.1.0" ExcludeAssets="runtime" PrivateAssets="all" />
+    <PackageReference Include="Microsoft.Build.Locator" Version="1.4.1" PrivateAssets="all" />
+
     <!-- Added as a temporary workaround to running from visual studio where the system version of dotnet is run. 
          Adding a loader context resolves many of the issues, but it still fails when MSBuild tries to write the json -->
     <PackageReference Include="Newtonsoft.Json" Version="13.0.1" />
