Redundancy cleanup in UBT Low Level Tests:
- test-specific targets and modules are constructed by RulesAssembly which removes a lot of spaghetti code
- no more excluding the Tests folder by default which is guaranteed to mislead users #fyi Jon.Nabozny
- test compilation in Tests is controlled through WITH_LOW_LEVEL_TESTS
- preparation code for switching tests between Catch2 / TestAutomation Fmwk

Nightly LLT PF 6260431391376845adeb4c40

#rnx

diff --git a/Engine/Source/Developer/LowLevelTestsRunner/LowLevelTestsRunner.Build.cs b/Engine/Source/Developer/LowLevelTestsRunner/LowLevelTestsRunner.Build.cs
--- a/Engine/Source/Developer/LowLevelTestsRunner/LowLevelTestsRunner.Build.cs
+++ b/Engine/Source/Developer/LowLevelTestsRunner/LowLevelTestsRunner.Build.cs
@@ -34,12 +34,18 @@ public class LowLevelTestsRunner : ModuleRules
 
 		if (Target.bCompileAgainstApplicationCore)
 		{
-			PrivateDependencyModuleNames.Add("ApplicationCore");
+			if (!PrivateDependencyModuleNames.Contains("ApplicationCore"))
+			{
+				PrivateDependencyModuleNames.Add("ApplicationCore");
+			}
 		}
 
 		if (Target.bCompileAgainstCoreUObject)
 		{
-			PrivateDependencyModuleNames.Add("CoreUObject");
+			if (!PrivateDependencyModuleNames.Contains("CoreUObject"))
+			{
+				PrivateDependencyModuleNames.Add("CoreUObject");
+			}
 		}
 	}
 }
\ No newline at end of file
diff --git a/Engine/Source/Programs/LowLevelTests/LowLevelTests.Build.cs b/Engine/Source/Programs/LowLevelTests/LowLevelTests.Build.cs
--- a/Engine/Source/Programs/LowLevelTests/LowLevelTests.Build.cs
+++ b/Engine/Source/Programs/LowLevelTests/LowLevelTests.Build.cs
@@ -48,6 +48,7 @@ public class LowLevelTests : ModuleRules
 			PrivateDependencyModuleNames.AddRange(
 				new string[] {
 				"Core",
+				"Cbor",
 				"CoreUObject",
 				"Projects",
 				"LowLevelTestsRunner"
diff --git a/Engine/Source/Programs/LowLevelTests/LowLevelTests.Target.cs b/Engine/Source/Programs/LowLevelTests/LowLevelTests.Target.cs
--- a/Engine/Source/Programs/LowLevelTests/LowLevelTests.Target.cs
+++ b/Engine/Source/Programs/LowLevelTests/LowLevelTests.Target.cs
@@ -17,12 +17,6 @@ public class LowLevelTestsTarget : TargetRules
 		Type = TargetType.Program;
 		LinkType = TargetLinkType.Monolithic;
 
-		// LowLevelTests to produce executable with Core tests
-		if (GetType() == typeof(LowLevelTestsTarget))
-		{
-			bIncludeAllTestsOverride = true;
-		}
-
 		bDeployAfterCompile = Target.Platform != UnrealTargetPlatform.Android;
 		bIsBuildingConsoleApplication = true;
 
@@ -59,6 +53,11 @@ public class LowLevelTestsTarget : TargetRules
 
 		VSTestRunSettingsFile = FileReference.Combine(Unreal.EngineDirectory, "Source", "Programs", "LowLevelTests", "vstest.runsettings");
 
+		if (this.GetType().Name.StartsWith("LowLevelTests"))
+		{
+			bIsTestTargetOverride = true;
+		}
+
 		SetupPreprocessorDefinitions(Target);
 	}
 
@@ -70,7 +69,6 @@ public class LowLevelTestsTarget : TargetRules
 
 	protected virtual void SetupPreprocessorDefinitions(TargetInfo Target)
 	{
-		GlobalDefinitions.Add("UE_LOW_LEVEL_TESTS=1");
 		GlobalDefinitions.Add("STATS=0");
 
 		// Platform specific setup
diff --git a/Engine/Source/Programs/UnrealBuildTool/Configuration/ModuleRules.cs b/Engine/Source/Programs/UnrealBuildTool/Configuration/ModuleRules.cs
--- a/Engine/Source/Programs/UnrealBuildTool/Configuration/ModuleRules.cs
+++ b/Engine/Source/Programs/UnrealBuildTool/Configuration/ModuleRules.cs
@@ -796,21 +796,6 @@ namespace UnrealBuildTool
 			}
 		}
 
-		/// <summary>
-		/// Decides whether to include the "Tests" folder at the root of the module as part of the compilation.
-		/// By default this folder is excluded from the module's compilation.
-		/// The "Tests" folder should be reserved for low level tests but there are exceptions to this rule.
-		/// </summary>
-		public bool bBuildIncludeTestsFolder
-		{
-			get { return bBuildIncludeTestsFolderOverride ?? false; }
-		}
-		/// <summary>
-		/// Set this override to true in derived target classes to include the "Tests" folder at the root of the module in the compilation.
-		/// </summary>
-		protected bool? bBuildIncludeTestsFolderOverride;
-
-
 		/// <summary>
 		/// The number of source files in this module before unity build will be activated for that module.  If set to
 		/// anything besides -1, will override the default setting which is controlled by MinGameModuleSourceFilesForUnityBuild
@@ -1412,6 +1397,76 @@ namespace UnrealBuildTool
 			return false;
 		}
 
+		/// <summary>
+		/// Prepares a module for building a low level tests executable.
+		/// If the target includes all tests or we're building a module as part of a test module chain, then we include the LowLevelTestsRunner dependency.
+		/// We also keep track of any Editor and Engine dependencies and add appropriate global definitions.
+		/// </summary>
+		internal void PrepareModuleForTests()
+		{
+			if (!IsTestModule && Target.IsTestTarget)
+			{
+				if (Name != "LowLevelTestsRunner" && System.IO.Directory.Exists(TestsDirectory))
+				{
+					if (!PrivateIncludePathModuleNames.Contains("LowLevelTestsRunner"))
+					{
+						PrivateIncludePathModuleNames.Add("LowLevelTestsRunner");
+					}
+				}
+				else if (Name == "LowLevelTestsRunner")
+				{
+					TestTargetRules.LowLevelTestsRunnerModule = this;
+				}
+
+				// If one of these modules is in the dependency graph, we must enable their appropriate global definitions
+				if (Name == "Editor")
+				{
+					TestTargetRules.bTestsRequireEditor = true;
+				}
+				if (Name == "Engine")
+				{
+					TestTargetRules.bTestsRequireEngine = true;
+				}
+				if (Name == "ApplicationCore")
+				{
+					TestTargetRules.bTestsRequireApplicationCore = true;
+				}
+				if (Name == "CoreUObject")
+				{
+					TestTargetRules.bTestsRequireCoreUObject = true;
+				}
+
+				if (TestTargetRules.LowLevelTestsRunnerModule != null)
+				{
+					if (TestTargetRules.bTestsRequireEditor && !TestTargetRules.LowLevelTestsRunnerModule.PrivateDependencyModuleNames.Contains("Editor"))
+					{
+						TestTargetRules.LowLevelTestsRunnerModule.PrivateDependencyModuleNames.Add("Editor");
+					}
+					if (TestTargetRules.bTestsRequireEngine && !TestTargetRules.LowLevelTestsRunnerModule.PrivateDependencyModuleNames.Contains("Engine"))
+					{
+						TestTargetRules.LowLevelTestsRunnerModule.PrivateDependencyModuleNames.Add("Engine");
+					}
+					if (TestTargetRules.bTestsRequireApplicationCore && !TestTargetRules.LowLevelTestsRunnerModule.PrivateDependencyModuleNames.Contains("ApplicationCore"))
+					{
+						TestTargetRules.LowLevelTestsRunnerModule.PrivateDependencyModuleNames.Add("ApplicationCore");
+					}
+					if (TestTargetRules.bTestsRequireCoreUObject && !TestTargetRules.LowLevelTestsRunnerModule.PrivateDependencyModuleNames.Contains("CoreUObject"))
+					{
+						TestTargetRules.LowLevelTestsRunnerModule.PrivateDependencyModuleNames.Add("CoreUObject");
+					}
+				}
+			}
+		}
+
+		internal bool IsTestModule
+		{
+			get { return bIsTestModuleOverride ?? false; }
+		}
+		/// <summary>
+		/// Whether this is a low level tests module.
+		/// </summary>
+		protected bool? bIsTestModuleOverride;
+
 		/// <summary>
 		/// Returns the module directory for a given subclass of the module (platform extensions add subclasses of ModuleRules to add in platform-specific settings)
 		/// </summary>
diff --git a/Engine/Source/Programs/UnrealBuildTool/Configuration/TargetDescriptor.cs b/Engine/Source/Programs/UnrealBuildTool/Configuration/TargetDescriptor.cs
--- a/Engine/Source/Programs/UnrealBuildTool/Configuration/TargetDescriptor.cs
+++ b/Engine/Source/Programs/UnrealBuildTool/Configuration/TargetDescriptor.cs
@@ -28,7 +28,11 @@ namespace UnrealBuildTool
 
 		public static string GetTestedTargetName(string Name)
 		{
-			return Name.Substring(0, Name.Length - TEST_TARGETS_SUFFIX.Length);
+			if (Name.EndsWith(TEST_TARGETS_SUFFIX))
+			{
+				return Name.Substring(0, Name.Length - TEST_TARGETS_SUFFIX.Length);
+			}
+			return Name;
 		}
 
 		/// <summary>
diff --git a/Engine/Source/Programs/UnrealBuildTool/Configuration/TargetRules.cs b/Engine/Source/Programs/UnrealBuildTool/Configuration/TargetRules.cs
--- a/Engine/Source/Programs/UnrealBuildTool/Configuration/TargetRules.cs
+++ b/Engine/Source/Programs/UnrealBuildTool/Configuration/TargetRules.cs
@@ -39,7 +39,7 @@ namespace UnrealBuildTool
 		/// <summary>
 		/// Program (standalone program, e.g. ShaderCompileWorker.exe, can be modular or monolithic depending on the program)
 		/// </summary>
-		Program,
+		Program
 	}
 
 	/// <summary>
@@ -274,20 +274,17 @@ namespace UnrealBuildTool
 
 		private readonly string DefaultName;
 
-		private TestTargetRules? TestTargetRules;
-
 		/// <summary>
-		/// Create and return test target executable on demand for this target and its dependencies' tests.
+		/// Whether this is a low level tests target.
 		/// </summary>
-		public TestTargetRules CreateOrGetTestTarget()
+		public bool IsTestTarget
 		{
-			if (TestTargetRules == null)
-			{
-				TargetInfo TestsTargetInfo = new TargetInfo(Name + TargetDescriptor.TEST_TARGETS_SUFFIX, Platform, Configuration, Architecture, ProjectFile, null);
-				TestTargetRules = new TestTargetRules(this, TestsTargetInfo);
-			}
-			return TestTargetRules;	
+			get { return bIsTestTargetOverride ?? false; }
 		}
+		/// <summary>
+		/// Override this boolean flag in inheriting classes for low level test targets.
+		/// </summary>
+		protected bool? bIsTestTargetOverride;
 
 		/// <summary>
 		/// File containing the general type for this target (not including platform/group)
@@ -430,18 +427,6 @@ namespace UnrealBuildTool
 		[CommandLine("-AllModules")]
 		public bool bBuildAllModules = false;
 
-		/// <summary>
-		/// Decides whether to compile with all the "Tests" folders from all dependent modules.
-		/// </summary>
-		public bool bIncludeAllTests
-		{
-			get { return bIncludeAllTestsOverride ?? false; }
-		}
-		/// <summary>
-		/// Set this override to true in derived target classes to compile with all the "Tests" folders from all dependent modules.
-		/// </summary>
-		protected bool? bIncludeAllTestsOverride;
-
 		/// <summary>
 		/// Set this to reference a VSTest run settings file from generated projects.
 		/// </summary>
@@ -1887,16 +1872,16 @@ namespace UnrealBuildTool
 			{
 				GlobalDefinitions.Add("UE_GAME=1");
 			}
-			else if(Type == global::UnrealBuildTool.TargetType.Client)
+			else if (Type == global::UnrealBuildTool.TargetType.Client)
 			{
 				GlobalDefinitions.Add("UE_GAME=1");
 				GlobalDefinitions.Add("UE_CLIENT=1");
 			}
-			else if(Type == global::UnrealBuildTool.TargetType.Editor)
+			else if (Type == global::UnrealBuildTool.TargetType.Editor)
 			{
 				GlobalDefinitions.Add("UE_EDITOR=1");
 			}
-			else if(Type == global::UnrealBuildTool.TargetType.Server)
+			else if (Type == global::UnrealBuildTool.TargetType.Server)
 			{
 				GlobalDefinitions.Add("UE_SERVER=1");
 				GlobalDefinitions.Add("USE_NULL_RHI=1");
@@ -2168,19 +2153,6 @@ namespace UnrealBuildTool
 			get { return Inner.Name; }
 		}
 
-		public TestTargetRules TestTarget { get { return Inner.CreateOrGetTestTarget(); } }
-
-		public TargetRules TestedTarget {
-			get
-			{
-				if (Inner is TestTargetRules)
-				{
-					return ((TestTargetRules)Inner).TestedTarget;
-				}
-				throw new Exception("Not a test target.");
-			}
-		}
-
 		internal FileReference File
 		{
 			get { return Inner.File!; }
@@ -2286,11 +2258,6 @@ namespace UnrealBuildTool
 			get { return Inner.bBuildAllModules; }
 		}
 
-		public bool bIncludeAllTests
-		{
-			get { return Inner.bIncludeAllTests; }
-		}
-		
 		public IEnumerable<string> AdditionalPlugins
 		{
 			get { return Inner.AdditionalPlugins; }
@@ -3202,9 +3169,9 @@ namespace UnrealBuildTool
 			get { return "../Binaries/ThirdParty/"; }
 		}
 
-		public bool IsTestTarget()
+		public bool IsTestTarget
 		{
-			return Inner is TestTargetRules;
+			get { return Inner.IsTestTarget; }
 		}
 
 		/// <summary>
diff --git a/Engine/Source/Programs/UnrealBuildTool/Configuration/TestModuleRules.cs b/Engine/Source/Programs/UnrealBuildTool/Configuration/TestModuleRules.cs
--- a/Engine/Source/Programs/UnrealBuildTool/Configuration/TestModuleRules.cs
+++ b/Engine/Source/Programs/UnrealBuildTool/Configuration/TestModuleRules.cs
@@ -18,17 +18,19 @@ namespace UnrealBuildTool
 		/// <summary>
 		/// Constructs a TestModuleRules object with an associated tested module.
 		/// </summary>
-		public TestModuleRules(ModuleRules TestedModule) : base(new ReadOnlyTargetRules(TestedModule.Target.TestTarget))
+		public TestModuleRules(ModuleRules TestedModule) : base(TestedModule.Target)
 		{
 			this.TestedModule = TestedModule;
 
+			bIsTestModuleOverride = true;
+
 			Name = TestedModule.Name + "Tests";
 			if (!string.IsNullOrEmpty(TestedModule.ShortName))
 			{
 				ShortName = TestedModule.ShortName + "Tests";
 			}
 
-			File = FileReference.Combine(UnrealBuildTool.EngineSourceDirectory, "UnrealBuildTool", "Configuration", "TestModuleRules.cs");
+			File = TestedModule.File;
 			Directory = DirectoryReference.Combine(TestedModule.Directory, "Tests");
 
 			Context = TestedModule.Context;
diff --git a/Engine/Source/Programs/UnrealBuildTool/Configuration/TestTargetRules.cs b/Engine/Source/Programs/UnrealBuildTool/Configuration/TestTargetRules.cs
--- a/Engine/Source/Programs/UnrealBuildTool/Configuration/TestTargetRules.cs
+++ b/Engine/Source/Programs/UnrealBuildTool/Configuration/TestTargetRules.cs
@@ -12,6 +12,31 @@ namespace UnrealBuildTool
 	/// </summary>
 	public class TestTargetRules : TargetRules
 	{
+		/// <summary>
+		/// Keeps track if low level tests executable must build with the Editor.
+		/// </summary>
+		public static bool bTestsRequireEditor = false;
+
+		/// <summary>
+		/// Keeps track if low level tests executable must build with the Engine.
+		/// </summary>
+		public static bool bTestsRequireEngine = false;
+
+		/// <summary>
+		/// Keeps track if low level tests executable must build with the ApplicationCore.
+		/// </summary>
+		public static bool bTestsRequireApplicationCore = false;
+
+		/// <summary>
+		/// Keeps track if low level tests executable must build with the CoreUObject.
+		/// </summary>
+		public static bool bTestsRequireCoreUObject = false;
+
+		/// <summary>
+		/// Keep track of low level test runner module instance.
+		/// </summary>
+		public static ModuleRules? LowLevelTestsRunnerModule;
+
 		/// <summary>
 		/// Associated tested target of this test target.
 		/// </summary>
@@ -28,12 +53,14 @@ namespace UnrealBuildTool
 				throw new BuildException("TestedTarget can't be of type TestTargetRules.");
 			}
 
+			bIsTestTargetOverride = true;
+
 			this.TestedTarget = TestedTarget;
 
 			DefaultBuildSettings = BuildSettingsVersion.V2;
 
 			ExeBinariesSubFolder = Name = TestedTarget.Name + "Tests";
-			TargetSourceFile = File = FileReference.Combine(UnrealBuildTool.EngineSourceDirectory, "UnrealBuildTool", "Configuration", "TestTargetRules.cs");
+			TargetSourceFile = File = TestedTarget.File;
 			LaunchModuleName = TestedTarget.LaunchModuleName + "Tests";
 
 			ManifestFileNames = TestedTarget.ManifestFileNames;
@@ -81,7 +108,6 @@ namespace UnrealBuildTool
 				bDebugBuildsActuallyUseDebugCRT = true;
 			}
 
-			GlobalDefinitions.Add("UE_LOW_LEVEL_TESTS=1");
 			GlobalDefinitions.Add("STATS=0");
 			GlobalDefinitions.Add("TEST_FOR_VALID_FILE_SYSTEM_MEMORY=0");
 
diff --git a/Engine/Source/Programs/UnrealBuildTool/Configuration/UEBuildBinary.cs b/Engine/Source/Programs/UnrealBuildTool/Configuration/UEBuildBinary.cs
--- a/Engine/Source/Programs/UnrealBuildTool/Configuration/UEBuildBinary.cs
+++ b/Engine/Source/Programs/UnrealBuildTool/Configuration/UEBuildBinary.cs
@@ -718,7 +718,7 @@ namespace UnrealBuildTool
 				{
 					// Compile each module.
 					Log.TraceVerbose("Compile module: " + Module.Name);
-					LinkInputFiles = Module.Compile(Target, ToolChain, BinaryCompileEnvironment, SpecificFilesToCompile, WorkingSet, Graph, Module.Rules.bBuildIncludeTestsFolder || Target.bIncludeAllTests);
+					LinkInputFiles = Module.Compile(Target, ToolChain, BinaryCompileEnvironment, SpecificFilesToCompile, WorkingSet, Graph);
 
 					// Save the module outputs. In monolithic builds, this is just the object files.
 					if (Target.LinkType == TargetLinkType.Monolithic)
diff --git a/Engine/Source/Programs/UnrealBuildTool/Configuration/UEBuildModule.cs b/Engine/Source/Programs/UnrealBuildTool/Configuration/UEBuildModule.cs
--- a/Engine/Source/Programs/UnrealBuildTool/Configuration/UEBuildModule.cs
+++ b/Engine/Source/Programs/UnrealBuildTool/Configuration/UEBuildModule.cs
@@ -852,7 +852,7 @@ namespace UnrealBuildTool
 		/// <summary>
 		/// Compiles the module, and returns a list of files output by the compiler.
 		/// </summary>
-		public virtual List<FileItem> Compile(ReadOnlyTargetRules Target, UEToolChain ToolChain, CppCompileEnvironment CompileEnvironment, List<FileReference> SpecificFilesToCompile, ISourceFileWorkingSet WorkingSet, IActionGraphBuilder Graph, bool bBuildIncludeTestsFolder = false)
+		public virtual List<FileItem> Compile(ReadOnlyTargetRules Target, UEToolChain ToolChain, CppCompileEnvironment CompileEnvironment, List<FileReference> SpecificFilesToCompile, ISourceFileWorkingSet WorkingSet, IActionGraphBuilder Graph)
 		{
 			// Generate type libraries for Windows
 			foreach(ModuleRules.TypeLibrary TypeLibrary in Rules.TypeLibraries)
@@ -921,7 +921,7 @@ namespace UnrealBuildTool
 			}
 		}
 
-		public delegate UEBuildModule CreateModuleDelegate(string Name, string ReferenceChain, bool IsTestModule = false);
+		public delegate UEBuildModule CreateModuleDelegate(string Name, string ReferenceChain);
 
 		/// <summary>
 		/// Public entry point to recursively create a module and all its dependencies
diff --git a/Engine/Source/Programs/UnrealBuildTool/Configuration/UEBuildModuleCPP.cs b/Engine/Source/Programs/UnrealBuildTool/Configuration/UEBuildModuleCPP.cs
--- a/Engine/Source/Programs/UnrealBuildTool/Configuration/UEBuildModuleCPP.cs
+++ b/Engine/Source/Programs/UnrealBuildTool/Configuration/UEBuildModuleCPP.cs
@@ -365,7 +365,7 @@ namespace UnrealBuildTool
 		}
 
 		// UEBuildModule interface.
-		public override List<FileItem> Compile(ReadOnlyTargetRules Target, UEToolChain ToolChain, CppCompileEnvironment BinaryCompileEnvironment, List<FileReference> SpecificFilesToCompile, ISourceFileWorkingSet WorkingSet, IActionGraphBuilder Graph, bool bBuildIncludeTestsFolder = false)
+		public override List<FileItem> Compile(ReadOnlyTargetRules Target, UEToolChain ToolChain, CppCompileEnvironment BinaryCompileEnvironment, List<FileReference> SpecificFilesToCompile, ISourceFileWorkingSet WorkingSet, IActionGraphBuilder Graph)
 		{
 			//UEBuildPlatform BuildPlatform = UEBuildPlatform.GetBuildPlatform(BinaryCompileEnvironment.Platform);
 
@@ -404,7 +404,7 @@ namespace UnrealBuildTool
 
 			// Find all the input files
 			Dictionary<DirectoryItem, FileItem[]> DirectoryToSourceFiles = new Dictionary<DirectoryItem, FileItem[]>();
-			InputFileCollection InputFiles = FindInputFiles(Target.Platform, DirectoryToSourceFiles, Target.IsTestTarget() || bBuildIncludeTestsFolder);
+			InputFileCollection InputFiles = FindInputFiles(Target.Platform, DirectoryToSourceFiles);
 
 			foreach (KeyValuePair<DirectoryItem, FileItem[]> Pair in DirectoryToSourceFiles)
 			{
@@ -1542,9 +1542,8 @@ namespace UnrealBuildTool
 		/// </summary>
 		/// <param name="Platform">The platform the module is being built for</param>
 		/// <param name="DirectoryToSourceFiles">Map of directory to source files inside it</param>
-		/// <param name="IncludeTests">Whether tests are included in the file list</param>
 		/// <returns>Set of source files that should be built</returns>
-		public InputFileCollection FindInputFiles(UnrealTargetPlatform Platform, Dictionary<DirectoryItem, FileItem[]> DirectoryToSourceFiles, bool IncludeTests = false)
+		public InputFileCollection FindInputFiles(UnrealTargetPlatform Platform, Dictionary<DirectoryItem, FileItem[]> DirectoryToSourceFiles)
 		{
 			ReadOnlyHashSet<string> ExcludedNames = UEBuildPlatform.GetBuildPlatform(Platform).GetExcludedFolderNames();
 
@@ -1554,32 +1553,7 @@ namespace UnrealBuildTool
 			foreach (DirectoryReference Dir in ModuleDirectories)
 			{
 				DirectoryItem ModuleDirectoryItem = DirectoryItem.GetItemByDirectoryReference(Dir);
-				if (!IncludeTests)
-				{
-					// Skip Tests folder at the base of the source folder if not including tests
-					foreach (DirectoryItem FirstLevelDir in ModuleDirectoryItem.EnumerateDirectories())
-					{
-						if (FirstLevelDir.Name == "Tests")
-						{
-							continue;
-						}
-						else
-						{
-							FindInputFilesFromDirectoryRecursive(FirstLevelDir, ExcludedNames, SourceDirectories, DirectoryToSourceFiles, InputFiles);
-						}
-					}
-					// Add remaining source files at the base of the source folder
-					FileItem[] SourceFiles = FindInputFilesFromDirectory(ModuleDirectoryItem, InputFiles);
-					if (SourceFiles.Length > 0)
-					{
-						SourceDirectories.Add(ModuleDirectoryItem.Location);
-					}
-					DirectoryToSourceFiles.Add(ModuleDirectoryItem, SourceFiles);
-				}
-				else
-				{
-					FindInputFilesFromDirectoryRecursive(ModuleDirectoryItem, ExcludedNames, SourceDirectories, DirectoryToSourceFiles, InputFiles);
-				}
+				FindInputFilesFromDirectoryRecursive(ModuleDirectoryItem, ExcludedNames, SourceDirectories, DirectoryToSourceFiles, InputFiles);
 			}
 
 			return InputFiles;
diff --git a/Engine/Source/Programs/UnrealBuildTool/Configuration/UEBuildTarget.cs b/Engine/Source/Programs/UnrealBuildTool/Configuration/UEBuildTarget.cs
--- a/Engine/Source/Programs/UnrealBuildTool/Configuration/UEBuildTarget.cs
+++ b/Engine/Source/Programs/UnrealBuildTool/Configuration/UEBuildTarget.cs
@@ -711,16 +711,7 @@ namespace UnrealBuildTool
 			TargetRules RulesObject;
 			using (GlobalTracer.Instance.BuildSpan("RulesAssembly.CreateTargetRules()").StartActive())
 			{
-				if (!Descriptor.IsTestsTarget)
-				{
-					RulesObject = RulesAssembly.CreateTargetRules(Descriptor.Name, Descriptor.Platform, Descriptor.Configuration, Descriptor.Architecture, Descriptor.ProjectFile, Descriptor.AdditionalArguments);
-				}
-				else
-				{
-					// Get testable target first, then use its tests executable target instead
-					RulesObject = RulesAssembly.CreateTargetRules(TargetDescriptor.GetTestedTargetName(Descriptor.Name), Descriptor.Platform, Descriptor.Configuration, Descriptor.Architecture, Descriptor.ProjectFile, Descriptor.AdditionalArguments);
-					RulesObject = RulesObject.CreateOrGetTestTarget();
-				}
+				RulesObject = RulesAssembly.CreateTargetRules(Descriptor.Name, Descriptor.Platform, Descriptor.Configuration, Descriptor.Architecture, Descriptor.ProjectFile, Descriptor.AdditionalArguments, Descriptor.IsTestsTarget);
 			}
 			if ((ProjectFileGenerator.bGenerateProjectFiles == false) && !RulesObject.GetSupportedPlatforms().Contains(Descriptor.Platform))
 			{
@@ -1105,31 +1096,6 @@ namespace UnrealBuildTool
 		/// </summary>
 		protected bool bCompileMonolithic = false;
 
-		/// <summary>
-		/// Keeps track if low level tests executable must build with the Editor.
-		/// </summary>
-		private bool bTestsRequireEditor = false;
-
-		/// <summary>
-		/// Keeps track if low level tests executable must build with the Engine.
-		/// </summary>
-		private bool bTestsRequireEngine = false;
-
-		/// <summary>
-		/// Keeps track if low level tests executable must build with the ApplicationCore.
-		/// </summary>
-		private bool bTestsRequireApplicationCore = false;
-
-		/// <summary>
-		/// Keeps track if low level tests executable must build with the CoreUObject.
-		/// </summary>
-		private bool bTestsRequireCoreUObject = false;
-
-		/// <summary>
-		/// Track low level tests runner module for configuration.
-		/// </summary>
-		private ModuleRules? LowLevelTestsRunnerModule;
-
 		/// <summary>
 		/// Used to keep track of all modules by name.
 		/// </summary>
@@ -3073,9 +3039,9 @@ namespace UnrealBuildTool
 				UEBuildModule Module = FindOrCreateModuleByName(ModuleName, PrecompileReferenceChain);
 				AllModules.Add(Module);
 				Module.RecursivelyCreateModules(
-					(string ModuleName, string ReferenceChain, bool IsTestModule) => 
+					(string ModuleName, string ReferenceChain) => 
 					{ 
-						UEBuildModule FoundModule = FindOrCreateModuleByName(ModuleName, ReferenceChain, IsTestModule);
+						UEBuildModule FoundModule = FindOrCreateModuleByName(ModuleName, ReferenceChain);
 						AllModules.Add(FoundModule);
 						return FoundModule; 
 					}, 
@@ -3688,7 +3654,7 @@ namespace UnrealBuildTool
 			}
 
 			// Create the launch module
-			UEBuildModuleCPP LaunchModule = FindOrCreateCppModuleByName(Rules.LaunchModuleName, TargetRulesFile.GetFileName(), Rules.IsTestTarget());
+			UEBuildModuleCPP LaunchModule = FindOrCreateCppModuleByName(Rules.LaunchModuleName, TargetRulesFile.GetFileName());
 
 			// Get the intermediate directory for the launch module directory. This can differ from the standard engine intermediate directory because it is always configuration-specific.
 			DirectoryReference IntermediateDirectory;
@@ -3878,18 +3844,20 @@ namespace UnrealBuildTool
 			GlobalCompileEnvironment.Definitions.Add("WITH_DEV_AUTOMATION_TESTS=" + (bCompileDevTests ? "1" : "0"));
 			GlobalCompileEnvironment.Definitions.Add("WITH_PERF_AUTOMATION_TESTS=" + (bCompilePerfTests ? "1" : "0"));
 
+			GlobalCompileEnvironment.Definitions.Add(String.Format("WITH_LOW_LEVEL_TESTS={0}", Rules.IsTestTarget ? "1" : "0"));
+
 			GlobalCompileEnvironment.Definitions.Add("UNICODE");
 			GlobalCompileEnvironment.Definitions.Add("_UNICODE");
 			GlobalCompileEnvironment.Definitions.Add("__UNREAL__");
 
 			GlobalCompileEnvironment.Definitions.Add(String.Format("IS_MONOLITHIC={0}", ShouldCompileMonolithic() ? "1" : "0"));
 
-			GlobalCompileEnvironment.Definitions.Add(String.Format("WITH_ENGINE={0}", Rules.bCompileAgainstEngine || bTestsRequireEngine ? "1" : "0"));
+			GlobalCompileEnvironment.Definitions.Add(String.Format("WITH_ENGINE={0}", Rules.bCompileAgainstEngine || TestTargetRules.bTestsRequireEngine ? "1" : "0"));
 			GlobalCompileEnvironment.Definitions.Add(String.Format("WITH_UNREAL_DEVELOPER_TOOLS={0}", Rules.bBuildDeveloperTools ? "1" : "0"));
 			GlobalCompileEnvironment.Definitions.Add(String.Format("WITH_UNREAL_TARGET_DEVELOPER_TOOLS={0}", Rules.bBuildTargetDeveloperTools ? "1" : "0"));
 
 			// Set a macro to control whether to initialize ApplicationCore. Command line utilities should not generally need this.
-			if (Rules.bCompileAgainstApplicationCore || bTestsRequireApplicationCore)
+			if (Rules.bCompileAgainstApplicationCore || TestTargetRules.bTestsRequireApplicationCore)
 			{
 				GlobalCompileEnvironment.Definitions.Add("WITH_APPLICATION_CORE=1");
 			}
@@ -3898,7 +3866,7 @@ namespace UnrealBuildTool
 				GlobalCompileEnvironment.Definitions.Add("WITH_APPLICATION_CORE=0");
 			}
 
-			if (Rules.bCompileAgainstCoreUObject || bTestsRequireCoreUObject)
+			if (Rules.bCompileAgainstCoreUObject || TestTargetRules.bTestsRequireCoreUObject)
 			{
 				GlobalCompileEnvironment.Definitions.Add("WITH_COREUOBJECT=1");
 			}
@@ -4006,7 +3974,7 @@ namespace UnrealBuildTool
 				GlobalCompileEnvironment.Definitions.Add("USE_ESTIMATED_UTCNOW=0");
 			}
 
-			if ((Rules.bCompileAgainstEditor && (Rules.Type == TargetType.Editor || Rules.Type == TargetType.Program)) || bTestsRequireEditor)
+			if ((Rules.bCompileAgainstEditor && (Rules.Type == TargetType.Editor || Rules.Type == TargetType.Program)) || TestTargetRules.bTestsRequireEditor)
 			{
 				GlobalCompileEnvironment.Definitions.Add("WITH_EDITOR=1");
 				GlobalCompileEnvironment.Definitions.Add("WITH_IOSTORE_IN_EDITOR=1");
@@ -4168,19 +4136,10 @@ namespace UnrealBuildTool
 		/// <summary>
 		/// Create a rules object for the given module, and set any default values for this target
 		/// </summary>
-		private ModuleRules CreateModuleRulesAndSetDefaults(string ModuleName, string ReferenceChain, bool IsTestModule = false)
+		private ModuleRules CreateModuleRulesAndSetDefaults(string ModuleName, string ReferenceChain)
 		{
 			// Create the rules from the assembly
-			ModuleRules RulesObject;
-			if (IsTestModule)
-			{
-				RulesObject = RulesAssembly.CreateModuleRules(TargetDescriptor.GetTestedTargetName(ModuleName), new ReadOnlyTargetRules(Rules.TestedTarget), ReferenceChain);
-				RulesObject = new TestModuleRules(RulesObject);
-			}
-			else
-			{
-				RulesObject = RulesAssembly.CreateModuleRules(ModuleName, Rules, ReferenceChain);
-			}
+			ModuleRules RulesObject = RulesAssembly.CreateModuleRules(ModuleName, Rules, ReferenceChain);
 
 			// Set whether the module requires an IMPLEMENT_MODULE macro
 			if (!RulesObject.bRequiresImplementModule.HasValue)
@@ -4296,20 +4255,16 @@ namespace UnrealBuildTool
 		/// </summary>
 		/// <param name="ModuleName">Name of the module</param>
 		/// <param name="ReferenceChain">Chain of references causing this module to be instantiated, for display in error messages</param>
-		/// <param name="IsTestModule">Whether it's a test module or not</param>
-		public UEBuildModule FindOrCreateModuleByName(string ModuleName, string ReferenceChain, bool IsTestModule = false)
+		public UEBuildModule FindOrCreateModuleByName(string ModuleName, string ReferenceChain)
 		{
 			UEBuildModule? Module;
 			if (!Modules.TryGetValue(ModuleName, out Module))
 			{
 				// @todo projectfiles: Cross-platform modules can appear here during project generation, but they may have already
 				//   been filtered out by the project generator.  This causes the projects to not be added to directories properly.
-				ModuleRules RulesObject = CreateModuleRulesAndSetDefaults(ModuleName, ReferenceChain, IsTestModule);
+				ModuleRules RulesObject = CreateModuleRulesAndSetDefaults(ModuleName, ReferenceChain);
 				DirectoryReference ModuleDirectory = RulesObject.File.Directory;
 
-				// If we're building a test executable, include test harness dependency
-				PrepareModuleForTests(IsTestModule, RulesObject);
-
 				// Clear the bUsePrecompiled flag if we're compiling a foreign plugin; since it's treated like an engine module, it will default to true in an installed build.
 				if (RulesObject.Plugin != null && RulesObject.Plugin.File == ForeignPlugin)
 				{
@@ -4421,78 +4376,15 @@ namespace UnrealBuildTool
 			return Module;
 		}
 
-		/// <summary>
-		/// Prepares a module for building a low level tests executable.
-		/// If the target includes all tests OR we're building a module as part of a test module chain, then we include the LowLevelTestsRunner dependency.
-		/// We also keep track of any Editor and Engine dependencies and add appropriate global definitions at compilation step.
-		/// </summary>
-		/// <param name="IsTestModule">Whether we're building a tests module or not.</param>
-		/// <param name="RulesObject">The module rules object.</param>
-		private void PrepareModuleForTests(bool IsTestModule, ModuleRules RulesObject)
-		{
-			if (!IsTestModule
-				&& Rules.IsTestTarget()
-				|| Rules.bIncludeAllTests)
-			{
-				if (RulesObject.Name != "LowLevelTestsRunner" && Directory.Exists(RulesObject.TestsDirectory))
-				{
-					RulesObject.PrivateIncludePathModuleNames.Add("LowLevelTestsRunner");
-				}
-				else if (RulesObject.Name == "LowLevelTestsRunner")
-				{
-					LowLevelTestsRunnerModule = RulesObject;
-				}
-
-				// If one of these modules is in the dependency graph, we must enable their appropriate global definitions
-				if (RulesObject.Name == "Editor")
-				{
-					bTestsRequireEditor = true;
-				}
-				if (RulesObject.Name == "Engine")
-				{
-					bTestsRequireEngine = true;
-				}
-				if (RulesObject.Name == "ApplicationCore")
-				{
-					bTestsRequireApplicationCore = true;
-				}
-				if (RulesObject.Name == "CoreUObject")
-				{
-					bTestsRequireCoreUObject = true;
-				}
-
-				if (LowLevelTestsRunnerModule != null)
-				{
-					if (bTestsRequireEditor && !LowLevelTestsRunnerModule.PrivateDependencyModuleNames.Contains("Editor"))
-					{
-						LowLevelTestsRunnerModule.PrivateDependencyModuleNames.Add("Editor");
-					}
-					if (bTestsRequireEngine && !LowLevelTestsRunnerModule.PrivateDependencyModuleNames.Contains("Engine"))
-					{
-						LowLevelTestsRunnerModule.PrivateDependencyModuleNames.Add("Engine");
-					}
-					if (bTestsRequireApplicationCore && !LowLevelTestsRunnerModule.PrivateDependencyModuleNames.Contains("ApplicationCore"))
-					{
-						LowLevelTestsRunnerModule.PrivateDependencyModuleNames.Add("ApplicationCore");
-					}
-					if (bTestsRequireCoreUObject && !LowLevelTestsRunnerModule.PrivateDependencyModuleNames.Contains("CoreUObject"))
-					{
-						LowLevelTestsRunnerModule.PrivateDependencyModuleNames.Add("CoreUObject");
-					}
-				}
-			}
-		}
-
 		/// <summary>
 		/// Constructs a new C++ module
 		/// </summary>
 		/// <param name="ModuleName">Name of the module</param>
 		/// <param name="ReferenceChain">Chain of references causing this module to be instantiated, for display in error messages</param>
-		/// <param name="IsTestModule">Whether it's a test module.</param>
 		/// <returns>New C++ module</returns>
-		public UEBuildModuleCPP FindOrCreateCppModuleByName(string ModuleName, string ReferenceChain, bool IsTestModule = false)
+		public UEBuildModuleCPP FindOrCreateCppModuleByName(string ModuleName, string ReferenceChain)
 		{
-			UEBuildModuleCPP? CppModule = FindOrCreateModuleByName(ModuleName, ReferenceChain, IsTestModule) as UEBuildModuleCPP;
+			UEBuildModuleCPP? CppModule = FindOrCreateModuleByName(ModuleName, ReferenceChain) as UEBuildModuleCPP;
 			if (CppModule == null)
 			{
 				throw new BuildException("'{0}' is not a C++ module (referenced via {1})", ModuleName, ReferenceChain);
diff --git a/Engine/Source/Programs/UnrealBuildTool/ProjectFiles/VisualStudio/VCProject.cs b/Engine/Source/Programs/UnrealBuildTool/ProjectFiles/VisualStudio/VCProject.cs
--- a/Engine/Source/Programs/UnrealBuildTool/ProjectFiles/VisualStudio/VCProject.cs
+++ b/Engine/Source/Programs/UnrealBuildTool/ProjectFiles/VisualStudio/VCProject.cs
@@ -512,14 +512,17 @@ namespace UnrealBuildTool
 			public string ProjectPlatformName;
 			public string ProjectConfigurationName;
 			public ProjectTarget? ProjectTarget;
+			public bool TestMode;
 
-			public ProjectConfigAndTargetCombination(UnrealTargetPlatform? InPlatform, UnrealTargetConfiguration InConfiguration, string InProjectPlatformName, string InProjectConfigurationName, ProjectTarget? InProjectTarget)
+			public ProjectConfigAndTargetCombination(UnrealTargetPlatform? InPlatform, UnrealTargetConfiguration InConfiguration, string InProjectPlatformName, string InProjectConfigurationName, ProjectTarget? InProjectTarget, bool InTestMode = false)
 			{
+				TestMode = InTestMode;
+
 				Platform = InPlatform;
 				Configuration = InConfiguration;
 				ProjectPlatformName = InProjectPlatformName;
-				ProjectConfigurationName = InProjectConfigurationName;
-				ProjectTarget = InProjectTarget;
+				ProjectConfigurationName = InProjectConfigurationName + (!InTestMode ? string.Empty : "_LowLevelTests");
+				ProjectTarget = InProjectTarget;				
 			}
 
 			public string? ProjectConfigurationAndPlatformName
@@ -642,6 +645,9 @@ namespace UnrealBuildTool
 				{
 					ProjectConfigAndTargetCombination StubCombination = new ProjectConfigAndTargetCombination(UnrealTargetPlatform.Parse(StubProjectPlatformName), UnrealTargetConfiguration.Unknown, StubProjectPlatformName, StubProjectConfigurationName, null);
 					ProjectConfigAndTargetCombinations.Add(StubCombination);
+
+					ProjectConfigAndTargetCombination TestModeStubCombination = new ProjectConfigAndTargetCombination(UnrealTargetPlatform.Parse(StubProjectPlatformName), UnrealTargetConfiguration.Unknown, StubProjectPlatformName, StubProjectConfigurationName, null, true);
+					ProjectConfigAndTargetCombinations.Add(TestModeStubCombination);
 				}
 				else
 				{
@@ -677,6 +683,10 @@ namespace UnrealBuildTool
 
 										ProjectConfigAndTargetCombination Combination = new ProjectConfigAndTargetCombination(Platform, Configuration, ProjectPlatformName, ProjectConfigurationName, ProjectTarget);
 										ProjectConfigAndTargetCombinations.Add(Combination);
+
+										// Configuration for building with -Mode=Test, produces separate low level tests executable
+										ProjectConfigAndTargetCombination TestModeCombination = new ProjectConfigAndTargetCombination(Platform, Configuration, ProjectPlatformName, ProjectConfigurationName, ProjectTarget, true);
+										ProjectConfigAndTargetCombinations.Add(TestModeCombination);
 									}
 								}
 							}
@@ -1727,6 +1737,11 @@ namespace UnrealBuildTool
 						BuildArguments.AppendFormat(" -Project={0}", UProjectPath);
 					}
 
+					if (Combination.TestMode)
+					{
+						BuildArguments.Append(" -Mode=Test");
+					}
+
 					List<string> ExtraTargets = new List<string>();
 					if (!bUsePrecompiled)
 					{
diff --git a/Engine/Source/Programs/UnrealBuildTool/System/RulesAssembly.cs b/Engine/Source/Programs/UnrealBuildTool/System/RulesAssembly.cs
--- a/Engine/Source/Programs/UnrealBuildTool/System/RulesAssembly.cs
+++ b/Engine/Source/Programs/UnrealBuildTool/System/RulesAssembly.cs
@@ -388,6 +388,11 @@ namespace UnrealBuildTool
 		/// <returns>Compiled module rule info</returns>
 		public ModuleRules CreateModuleRules(string ModuleName, ReadOnlyTargetRules Target, string ReferenceChain)
 		{
+			if (Target.IsTestTarget)
+			{
+				ModuleName = TargetDescriptor.GetTestedTargetName(ModuleName);
+			}
+
 			// Currently, we expect the user's rules object type name to be the same as the module name
 			string ModuleTypeName = ModuleName;
 
@@ -505,6 +510,13 @@ namespace UnrealBuildTool
 				}
 				Constructor.Invoke(RulesObject, new object[] { Target });
 
+				if (Target.IsTestTarget && Target.LaunchModuleName != null && ModuleName == TargetDescriptor.GetTestedTargetName(Target.LaunchModuleName))
+				{
+					RulesObject = new TestModuleRules(RulesObject);
+				}
+
+				RulesObject.PrepareModuleForTests();
+
 				return RulesObject;
 			}
 			catch (Exception Ex)
@@ -519,8 +531,9 @@ namespace UnrealBuildTool
 		/// </summary>
 		/// <param name="TypeName">Type name of the target rules</param>
 		/// <param name="TargetInfo">Target configuration information to pass to the constructor</param>
+		/// <param name="IsTestTarget">If building a low level tests target</param>
 		/// <returns>Instance of the corresponding TargetRules</returns>
-		protected TargetRules CreateTargetRulesInstance(string TypeName, TargetInfo TargetInfo)
+		protected TargetRules CreateTargetRulesInstance(string TypeName, TargetInfo TargetInfo, bool IsTestTarget = false)
 		{
 			// The build module must define a type named '<TargetName>Target' that derives from our 'TargetRules' type.  
 			Type? BaseRulesType = CompiledAssembly?.GetType(TypeName);
@@ -674,6 +687,11 @@ namespace UnrealBuildTool
 				throw new BuildException(String.Format("{0}: {1} does not support modular builds", Rules.Name, Rules.Platform));
 			}
 
+			if (IsTestTarget)
+			{
+				Rules = new TestTargetRules(Rules, TargetInfo);
+			}
+
 			return Rules;
 		}
 
@@ -686,9 +704,14 @@ namespace UnrealBuildTool
 		/// <param name="Architecture">Architecture being built</param>
 		/// <param name="ProjectFile">Path to the project file for this target</param>
 		/// <param name="Arguments">Command line arguments for this target</param>
+		/// <param name="IsTestTarget">If building a low level test target</param>
 		/// <returns>The build target rules for the specified target</returns>
-		public TargetRules CreateTargetRules(string TargetName, UnrealTargetPlatform Platform, UnrealTargetConfiguration Configuration, string Architecture, FileReference? ProjectFile, CommandLineArguments? Arguments)
+		public TargetRules CreateTargetRules(string TargetName, UnrealTargetPlatform Platform, UnrealTargetConfiguration Configuration, string Architecture, FileReference? ProjectFile, CommandLineArguments? Arguments, bool IsTestTarget = false)
 		{
+			if (IsTestTarget)
+			{
+				TargetName = TargetDescriptor.GetTestedTargetName(TargetName);
+			}
 			bool bFoundTargetName = TargetNameToTargetFile.ContainsKey(TargetName);
 			if (bFoundTargetName == false)
 			{
@@ -713,7 +736,7 @@ namespace UnrealBuildTool
 				}
 				else
 				{
-					return Parent.CreateTargetRules(TargetName, Platform, Configuration, Architecture, ProjectFile, Arguments);
+					return Parent.CreateTargetRules(TargetName, Platform, Configuration, Architecture, ProjectFile, Arguments, IsTestTarget);
 				}
 			}
 
@@ -721,7 +744,7 @@ namespace UnrealBuildTool
 			string TargetTypeName = TargetName + "Target";
 
 			// The build module must define a type named '<TargetName>Target' that derives from our 'TargetRules' type.  
-			return CreateTargetRulesInstance(TargetTypeName, new TargetInfo(TargetName, Platform, Configuration, Architecture, ProjectFile, Arguments));
+			return CreateTargetRulesInstance(TargetTypeName, new TargetInfo(TargetName, Platform, Configuration, Architecture, ProjectFile, Arguments), IsTestTarget);
 		}
 
 		/// <summary>
diff --git a/Engine/Source/Runtime/Cbor/Tests/CborTests.cpp b/Engine/Source/Runtime/Cbor/Tests/CborTests.cpp
--- a/Engine/Source/Runtime/Cbor/Tests/CborTests.cpp
+++ b/Engine/Source/Runtime/Cbor/Tests/CborTests.cpp
@@ -1,5 +1,7 @@
 // Copyright Epic Games, Inc. All Rights Reserved.
 
+#if WITH_LOW_LEVEL_TESTS
+
 #include "CoreMinimal.h"
 #include "Misc/AutomationTest.h"
 #include "Serialization/MemoryReader.h"
@@ -387,4 +389,6 @@ TEST_CASE("CBor::FCbor::ByteSwap Performance Test", "[CBor][Perf]")
 	double Margin = 0.5;
 	CHECK(( WriteRatio >= 1.0 - Margin && WriteRatio <= 1.0 + Margin)==true);
 	CHECK(( ReadRatio  >= 1.0 - Margin && ReadRatio  <= 1.0 + Margin)==true);
-}
\ No newline at end of file
+}
+
+#endif
\ No newline at end of file
diff --git a/Engine/Source/Runtime/Core/Private/Misc/AutomationTest.cpp b/Engine/Source/Runtime/Core/Private/Misc/AutomationTest.cpp
--- a/Engine/Source/Runtime/Core/Private/Misc/AutomationTest.cpp
+++ b/Engine/Source/Runtime/Core/Private/Misc/AutomationTest.cpp
@@ -1199,6 +1199,24 @@ void FAutomationTestBase::AddExpectedError(FString ExpectedErrorPattern, EAutoma
 	}
 }
 
+uint32 FAutomationTestBase::ExtractAutomationTestFlags(FString InTagNotation)
+{
+	uint32 Result = 0;
+	TArray<FString> OutputParts;
+	InTagNotation
+		.Replace(TEXT("["), TEXT(""))
+		.Replace(TEXT("]"), TEXT(";"))
+		.ParseIntoArray(OutputParts, TEXT(";"), true);
+	for (auto it = OutputParts.begin(); it != OutputParts.end(); ++it)
+	{
+		if (FStringToEAutomationTestFlagMap.Contains(*it))
+		{
+			Result |= FStringToEAutomationTestFlagMap[*it];
+		}
+	}
+	return Result;
+}
+
 void FAutomationTestBase::GetExpectedErrors(TArray<FAutomationExpectedError>& OutInfo) const
 {
 	OutInfo = ExpectedErrors;
diff --git a/Engine/Source/Runtime/Core/Private/Unix/UnixPlatformStackWalk.cpp b/Engine/Source/Runtime/Core/Private/Unix/UnixPlatformStackWalk.cpp
--- a/Engine/Source/Runtime/Core/Private/Unix/UnixPlatformStackWalk.cpp
+++ b/Engine/Source/Runtime/Core/Private/Unix/UnixPlatformStackWalk.cpp
@@ -328,7 +328,7 @@ namespace
 			if (FPaths::IsRelative(info.dli_fname))
 			{
 				FCStringAnsi::Strcpy(ModuleSymbolPath, TCHAR_TO_UTF8(FPlatformProcess::BaseDir()));
-#ifdef UE_LOW_LEVEL_TESTS
+#ifdef WITH_LOW_LEVEL_TESTS
 				// Low level tests live one level above the base directory in a folder <ModuleName>Tests
 				// Sometimes this folder can also be just <ModuleName> if the target was compiled with the tests
 				if (FPaths::DirectoryExists(FPaths::Combine(ModuleSymbolPath, TCHAR_TO_UTF8(*FPaths::GetBaseFilename(out_SymbolInfo.ModuleName)))))
diff --git a/Engine/Source/Runtime/Core/Public/Misc/AutomationTest.h b/Engine/Source/Runtime/Core/Public/Misc/AutomationTest.h
--- a/Engine/Source/Runtime/Core/Public/Misc/AutomationTest.h
+++ b/Engine/Source/Runtime/Core/Public/Misc/AutomationTest.h
@@ -47,7 +47,10 @@
 		Info.LineNumber = 1;																				\
 	}
 
-/** Flags for specifying automation test requirements/behavior */
+/**
+* Flags for specifying automation test requirements/behavior
+* Update FStringToEAutomationTestFlagMap when updating this enum.
+*/
 namespace EAutomationTestFlags
 {
 	enum Type
@@ -1787,6 +1790,37 @@ protected:
 	/** Sets the parameter context of the test. */
 	virtual void SetTestContext(FString Context) { TestParameterContext = Context; }
 
+	/** Extracts a combined EAutomationTestFlags value from a string representation using tag notation "[Filter_1]...[Filter_n][Tag_1]...[Tag_m]" */
+	uint32 ExtractAutomationTestFlags(FString InTagNotation);
+
+	/** String to EAutomationTestFlags map */
+	const TMap<FString, uint32> FStringToEAutomationTestFlagMap = {
+		{ TEXT("EditorContext"), EAutomationTestFlags::EditorContext},
+		{ TEXT("ClientContext"), EAutomationTestFlags::ClientContext},
+		{ TEXT("ServerContext"), EAutomationTestFlags::ServerContext},
+		{ TEXT("CommandletContext"), EAutomationTestFlags::CommandletContext},
+		{ TEXT("ApplicationContextMask"), EAutomationTestFlags::ApplicationContextMask},
+		{ TEXT("NonNullRHI"), EAutomationTestFlags::NonNullRHI},
+		{ TEXT("RequiresUser"), EAutomationTestFlags::RequiresUser},
+		{ TEXT("FeatureMask"), EAutomationTestFlags::FeatureMask},
+		{ TEXT("Disabled"), EAutomationTestFlags::Disabled},
+		{ TEXT("CriticalPriority"), EAutomationTestFlags::CriticalPriority},
+		{ TEXT("HighPriority"), EAutomationTestFlags::HighPriority},
+		{ TEXT("HighPriorityAndAbove"), EAutomationTestFlags::HighPriorityAndAbove},
+		{ TEXT("MediumPriority"), EAutomationTestFlags::MediumPriority},
+		{ TEXT("MediumPriorityAndAbove"), EAutomationTestFlags::MediumPriorityAndAbove},
+		{ TEXT("LowPriority"), EAutomationTestFlags::LowPriority},
+		{ TEXT("PriorityMask"), EAutomationTestFlags::PriorityMask},
+		{ TEXT("SmokeFilter"), EAutomationTestFlags::SmokeFilter},
+		{ TEXT("EngineFilter"), EAutomationTestFlags::EngineFilter},
+		{ TEXT("ProductFilter"), EAutomationTestFlags::ProductFilter},
+		{ TEXT("PerfFilter"), EAutomationTestFlags::PerfFilter},
+		{ TEXT("SmokeFilter"), EAutomationTestFlags::SmokeFilter},
+		{ TEXT("StressFilter"), EAutomationTestFlags::StressFilter},
+		{ TEXT("NegativeFilter"), EAutomationTestFlags::NegativeFilter},
+		{ TEXT("FilterMask"), EAutomationTestFlags::FilterMask}
+	};
+
 protected:
 
 	//Flag to indicate if this is a complex task
@@ -3086,6 +3120,50 @@ public: \
 		virtual FString GetBeautifiedTestName() const override { return PrettyName; } \
 	};
 
+/**
+* This macro works the same way IMPLEMENT_SIMPLE_AUTOMATION_TEST_PRIVATE does with the following specifications:
+*	- TFalgs is expected to be declared as a string and it can include extra test tags along with EAutomationTestFlags flags
+		e.g. "[Flag_1]...[Flag_n][Tag_1][Tag_2]...[Tag_m]". The flags will be converted to an appropriate combined EAutomationTestFlags value
+*	- PrettyName is declared using double column notation "Namespace_1::Namespace_2:: ... ::Namespace_n::TestName"
+*		and it is converted to its dot notation "Namespace_1.Namespace_2. ... .Namespace_n.TestName"
+*	- It is mainly used when declaring tests that can run both as Low Level Tests with Catch2 or Functional Tests with FAutomationTestFramework
+*/
+#define IMPLEMENT_SIMPLE_AUTOMATION_TEST_PRIVATE_LLT( TClass, TBaseClass, PrettyName, TFlags, FileName, LineNumber ) \
+	class TClass : public TBaseClass \
+	{ \
+	public: \
+		TClass( const FString& InName ) \
+		:TBaseClass( InName, false ) {\
+			TestFlags = ExtractAutomationTestFlags(TFlags);\
+			PrettyNameDotNotation = PrettyName.Replace(TEXT("::"), TEXT("."));\
+			assert(TestFlags & EAutomationTestFlags::ApplicationContextMask, "AutomationTest has no application flag.  It shouldn't run.  See AutomationTest.h."); \
+			assert(	((TestFlags & EAutomationTestFlags::FilterMask) == EAutomationTestFlags::SmokeFilter) || \
+					((TestFlags & EAutomationTestFlags::FilterMask) == EAutomationTestFlags::EngineFilter) || \
+					((TestFlags & EAutomationTestFlags::FilterMask) == EAutomationTestFlags::ProductFilter) || \
+					((TestFlags & EAutomationTestFlags::FilterMask) == EAutomationTestFlags::PerfFilter) || \
+					((TestFlags & EAutomationTestFlags::FilterMask) == EAutomationTestFlags::StressFilter) || \
+					((TestFlags & EAutomationTestFlags::FilterMask) == EAutomationTestFlags::NegativeFilter), \
+					"All AutomationTests must have exactly 1 filter type specified.  See AutomationTest.h."); \
+		} \
+		virtual uint32 GetTestFlags() const override { return TestFlags; } \
+		virtual bool IsStressTest() const { return false; } \
+		virtual uint32 GetRequiredDeviceNum() const override { return 1; } \
+		virtual FString GetTestSourceFileName() const override { return FileName; } \
+		virtual int32 GetTestSourceFileLine() const override { return LineNumber; } \
+	protected: \
+		virtual void GetTests(TArray<FString>& OutBeautifiedNames, TArray <FString>& OutTestCommands) const override \
+		{ \
+			OutBeautifiedNames.Add(PrettyNameDotNotation)); \
+			OutTestCommands.Add(FString()); \
+		} \
+		virtual bool RunTest(const FString& Parameters) override; \
+		virtual FString GetBeautifiedTestName() const override { return PrettyNameDotNotation; } \
+	private: \
+		uint32 TestFlags;\
+		FString PrettyNameDotNotation;\
+	};
+
+
 #define IMPLEMENT_COMPLEX_AUTOMATION_TEST_PRIVATE( TClass, TBaseClass, PrettyName, TFlags, FileName, LineNumber ) \
 	class TClass : public TBaseClass \
 	{ \
diff --git a/Engine/Source/Runtime/Core/Tests/Compression/CompressedBufferTest.cpp b/Engine/Source/Runtime/Core/Tests/Compression/CompressedBufferTest.cpp
--- a/Engine/Source/Runtime/Core/Tests/Compression/CompressedBufferTest.cpp
+++ b/Engine/Source/Runtime/Core/Tests/Compression/CompressedBufferTest.cpp
@@ -1,5 +1,7 @@
 // Copyright Epic Games, Inc. All Rights Reserved.
 
+#if WITH_LOW_LEVEL_TESTS
+
 #include "Compression/CompressedBuffer.h"
 
 #include "Algo/Compare.h"
@@ -241,3 +243,5 @@ TEST_CASE("Core::Compression::CompressedBuffer::Decompress", "[Core][Compression
 		}
 	}
 }
+
+#endif
\ No newline at end of file
diff --git a/Engine/Source/Runtime/Core/Tests/Containers/DequeTest.cpp b/Engine/Source/Runtime/Core/Tests/Containers/DequeTest.cpp
--- a/Engine/Source/Runtime/Core/Tests/Containers/DequeTest.cpp
+++ b/Engine/Source/Runtime/Core/Tests/Containers/DequeTest.cpp
@@ -1,5 +1,7 @@
 // Copyright Epic Games, Inc. All Rights Reserved.
 
+#if WITH_LOW_LEVEL_TESTS
+
 #include "Containers/Deque.h"
 #include "CoreMinimal.h"
 #include "TokenTest.h"
@@ -928,3 +930,5 @@ TEST_CASE("Core::Containers::TDeque::Assign from std initializer_list", "[Core][
 }
 
 }  // namespace Deque
+
+#endif
diff --git a/Engine/Source/Runtime/Core/Tests/Containers/TokenTest.cpp b/Engine/Source/Runtime/Core/Tests/Containers/TokenTest.cpp
--- a/Engine/Source/Runtime/Core/Tests/Containers/TokenTest.cpp
+++ b/Engine/Source/Runtime/Core/Tests/Containers/TokenTest.cpp
@@ -1,5 +1,7 @@
 // Copyright Epic Games, Inc. All Rights Reserved.
 
+#if WITH_LOW_LEVEL_TESTS
+
 #include "TokenTest.h"
 #include "TestHarness.h"
 
@@ -108,3 +110,5 @@ TEST_CASE("Core::Containers::TToken::Move assignment", "[Core][Containers][TToke
 	CHECK(int32Token::NumCopyAssignmentCalls() == 0);
 	CHECK(int32Token::NumMoveAssignmentCalls() == 1);
 }
+
+#endif
\ No newline at end of file
diff --git a/Engine/Source/Runtime/Core/Tests/Hash/Blake3Test.cpp b/Engine/Source/Runtime/Core/Tests/Hash/Blake3Test.cpp
--- a/Engine/Source/Runtime/Core/Tests/Hash/Blake3Test.cpp
+++ b/Engine/Source/Runtime/Core/Tests/Hash/Blake3Test.cpp
@@ -1,5 +1,7 @@
 // Copyright Epic Games, Inc. All Rights Reserved.
 
+#if WITH_LOW_LEVEL_TESTS
+
 #include "Hash/Blake3.h"
 
 #include "Misc/StringBuilder.h"
@@ -59,3 +61,5 @@ TEST_CASE("Core::Hash::Blake3", "[Core][Hash][Smoke]")
 		CHECK(HashHex.ToView() == TestCase.Hash);
 	}
 }
+
+#endif
\ No newline at end of file
diff --git a/Engine/Source/Runtime/Core/Tests/Math/ColorTests.cpp b/Engine/Source/Runtime/Core/Tests/Math/ColorTests.cpp
--- a/Engine/Source/Runtime/Core/Tests/Math/ColorTests.cpp
+++ b/Engine/Source/Runtime/Core/Tests/Math/ColorTests.cpp
@@ -1,4 +1,7 @@
 // Copyright Epic Games, Inc. All Rights Reserved.
+
+#if WITH_LOW_LEVEL_TESTS
+
 #include "CoreMinimal.h"
 #include "Math/Color.h"
 #include "TestHarness.h"
@@ -63,3 +66,5 @@ TEST_CASE("Core::Math::FColor::Smoke Test", "[Core][Math][Smoke]")
 		CHECK(1 == 1);
 	}
 }
+
+#endif
\ No newline at end of file
diff --git a/Engine/Source/Runtime/Core/Tests/Misc/StringBuilderTest.cpp b/Engine/Source/Runtime/Core/Tests/Misc/StringBuilderTest.cpp
--- a/Engine/Source/Runtime/Core/Tests/Misc/StringBuilderTest.cpp
+++ b/Engine/Source/Runtime/Core/Tests/Misc/StringBuilderTest.cpp
@@ -1,5 +1,7 @@
 // Copyright Epic Games, Inc. All Rights Reserved.
 
+#if WITH_LOW_LEVEL_TESTS
+
 #include "Misc/StringBuilder.h"
 
 #include "Containers/UnrealString.h"
@@ -279,3 +281,5 @@ TEST_CASE("Core::String::StringBuilder", "[Core][String][Smoke]")
 		CHECK(Builder == ANSITEXTVIEW(""));
 	}
 }
+
+#endif
\ No newline at end of file
diff --git a/Engine/Source/Runtime/Core/Tests/Serialization/VarIntTest.cpp b/Engine/Source/Runtime/Core/Tests/Serialization/VarIntTest.cpp
--- a/Engine/Source/Runtime/Core/Tests/Serialization/VarIntTest.cpp
+++ b/Engine/Source/Runtime/Core/Tests/Serialization/VarIntTest.cpp
@@ -1,5 +1,7 @@
 // Copyright Epic Games, Inc. All Rights Reserved.
 
+#if WITH_LOW_LEVEL_TESTS
+
 #include "Serialization/VarInt.h"
 
 #include "Serialization/BufferReader.h"
@@ -278,3 +280,5 @@ TEST_CASE("Core::Serialization::VarInt::Serialize", "[Core][Serialization][Smoke
 		CHECK(ReadAr.Tell() == ReadByteCount);
 	}
 }
+
+#endif
\ No newline at end of file
diff --git a/Engine/Source/Runtime/Core/Tests/String/EscapeTest.cpp b/Engine/Source/Runtime/Core/Tests/String/EscapeTest.cpp
--- a/Engine/Source/Runtime/Core/Tests/String/EscapeTest.cpp
+++ b/Engine/Source/Runtime/Core/Tests/String/EscapeTest.cpp
@@ -1,5 +1,7 @@
 // Copyright Epic Games, Inc. All Rights Reserved.
 
+#if WITH_LOW_LEVEL_TESTS
+
 #include "String/Escape.h"
 
 #include "Containers/StringView.h"
@@ -35,3 +37,5 @@ TEST_CASE("Core::String::Escape", "[Core][String][Smoke]")
 }
 
 } // UE::String
+
+#endif
\ No newline at end of file
diff --git a/Engine/Source/Runtime/Core/Tests/String/FindTest.cpp b/Engine/Source/Runtime/Core/Tests/String/FindTest.cpp
--- a/Engine/Source/Runtime/Core/Tests/String/FindTest.cpp
+++ b/Engine/Source/Runtime/Core/Tests/String/FindTest.cpp
@@ -1,5 +1,7 @@
 // Copyright Epic Games, Inc. All Rights Reserved.
 
+#if WITH_LOW_LEVEL_TESTS
+
 #include "String/Find.h"
 
 #include "Containers/StringView.h"
@@ -222,3 +224,5 @@ TEST_CASE("Core::String::Find", "[Core][String][Smoke]")
 }
 
 } // UE::String
+
+#endif
\ No newline at end of file
diff --git a/Engine/Source/Runtime/Core/Tests/String/JoinTest.cpp b/Engine/Source/Runtime/Core/Tests/String/JoinTest.cpp
--- a/Engine/Source/Runtime/Core/Tests/String/JoinTest.cpp
+++ b/Engine/Source/Runtime/Core/Tests/String/JoinTest.cpp
@@ -1,5 +1,7 @@
 // Copyright Epic Games, Inc. All Rights Reserved.
 
+#if WITH_LOW_LEVEL_TESTS
+
 #include "String/Join.h"
 
 #include "Containers/StringView.h"
@@ -36,3 +38,5 @@ TEST_CASE("Core::String::Join", "[Core][String][Smoke]")
 }
 
 } // UE::String
+
+#endif
\ No newline at end of file
diff --git a/Engine/Source/Runtime/Core/Tests/String/ParseLinesTest.cpp b/Engine/Source/Runtime/Core/Tests/String/ParseLinesTest.cpp
--- a/Engine/Source/Runtime/Core/Tests/String/ParseLinesTest.cpp
+++ b/Engine/Source/Runtime/Core/Tests/String/ParseLinesTest.cpp
@@ -1,5 +1,7 @@
 // Copyright Epic Games, Inc. All Rights Reserved.
 
+#if WITH_LOW_LEVEL_TESTS
+
 #include "String/ParseLines.h"
 
 #include "Algo/Compare.h"
@@ -80,3 +82,5 @@ TEST_CASE("Core::String::ParseLines", "[Core][String][Smoke]")
 }
 
 } // UE::String
+
+#endif
\ No newline at end of file
diff --git a/Engine/Source/Runtime/Core/Tests/String/ParseTokensTest.cpp b/Engine/Source/Runtime/Core/Tests/String/ParseTokensTest.cpp
--- a/Engine/Source/Runtime/Core/Tests/String/ParseTokensTest.cpp
+++ b/Engine/Source/Runtime/Core/Tests/String/ParseTokensTest.cpp
@@ -1,5 +1,7 @@
 // Copyright Epic Games, Inc. All Rights Reserved.
 
+#if WITH_LOW_LEVEL_TESTS
+
 #include "String/ParseTokens.h"
 
 #include "Algo/Compare.h"
@@ -162,3 +164,5 @@ TEST_CASE("Core::String::ParseTokens", "[Core][String][Smoke]")
 }
 
 } // UE::String
+
+#endif
\ No newline at end of file
diff --git a/Engine/Source/Runtime/Core/Tests/Templates/ValueOrErrorTest.cpp b/Engine/Source/Runtime/Core/Tests/Templates/ValueOrErrorTest.cpp
--- a/Engine/Source/Runtime/Core/Tests/Templates/ValueOrErrorTest.cpp
+++ b/Engine/Source/Runtime/Core/Tests/Templates/ValueOrErrorTest.cpp
@@ -1,5 +1,7 @@
 // Copyright Epic Games, Inc. All Rights Reserved.
 
+#if WITH_LOW_LEVEL_TESTS
+
 #include "Templates/ValueOrError.h"
 
 #include "TestHarness.h"
@@ -229,3 +231,5 @@ TEST_CASE("Core::Templates::ValueOrError", "[Core][Templates][Smoke]")
 	CHECK(ValueCount == 0);
 	CHECK(ErrorCount == 0);
 }
+
+#endif
\ No newline at end of file
diff --git a/Engine/Source/Runtime/CoreUObject/Tests/ObjectHandleTest.cpp b/Engine/Source/Runtime/CoreUObject/Tests/ObjectHandleTest.cpp
--- a/Engine/Source/Runtime/CoreUObject/Tests/ObjectHandleTest.cpp
+++ b/Engine/Source/Runtime/CoreUObject/Tests/ObjectHandleTest.cpp
@@ -1,5 +1,7 @@
 // Copyright Epic Games, Inc. All Rights Reserved.
 
+#if WITH_LOW_LEVEL_TESTS
+
 #include "UObject/ObjectHandle.h"
 
 #include "HAL/PlatformProperties.h"
@@ -195,3 +197,5 @@ TEST_CASE_METHOD(FObjectHandleTestBase, "CoreUObject::FObjectHandle::Resolve Mal
 	TestResolveFailure(FPackedObjectRef { 0xEFEF'EFEF'EFEF'EFEFull });
 }
 #endif // UE_WITH_OBJECT_HANDLE_LATE_RESOLVE
+
+#endif
\ No newline at end of file
diff --git a/Engine/Source/Runtime/CoreUObject/Tests/ObjectPtrTest.cpp b/Engine/Source/Runtime/CoreUObject/Tests/ObjectPtrTest.cpp
--- a/Engine/Source/Runtime/CoreUObject/Tests/ObjectPtrTest.cpp
+++ b/Engine/Source/Runtime/CoreUObject/Tests/ObjectPtrTest.cpp
@@ -1,5 +1,7 @@
 // Copyright Epic Games, Inc. All Rights Reserved.
 
+#if WITH_LOW_LEVEL_TESTS
+
 #include "UObject/ObjectPtr.h"
 
 #include "ObjectRefTrackingTestBase.h"
@@ -361,3 +363,5 @@ class UForwardDeclaredObjDerived: public UObject {};
 class FForwardDeclaredNotObjDerived {};
 
 }
+
+#endif
\ No newline at end of file
diff --git a/Engine/Source/Runtime/CoreUObject/Tests/ObjectRefTrackingTestBase.cpp b/Engine/Source/Runtime/CoreUObject/Tests/ObjectRefTrackingTestBase.cpp
--- a/Engine/Source/Runtime/CoreUObject/Tests/ObjectRefTrackingTestBase.cpp
+++ b/Engine/Source/Runtime/CoreUObject/Tests/ObjectRefTrackingTestBase.cpp
@@ -1,5 +1,7 @@
 // Copyright Epic Games, Inc. All Rights Reserved.
 
+#if WITH_LOW_LEVEL_TESTS
+
 #include "ObjectRefTrackingTestBase.h"
 
 #if UE_WITH_OBJECT_HANDLE_TRACKING
@@ -8,4 +10,6 @@ ObjectHandleReadFunction* FObjectRefTrackingTestBase::PrevReadFunc = nullptr;
 #endif
 thread_local uint32 FObjectRefTrackingTestBase::NumResolves = 0;
 thread_local uint32 FObjectRefTrackingTestBase::NumFailedResolves = 0;
-thread_local uint32 FObjectRefTrackingTestBase::NumReads = 0;
\ No newline at end of file
+thread_local uint32 FObjectRefTrackingTestBase::NumReads = 0;
+
+#endif
\ No newline at end of file
diff --git a/Engine/Source/Runtime/CoreUObject/Tests/ObjectRefTrackingTestBase.h b/Engine/Source/Runtime/CoreUObject/Tests/ObjectRefTrackingTestBase.h
--- a/Engine/Source/Runtime/CoreUObject/Tests/ObjectRefTrackingTestBase.h
+++ b/Engine/Source/Runtime/CoreUObject/Tests/ObjectRefTrackingTestBase.h
@@ -2,6 +2,8 @@
 
 #pragma once
 
+#if WITH_LOW_LEVEL_TESTS
+
 #include "TestHarness.h"
 
 #include "CoreMinimal.h"
@@ -129,3 +131,4 @@ private:
 	static thread_local uint32 NumReads;
 };
 
+#endif
\ No newline at end of file
diff --git a/Engine/Source/Runtime/CoreUObject/Tests/Serialization/BulkDataTests.cpp b/Engine/Source/Runtime/CoreUObject/Tests/Serialization/BulkDataTests.cpp
--- a/Engine/Source/Runtime/CoreUObject/Tests/Serialization/BulkDataTests.cpp
+++ b/Engine/Source/Runtime/CoreUObject/Tests/Serialization/BulkDataTests.cpp
@@ -1,5 +1,7 @@
 // Copyright Epic Games, Inc. All Rights Reserved.
 
+#if WITH_LOW_LEVEL_TESTS
+
 #include "TestHarness.h"
 #include "Serialization/BulkData.h"
 
@@ -75,3 +77,4 @@ namespace BulkDataTest
 	}
 }
 
+#endif
\ No newline at end of file
