PR #9096: Fix missing includes rider (Contributed by [REDACTED_PERSON])
#tests Tested in Rider

diff --git a/Engine/Source/Programs/UnrealBuildTool/ProjectFiles/Rider/RiderProjectFile.cs b/Engine/Source/Programs/UnrealBuildTool/ProjectFiles/Rider/RiderProjectFile.cs
--- a/Engine/Source/Programs/UnrealBuildTool/ProjectFiles/Rider/RiderProjectFile.cs
+++ b/Engine/Source/Programs/UnrealBuildTool/ProjectFiles/Rider/RiderProjectFile.cs
@@ -192,13 +192,14 @@ namespace UnrealBuildTool
 					if(ModuleNames.Add(Module.Name))
 					{
 						Writer.WriteObjectStart(Module.Name);
-						ExportModule(Module, Binary.OutputDir, Target.GetExecutableDir(), Writer);
 						UEBuildModuleCPP? ModuleCpp = Module as UEBuildModuleCPP;
 						if (ModuleCpp != null)
 						{
 							CppCompileEnvironment ModuleCompileEnvironment = ModuleCpp.CreateCompileEnvironmentForIntellisense(Target.Rules, BinaryCompileEnvironment);
 							ExportModuleCpp(ModuleCpp, ModuleCompileEnvironment, Writer);
+							Module.PrivateIncludePaths.UnionWith(ModuleCompileEnvironment.UserIncludePaths);
 						}
+						ExportModule(Module, Binary.OutputDir, Target.GetExecutableDir(), Writer);
 						Writer.WriteObjectEnd();
 					}
 				}
