UnrealBuildTool: Retry failed HordeExecute remote tasks. Additional logging

#rnx

diff --git a/Engine/Source/Programs/UnrealBuildTool/Executors/Experimental/HordeExecutor.cs b/Engine/Source/Programs/UnrealBuildTool/Executors/Experimental/HordeExecutor.cs
--- a/Engine/Source/Programs/UnrealBuildTool/Executors/Experimental/HordeExecutor.cs
+++ b/Engine/Source/Programs/UnrealBuildTool/Executors/Experimental/HordeExecutor.cs
@@ -31,7 +31,6 @@ namespace UnrealBuildTool
 	/// * Investigate overhead when connecting from Zen to Horde
 	/// * Local process racing when remote tasks are blocking more tasks from starting, or at the end of a build.
 	/// * Print remote status as tasks complete, print more concise and useful summary at the end of the build
-	/// * Better error handling when remote tasks fail, locally retry failed actions
 	/// * Additional toolchain platform support
 	/// * Dynamic control over the number of remote processes to run at once. Needs knowledge about availibility of remote compute resources, if the pending remote queue is overloaded and by how much, etc.
 	/// </summary>
@@ -49,6 +48,12 @@ namespace UnrealBuildTool
 		[XmlConfigFile]
 		private static bool RemoteProcessOnly = false;
 
+		/// <summary>
+		/// Retry actions locally when failed remotely.
+		/// </summary>
+		[XmlConfigFile]
+		private static bool RetryFailedRemote = true;
+
 		private static readonly Lazy<HttpClient> LazyZenHttpClient = new Lazy<HttpClient>();
 		private static HttpClient ZenHttpClient { get { return LazyZenHttpClient.Value; } }
 
@@ -252,6 +257,13 @@ namespace UnrealBuildTool
 			public RemoteExecuteResults(List<string> LogLines, int ExitCode, TimeSpan ExecutionTime, TimeSpan ProcessorTime, Dictionary<string, DateTime> Timepoints) : base(LogLines, ExitCode, ExecutionTime, ProcessorTime)
 			{
 				this.Timepoints = Timepoints;
+				AdditionalDescription = "Remote";
+			}
+
+			public RemoteExecuteResults(List<string> LogLines, int ExitCode) : base(LogLines, ExitCode)
+			{
+				Timepoints = new Dictionary<string, DateTime>();
+				AdditionalDescription = "Remote";
 			}
 		}
 
@@ -334,13 +346,9 @@ namespace UnrealBuildTool
 				LogTasks.Add(LogTask);
 			}
 
-			Task SummaryTask = Task.Factory.ContinueWhenAll(LogTasks.ToArray(), (AntecedentTasks) => TraceSummary(ExecuteTasks, ProcessGroup), CancellationToken);
-			SummaryTask = SummaryTask.ContinueWith(antecedent => TraceRemoteSummary(ExecuteTasks));
-			while (!SummaryTask.IsCompleted)
-			{
-				SummaryTask.Wait(1000);
-				Log.TraceInformation($"Actions Running Local:{ActualNumParallelProcesses - MaxProcessSemaphore.CurrentCount} Remote:{NumRemoteParallelProcesses - MaxRemoteProcessSemaphore.CurrentCount}");
-			}
+			Task.Factory.ContinueWhenAll(LogTasks.ToArray(), (AntecedentTasks) => TraceSummary(ExecuteTasks, ProcessGroup), CancellationToken)
+				.ContinueWith(antecedent => TraceRemoteSummary(ExecuteTasks))
+				.Wait();
 
 			// Return if all tasks succeeded
 			return ExecuteTasks.Values.All(x => x.Result.ExitCode == 0);
@@ -376,40 +384,64 @@ namespace UnrealBuildTool
 			TraceStats(RemoteExecuteResults, "horde prepare sandbox", "horde-execution-start", "horde-execution-download-input");
 			TraceStats(RemoteExecuteResults, "horde execute", "horde-execution-download-input", "horde-execution-execute");
 			TraceStats(RemoteExecuteResults, "horde upload results", "horde-execution-execute", "horde-execution-upload-ref");
-			TraceStats(RemoteExecuteResults, "zen complete", "horde-queue-complete", "zen-queue-complete");
+			TraceStats(RemoteExecuteResults, "zen queue results wait", "zen-complete-queue-added", "zen-complete-queue-dispatched");
+			TraceStats(RemoteExecuteResults, "zen get horde", "zen-complete-queue-dispatched", "zen-storage-get-blobs");
+			TraceStats(RemoteExecuteResults, "zen finalize", "zen-storage-get-blobs", "zen-queue-complete");
 			Log.TraceInformation("");
 		}
 
 		private static Task<ExecuteResults> CreateRemoteExecuteTask(LinkedAction Action, Dictionary<LinkedAction, Task<ExecuteResults>> ExecuteTasks, ManagedProcessGroup ProcessGroup, SemaphoreSlim MaxProcessSemaphore, SemaphoreSlim MaxRemoteProcessSemaphore, CancellationToken CancellationToken)
 		{
+			Task<ExecuteResults> ActionTask;
 			if (Action.PrerequisiteActions.Count == 0)
 			{
-				return Task.Factory.StartNew(
+				ActionTask = Task.Factory.StartNew(
 					() => ExecuteRemoteAction(Array.Empty<Task<ExecuteResults>>(), Action, ProcessGroup, MaxProcessSemaphore, MaxRemoteProcessSemaphore, CancellationToken),
 					CancellationToken,
-					TaskCreationOptions.LongRunning,
+					TaskCreationOptions.LongRunning | TaskCreationOptions.PreferFairness,
 					TaskScheduler.Current
 				).Unwrap();
 			}
-
-			// Create tasks for any preresquite actions if they don't exist already
-			List<Task<ExecuteResults>> PrerequisiteTasks = new List<Task<ExecuteResults>>();
-			foreach (var PrerequisiteAction in Action.PrerequisiteActions)
+			else
 			{
-				if (!ExecuteTasks.ContainsKey(PrerequisiteAction))
+
+				// Create tasks for any preresquite actions if they don't exist already
+				List<Task<ExecuteResults>> PrerequisiteTasks = new List<Task<ExecuteResults>>();
+				foreach (var PrerequisiteAction in Action.PrerequisiteActions)
 				{
-					ExecuteTasks.Add(PrerequisiteAction, CreateRemoteExecuteTask(PrerequisiteAction, ExecuteTasks, ProcessGroup, MaxProcessSemaphore, MaxRemoteProcessSemaphore, CancellationToken));
+					if (!ExecuteTasks.ContainsKey(PrerequisiteAction))
+					{
+						ExecuteTasks.Add(PrerequisiteAction, CreateRemoteExecuteTask(PrerequisiteAction, ExecuteTasks, ProcessGroup, MaxProcessSemaphore, MaxRemoteProcessSemaphore, CancellationToken));
+					}
+					PrerequisiteTasks.Add(ExecuteTasks[PrerequisiteAction]);
 				}
-				PrerequisiteTasks.Add(ExecuteTasks[PrerequisiteAction]);
+
+				ActionTask = Task.Factory.ContinueWhenAll(
+					PrerequisiteTasks.ToArray(),
+					(AntecedentTasks) => ExecuteRemoteAction(AntecedentTasks, Action, ProcessGroup, MaxProcessSemaphore, MaxRemoteProcessSemaphore, CancellationToken),
+					CancellationToken,
+					TaskContinuationOptions.LongRunning | TaskContinuationOptions.PreferFairness,
+					TaskScheduler.Current
+				).Unwrap();
 			}
 
-			return Task.Factory.ContinueWhenAll(
-				PrerequisiteTasks.ToArray(),
-				(AntecedentTasks) => ExecuteRemoteAction(AntecedentTasks, Action, ProcessGroup, MaxProcessSemaphore, MaxRemoteProcessSemaphore, CancellationToken),
+			if (RetryFailedRemote)
+			{
+				return ActionTask.ContinueWith(ancendent =>
+				{
+					if (ancendent.IsCanceled || ancendent.Result.ExitCode == 0 || ancendent.Result is not RemoteExecuteResults)
+					{
+						return Task.FromResult(ancendent.Result);
+					}
+					Log.TraceInformation($"Remote execute for {(Action.CommandDescription ?? Action.CommandPath.GetFileNameWithoutExtension())} {Action.StatusDescription} failed, rescheduling for local execution.");
+					return ExecuteAction(Array.Empty<Task<ExecuteResults>>(), Action, ProcessGroup, MaxProcessSemaphore, CancellationToken);
+				},
 				CancellationToken,
-				TaskContinuationOptions.LongRunning,
-				TaskScheduler.Current
-			).Unwrap();
+				TaskContinuationOptions.LongRunning | TaskContinuationOptions.PreferFairness,
+				TaskScheduler.Current).Unwrap();
+			}
+
+			return ActionTask;
 		}
 
 		private static HashSet<FileReference>? GatherDependencies(FileItem DependencyListFile, IEnumerable<FileItem> PrerequisiteItems)
@@ -511,6 +543,14 @@ namespace UnrealBuildTool
 			}
 			catch (Exception Ex)
 			{
+				if (RemoteSemaphoreTask?.Status == TaskStatus.RanToCompletion && RemoteSemaphoreTask?.Result == true)
+				{
+					if (!RetryFailedRemote)
+                    {
+						Log.WriteException(Ex, null);
+                    }
+					return new RemoteExecuteResults(new List<string>(), int.MaxValue);
+				}
 				Log.WriteException(Ex, null);
 				return new ExecuteResults(new List<string>(), int.MaxValue);
 			}
@@ -528,7 +568,7 @@ namespace UnrealBuildTool
 			}
 		}
 
-		private static async Task<ExecuteResults> RunRemoteAction(LinkedAction Action, HashSet<FileReference> DependencyItems, CancellationToken CancellationToken)
+		private static async Task<RemoteExecuteResults> RunRemoteAction(LinkedAction Action, HashSet<FileReference> DependencyItems, CancellationToken CancellationToken)
 		{
 			Dictionary<string, DateTime> Timepoints = new Dictionary<string, DateTime>();
 			Timepoints["ubt-queue-dispatched"] = DateTime.UtcNow;
diff --git a/Engine/Source/Programs/UnrealBuildTool/Executors/ParallelExecutor.cs b/Engine/Source/Programs/UnrealBuildTool/Executors/ParallelExecutor.cs
--- a/Engine/Source/Programs/UnrealBuildTool/Executors/ParallelExecutor.cs
+++ b/Engine/Source/Programs/UnrealBuildTool/Executors/ParallelExecutor.cs
@@ -130,6 +130,7 @@ namespace UnrealBuildTool
 			public int ExitCode { get; private set; }
 			public TimeSpan ExecutionTime { get; private set; }
 			public TimeSpan ProcessorTime { get; private set; }
+			public string? AdditionalDescription { get; protected set; } = null;
 
 			public ExecuteResults(List<string> LogLines, int ExitCode, TimeSpan ExecutionTime, TimeSpan ProcessorTime)
 			{
@@ -295,6 +296,7 @@ namespace UnrealBuildTool
 			int ExitCode = int.MaxValue;
 			TimeSpan ExecutionTime = TimeSpan.Zero;
 			TimeSpan ProcessorTime = TimeSpan.Zero;
+			string? AdditionalDescription = null;
 			if (ExecuteTask.Status == TaskStatus.RanToCompletion)
 			{
 				ExecuteResults ExecuteTaskResult = ExecuteTask.Result;
@@ -302,6 +304,7 @@ namespace UnrealBuildTool
 				ExitCode = ExecuteTaskResult.ExitCode;
 				ExecutionTime = ExecuteTaskResult.ExecutionTime;
 				ProcessorTime = ExecuteTaskResult.ProcessorTime;
+				AdditionalDescription = ExecuteTaskResult.AdditionalDescription;
 			}
 
 			// Write it to the log
@@ -314,6 +317,10 @@ namespace UnrealBuildTool
 			{
 				Description = $"{(Action.CommandDescription ?? Action.CommandPath.GetFileNameWithoutExtension())} {LogLines[0]}".Trim();
 			}
+			if (!string.IsNullOrEmpty(AdditionalDescription))
+            {
+				Description = $"{AdditionalDescription} {Description}";
+			}
 
 			lock (ProgressWriter)
 			{
