VCToolchain and linker visibility changes required to enable modular builds for MS console platforms

diff --git a/Engine/Source/Programs/UnrealBuildTool/Platform/Windows/VCToolChain.cs b/Engine/Source/Programs/UnrealBuildTool/Platform/Windows/VCToolChain.cs
--- a/Engine/Source/Programs/UnrealBuildTool/Platform/Windows/VCToolChain.cs
+++ b/Engine/Source/Programs/UnrealBuildTool/Platform/Windows/VCToolChain.cs
@@ -1039,11 +1039,11 @@ namespace UnrealBuildTool
 
 				// Allow delay-loaded DLLs to be explicitly unloaded.
 				Arguments.Add("/DELAY:UNLOAD");
+			}
 
-				if (LinkEnvironment.bIsBuildingDLL)
-				{
-					Arguments.Add("/DLL");
-				}
+			if (LinkEnvironment.bIsBuildingDLL)
+			{
+				Arguments.Add("/DLL");
 			}
 
 			// Don't embed the full PDB path; we want to be able to move binaries elsewhere. They will always be side by side.
