Fix iOS code signing with certificates issued for organizations or individuals with non-latin characters in their names.

diff --git a/Engine/Source/Programs/IOS/iPhonePackager/CodeSigning.cs b/Engine/Source/Programs/IOS/iPhonePackager/CodeSigning.cs
--- a/Engine/Source/Programs/IOS/iPhonePackager/CodeSigning.cs
+++ b/Engine/Source/Programs/IOS/iPhonePackager/CodeSigning.cs
@@ -236,6 +236,7 @@ namespace iPhonePackager
 					CertTool.StartInfo.UseShellExecute = false;
 					CertTool.StartInfo.Arguments = string.Format("find-certificate -a -c \"{0}\" -p", SearchPrefix);
 					CertTool.StartInfo.RedirectStandardOutput = true;
+					CertTool.StartInfo.StandardOutputEncoding = Encoding.UTF8;
 					CertTool.OutputDataReceived += new DataReceivedEventHandler(OutputReceivedCertToolProcessCall);
 					CertTool.Start();
 					CertTool.BeginOutputReadLine();
@@ -418,6 +419,7 @@ namespace iPhonePackager
 				CertTool.StartInfo.UseShellExecute = false;
 				CertTool.StartInfo.Arguments = "find-identity -p codesigning -v";
 				CertTool.StartInfo.RedirectStandardOutput = true;
+				CertTool.StartInfo.StandardOutputEncoding = Encoding.UTF8;
 				CertTool.OutputDataReceived += new DataReceivedEventHandler (OutputReceivedCertToolProcessCall);
 				CertTool.Start ();
 				CertTool.BeginOutputReadLine ();
diff --git a/Engine/Source/Programs/IOS/iPhonePackager/Utilities.cs b/Engine/Source/Programs/IOS/iPhonePackager/Utilities.cs
--- a/Engine/Source/Programs/IOS/iPhonePackager/Utilities.cs
+++ b/Engine/Source/Programs/IOS/iPhonePackager/Utilities.cs
@@ -957,6 +957,10 @@ namespace iPhonePackager
 			CertChain[0] = new X509CertificateEntry(BouncyCert);
 
 			AsymmetricCipherKeyPair KeyPair = LoadKeyPairFromDiskBouncy(KeyFilename);
+			if (KeyPair == null || KeyPair.Private == null)
+			{
+				throw new InvalidDataException("The key pair provided does contain a private key.  Make sure you provide the same key pair that was used to generate the original certificate signing request");
+			}
 
 			Store.SetKeyEntry(FriendlyName, new AsymmetricKeyEntry(KeyPair.Private), CertChain);
 
diff --git a/Engine/Source/Programs/IOS/iPhonePackager/iPhonePackager.cs b/Engine/Source/Programs/IOS/iPhonePackager/iPhonePackager.cs
--- a/Engine/Source/Programs/IOS/iPhonePackager/iPhonePackager.cs
+++ b/Engine/Source/Programs/IOS/iPhonePackager/iPhonePackager.cs
@@ -790,6 +790,9 @@ namespace iPhonePackager
 
             AppDomain.CurrentDomain.AssemblyResolve += new ResolveEventHandler(CurrentDomain_AssemblyResolve);
 
+			// Force UTF-8 encoding so UBT and the engine can parse the IPP output consistently.
+			Console.OutputEncoding = Encoding.UTF8;
+
             // A simple, top-level try-catch block
             try
             {
diff --git a/Engine/Source/Programs/UnrealBuildTool/Platform/IOS/UEBuildIOS.cs b/Engine/Source/Programs/UnrealBuildTool/Platform/IOS/UEBuildIOS.cs
--- a/Engine/Source/Programs/UnrealBuildTool/Platform/IOS/UEBuildIOS.cs
+++ b/Engine/Source/Programs/UnrealBuildTool/Platform/IOS/UEBuildIOS.cs
@@ -471,7 +471,8 @@ namespace UnrealBuildTool
                 Process IPPProcess = new Process();
 
 				string IPPCmd = "certificates " + ((ProjectFile != null) ? ("\"" + ProjectFile.ToString() + "\"") : "Engine") + " -bundlename " + ProjectSettings.BundleIdentifier + (bForDistribtion ? " -distribution" : "");
-				
+				IPPProcess.StartInfo.StandardOutputEncoding = Encoding.UTF8;
+				IPPProcess.StartInfo.StandardErrorEncoding = Encoding.UTF8;
 				IPPProcess.StartInfo.WorkingDirectory = Unreal.EngineDirectory.ToString();
 				IPPProcess.OutputDataReceived += new DataReceivedEventHandler(IPPDataReceivedHandler);
 				IPPProcess.ErrorDataReceived += new DataReceivedEventHandler(IPPDataReceivedHandler);
@@ -524,6 +525,8 @@ namespace UnrealBuildTool
                 Log.TraceInformation("Provision not specified or not found for " + ((ProjectFile != null) ? ProjectFile.GetFileNameWithoutAnyExtensions() : "UnrealGame") + ", searching for compatible match...");
                 Process IPPProcess = new Process();
 
+				IPPProcess.StartInfo.StandardOutputEncoding = Encoding.UTF8;
+				IPPProcess.StartInfo.StandardErrorEncoding = Encoding.UTF8;
 				IPPProcess.OutputDataReceived += new DataReceivedEventHandler(IPPDataReceivedHandler);
 				IPPProcess.ErrorDataReceived += new DataReceivedEventHandler(IPPDataReceivedHandler);
 				IPPProcess.StartInfo.WorkingDirectory = Unreal.EngineDirectory.ToString();
@@ -1255,4 +1258,4 @@ namespace UnrealBuildTool
 			UEBuildPlatform.RegisterPlatformWithGroup(UnrealTargetPlatform.IOS, UnrealPlatformGroup.IOS);
 		}
 	}
-}
\ No newline at end of file
+}
