Fixed compile error when using installed builds.  Generated include files were missing.

diff --git a/Engine/Source/Programs/UnrealBuildTool/Configuration/UEBuildBinary.cs b/Engine/Source/Programs/UnrealBuildTool/Configuration/UEBuildBinary.cs
--- a/Engine/Source/Programs/UnrealBuildTool/Configuration/UEBuildBinary.cs
+++ b/Engine/Source/Programs/UnrealBuildTool/Configuration/UEBuildBinary.cs
@@ -466,7 +466,7 @@ namespace UnrealBuildTool
 				{
 					if(Module.GeneratedCodeDirectory != null && DirectoryReference.Exists(Module.GeneratedCodeDirectory))
 					{
-						foreach(FileReference GeneratedCodeFile in DirectoryReference.EnumerateFiles(Module.GeneratedCodeDirectory))
+						foreach(FileReference GeneratedCodeFile in DirectoryReference.EnumerateFiles(Module.GeneratedCodeDirectory, "*", SearchOption.AllDirectories))
 						{
 							// Exclude timestamp files, since they're always updated and cause collisions between builds
 							if(!GeneratedCodeFile.GetFileName().Equals("Timestamp", StringComparison.OrdinalIgnoreCase) && !GeneratedCodeFile.HasExtension(".cpp"))
