- Adding comments to the PlatformSDK files to note the version numbers are now in the Versions.cs files
#p4v-preflight-copy 19822331

diff --git a/Engine/Source/Programs/UnrealBuildTool/Platform/Android/AndroidPlatformSDK.cs b/Engine/Source/Programs/UnrealBuildTool/Platform/Android/AndroidPlatformSDK.cs
--- a/Engine/Source/Programs/UnrealBuildTool/Platform/Android/AndroidPlatformSDK.cs
+++ b/Engine/Source/Programs/UnrealBuildTool/Platform/Android/AndroidPlatformSDK.cs
@@ -7,6 +7,11 @@ using EpicGames.Core;
 using System.Text.RegularExpressions;
 using System.Linq;
 
+///////////////////////////////////////////////////////////////////
+// If you are looking for supported version numbers, look in the
+// AndroidPlatformSDK.Versions.cs file next to this file
+///////////////////////////////////////////////////////////////////
+
 namespace UnrealBuildTool
 {
 	partial class AndroidPlatformSDK : UEBuildPlatformSDK
diff --git a/Engine/Source/Programs/UnrealBuildTool/Platform/Linux/LinuxPlatformSDK.cs b/Engine/Source/Programs/UnrealBuildTool/Platform/Linux/LinuxPlatformSDK.cs
--- a/Engine/Source/Programs/UnrealBuildTool/Platform/Linux/LinuxPlatformSDK.cs
+++ b/Engine/Source/Programs/UnrealBuildTool/Platform/Linux/LinuxPlatformSDK.cs
@@ -7,6 +7,11 @@ using EpicGames.Core;
 using System.Text.RegularExpressions;
 using UnrealBuildBase;
 
+///////////////////////////////////////////////////////////////////
+// If you are looking for supported version numbers, look in the
+// LinuxPlatformSDK.Versions.cs file next to this file
+///////////////////////////////////////////////////////////////////
+
 namespace UnrealBuildTool
 {
 	partial class LinuxPlatformSDK : UEBuildPlatformSDK
diff --git a/Engine/Source/Programs/UnrealBuildTool/Platform/Mac/ApplePlatformSDK.Versions.cs b/Engine/Source/Programs/UnrealBuildTool/Platform/Mac/ApplePlatformSDK.Versions.cs
new file mode 100644
--- /dev/null
+++ b/Engine/Source/Programs/UnrealBuildTool/Platform/Mac/ApplePlatformSDK.Versions.cs
@@ -0,0 +1,40 @@
+// Copyright Epic Games, Inc. All Rights Reserved.
+
+using System;
+using System.Collections.Generic;
+using System.IO;
+using EpicGames.Core;
+using System.Text.RegularExpressions;
+using Microsoft.Win32;
+using System.Diagnostics;
+
+namespace UnrealBuildTool
+{
+	internal partial class ApplePlatformSDK : UEBuildPlatformSDK
+	{
+		public override string GetMainVersion()
+		{
+			return "13.3";
+		}
+
+		protected override void GetValidVersionRange(out string MinVersion, out string MaxVersion)
+		{
+			if (RuntimePlatform.IsMac)
+			{
+				MinVersion = "12.4.0";
+				MaxVersion = "13.9.9";
+			}
+			else
+			{
+				// @todo turnkey: these are MobileDevice .dll versions in Windows - to get the iTunes app version (12.3.4.1 etc) would need to hunt down the .exe
+				MinVersion = "1100.0.0.0";
+				MaxVersion = "1399.0";
+			}
+		}
+
+		protected override void GetValidSoftwareVersionRange(out string? MinVersion, out string? MaxVersion)
+		{
+			MinVersion = MaxVersion = null;
+		}
+	}
+}
\ No newline at end of file
diff --git a/Engine/Source/Programs/UnrealBuildTool/Platform/Mac/ApplePlatformSDK.cs b/Engine/Source/Programs/UnrealBuildTool/Platform/Mac/ApplePlatformSDK.cs
--- a/Engine/Source/Programs/UnrealBuildTool/Platform/Mac/ApplePlatformSDK.cs
+++ b/Engine/Source/Programs/UnrealBuildTool/Platform/Mac/ApplePlatformSDK.cs
@@ -8,35 +8,16 @@ using System.Text.RegularExpressions;
 using Microsoft.Win32;
 using System.Diagnostics;
 
+///////////////////////////////////////////////////////////////////
+// If you are looking for supported version numbers, look in the
+// ApplePlatformSDK.Versions.cs file next to this file, and
+// als Mac/IOSPlatformSDK.Versions.cs
+///////////////////////////////////////////////////////////////////
+
 namespace UnrealBuildTool
 {
-	internal class ApplePlatformSDK : UEBuildPlatformSDK
+	internal partial class ApplePlatformSDK : UEBuildPlatformSDK
 	{
-		public override string GetMainVersion()
-		{
-			return "13.3";
-		}
-
-		protected override void GetValidVersionRange(out string MinVersion, out string MaxVersion)
-		{
-			if (RuntimePlatform.IsMac)
-			{
-				MinVersion = "12.4.0";
-				MaxVersion = "13.9.9";
-			}
-			else
-			{
-				// @todo turnkey: these are MobileDevice .dll versions in Windows - to get the iTunes app version (12.3.4.1 etc) would need to hunt down the .exe
-				MinVersion = "1100.0.0.0";
-				MaxVersion = "1399.0";
-			}
-		}
-
-		protected override void GetValidSoftwareVersionRange(out string? MinVersion, out string? MaxVersion)
-		{
-			MinVersion = MaxVersion = null;
-		}
-
 		public override bool TryConvertVersionToInt(string? StringValue, out UInt64 OutValue, string? Hint)
 		{
 			OutValue = 0;
diff --git a/Engine/Source/Programs/UnrealBuildTool/Platform/Mac/MacPlatformSDK.Versions.cs b/Engine/Source/Programs/UnrealBuildTool/Platform/Mac/MacPlatformSDK.Versions.cs
--- a/Engine/Source/Programs/UnrealBuildTool/Platform/Mac/MacPlatformSDK.Versions.cs
+++ b/Engine/Source/Programs/UnrealBuildTool/Platform/Mac/MacPlatformSDK.Versions.cs
@@ -6,29 +6,10 @@ namespace UnrealBuildTool
 {
 	internal partial class MacPlatformSDK : ApplePlatformSDK
 	{
-		public override string GetMainVersion()
+		protected override void GetValidSoftwareVersionRange(out string MinVersion, out string? MaxVersion)
 		{
-			return "11.5";
-		}
-
-		protected override void GetValidVersionRange(out string MinVersion, out string MaxVersion)
-		{
-			if (RuntimePlatform.IsMac)
-			{
-				MinVersion = "11.0.0";
-				MaxVersion = "13.9.9";
-			}
-			else
-			{
-				// @todo turnkey: these are MobileDevice .dll versions in Windows - to get the iTunes app version (12.3.4.1 etc) would need to hunt down the .exe
-				MinVersion = "1100.0.0.0";
-				MaxVersion = "1399.0";
-			}
-		}
-
-		protected override void GetValidSoftwareVersionRange(out string? MinVersion, out string? MaxVersion)
-		{
-			MinVersion = MaxVersion = null;
+			MinVersion = "10.15.7";
+			MaxVersion = null;
 		}
 	}
 }
\ No newline at end of file
diff --git a/Engine/Source/Programs/UnrealBuildTool/Platform/Windows/MicrosoftPlatformSDK.cs b/Engine/Source/Programs/UnrealBuildTool/Platform/Windows/MicrosoftPlatformSDK.cs
--- a/Engine/Source/Programs/UnrealBuildTool/Platform/Windows/MicrosoftPlatformSDK.cs
+++ b/Engine/Source/Programs/UnrealBuildTool/Platform/Windows/MicrosoftPlatformSDK.cs
@@ -15,6 +15,11 @@ using System.Buffers.Binary;
 using UnrealBuildBase;
 using System.Runtime.Versioning;
 
+///////////////////////////////////////////////////////////////////
+// If you are looking for supported version numbers, look in the
+// MicrosoftPlatformSDK.Versions.cs file next to this file
+///////////////////////////////////////////////////////////////////
+
 namespace UnrealBuildTool
 {
 	internal partial class MicrosoftPlatformSDK : UEBuildPlatformSDK
