HoloLens can read build resources from platform extension folders.

#rnx

diff --git a/Engine/Source/Programs/UnrealBuildTool/Platform/HoloLens/HoloLensDeploy.cs b/Engine/Source/Programs/UnrealBuildTool/Platform/HoloLens/HoloLensDeploy.cs
--- a/Engine/Source/Programs/UnrealBuildTool/Platform/HoloLens/HoloLensDeploy.cs
+++ b/Engine/Source/Programs/UnrealBuildTool/Platform/HoloLens/HoloLensDeploy.cs
@@ -222,6 +222,11 @@ namespace UnrealBuildTool
 			string OutputAppX = Path.Combine(OutputDirectory, OutputName + Extension);
 			string SigningCertificate = @"Build\HoloLens\SigningCertificate.pfx";
 			string SigningCertificatePath = Path.Combine(Receipt.ProjectDir != null ? Receipt.ProjectDir.FullName : Unreal.EngineDirectory.FullName, SigningCertificate);
+			if (!File.Exists(SigningCertificatePath))
+			{
+				SigningCertificate = @"Platforms\HoloLens\Build\SigningCertificate.pfx";
+				SigningCertificatePath = Path.Combine(Receipt.ProjectDir != null ? Receipt.ProjectDir.FullName : Unreal.EngineDirectory.FullName, SigningCertificate);
+			}
 
 			string MapFilename = Path.Combine(IntermediateDirectory, OutputName + ".pkgmap");
 			var LocalRoot = Receipt.ProjectDir;
diff --git a/Engine/Source/Programs/UnrealBuildTool/Platform/HoloLens/ManifestGenerator.cs b/Engine/Source/Programs/UnrealBuildTool/Platform/HoloLens/ManifestGenerator.cs
--- a/Engine/Source/Programs/UnrealBuildTool/Platform/HoloLens/ManifestGenerator.cs
+++ b/Engine/Source/Programs/UnrealBuildTool/Platform/HoloLens/ManifestGenerator.cs
@@ -36,7 +36,6 @@ namespace UnrealBuildTool
 		//private UnrealTargetPlatform Platform;
 		private WindowsArchitecture Architecture;
 		private string? TargetSettings;
-		private string? BuildResourceProjectRelativePath;
 		private string? ProjectPath;
 		private string? OutputPath;
 		private string? IntermediatePath;
@@ -425,17 +424,30 @@ namespace UnrealBuildTool
 		private bool CopyAndReplaceBinaryIntermediate(string ResourceFileName, bool AllowEngineFallback = true, Action<string, string>? CopyOp = null)
 		{
 			string TargetPath = Path.Combine(IntermediatePath!, BuildResourceSubPath);
-			string SourcePath = Path.Combine(ProjectPath!, BuildResourceProjectRelativePath!, BuildResourceSubPath);
 
 			CopyOp = CopyOp ?? File.Copy;
 
-			// Try falling back to the engine defaults if requested
+			// look in project normal Build location
+			string SourcePath = Path.Combine(ProjectPath!, "Build", "HoloLens", BuildResourceSubPath);
 			bool bFileExists = File.Exists(Path.Combine(SourcePath, ResourceFileName));
+
+			// look in Platform Extensions next
 			if (!bFileExists)
 			{
-				if (AllowEngineFallback)
+				SourcePath = Path.Combine(ProjectPath!, "Platforms", "HoloLens", "Build", BuildResourceSubPath);
+				bFileExists = File.Exists(Path.Combine(SourcePath, ResourceFileName));
+			}
+
+			// look in Engine, if allowed
+			if (!bFileExists && AllowEngineFallback)
+			{
+				SourcePath = Path.Combine(Unreal.EngineDirectory.FullName, "Build", "HoloLens", EngineResourceSubPath);
+				bFileExists = File.Exists(Path.Combine(SourcePath, ResourceFileName));
+
+				// look in Platform extensions too
+				if (!bFileExists)
 				{
-					SourcePath = Path.Combine(Unreal.EngineDirectory.FullName, BuildResourceProjectRelativePath!, EngineResourceSubPath);
+					SourcePath = Path.Combine(Unreal.EngineDirectory.FullName, "Platforms", "HoloLens", "Build", EngineResourceSubPath);
 					bFileExists = File.Exists(Path.Combine(SourcePath, ResourceFileName));
 				}
 			}
@@ -777,7 +789,6 @@ namespace UnrealBuildTool
 			//Platform = TargetPlatform;
 			Architecture = TargetArchitecture;
 			TargetSettings = "/Script/HoloLensPlatformEditor.HoloLensTargetSettings";
-			BuildResourceProjectRelativePath = "Build\\HoloLens";
 			BuildResourceSubPath = "Resources";
 			StoreResourceSubPath = WindowsExports.GetArchitectureSubpath(Architecture) + "\\" + BuildResourceSubPath;
 
@@ -1159,7 +1170,6 @@ namespace UnrealBuildTool
 			WinMDReferences = new List<WinMDRegistrationInfo>();
 			//Platform = TargetPlatform;
 			TargetSettings = "/Script/HoloLensPlatformEditor.HoloLensTargetSettings";
-			BuildResourceProjectRelativePath = "Build\\HoloLens";
 			BuildResourceSubPath = "AssetsResources";
 			StoreResourceSubPath =  BuildResourceSubPath;
 
