Add an option to discard circular dependency validation for non engine modules
#rnx

diff --git a/Engine/Source/Programs/UnrealBuildTool/Configuration/ModuleRules.cs b/Engine/Source/Programs/UnrealBuildTool/Configuration/ModuleRules.cs
--- a/Engine/Source/Programs/UnrealBuildTool/Configuration/ModuleRules.cs
+++ b/Engine/Source/Programs/UnrealBuildTool/Configuration/ModuleRules.cs
@@ -1055,6 +1055,13 @@ namespace UnrealBuildTool
 			get { return bLegacyPublicIncludePathsPrivate ?? ((DefaultBuildSettings < BuildSettingsVersion.V2) ? Target.bLegacyPublicIncludePaths : false); }
 		}
 		private bool? bLegacyPublicIncludePathsPrivate;
+		
+		/// <summary>
+		/// Whether circular dependencies will be validated against the allow list
+		/// Circular module dependencies result in slower builds. Disabling this option is strongly discouraged.
+        /// This option is ignored for Engine modules which will always be validated against the allow list.
+		/// </summary>
+		public bool bValidateCircularDependencies = true;
 
 		/// <summary>
 		/// Which stanard to use for compiling this module
diff --git a/Engine/Source/Programs/UnrealBuildTool/Configuration/UEBuildModuleCPP.cs b/Engine/Source/Programs/UnrealBuildTool/Configuration/UEBuildModuleCPP.cs
--- a/Engine/Source/Programs/UnrealBuildTool/Configuration/UEBuildModuleCPP.cs
+++ b/Engine/Source/Programs/UnrealBuildTool/Configuration/UEBuildModuleCPP.cs
@@ -226,11 +226,18 @@ namespace UnrealBuildTool
 				Log.TraceVerbose("Compile Env {0}: {1}", Name, Def);
 			}
 
-			foreach(string CircularlyReferencedModuleName in Rules.CircularlyReferencedDependentModules)
+			if (Rules.bValidateCircularDependencies || Rules.bTreatAsEngineModule)
 			{
-				if(CircularlyReferencedModuleName != "BlueprintContext" && !CircularDependenciesAllowList.Any(x => x.Key == Name && x.Value == CircularlyReferencedModuleName))
+				foreach (string CircularlyReferencedModuleName in Rules.CircularlyReferencedDependentModules)
 				{
-					Log.TraceWarning("Found reference between '{0}' and '{1}'. Support for circular references is being phased out; please do not introduce new ones.", Name, CircularlyReferencedModuleName);
+					if (CircularlyReferencedModuleName != "BlueprintContext" &&
+					    !CircularDependenciesAllowList.Any(x =>
+						    x.Key == Name && x.Value == CircularlyReferencedModuleName))
+					{
+						Log.TraceWarning(
+							"Found reference between '{0}' and '{1}'. Support for circular references is being phased out; please do not introduce new ones.",
+							Name, CircularlyReferencedModuleName);
+					}
 				}
 			}
 
