UnrealBuildTool: Debug flag to launch zenserver

#rnx

diff --git a/Engine/Source/Programs/UnrealBuildTool/Executors/Experimental/HordeExecutor.cs b/Engine/Source/Programs/UnrealBuildTool/Executors/Experimental/HordeExecutor.cs
--- a/Engine/Source/Programs/UnrealBuildTool/Executors/Experimental/HordeExecutor.cs
+++ b/Engine/Source/Programs/UnrealBuildTool/Executors/Experimental/HordeExecutor.cs
@@ -8,6 +8,7 @@ using System.Diagnostics;
 using System.IO;
 using System.Linq;
 using System.Net.Http;
+using System.Runtime.InteropServices;
 using System.Text;
 using System.Threading;
 using System.Threading.Tasks;
@@ -25,12 +26,13 @@ namespace UnrealBuildTool
 	///						  This will add additional overhead to the build and is not recommended unless using Horde compute remote distribution.
 	/// * -NoPCH:			  Cannot build with pch, as absolute paths are embedded into the .gch file for included headers, and these paths must exist exactly.
 	/// * -HordeCompute:	  Flag to enable this executor
-	/// * ZenServer:	      Must be running with the correct credentials to connect to both Horde Storage to upload files and Horde to make execution requests.
+	/// * ZenServer:		  Must be running with the correct credentials to connect to both Horde Storage to upload files and Horde to make execution requests.
+	///						  Can be launched via -LaunchZenServer.
 	///
 	/// TODO:
 	/// * Investigate overhead when connecting from Zen to Horde
 	/// * Local process racing when remote tasks are blocking more tasks from starting, or at the end of a build.
-	/// * Print remote status as tasks complete, print more concise and useful summary at the end of the build
+	/// * Print more concise and useful summary at the end of the build
 	/// * Additional toolchain platform support
 	/// * Dynamic control over the number of remote processes to run at once. Needs knowledge about availibility of remote compute resources, if the pending remote queue is overloaded and by how much, etc.
 	/// </summary>
@@ -54,6 +56,11 @@ namespace UnrealBuildTool
 		[XmlConfigFile]
 		private static bool RetryFailedRemote = true;
 
+		/// <summary>
+		/// Temporary debug flag to allow launching zenserver.
+		/// </summary>
+		private static bool LaunchZenServer = Environment.CommandLine.ToLower().Contains("-launchzenserver");
+
 		private static readonly Lazy<HttpClient> LazyZenHttpClient = new Lazy<HttpClient>();
 		private static HttpClient ZenHttpClient { get { return LazyZenHttpClient.Value; } }
 
@@ -289,6 +296,28 @@ namespace UnrealBuildTool
 		/// </summary>
 		/// <returns>True if the task executor can be used</returns>
 		public static new bool IsAvailable()
+		{
+			try
+			{
+				if (LaunchZenServer)
+				{
+					ConfigHierarchy EngineConfig = ConfigCache.ReadHierarchy(ConfigHierarchyType.Engine, DirectoryReference.Combine(Unreal.EngineDirectory, "Programs", "UnrealBuildTool"), BuildHostPlatform.Current.Platform);
+
+					if (EngineConfig.GetBool("Zen", "AutoLaunch", out bool AutoLaunch) && AutoLaunch)
+					{
+						return true;
+					}
+				}
+
+				return IsReady();
+			}
+			catch
+			{
+				return false;
+			}
+		}
+
+		private static bool IsReady()
 		{
 			try
 			{
@@ -302,6 +331,44 @@ namespace UnrealBuildTool
 			}
 		}
 
+		private static bool LaunchZen(ManagedProcessGroup ProcessGroup, [System.Diagnostics.CodeAnalysis.NotNullWhen(true)] out ManagedProcess? ZenProcess)
+		{
+			ZenProcess = null;
+			try
+			{
+
+				ConfigHierarchy EngineConfig = ConfigCache.ReadHierarchy(ConfigHierarchyType.Engine, DirectoryReference.Combine(Unreal.EngineDirectory, "Programs", "UnrealBuildTool"), BuildHostPlatform.Current.Platform);
+
+				DirectoryReference? DataPathReference = null;
+				if (EngineConfig.GetString("Zen.AutoLaunch", "ExtraArgs", out string ExtraArgs) && EngineConfig.GetString("Zen.AutoLaunch", "DataPath", out string DataPath))
+				{
+					DataPathReference = DirectoryReference.FromString(DataPath.Replace("%APPSETTINGSDIR%", $"{Path.Combine(Environment.GetFolderPath(Environment.SpecialFolder.CommonApplicationData), "Epic")}{Path.DirectorySeparatorChar}"));
+				}
+
+				if (DataPathReference == null || DataPathReference.ParentDirectory == null)
+				{
+					return false;
+				}
+
+				// Copy Zen to DataPath
+				FileReference ZenSourcePath = FileReference.Combine(Unreal.EngineDirectory, "Binaries", BuildHostPlatform.Current.Platform.ToString(), $"zenserver{(RuntimeInformation.IsOSPlatform(OSPlatform.Windows) ? ".exe" : "")}");
+				FileReference ZenDestPath = FileReference.Combine(DataPathReference.ParentDirectory, "Install", ZenSourcePath.GetFileName());
+
+				DirectoryReference.CreateDirectory(DataPathReference);
+				DirectoryReference.CreateDirectory(ZenDestPath.Directory);
+				FileReference.MakeWriteable(ZenDestPath);
+				FileReference.Copy(ZenSourcePath, ZenDestPath, true);
+
+				ZenProcess = new ManagedProcess(ProcessGroup, ZenDestPath.FullName, $"--port 1337 --data-dir \"{DataPathReference.FullName}\" {ExtraArgs}", ZenDestPath.Directory.FullName, null, ProcessPriorityClass.Normal);
+
+				return true;
+			}
+			catch
+			{
+				return false;
+			}
+		}
+
 		/// <summary>
 		/// Executes the specified actions locally.
 		/// </summary>
@@ -313,6 +380,20 @@ namespace UnrealBuildTool
 			int ActualNumParallelProcesses = Math.Min(TotalActions, NumParallelProcesses);
 
 			using ManagedProcessGroup ProcessGroup = new ManagedProcessGroup();
+			ManagedProcess? ZenProcess = null;
+			if (!IsReady())
+			{
+				if (LaunchZenServer && LaunchZen(ProcessGroup, out ZenProcess))
+				{
+					Task.Delay(1000).Wait();
+				}
+				if (!IsReady())
+				{
+					Log.TraceInformationOnce("Unable to establish connection to zenserver, disabling HordeExecutor");
+					return base.ExecuteActions(InputActions);
+				}
+			}
+
 			using SemaphoreSlim MaxProcessSemaphore = new SemaphoreSlim(ActualNumParallelProcesses, ActualNumParallelProcesses);
 			using SemaphoreSlim MaxRemoteProcessSemaphore = new SemaphoreSlim(NumRemoteParallelProcesses, NumRemoteParallelProcesses);
 			using ProgressWriter ProgressWriter = new ProgressWriter("Compiling C++ source code...", false);
@@ -546,9 +627,9 @@ namespace UnrealBuildTool
 				if (RemoteSemaphoreTask?.Status == TaskStatus.RanToCompletion && RemoteSemaphoreTask?.Result == true)
 				{
 					if (!RetryFailedRemote)
-                    {
+					{
 						Log.WriteException(Ex, null);
-                    }
+					}
 					return new RemoteExecuteResults(new List<string>(), int.MaxValue);
 				}
 				Log.WriteException(Ex, null);
