UnrealBuildTool: Mac and Linux -fpch-validate-input-files-content

This flag validates pch inputs using content only when the mtime check fails, allowing the build to not fail if the content does not actually differ.

#rnx

diff --git a/Engine/Source/Programs/UnrealBuildTool/Platform/Linux/LinuxToolChain.cs b/Engine/Source/Programs/UnrealBuildTool/Platform/Linux/LinuxToolChain.cs
--- a/Engine/Source/Programs/UnrealBuildTool/Platform/Linux/LinuxToolChain.cs
+++ b/Engine/Source/Programs/UnrealBuildTool/Platform/Linux/LinuxToolChain.cs
@@ -926,6 +926,10 @@ namespace UnrealBuildTool
 		{
 			string Result = "";
 			Result += " -x c++-header";
+			if (CompilerVersionGreaterOrEqual(11, 0, 0))
+			{
+				Result += " -fpch-validate-input-files-content";
+			}
 			Result += GetCompilerStandardVersion_CPP(CompileEnvironment);
 			return Result;
 		}
@@ -1279,6 +1283,10 @@ namespace UnrealBuildTool
 
 			if (CompileEnvironment.PrecompiledHeaderAction == PrecompiledHeaderAction.Include)
 			{
+				if (CompilerVersionGreaterOrEqual(11, 0, 0))
+				{
+					PCHArguments += " -fpch-validate-input-files-content";
+				}
 				PCHArguments += string.Format(" -include \"{0}\"", NormalizeCommandLinePath(CompileEnvironment.PrecompiledHeaderIncludeFilename!));
 			}
 
diff --git a/Engine/Source/Programs/UnrealBuildTool/Platform/Mac/MacToolChain.cs b/Engine/Source/Programs/UnrealBuildTool/Platform/Mac/MacToolChain.cs
--- a/Engine/Source/Programs/UnrealBuildTool/Platform/Mac/MacToolChain.cs
+++ b/Engine/Source/Programs/UnrealBuildTool/Platform/Mac/MacToolChain.cs
@@ -462,6 +462,10 @@ namespace UnrealBuildTool
 				PCHArguments.Append(" -include \"");
 				PCHArguments.Append(NormalizeCommandLinePath(CompileEnvironment.PrecompiledHeaderIncludeFilename!));
 				PCHArguments.Append("\"");
+				if (GetClangVersion().Major >= 11)
+				{
+					PCHArguments.Append(" -fpch-validate-input-files-content");
+				}
 			}
 
 			// Add include paths to the argument list.
@@ -510,6 +514,10 @@ namespace UnrealBuildTool
 				{
 					// Compile the file as a C++ PCH.
 					FileArguments += GetCompileArguments_PCH(CompileEnvironment);
+					if (GetClangVersion().Major >= 11)
+					{
+						FileArguments += " -fpch-validate-input-files-content";
+					}
 					FileArguments += GetRTTIFlag(CompileEnvironment);
 				}
 				else if (Extension == ".C")
