Add -fno-delete-null-pointer-checks to disable clang13 optimization on newer Android NDK toolchains

#android

diff --git a/Engine/Source/Programs/UnrealBuildTool/Platform/Android/AndroidToolChain.cs b/Engine/Source/Programs/UnrealBuildTool/Platform/Android/AndroidToolChain.cs
--- a/Engine/Source/Programs/UnrealBuildTool/Platform/Android/AndroidToolChain.cs
+++ b/Engine/Source/Programs/UnrealBuildTool/Platform/Android/AndroidToolChain.cs
@@ -638,6 +638,12 @@ namespace UnrealBuildTool
 				//Resul += " -fvisibility-inlines-hidden-static-local-var";
 			}
 
+			if (CompilerVersionGreaterOrEqual(12, 0, 0))
+			{
+				// The 'this' vs nullptr comparisons get optimized away for newer versions of Clang, which is undesirable until we refactor these checks.
+				Result += " -fno-delete-null-pointer-checks";
+			}
+
 			if (CompileEnvironment.DeprecationWarningLevel == WarningLevel.Error)
 			{
 				Result += " -Werror=deprecated-declarations";
