Change direct reference to Desktop platforms to allow dynamic linking on consoles

diff --git a/Engine/Source/Programs/UnrealBuildTool/Configuration/UEBuildBinary.cs b/Engine/Source/Programs/UnrealBuildTool/Configuration/UEBuildBinary.cs
--- a/Engine/Source/Programs/UnrealBuildTool/Configuration/UEBuildBinary.cs
+++ b/Engine/Source/Programs/UnrealBuildTool/Configuration/UEBuildBinary.cs
@@ -205,9 +205,7 @@ namespace UnrealBuildTool
 				// Mark the link environment as cross-referenced.
 				BinaryLinkEnvironment.bIsCrossReferenced = true;
 
-				if (BinaryLinkEnvironment.Platform != UnrealTargetPlatform.Mac &&
-					BinaryLinkEnvironment.Platform != UnrealTargetPlatform.Linux &&
-					BinaryLinkEnvironment.Platform != UnrealTargetPlatform.LinuxArm64)
+				if (BinaryLinkEnvironment.Platform.IsInGroup(UnrealPlatformGroup.Microsoft))
 				{
 					// Create the import library.
 					OutputFiles.AddRange(ToolChain.LinkAllFiles(BinaryLinkEnvironment, true, Graph));
@@ -333,10 +331,7 @@ namespace UnrealBuildTool
 				foreach (FileReference OutputFilePath in OutputFilePaths)
 				{
 					FileReference LibraryFileName;
-					if (Type == UEBuildBinaryType.StaticLibrary ||
-						DependentLinkEnvironment.Platform == UnrealTargetPlatform.Mac ||
-						DependentLinkEnvironment.Platform == UnrealTargetPlatform.Linux ||
-						DependentLinkEnvironment.Platform == UnrealTargetPlatform.LinuxArm64)
+					if (Type == UEBuildBinaryType.StaticLibrary || !DependentLinkEnvironment.Platform.IsInGroup(UnrealPlatformGroup.Microsoft))
 					{
 						LibraryFileName = OutputFilePath;
 					}
