Add -fno-delete-null-pointer-checks to disable clang13 optimization on iOS

diff --git a/Engine/Source/Programs/UnrealBuildTool/Platform/IOS/IOSToolChain.cs b/Engine/Source/Programs/UnrealBuildTool/Platform/IOS/IOSToolChain.cs
--- a/Engine/Source/Programs/UnrealBuildTool/Platform/IOS/IOSToolChain.cs
+++ b/Engine/Source/Programs/UnrealBuildTool/Platform/IOS/IOSToolChain.cs
@@ -1,4 +1,4 @@
-// Copyright Epic Games, Inc. All Rights Reserved.
+﻿// Copyright Epic Games, Inc. All Rights Reserved.
 
 
 using System;
@@ -286,6 +286,9 @@ namespace UnrealBuildTool
 			if (GetClangVersion().Major >= 12)
 			{
 				Result += " -Wno-range-loop-analysis";
+
+				// We have 'this' vs nullptr comparisons that get optimized away for newer versions of Clang, which is undesirable until we refactor these checks.
+				Result += " -fno-delete-null-pointer-checks";
 			}
 
 			if (CompileEnvironment.ShadowVariableWarningLevel != WarningLevel.Off)
