UnrealBuildTool: Add logging when zenserver fails to launch

#rnx

diff --git a/Engine/Source/Programs/UnrealBuildTool/Executors/Experimental/HordeExecutor.cs b/Engine/Source/Programs/UnrealBuildTool/Executors/Experimental/HordeExecutor.cs
--- a/Engine/Source/Programs/UnrealBuildTool/Executors/Experimental/HordeExecutor.cs
+++ b/Engine/Source/Programs/UnrealBuildTool/Executors/Experimental/HordeExecutor.cs
@@ -336,17 +336,32 @@ namespace UnrealBuildTool
 			ZenProcess = null;
 			try
 			{
-
 				ConfigHierarchy EngineConfig = ConfigCache.ReadHierarchy(ConfigHierarchyType.Engine, DirectoryReference.Combine(Unreal.EngineDirectory, "Programs", "UnrealBuildTool"), BuildHostPlatform.Current.Platform);
 
+				if (!EngineConfig.GetBool("Zen", "AutoLaunch", out bool AutoLaunch) || !AutoLaunch)
+				{
+					Log.TraceWarning("Unable to launch zenserver: '[Zen] AutoLaunch' not enabled in Engine config");
+					return false;
+				}
+
 				DirectoryReference? DataPathReference = null;
-				if (EngineConfig.GetString("Zen.AutoLaunch", "ExtraArgs", out string ExtraArgs) && EngineConfig.GetString("Zen.AutoLaunch", "DataPath", out string DataPath))
+				if (!EngineConfig.GetString("Zen.AutoLaunch", "DataPath", out string DataPath) || string.IsNullOrEmpty(DataPath))
+				{
+					Log.TraceWarning("Unable to launch zenserver: '[Zen.AutoLaunch] DataPath' not set in Engine config");
+					return false;
+				}
+
+				DataPathReference = DirectoryReference.FromString(DataPath.Replace("%APPSETTINGSDIR%", $"{Path.Combine(Environment.GetFolderPath(Environment.SpecialFolder.CommonApplicationData), "Epic")}{Path.DirectorySeparatorChar}"));
+
+				if (EngineConfig.GetString("Zen.AutoLaunch", "ExtraArgs", out string ExtraArgs))
 				{
-					DataPathReference = DirectoryReference.FromString(DataPath.Replace("%APPSETTINGSDIR%", $"{Path.Combine(Environment.GetFolderPath(Environment.SpecialFolder.CommonApplicationData), "Epic")}{Path.DirectorySeparatorChar}"));
+					Log.TraceWarning("Unable to launch zenserver: '[Zen.AutoLaunch] ExtraArgs' not set in Engine config");
+					return false;
 				}
 
 				if (DataPathReference == null || DataPathReference.ParentDirectory == null)
 				{
+					Log.TraceWarning("Unable to launch zenserver: data-dir not found");
 					return false;
 				}
 
@@ -354,17 +369,20 @@ namespace UnrealBuildTool
 				FileReference ZenSourcePath = FileReference.Combine(Unreal.EngineDirectory, "Binaries", BuildHostPlatform.Current.Platform.ToString(), $"zenserver{(RuntimeInformation.IsOSPlatform(OSPlatform.Windows) ? ".exe" : "")}");
 				FileReference ZenDestPath = FileReference.Combine(DataPathReference.ParentDirectory, "Install", ZenSourcePath.GetFileName());
 
+				Log.TraceInformation($"Copying zenserver to '{ZenDestPath}'");
 				DirectoryReference.CreateDirectory(DataPathReference);
 				DirectoryReference.CreateDirectory(ZenDestPath.Directory);
 				FileReference.MakeWriteable(ZenDestPath);
 				FileReference.Copy(ZenSourcePath, ZenDestPath, true);
 
+				Log.TraceInformation($"Launching zenserver with data-dir '{DataPathReference.FullName}'...");
 				ZenProcess = new ManagedProcess(ProcessGroup, ZenDestPath.FullName, $"--port 1337 --data-dir \"{DataPathReference.FullName}\" {ExtraArgs}", ZenDestPath.Directory.FullName, null, ProcessPriorityClass.Normal);
 
 				return true;
 			}
-			catch
+			catch (Exception Ex)
 			{
+				Log.WriteException(Ex, null);
 				return false;
 			}
 		}
@@ -389,6 +407,11 @@ namespace UnrealBuildTool
 				}
 				if (!IsReady())
 				{
+					if (LaunchZenServer)
+					{
+						Log.TraceError("Unable to establish connection to zenserver");
+						return false;
+					}
 					Log.TraceInformationOnce("Unable to establish connection to zenserver, disabling HordeExecutor");
 					return base.ExecuteActions(InputActions);
 				}
