Minor refactor to UBT plugins.
Added support for props file for UBT plugins.  Needed for packaging of plugins having a ubt plugin

#rnx

diff --git a/Engine/Source/Programs/Shared/EpicGames.Build/System/CompileScriptModules.cs b/Engine/Source/Programs/Shared/EpicGames.Build/System/CompileScriptModules.cs
--- a/Engine/Source/Programs/Shared/EpicGames.Build/System/CompileScriptModules.cs
+++ b/Engine/Source/Programs/Shared/EpicGames.Build/System/CompileScriptModules.cs
@@ -109,6 +109,31 @@ namespace UnrealBuildBase
 			return GetTargetPaths(Build(RulesFileType, FoundProjects, BaseDirectories, bForceCompile, bNoCompile, bUseBuildRecords, out bBuildSuccess, OnBuildingProjects));
 		}
 
+		/// <summary>
+		/// Test to see if all the given projects are up-to-date
+		/// </summary>
+		/// <param name="FoundProjects">Collection of projects to test</param>
+		/// <param name="BaseDirectories">Base directories of the projects</param>
+		/// <returns>True if all of the projects are up to date</returns>
+		public static bool AreScriptModulesUpToDate(HashSet<FileReference> FoundProjects, List<DirectoryReference> BaseDirectories)
+		{
+			CsProjBuildHook Hook = new Hook();
+
+			// Load existing build records, validating them only if (re)compiling script projects is an option
+			Dictionary<FileReference, (CsProjBuildRecord BuildRecord, FileReference BuildRecordPath)> ExistingBuildRecords = LoadExistingBuildRecords(BaseDirectories);
+
+			Dictionary<FileReference, (CsProjBuildRecord BuildRecord, FileReference BuildRecordPath)> ValidBuildRecords = new Dictionary<FileReference, (CsProjBuildRecord, FileReference)>(ExistingBuildRecords.Count);
+			Dictionary<FileReference, (CsProjBuildRecord BuildRecord, FileReference BuildRecordPath)> InvalidBuildRecords = new Dictionary<FileReference, (CsProjBuildRecord, FileReference)>(ExistingBuildRecords.Count);
+
+			foreach (FileReference Project in FoundProjects)
+			{
+				ValidateBuildRecordRecursively(ValidBuildRecords, InvalidBuildRecords, ExistingBuildRecords, Project, Hook);
+			}
+
+			// If all found records are valid, we can return their targets directly
+			return FoundProjects.All(x => ValidBuildRecords.ContainsKey(x));
+		}
+
 		/// <summary>
 		/// Locates script modules, builds them if necessary, returns set of .dll files
 		/// </summary>
diff --git a/Engine/Source/Programs/Shared/EpicGames.Build/System/EnumeratePlugins.cs b/Engine/Source/Programs/Shared/EpicGames.Build/System/EnumeratePlugins.cs
--- a/Engine/Source/Programs/Shared/EpicGames.Build/System/EnumeratePlugins.cs
+++ b/Engine/Source/Programs/Shared/EpicGames.Build/System/EnumeratePlugins.cs
@@ -15,6 +15,12 @@ namespace UnrealBuildBase
 		/// </summary>
 		static Dictionary<DirectoryReference, List<FileReference>> PluginFileCache = new Dictionary<DirectoryReference, List<FileReference>>();
 
+		/// <summary>
+		/// Cache of all UBT plugins that have been built.  When generating projects, there is a good chance that a UBT project
+		/// might need to be rebuilt.
+		/// </summary>
+		static Dictionary<FileReference, FileReference> UbtPluginFileCache = new Dictionary<FileReference, FileReference>();
+
 		/// <summary>
 		/// Invalidate cached plugin data so that we can pickup new things
 		/// Warning: Will make subsequent plugin lookups and directory scans slow until the caches are repopulated
@@ -104,13 +110,9 @@ namespace UnrealBuildBase
 		/// Enumerates all the UBT plugin files available to the given project
 		/// </summary>
 		/// <param name="ProjectFile">Path to the project file</param>
-		/// <param name="Plugins">Collection of found plugins</param>
-		/// <returns>True if the plugins compiled</returns>
-		public static bool EnumerateUbtPlugins(FileReference? ProjectFile, out (FileReference ProjectFile, FileReference TargetAssembly)[]? Plugins)
+		/// <returns>The collection of UBT project files</returns>
+		public static List<FileReference> EnumerateUbtPlugins(FileReference? ProjectFile)
 		{
-			bool bBuildSuccess = true;
-			Plugins = null;
-
 			List<DirectoryReference> PluginDirectories = new List<DirectoryReference>();
 			PluginDirectories.AddRange(Unreal.GetExtensionDirs(Unreal.EngineDirectory, "Source"));
 			PluginDirectories.AddRange(Unreal.GetExtensionDirs(Unreal.EngineDirectory, "Plugins"));
@@ -120,23 +122,36 @@ namespace UnrealBuildBase
 				PluginDirectories.AddRange(Unreal.GetExtensionDirs(ProjectFile.Directory, "Plugins"));
 				PluginDirectories.AddRange(Unreal.GetExtensionDirs(ProjectFile.Directory, "Mods"));
 			}
+			List<FileReference> Plugins = UnrealBuildBase.Rules.FindAllRulesFiles(PluginDirectories, UnrealBuildBase.Rules.RulesFileType.UbtPlugin);
+			Plugins.SortBy(P => P.FullName);
+			return Plugins;
+		}
+
+		/// <summary>
+		/// Build the given collection of plugins
+		/// </summary>
+		/// <param name="ProjectFile">Path to the project file</param>
+		/// <param name="FoundPlugins">Collection of plugins to build</param>
+		/// <param name="Plugins">Collection of built plugins</param>
+		/// <returns>True if the plugins compiled</returns>
+		public static bool BuildUbtPlugins(FileReference? ProjectFile, IEnumerable<FileReference> FoundPlugins,
+			out (FileReference ProjectFile, FileReference TargetAssembly)[]? Plugins)
+		{
+			bool bBuildSuccess = true;
+			Plugins = null;
 
-			// Scan for UBT plugin projects
-			HashSet<FileReference> ScannedPlugins = new HashSet<FileReference>(UnrealBuildBase.Rules.FindAllRulesFiles(
-				PluginDirectories, UnrealBuildBase.Rules.RulesFileType.UbtPlugin));
+			// Collect the list of plugins already build
+			List<FileReference> Built = FoundPlugins.Where(P => !UbtPluginFileCache.ContainsKey(P)).ToList();
 
-			// Build the located plugins
-			Dictionary<FileReference, (CsProjBuildRecord BuildRecord, FileReference BuildRecordFile)>? BuiltPlugins = null;
-			if (ScannedPlugins.Count > 0)
-			{
-				List<DirectoryReference> BaseDirectories = new List<DirectoryReference>(2);
-				BaseDirectories.Add(Unreal.EngineDirectory);
-				if (ProjectFile != null)
-				{
-					BaseDirectories.Add(ProjectFile.Directory);
-				}
+			// Collect the list of plugins not build
+			IEnumerable<FileReference> NotBuilt = FoundPlugins.Where(P => !UbtPluginFileCache.ContainsKey(P));
 
-				BuiltPlugins = CompileScriptModule.Build(UnrealBuildBase.Rules.RulesFileType.UbtPlugin, ScannedPlugins, BaseDirectories, false, false, true, out bBuildSuccess,
+			// If we have anything left to build, then build it
+			if (NotBuilt.Any())
+			{
+				Dictionary<FileReference, (CsProjBuildRecord BuildRecord, FileReference BuildRecordFile)>? BuiltPlugins =
+					CompileScriptModule.Build(UnrealBuildBase.Rules.RulesFileType.UbtPlugin, new HashSet<FileReference>(NotBuilt),
+					GetBaseDirectories(ProjectFile), false, false, true, out bBuildSuccess,
 					Count =>
 					{
 						if (Log.OutputFile != null)
@@ -149,10 +164,28 @@ namespace UnrealBuildBase
 						}
 					}
 				);
-
-				Plugins = BuiltPlugins.Select(P => (P.Key, CompileScriptModule.GetTargetPath(P))).OrderBy(X => X.Key.FullName).ToArray();
+				
+				// Add any built plugins back into the cache
+				foreach (KeyValuePair<FileReference, (CsProjBuildRecord BuildRecord, FileReference BuildRecordFile)> BuiltPlugin in BuiltPlugins)
+				{
+					Built.Add(BuiltPlugin.Key);
+					UbtPluginFileCache.Add(BuiltPlugin.Key, CompileScriptModule.GetTargetPath(BuiltPlugin));
+				}
 			}
+
+			Plugins = Built.Select(P => (P, UbtPluginFileCache[P])).OrderBy(X => X.P.FullName).ToArray();
 			return bBuildSuccess;
 		}
+
+		static List<DirectoryReference> GetBaseDirectories(FileReference? ProjectFile)
+		{
+			List<DirectoryReference> BaseDirectories = new List<DirectoryReference>(2);
+			BaseDirectories.Add(Unreal.EngineDirectory);
+			if (ProjectFile != null)
+			{
+				BaseDirectories.Add(ProjectFile.Directory);
+			}
+			return BaseDirectories;
+		}
 	}
 }
diff --git a/Engine/Source/Programs/UnrealBuildTool/Configuration/UEBuildTarget.cs b/Engine/Source/Programs/UnrealBuildTool/Configuration/UEBuildTarget.cs
--- a/Engine/Source/Programs/UnrealBuildTool/Configuration/UEBuildTarget.cs
+++ b/Engine/Source/Programs/UnrealBuildTool/Configuration/UEBuildTarget.cs
@@ -1072,14 +1072,19 @@ namespace UnrealBuildTool
 		public FileReference? ForeignPlugin;
 
 		/// <summary>
-		/// Collection of all UBT plugin project files
+		/// Collection of all UBT plugins (project files)
 		/// </summary>
 		public List<FileReference>? UbtPlugins;
 
 		/// <summary>
-		/// Collection of all located UHT plugin assemblies
+		/// Collection of all enabled UBT plugins (project files)
 		/// </summary>
-		public List<FileReference>? UhtPlugins;
+		public List<FileReference>? EnabledUbtPlugins;
+
+		/// <summary>
+		/// Collection of all enabled UHT plugin assemblies
+		/// </summary>
+		public List<FileReference>? EnabledUhtPlugins;
 
 		/// <summary>
 		/// All application binaries; may include binaries not built by this target.
@@ -1776,7 +1781,8 @@ namespace UnrealBuildTool
 			// Create the makefile
 			string ExternalMetadata = UEBuildPlatform.GetBuildPlatform(Platform).GetExternalBuildMetadata(ProjectFile);
 			TargetMakefile Makefile = new TargetMakefile(ExternalMetadata, Binaries[0].OutputFilePaths[0], ReceiptFileName, ProjectIntermediateDirectory, TargetType, 
-				Rules.ConfigValueTracker, bDeployAfterCompile, bHasProjectScriptPlugin, bHasRequiredProjectScriptPlugin, UbtPlugins?.ToArray(), UhtPlugins?.ToArray());
+				Rules.ConfigValueTracker, bDeployAfterCompile, bHasProjectScriptPlugin, bHasRequiredProjectScriptPlugin, UbtPlugins?.ToArray(), EnabledUbtPlugins?.ToArray(), 
+				EnabledUhtPlugins?.ToArray());
 
 			// Get diagnostic info to be printed before each build
 			TargetToolChain.GetVersionInfo(Makefile.Diagnostics);
@@ -3591,16 +3597,29 @@ namespace UnrealBuildTool
 			}
 
 			// Collect the plugins
-			(FileReference ProjectFile, FileReference TargetAssembly)[]? BuiltPlugins;
-			if (!PluginsBase.EnumerateUbtPlugins(ProjectFile, out BuiltPlugins))
-			{
-				throw new BuildException("One or more UBT plugins failed to compile.");
-			}
+			UbtPlugins = PluginsBase.EnumerateUbtPlugins(ProjectFile);
 
-			// Save the collection
-			if (BuiltPlugins != null)
+			// If we found possible plugins
+			(FileReference ProjectFile, FileReference TargetAssembly)[]? BuiltPlugins = null;
+			if (UbtPlugins != null)
 			{
-				UbtPlugins = BuiltPlugins.Select(P => P.ProjectFile).ToList();
+
+				// Filter the UBT plugins based on the enabled plugins
+				EnabledUbtPlugins = new List<FileReference>();
+				EnabledUhtPlugins = new List<FileReference>();
+				foreach (UEBuildPlugin Plugin in BuildPlugins)
+				{
+					PluginInfo Info = Plugin.Info;
+					EnabledUbtPlugins.AddRange(UbtPlugins.Where(P => P.IsUnderDirectory(Info.Directory)));
+				}
+				EnabledUbtPlugins.SortBy(P => P.FullName);
+
+				// Build the plugins
+				bool bCompiled = PluginsBase.BuildUbtPlugins(ProjectFile, EnabledUbtPlugins, out BuiltPlugins);
+				if (!bCompiled)
+				{
+					throw new BuildException("Not all plugins compiled");
+				}
 			}
 
 			// For all the plugins 
@@ -3618,11 +3637,7 @@ namespace UnrealBuildTool
 						{
 							if (UhtTables.IsUhtPlugin(BuiltPlugin.TargetAssembly.FullName))
 							{
-								if (UhtPlugins == null)
-								{
-									UhtPlugins = new List<FileReference>();
-								}
-								UhtPlugins.Add(BuiltPlugin.TargetAssembly);
+								EnabledUhtPlugins!.Add(BuiltPlugin.TargetAssembly);
 								bUbtPluginFound = true;
 								break;
 							}
@@ -3639,6 +3654,11 @@ namespace UnrealBuildTool
 						bHasRequiredProjectScriptPlugin = true;
 					}
 				}
+
+				if (EnabledUhtPlugins != null)
+				{
+					EnabledUhtPlugins.SortBy(P => P.FullName);
+				}
 			}
 		}
 
diff --git a/Engine/Source/Programs/UnrealBuildTool/ProjectFiles/ProjectFileGenerator.cs b/Engine/Source/Programs/UnrealBuildTool/ProjectFiles/ProjectFileGenerator.cs
--- a/Engine/Source/Programs/UnrealBuildTool/ProjectFiles/ProjectFileGenerator.cs
+++ b/Engine/Source/Programs/UnrealBuildTool/ProjectFiles/ProjectFileGenerator.cs
@@ -460,10 +460,13 @@ namespace UnrealBuildTool
 					AddExistingProjectFile(Project, bForceDevelopmentConfiguration: true);
 					AddedProjectFiles.Add(Project);
 
+					// Always create a props file for UBT plugins.  Makes packaging a plugin easier if
+					// it happens to be in the engine
+					bool bCreatePropsFile = RulesFileType == Rules.RulesFileType.UbtPlugin;
+
 					if (!ProjectFile.IsUnderDirectory(Unreal.EngineDirectory))
 					{
-						FileReference PropsFile = new FileReference(ProjectFile.FullName + ".props");
-						CreateProjectPropsFile(PropsFile);
+						bCreatePropsFile = true;
 
 						if (ProjectFile.IsUnderDirectory(SamplesDirectory))
 						{
@@ -478,6 +481,12 @@ namespace UnrealBuildTool
 					{
 						Folder.ChildProjects.Add(Project);
 					}
+
+					if (bCreatePropsFile)
+					{
+						FileReference PropsFile = new FileReference(ProjectFile.FullName + ".props");
+						CreateProjectPropsFile(PropsFile);
+					}
 				}
 			}
 		}
diff --git a/Engine/Source/Programs/UnrealBuildTool/System/ExternalExecution.cs b/Engine/Source/Programs/UnrealBuildTool/System/ExternalExecution.cs
--- a/Engine/Source/Programs/UnrealBuildTool/System/ExternalExecution.cs
+++ b/Engine/Source/Programs/UnrealBuildTool/System/ExternalExecution.cs
@@ -597,9 +597,9 @@ namespace UnrealBuildTool
 		/// Gets the path to the receipt for UHT
 		/// </summary>
 		/// <returns>Path to the UHT receipt</returns>
-		public static FileReference GetHeaderToolReceiptFile(FileReference? ProjectFile, UnrealTargetConfiguration Configuration, bool bHasProjectScriptPlugin, FileReference[]? UhtPlugins)
+		public static FileReference GetHeaderToolReceiptFile(FileReference? ProjectFile, UnrealTargetConfiguration Configuration, bool bHasProjectScriptPlugin, FileReference[]? EnabledUhtPlugins)
 		{
-			if ((bHasProjectScriptPlugin || UhtPlugins != null) && ProjectFile != null)
+			if ((bHasProjectScriptPlugin || (EnabledUhtPlugins != null && EnabledUhtPlugins.Length > 0)) && ProjectFile != null)
 			{
 				return TargetReceipt.GetDefaultPath(ProjectFile.Directory, "UnrealHeaderTool", BuildHostPlatform.Current.Platform, Configuration, "");
 			}
@@ -934,7 +934,7 @@ namespace UnrealBuildTool
 			}
 			else
 			{
-				ExecuteInternalHeaderToolIfNecessary(BuildConfiguration, ProjectFile, Makefile, TargetName, Makefile.UhtPlugins, WorkingSet);
+				ExecuteInternalHeaderToolIfNecessary(BuildConfiguration, ProjectFile, Makefile, TargetName, WorkingSet);
 			}
 		}
 
@@ -992,9 +992,9 @@ namespace UnrealBuildTool
 			}
 
 			List<string> UhtPlugins = new List<string>();
-			if (Makefile.UhtPlugins != null)
+			if (Makefile.EnabledUhtPlugins != null)
 			{
-				foreach (FileReference UhtPlugin in Makefile.UhtPlugins)
+				foreach (FileReference UhtPlugin in Makefile.EnabledUhtPlugins)
 				{
 					UhtPlugins.Add(UhtPlugin.FullName);
 				}
@@ -1048,7 +1048,7 @@ namespace UnrealBuildTool
 				UnrealTargetConfiguration UHTConfig = BuildConfiguration.bForceDebugUnrealHeaderTool ? UnrealTargetConfiguration.Debug : UnrealTargetConfiguration.Development;
 
 				// Figure out the receipt path
-				FileReference HeaderToolReceipt = GetHeaderToolReceiptFile(ProjectFile, UHTConfig, Makefile.bHasProjectScriptPlugin, Makefile.UhtPlugins);
+				FileReference HeaderToolReceipt = GetHeaderToolReceiptFile(ProjectFile, UHTConfig, Makefile.bHasProjectScriptPlugin, Makefile.EnabledUhtPlugins);
 
 				// check if UHT is out of date
 				DateTime HeaderToolTimestampUtc = DateTime.MaxValue;
@@ -1246,7 +1246,7 @@ namespace UnrealBuildTool
 		/// Performs any early outs if headers need no changes, given the UObject modules, tool path, game name, and configuration
 		/// </summary>
 		private static void ExecuteInternalHeaderToolIfNecessary(BuildConfiguration BuildConfiguration, FileReference? ProjectFile, 
-			TargetMakefile Makefile, string TargetName, FileReference[]? UhtPlugins, ISourceFileWorkingSet WorkingSet)
+			TargetMakefile Makefile, string TargetName, ISourceFileWorkingSet WorkingSet)
 		{
 			if (ProgressWriter.bWriteMarkup)
 			{
@@ -1259,7 +1259,7 @@ namespace UnrealBuildTool
 				UnrealTargetConfiguration UHTConfig = BuildConfiguration.bForceDebugUnrealHeaderTool ? UnrealTargetConfiguration.Debug : UnrealTargetConfiguration.Development;
 
 				// Figure out the receipt path
-				FileReference HeaderToolReceipt = GetHeaderToolReceiptFile(ProjectFile, UHTConfig, Makefile.bHasProjectScriptPlugin, Makefile.UhtPlugins);
+				FileReference HeaderToolReceipt = GetHeaderToolReceiptFile(ProjectFile, UHTConfig, Makefile.bHasProjectScriptPlugin, Makefile.EnabledUhtPlugins);
 
 				// Get UHT assembly timestamp
 				DateTime CompositeTimestamp = new FileInfo(Assembly.GetExecutingAssembly().Location).LastWriteTimeUtc;
@@ -1268,9 +1268,9 @@ namespace UnrealBuildTool
 				{
 					CompositeTimestamp = Timestamp;
 				}
-				if (UhtPlugins != null)
+				if (Makefile.EnabledUhtPlugins != null)
 				{
-					foreach (FileReference Plugin in UhtPlugins)
+					foreach (FileReference Plugin in Makefile.EnabledUhtPlugins)
 					{
 						Timestamp = new FileInfo(Plugin.FullName).LastWriteTimeUtc;
 						if (CompositeTimestamp < Timestamp)
diff --git a/Engine/Source/Programs/UnrealBuildTool/System/TargetMakefile.cs b/Engine/Source/Programs/UnrealBuildTool/System/TargetMakefile.cs
--- a/Engine/Source/Programs/UnrealBuildTool/System/TargetMakefile.cs
+++ b/Engine/Source/Programs/UnrealBuildTool/System/TargetMakefile.cs
@@ -20,7 +20,7 @@ namespace UnrealBuildTool
 		/// <summary>
 		/// The version number to write
 		/// </summary>
-		public const int CurrentVersion = 29;
+		public const int CurrentVersion = 123456;
 
 		/// <summary>
 		/// The time at which the makefile was created
@@ -83,14 +83,19 @@ namespace UnrealBuildTool
 		public bool bHasRequiredProjectScriptPlugin;
 
 		/// <summary>
-		/// Collection of UBT C# plugins (project file)
+		/// Collection of all located UBT C# plugins regardless of if they are in an enabled plugin (project file)
 		/// </summary>
 		public FileReference[]? UbtPlugins;
 
 		/// <summary>
-		/// Collection of UHT C# plugins (target assembly file)
+		/// Collection of UBT C# plugins contained in enabled plugins (project file)
 		/// </summary>
-		public FileReference[]? UhtPlugins;
+		public FileReference[]? EnabledUbtPlugins;
+
+		/// <summary>
+		/// Collection of UHT C# plugins contained in enabled plugins (target assembly file)
+		/// </summary>
+		public FileReference[]? EnabledUhtPlugins;
 
 		/// <summary>
 		/// The array of command-line arguments. The makefile will be invalidated whenever these change.
@@ -213,10 +218,11 @@ namespace UnrealBuildTool
 		/// <param name="bHasProjectScriptPlugin">Whether the target has a project script plugin</param>
 		/// <param name="bHasRequiredProjectScriptPlugin">If true, we have a script plugin without a matching UBT plugin</param>
 		/// <param name="UbtPlugins">Collection of UBT plugins</param>
-		/// <param name="UhtPlugins">Collection of UBT plugins for UHT</param>
+		/// <param name="EnabledUbtPlugins">Collection of UBT plugins</param>
+		/// <param name="EnabledUhtPlugins">Collection of UBT plugins for UHT</param>
 		public TargetMakefile(string ExternalMetadata, FileReference ExecutableFile, FileReference ReceiptFile, DirectoryReference ProjectIntermediateDirectory, 
 			TargetType TargetType, ConfigValueTracker ConfigValueTracker, bool bDeployAfterCompile, bool bHasProjectScriptPlugin,
-			bool bHasRequiredProjectScriptPlugin, FileReference[]? UbtPlugins, FileReference[]? UhtPlugins)
+			bool bHasRequiredProjectScriptPlugin, FileReference[]? UbtPlugins, FileReference[]? EnabledUbtPlugins, FileReference[]? EnabledUhtPlugins)
 		{
 			this.CreateTimeUtc = UnrealBuildTool.StartTimeUtc;
 			this.ModifiedTimeUtc = CreateTimeUtc;
@@ -231,7 +237,8 @@ namespace UnrealBuildTool
 			this.bHasProjectScriptPlugin = bHasProjectScriptPlugin;
 			this.bHasRequiredProjectScriptPlugin = bHasRequiredProjectScriptPlugin;
 			this.UbtPlugins = UbtPlugins;
-			this.UhtPlugins = UhtPlugins;
+			this.EnabledUbtPlugins = EnabledUbtPlugins;
+			this.EnabledUhtPlugins = EnabledUhtPlugins;
 			this.Actions = new List<IExternalAction>();
 			this.OutputItems = new List<FileItem>();
 			this.ModuleNameToOutputItems = new Dictionary<string, FileItem[]>(StringComparer.OrdinalIgnoreCase);
@@ -270,7 +277,8 @@ namespace UnrealBuildTool
 			bHasProjectScriptPlugin = Reader.ReadBool();
 			bHasRequiredProjectScriptPlugin = Reader.ReadBool();
 			UbtPlugins = Reader.ReadArray(() => Reader.ReadFileReference())!;
-			UhtPlugins = Reader.ReadArray(() => Reader.ReadFileReference())!;
+			EnabledUbtPlugins = Reader.ReadArray(() => Reader.ReadFileReference())!;
+			EnabledUhtPlugins = Reader.ReadArray(() => Reader.ReadFileReference())!;
 			AdditionalArguments = Reader.ReadArray(() => Reader.ReadString())!;
 			UHTAdditionalArguments = Reader.ReadArray(() => Reader.ReadString())!;
 			PreBuildScripts = Reader.ReadArray(() => Reader.ReadFileReference())!;
@@ -313,7 +321,8 @@ namespace UnrealBuildTool
 			Writer.WriteBool(bHasProjectScriptPlugin);
 			Writer.WriteBool(bHasRequiredProjectScriptPlugin);
 			Writer.WriteArray(UbtPlugins, Item => Writer.WriteFileReference(Item));
-			Writer.WriteArray(UhtPlugins, Item => Writer.WriteFileReference(Item));
+			Writer.WriteArray(EnabledUbtPlugins, Item => Writer.WriteFileReference(Item));
+			Writer.WriteArray(EnabledUhtPlugins, Item => Writer.WriteFileReference(Item));
 			Writer.WriteArray(AdditionalArguments, Item => Writer.WriteString(Item));
 			Writer.WriteArray(UHTAdditionalArguments, Item => Writer.WriteString(Item));
 			Writer.WriteArray(PreBuildScripts, Item => Writer.WriteFileReference(Item));
@@ -368,7 +377,6 @@ namespace UnrealBuildTool
 		/// <returns>The loaded makefile, or null if it failed for some reason.  On failure, the 'ReasonNotLoaded' variable will contain information about why</returns>
 		public static TargetMakefile? Load(FileReference MakefilePath, FileReference? ProjectFile, UnrealTargetPlatform Platform, string[] Arguments, out string? ReasonNotLoaded)
 		{
-			(FileReference ProjectFile, FileReference TargetAssembly)[]? Plugins;
 			FileInfo MakefileInfo;
 			using (GlobalTracer.Instance.BuildSpan("Checking dependent timestamps").StartActive())
 			{
@@ -435,33 +443,6 @@ namespace UnrealBuildTool
 						return null;
 					}
 				}
-
-				// Check for any out of date UBT plugins
-				{
-					if (PluginsBase.EnumerateUbtPlugins(ProjectFile, out Plugins))
-					{
-						if (Plugins != null)
-						{
-							foreach ((FileReference ProjectFile, FileReference TargetAssembly) Plugin in Plugins)
-							{
-								DateTime PluginTimestamp = new FileInfo(Plugin.TargetAssembly.FullName).LastWriteTime;
-								if (MakefileInfo.LastWriteTime.CompareTo(PluginTimestamp) < 0)
-								{
-									// UnrealBuildTool was compiled more recently than the makefile
-									Log.TraceLog($"Makefile is older than {Plugin.TargetAssembly.GetFileNameWithoutExtension()} UBT plugin assembly, ignoring it");
-									ReasonNotLoaded = "UBT plugin assembly is newer";
-									return null;
-								}
-							}
-						}
-					}
-					else
-					{
-						Log.TraceLog($"Error collecting plugins");
-						ReasonNotLoaded = "Error collecting plugins";
-						return null;
-					}
-				}
 			}
 
 			TargetMakefile Makefile;
@@ -527,21 +508,34 @@ namespace UnrealBuildTool
 
 				// Check to see if the number of plugins changed
 				{
-					// Make sure the list of plugins hasn't changed
-					int NewCount = Plugins != null ? Plugins.Length : 0;
-					int OldCount = Makefile.UbtPlugins != null ? Makefile.UbtPlugins.Length : 0;
-					bool bSame = NewCount == OldCount;
-					if (bSame && Plugins != null)
-					{
-						bSame = Enumerable.SequenceEqual(Plugins.Select(P => P.ProjectFile), Makefile.UbtPlugins!);
-					}
-					if (!bSame)
+					FileReference[]? Plugins = PluginsBase.EnumerateUbtPlugins(ProjectFile)?.ToArray();
+
+					if ((Plugins != null) != (Makefile.UbtPlugins != null) ||
+						(Plugins != null && !Enumerable.SequenceEqual(Plugins, Makefile.UbtPlugins!)))
 					{
 						Log.TraceLog("Available UBT plugins changed");
 						ReasonNotLoaded = "Available UBT plugins changed";
 						return null;
 					}
 				}
+
+				// Check to make sure the enabled UBT plugins are up to date
+				if (Makefile.EnabledUbtPlugins != null)
+				{
+					List<DirectoryReference> BaseDirectories = new List<DirectoryReference>();
+					BaseDirectories.Add(Unreal.EngineDirectory);
+					if (ProjectFile != null)
+					{
+						BaseDirectories.Add(ProjectFile.Directory);
+					}
+					HashSet<FileReference> EnabledPlugins = new HashSet<FileReference>(Makefile.EnabledUbtPlugins);
+					if (!CompileScriptModule.AreScriptModulesUpToDate(EnabledPlugins, BaseDirectories))
+					{
+						Log.TraceLog("Enabled UBT plugins need to be recompiled");
+						ReasonNotLoaded = "Enabled UBT plugins need to be recompiled";
+						return null;
+					}
+				}
 			}
 
 			// The makefile is ok
