Low level tests further refactoring and test resource deployment
#rnx

diff --git a/Engine/Build/LowLevelTests_GenProps.xml b/Engine/Build/LowLevelTests_GenProps.xml
--- a/Engine/Build/LowLevelTests_GenProps.xml
+++ b/Engine/Build/LowLevelTests_GenProps.xml
@@ -1,13 +1,18 @@
 ﻿<?xml version="1.0" encoding="utf-8"?>
 <BuildGraph xmlns:xsi="http://www.w3.org/2001/XMLSchema-instance" xsi:schemaLocation="http://www.epicgames.com/BuildGraph ../../Build/Graph/Schema.xsd" xmlns="http://www.epicgames.com/BuildGraph">
   <Property Name="TestNames" Value="HeadlessChaos" />
+  <Property Name="SelfDisabled" Value="False" />
   <Property Name="HeadlessChaosDisabled" Value="False" />
+  <Property Name="SelfShort" Value="Self" />
   <Property Name="HeadlessChaosShort" Value="Headless Chaos" />
+  <Property Name="SelfTarget" Value="LowLevelTests" />
   <Property Name="HeadlessChaosTarget" Value="HeadlessChaos" />
+  <Property Name="SelfBinariesRelative" Value="Engine\Binaries" />
   <Property Name="HeadlessChaosBinariesRelative" Value="Engine\Binaries" />
-  <Property Name="HeadlessChaosResources" Value="" />
-  <Property Name="HeadlessChaosResourcesDest" Value="" />
+  <Property Name="SelfReportType" Value="console" />
   <Property Name="HeadlessChaosReportType" Value="console" />
+  <Option Name="RunSelfTests" DefaultValue="False" Description="Run Self Tests" />
   <Option Name="RunHeadlessChaosTests" DefaultValue="False" Description="Run Headless Chaos Tests" />
+  <Property Name="SelfSupportedPlatforms" Value="Win64;Mac;Linux;Android" />
   <Property Name="HeadlessChaosSupportedPlatforms" Value="Win64" />
 </BuildGraph>
\ No newline at end of file
diff --git a/Engine/Source/Developer/LowLevelTestsRunner/LowLevelTestsRunner.Build.cs b/Engine/Source/Developer/LowLevelTestsRunner/LowLevelTestsRunner.Build.cs
--- a/Engine/Source/Developer/LowLevelTestsRunner/LowLevelTestsRunner.Build.cs
+++ b/Engine/Source/Developer/LowLevelTestsRunner/LowLevelTestsRunner.Build.cs
@@ -31,21 +31,5 @@ public class LowLevelTestsRunner : ModuleRules
 				"Projects"
 			}
 		);
-
-		if (Target.bCompileAgainstApplicationCore)
-		{
-			if (!PrivateDependencyModuleNames.Contains("ApplicationCore"))
-			{
-				PrivateDependencyModuleNames.Add("ApplicationCore");
-			}
-		}
-
-		if (Target.bCompileAgainstCoreUObject)
-		{
-			if (!PrivateDependencyModuleNames.Contains("CoreUObject"))
-			{
-				PrivateDependencyModuleNames.Add("CoreUObject");
-			}
-		}
 	}
 }
\ No newline at end of file
diff --git a/Engine/Source/Programs/HeadlessChaos/HeadlessChaos.Build.cs b/Engine/Source/Programs/HeadlessChaos/HeadlessChaos.Build.cs
--- a/Engine/Source/Programs/HeadlessChaos/HeadlessChaos.Build.cs
+++ b/Engine/Source/Programs/HeadlessChaos/HeadlessChaos.Build.cs
@@ -1,17 +1,11 @@
 // Copyright Epic Games, Inc. All Rights Reserved.
 
-using System.IO;
 using UnrealBuildTool;
 
 [SupportedPlatforms("Win64")]
-public class HeadlessChaos : LowLevelTests
+public class HeadlessChaos : TestModuleRules
 {
-	public override string TestName => "HeadlessChaos";
-	public override string TestShortName => "Headless Chaos";
-
-	public override bool UsesCatch2 => false;
-
-	public HeadlessChaos(ReadOnlyTargetRules Target) : base(Target)
+	public HeadlessChaos(ReadOnlyTargetRules Target) : base(Target, false)
 	{
 		PublicIncludePaths.Add("Runtime/Launch/Public");
 
@@ -55,5 +49,7 @@ public class HeadlessChaos : LowLevelTests
 		}
 
 		PrivateDefinitions.Add("CHAOS_INCLUDE_LEVEL_1=1");
+
+		UpdateBuildGraphPropertiesFile(new Metadata("HeadlessChaos", "Headless Chaos"));
 	}
 }
diff --git a/Engine/Source/Programs/HeadlessChaos/HeadlessChaos.Target.cs b/Engine/Source/Programs/HeadlessChaos/HeadlessChaos.Target.cs
--- a/Engine/Source/Programs/HeadlessChaos/HeadlessChaos.Target.cs
+++ b/Engine/Source/Programs/HeadlessChaos/HeadlessChaos.Target.cs
@@ -1,20 +1,14 @@
 // Copyright Epic Games, Inc. All Rights Reserved.
 
 using UnrealBuildTool;
-using System.Collections.Generic;
 
 [SupportedPlatforms("Win64")]
-public class HeadlessChaosTarget : TargetRules
+public class HeadlessChaosTarget : TestTargetRules
 {
 	public HeadlessChaosTarget(TargetInfo Target) : base(Target)
 	{
-		Type = TargetType.Program;
-		LinkType = TargetLinkType.Monolithic;
-
 		ExeBinariesSubFolder = LaunchModuleName = "HeadlessChaos";
 
-		bBuildDeveloperTools = false;
-
 		// HeadlessChaos doesn't ever compile with the engine linked in
 		bCompileAgainstEngine = false;
 		bCompileAgainstCoreUObject = true;
diff --git a/Engine/Source/Programs/LowLevelTests/LowLevelTests.Build.cs b/Engine/Source/Programs/LowLevelTests/LowLevelTests.Build.cs
--- a/Engine/Source/Programs/LowLevelTests/LowLevelTests.Build.cs
+++ b/Engine/Source/Programs/LowLevelTests/LowLevelTests.Build.cs
@@ -1,278 +1,18 @@
 // Copyright Epic Games, Inc. All Rights Reserved.
 using System;
-using System.Collections.Generic;
-using System.IO;
-using System.Linq;
-using System.Xml.Linq;
-using UnrealBuildBase;
 using UnrealBuildTool;
 
-public class LowLevelTests : ModuleRules
+public class LowLevelTests : TestModuleRules
 {
-	private readonly XNamespace BuildGraphNamespace = XNamespace.Get("http://www.epicgames.com/BuildGraph");
-	private readonly XNamespace SchemaInstance = XNamespace.Get("http://www.w3.org/2001/XMLSchema-instance");
-	private readonly XNamespace SchemaLocation = XNamespace.Get("http://www.epicgames.com/BuildGraph ../../Build/Graph/Schema.xsd");
-
-	public virtual string TestName => throw new Exception("TestName must be overwritten in subclasses.");
-	public virtual string TestShortName => throw new Exception("TestShortName must be overwritten in subclasses.");
-	
-	public virtual string ReportType => "console";
-	public virtual string ResourcesPath => string.Empty;
-
-	public virtual bool UsesCatch2 => true;
-
-	public virtual bool Disabled => false;
-
-	public LowLevelTests(ReadOnlyTargetRules Target) : base(Target)
+	public LowLevelTests(ReadOnlyTargetRules Target) : base(Target, true)
 	{
-		if (UsesCatch2)
-		{
-			PCHUsage = PCHUsageMode.NoPCHs;
-			PrecompileForTargets = PrecompileTargetsType.None;
-
-			if (Target.Configuration == UnrealTargetConfiguration.Debug && Target.Platform == UnrealTargetPlatform.Linux)
-			{
-				OptimizeCode = CodeOptimization.Never;
-			}
-
-			bAllowConfidentialPlatformDefines = true;
-			bLegalToDistributeObjectCode = true;
-
-			// Required false for catch.hpp
-			bUseUnity = false;
-
-			// Disable exception handling so that tests can assert for exceptions
-			bEnableObjCExceptions = false;
-			bEnableExceptions = false;
-		
-			PrivateDependencyModuleNames.AddRange(
-				new string[] {
+		PrivateDependencyModuleNames.AddRange(
+			new string[] {
 				"Core",
 				"Cbor",
-				"CoreUObject",
-				"Projects",
-				"LowLevelTestsRunner"
-				});
-
-			PrivateIncludePaths.Add("Runtime/Launch/Private");
-			PrivateIncludePathModuleNames.Add("Launch");
-
-			// Platforms specific setup
-			if (Target.Platform == UnrealTargetPlatform.Android)
-			{
-				PublicDefinitions.Add("CATCH_CONFIG_NOSTDOUT");
-			}
-		}
-
-		if (this.GetType() != typeof(LowLevelTests))
-		{
-			string SetupFile = Path.Combine("Private", "setup.cpp");
-			string TeardownFile = Path.Combine("Private", "teardown.cpp");
-			if (File.Exists(Path.Combine(ModuleDirectory.ToString(), SetupFile)))
-			{
-				BuildOrderSettings.AddBuildOrderOverride(SetupFile, SourceFileBuildOrder.First);
-			}
-			if (File.Exists(Path.Combine(ModuleDirectory.ToString(), TeardownFile)))
-			{
-				BuildOrderSettings.AddBuildOrderOverride(TeardownFile, SourceFileBuildOrder.Last);
-			}
-
-			UpdateBuildGraphPropertiesFile();
-
-			PrivateDependencyModuleNames.AddRange(
-				new string[] {
-				"Core",
-				"CoreUObject",
-				});
-		}
-	}
-
-	/// <summary>
-	/// Generates or updates include file for LowLevelTests.xml containing test flags: name, short name, target name, relative binaries path, supported platforms etc.
-	/// <paramref name="TestModule">The test module build class that inherits form LowLevelTests</paramref>
-	/// </summary>
-	private void UpdateBuildGraphPropertiesFile()
-	{
-		bool IsPublic = false;
-		string GeneratedPropertiesScriptFile;
-
-		string RestrictedFolder = Path.Combine(Unreal.EngineDirectory.FullName, "Restricted");
-		string NotForLicenseesFolder = Path.Combine(RestrictedFolder, "NotForLicensees");
-		string NonPublicFolder = Path.Combine(NotForLicenseesFolder, "Build");
-		string NonPublicPath = Path.Combine(NonPublicFolder, "LowLevelTests_GenProps.xml");
-
-		if (IsRestrictedPath(ModuleDirectory))
-		{
-			GeneratedPropertiesScriptFile = NonPublicPath;
-		}
-		else
-		{
-			IsPublic = true;
-			GeneratedPropertiesScriptFile = Path.Combine(Unreal.EngineDirectory.FullName, "Build", "LowLevelTests_GenProps.xml");
-		}
-
-		// UE-133126
-		if (File.Exists(NonPublicPath) && Directory.GetFileSystemEntries(NonPublicFolder).Length == 1)
-		{
-			File.Delete(NonPublicPath);
-			Directory.Delete(NonPublicFolder);
-			Directory.Delete(NotForLicenseesFolder);
-			Directory.Delete(RestrictedFolder);
-		}
-
-		if (!File.Exists(GeneratedPropertiesScriptFile))
-		{
-			Directory.CreateDirectory(Path.GetDirectoryName(GeneratedPropertiesScriptFile));
-			using (FileStream fileStream = File.Create(GeneratedPropertiesScriptFile))
-			{
-				XDocument initFile = new XDocument(new XElement(BuildGraphNamespace + "BuildGraph", new XAttribute(XNamespace.Xmlns + "xsi", SchemaInstance), new XAttribute(SchemaInstance + "schemaLocation", SchemaLocation)));
-				initFile.Root.Add(
-					new XElement(
-						BuildGraphNamespace + "Property",
-						new XAttribute("Name", "TestNames" + (!IsPublic ? "Restricted" : "")),
-						new XAttribute("Value", string.Empty)));
+				"CoreUObject"
+			});
 
-				initFile.Save(fileStream);
-			}
-		}
-
-		// All relevant properties
-		string TestTargetName = Target.LaunchModuleName;
-		string TestBinariesPath = TryGetBinariesPath();
-		string TestResourcesPath = ResourcesPath;
-
-		// Do not save full paths
-		if (Path.IsPathRooted(TestBinariesPath))
-		{
-			TestBinariesPath = Path.GetRelativePath(Unreal.RootDirectory.FullName, TestBinariesPath);
-		}
-		if (Path.IsPathRooted(ResourcesPath))
-		{
-			TestResourcesPath = Path.GetRelativePath(Unreal.RootDirectory.FullName, ResourcesPath);
-		}
-
-		XDocument GenPropsDoc = XDocument.Load(GeneratedPropertiesScriptFile);
-		XElement Root = GenPropsDoc.Root;
-		// First descendant must be TestNames
-		XElement TestNames = (XElement)Root.FirstNode;
-		string AllTestNames = TestNames.Attribute("Value").Value;
-		if (!AllTestNames.Contains(TestName))
-		{
-			if (string.IsNullOrEmpty(AllTestNames))
-			{
-				AllTestNames += TestName;
-			}
-			else if (!AllTestNames.Contains(TestName))
-			{
-				AllTestNames += ";" + TestName;
-			}
-		}
-		TestNames.Attribute("Value").SetValue(AllTestNames);
-
-		XElement lastUpdatedNode = TestNames;
-		InsertOrUpdateTestFlag(ref lastUpdatedNode, TestName, "Disabled", Disabled.ToString());
-		InsertOrUpdateTestFlag(ref lastUpdatedNode, TestName, "Short", TestShortName);
-		InsertOrUpdateTestFlag(ref lastUpdatedNode, TestName, "Target", TestTargetName);
-		InsertOrUpdateTestFlag(ref lastUpdatedNode, TestName, "BinariesRelative", TestBinariesPath);
-		InsertOrUpdateTestFlag(ref lastUpdatedNode, TestName, "Resources", TestResourcesPath);
-		InsertOrUpdateTestFlag(ref lastUpdatedNode, TestName, "ResourcesDest", Path.GetFileName(TestResourcesPath));
-		InsertOrUpdateTestFlag(ref lastUpdatedNode, TestName, "ReportType", ReportType.ToString());
-
-		InsertOrUpdateTestOption(ref lastUpdatedNode, TestName, TestShortName, "Run", "Tests", false.ToString());
-
-		List<UnrealTargetPlatform> AllSupportedPlatforms = new List<UnrealTargetPlatform>();
-		var SupportedPlatforms = GetType().GetCustomAttributes(typeof(SupportedPlatformsAttribute), false);
-		// If none specified we assume all platforms are supported by default
-		if (SupportedPlatforms.Length == 0)
-		{
-			AllSupportedPlatforms.AddRange(UnrealTargetPlatform.GetValidPlatforms().ToList());
-		}
-		else
-		{
-			foreach (var Platform in SupportedPlatforms)
-			{
-				AllSupportedPlatforms.AddRange(((SupportedPlatformsAttribute)Platform).Platforms);
-			}
-		}
-
-		InsertOrUpdateTestFlag(ref lastUpdatedNode, TestName, "SupportedPlatforms", AllSupportedPlatforms.Aggregate("", (current, next) => (current == "" ? next.ToString() : current + ";" + next.ToString())));
-
-		try
-		{
-			GenPropsDoc.Save(GeneratedPropertiesScriptFile);
-		}
-		catch (UnauthorizedAccessException)
-		{
-			// Expected on build machines.
-			// TODO: Ability to build for generate files and runnable tests.
-		}
-	}
-
-	private bool IsRestrictedPath(string ModuleDirectory)
-	{
-		foreach(string RestrictedFolderName in RestrictedFolder.GetNames())
-		{
-			if (ModuleDirectory.Contains(RestrictedFolderName))
-			{
-				return true;
-			}
-		}
-
-		return false;
-	}
-
-	private string TryGetBinariesPath()
-	{
-		int SourceFolderIndex = ModuleDirectory.IndexOf("Source");
-		if (SourceFolderIndex < 0)
-		{
-			throw new Exception("Could not detect source folder path for module " + GetType());
-		}
-		return ModuleDirectory.Substring(0, SourceFolderIndex) + "Binaries";
-	}
-
-	private void InsertOrUpdateTestFlag(ref XElement ElementUpsertAfter, string TestName, string FlagSuffix, string FlagValue)
-	{
-		IEnumerable<XElement> NextChunk = ElementUpsertAfter.ElementsAfterSelf(BuildGraphNamespace + "Property")
-			.Where(prop => prop.Attribute("Name").Value.EndsWith(FlagSuffix));
-		if (NextChunk
-			.Where(prop => prop.Attribute("Name").Value == TestName + FlagSuffix)
-			.Count() == 0)
-		{
-			XElement ElementInsert = new XElement(BuildGraphNamespace + "Property");
-			ElementInsert.SetAttributeValue("Name", TestName + FlagSuffix);
-			ElementInsert.SetAttributeValue("Value", FlagValue);
-			ElementUpsertAfter.AddAfterSelf(ElementInsert);
-		}
-		else
-		{
-			NextChunk
-				.Where(prop => prop.Attribute("Name").Value == TestName + FlagSuffix).First().SetAttributeValue("Value", FlagValue);
-		}
-		ElementUpsertAfter = NextChunk.Last();
-	}
-
-	private void InsertOrUpdateTestOption(ref XElement ElementUpsertAfter, string TestName, string TestShortName, string OptionPrefix, string OptionSuffix, string DefaultValue)
-	{
-		IEnumerable<XElement> NextChunk = ElementUpsertAfter.ElementsAfterSelf(BuildGraphNamespace + "Option")
-			.Where(prop => prop.Attribute("Name").Value.StartsWith(OptionPrefix) && prop.Attribute("Name").Value.EndsWith(OptionSuffix));
-		if (NextChunk
-			.Where(prop => prop.Attribute("Name").Value == OptionPrefix + TestName + OptionSuffix)
-			.Count() == 0)
-		{
-			XElement ElementInsert = new XElement(BuildGraphNamespace + "Option");
-			ElementInsert.SetAttributeValue("Name", OptionPrefix + TestName + OptionSuffix);
-			ElementInsert.SetAttributeValue("DefaultValue", DefaultValue);
-			ElementInsert.SetAttributeValue("Description", string.Format("{0} {1} {2}", OptionPrefix, TestShortName, OptionSuffix));
-			ElementUpsertAfter.AddAfterSelf(ElementInsert);
-		}
-		else
-		{
-			XElement ElementUpdate = NextChunk
-				.Where(prop => prop.Attribute("Name").Value == OptionPrefix + TestName + OptionSuffix).First();
-			ElementUpdate.SetAttributeValue("Description", string.Format("{0} {1} {2}", OptionPrefix, TestShortName, OptionSuffix));
-			ElementUpdate.SetAttributeValue("DefaultValue", DefaultValue);
-		}
-		ElementUpsertAfter = NextChunk.Last();
+		UpdateBuildGraphPropertiesFile(new Metadata("Self"), false);
 	}
 }
\ No newline at end of file
diff --git a/Engine/Source/Programs/LowLevelTests/LowLevelTests.Target.cs b/Engine/Source/Programs/LowLevelTests/LowLevelTests.Target.cs
--- a/Engine/Source/Programs/LowLevelTests/LowLevelTests.Target.cs
+++ b/Engine/Source/Programs/LowLevelTests/LowLevelTests.Target.cs
@@ -1,96 +1,12 @@
 // Copyright Epic Games, Inc. All Rights Reserved.
 
-using System.IO;
-using EpicGames.Core;
-using UnrealBuildBase;
 using UnrealBuildTool;
 
 [SupportedPlatforms(UnrealPlatformClass.All)]
-public class LowLevelTestsTarget : TargetRules
+public class LowLevelTestsTarget : TestTargetRules
 {
 	public LowLevelTestsTarget(TargetInfo Target) : base(Target)
 	{
-		ExeBinariesSubFolder = LaunchModuleName = "LowLevelTests";
-
-		SolutionDirectory = "Programs/LowLevelTests";
-
-		Type = TargetType.Program;
-		LinkType = TargetLinkType.Monolithic;
-
-		bDeployAfterCompile = Target.Platform != UnrealTargetPlatform.Android;
-		bIsBuildingConsoleApplication = true;
-
-		// Disabling default true flags that aren't necessary for tests
-
-		// Lean and Mean mode
-		bBuildDeveloperTools = false;
-		
-		// No localization
-		bCompileICU = false;
-
-		// No need for shaders by default
-		bForceBuildShaderFormats = false;
-
-		// Do not link against the engine, no Chromium Embedded Framework etc.
-		bCompileAgainstEngine = false;
-		bCompileCEF3 = false;
-		bCompileAgainstCoreUObject = false;
-		bCompileAgainstApplicationCore = false;
-		bUseLoggingInShipping = true;
-		bForceEnableExceptions = true;
-
-		bool bDebugOrDevelopment = Target.Configuration == UnrealTargetConfiguration.Debug || Target.Configuration == UnrealTargetConfiguration.Development;
-		bBuildWithEditorOnlyData = Target.Platform.IsInGroup(UnrealPlatformGroup.Desktop) && bDebugOrDevelopment;
-
-		// Disable malloc profiling in tests
-		bUseMallocProfiler = false;
-
-		// Useful for debugging test failures
-		if (Target.Configuration == UnrealTargetConfiguration.Debug)
-		{
-			bDebugBuildsActuallyUseDebugCRT = true;
-		}
-
-		VSTestRunSettingsFile = FileReference.Combine(Unreal.EngineDirectory, "Source", "Programs", "LowLevelTests", "vstest.runsettings");
-
-		if (this.GetType().Name.StartsWith("LowLevelTests"))
-		{
-			bIsTestTargetOverride = true;
-		}
-
-		SetupPreprocessorDefinitions(Target);
-	}
-
-	protected void SetupModule()
-	{
-		LaunchModuleName = this.GetType().Name.Replace("Target", string.Empty);
-		ExeBinariesSubFolder = LaunchModuleName;
-	}
-
-	protected virtual void SetupPreprocessorDefinitions(TargetInfo Target)
-	{
-		GlobalDefinitions.Add("STATS=0");
-
-		// Platform specific setup
-		if (Target.Platform == UnrealTargetPlatform.Android)
-		{
-			UndecoratedConfiguration = Target.Configuration;
-
-			string VersionScriptFile = Path.Combine(Directory.GetCurrentDirectory(), @"Developer\LowLevelTestsRunner\Private\Platform\Android\HideSymbols.ldscript");
-			AdditionalLinkerArguments = " -Wl,--version-script=" + VersionScriptFile;
-
-			GlobalDefinitions.Add("USE_ANDROID_INPUT=0");
-			GlobalDefinitions.Add("USE_ANDROID_OPENGL=0");
-			GlobalDefinitions.Add("USE_ANDROID_LAUNCH=0");
-			GlobalDefinitions.Add("USE_ANDROID_JNI=0");
-		}
-		else if (Target.Platform == UnrealTargetPlatform.IOS)
-		{
-			GlobalDefinitions.Add("HAS_METAL=0");
-
-			bIsBuildingConsoleApplication = false;
-			// Required for IOS, but needs to fix compilation errors
-			bCompileAgainstApplicationCore = true;
-		}
+		bWithLowLevelTestsOverride = true;
 	}
 }
diff --git a/Engine/Source/Programs/UnrealBuildTool/Configuration/ModuleRules.cs b/Engine/Source/Programs/UnrealBuildTool/Configuration/ModuleRules.cs
--- a/Engine/Source/Programs/UnrealBuildTool/Configuration/ModuleRules.cs
+++ b/Engine/Source/Programs/UnrealBuildTool/Configuration/ModuleRules.cs
@@ -1406,61 +1406,60 @@ namespace UnrealBuildTool
 
 		/// <summary>
 		/// Prepares a module for building a low level tests executable.
-		/// If the target includes all tests or we're building a module as part of a test module chain, then we include the LowLevelTestsRunner dependency.
-		/// We also keep track of any Editor and Engine dependencies and add appropriate global definitions.
+		/// If we're building a module as part of a test module chain and there's a Tests folder with low level tests, then they require the LowLevelTestsRunner dependency.
+		/// We also keep track of any Editor, Engine and other conditionally compiled dependencies.
 		/// </summary>
 		internal void PrepareModuleForTests()
 		{
-			if (!IsTestModule && Target.IsTestTarget)
+			if (Name != "LowLevelTestsRunner" && System.IO.Directory.Exists(TestsDirectory))
 			{
-				if (Name != "LowLevelTestsRunner" && System.IO.Directory.Exists(TestsDirectory))
+				if (!PrivateIncludePathModuleNames.Contains("LowLevelTestsRunner"))
 				{
-					if (!PrivateIncludePathModuleNames.Contains("LowLevelTestsRunner"))
-					{
-						PrivateIncludePathModuleNames.Add("LowLevelTestsRunner");
-					}
-				}
-				else if (Name == "LowLevelTestsRunner")
-				{
-					TestTargetRules.LowLevelTestsRunnerModule = this;
+					PrivateIncludePathModuleNames.Add("LowLevelTestsRunner");
 				}
+			}
+			else if (Name == "LowLevelTestsRunner")
+			{
+				TestTargetRules.LowLevelTestsRunnerModule = this;
+			}
 
-				// If one of these modules is in the dependency graph, we must enable their appropriate global definitions
-				if (Name == "Editor")
-				{
-					TestTargetRules.bTestsRequireEditor = true;
-				}
-				if (Name == "Engine")
+			// If one of these modules is in the dependency graph, we must enable their appropriate global definitions
+			if (Name == "UnrealEd")
+			{
+				// TODO: more support code required...
+				TestTargetRules.bTestsRequireEditor = false;
+			}
+			if (Name == "Engine")
+			{
+				// TODO: more support code required...
+				TestTargetRules.bTestsRequireEngine = false;
+			}
+			if (Name == "ApplicationCore")
+			{
+				TestTargetRules.bTestsRequireApplicationCore = true;
+			}
+			if (Name == "CoreUObject")
+			{
+				TestTargetRules.bTestsRequireCoreUObject = true;
+			}
+
+			if (TestTargetRules.LowLevelTestsRunnerModule != null)
+			{
+				if (TestTargetRules.bTestsRequireEditor && !TestTargetRules.LowLevelTestsRunnerModule.PrivateDependencyModuleNames.Contains("UnrealEd"))
 				{
-					TestTargetRules.bTestsRequireEngine = true;
+					TestTargetRules.LowLevelTestsRunnerModule.PrivateDependencyModuleNames.Add("UnrealEd");
 				}
-				if (Name == "ApplicationCore")
+				if (TestTargetRules.bTestsRequireEngine && !TestTargetRules.LowLevelTestsRunnerModule.PrivateDependencyModuleNames.Contains("Engine"))
 				{
-					TestTargetRules.bTestsRequireApplicationCore = true;
+					TestTargetRules.LowLevelTestsRunnerModule.PrivateDependencyModuleNames.Add("Engine");
 				}
-				if (Name == "CoreUObject")
+				if (TestTargetRules.bTestsRequireApplicationCore && !TestTargetRules.LowLevelTestsRunnerModule.PrivateDependencyModuleNames.Contains("ApplicationCore"))
 				{
-					TestTargetRules.bTestsRequireCoreUObject = true;
+					TestTargetRules.LowLevelTestsRunnerModule.PrivateDependencyModuleNames.Add("ApplicationCore");
 				}
-
-				if (TestTargetRules.LowLevelTestsRunnerModule != null)
+				if (TestTargetRules.bTestsRequireCoreUObject && !TestTargetRules.LowLevelTestsRunnerModule.PrivateDependencyModuleNames.Contains("CoreUObject"))
 				{
-					if (TestTargetRules.bTestsRequireEditor && !TestTargetRules.LowLevelTestsRunnerModule.PrivateDependencyModuleNames.Contains("Editor"))
-					{
-						TestTargetRules.LowLevelTestsRunnerModule.PrivateDependencyModuleNames.Add("Editor");
-					}
-					if (TestTargetRules.bTestsRequireEngine && !TestTargetRules.LowLevelTestsRunnerModule.PrivateDependencyModuleNames.Contains("Engine"))
-					{
-						TestTargetRules.LowLevelTestsRunnerModule.PrivateDependencyModuleNames.Add("Engine");
-					}
-					if (TestTargetRules.bTestsRequireApplicationCore && !TestTargetRules.LowLevelTestsRunnerModule.PrivateDependencyModuleNames.Contains("ApplicationCore"))
-					{
-						TestTargetRules.LowLevelTestsRunnerModule.PrivateDependencyModuleNames.Add("ApplicationCore");
-					}
-					if (TestTargetRules.bTestsRequireCoreUObject && !TestTargetRules.LowLevelTestsRunnerModule.PrivateDependencyModuleNames.Contains("CoreUObject"))
-					{
-						TestTargetRules.LowLevelTestsRunnerModule.PrivateDependencyModuleNames.Add("CoreUObject");
-					}
+					TestTargetRules.LowLevelTestsRunnerModule.PrivateDependencyModuleNames.Add("CoreUObject");
 				}
 			}
 		}
diff --git a/Engine/Source/Programs/UnrealBuildTool/Configuration/TargetRules.cs b/Engine/Source/Programs/UnrealBuildTool/Configuration/TargetRules.cs
--- a/Engine/Source/Programs/UnrealBuildTool/Configuration/TargetRules.cs
+++ b/Engine/Source/Programs/UnrealBuildTool/Configuration/TargetRules.cs
@@ -279,12 +279,41 @@ namespace UnrealBuildTool
 		/// </summary>
 		public bool IsTestTarget
 		{
-			get { return bIsTestTargetOverride ?? false; }
+			get { return bIsTestTargetOverride; }
 		}
 		/// <summary>
 		/// Override this boolean flag in inheriting classes for low level test targets.
 		/// </summary>
-		protected bool? bIsTestTargetOverride;
+		protected bool bIsTestTargetOverride;
+
+		/// <summary>
+		/// Whether this is a test target explicitly defined.
+		/// Explicitley defined test targets always inherit from TestTargetRules and define their own tests.
+		/// Implicit test targets are created from existing targets when building with -Mode=Test and they include tests from all dependencies.
+		/// </summary>
+		public bool ExplicitTestsTarget
+		{
+			get { return bExplicitTestsTargetOverride; }
+		}
+		/// <summary>
+		/// This flag is automatically set when classes inherit from TestTargetRules.
+		/// </summary>
+		protected bool bExplicitTestsTargetOverride;
+
+		/// <summary>
+		/// Controls the value of WITH_LOW_LEVEL_TESTS that dictates whether module-specific low level tests are compiled in or not.
+		/// </summary>
+		public bool WithLowLevelTests
+		{
+			get
+			{
+				return (IsTestTarget && !ExplicitTestsTarget) || bWithLowLevelTestsOverride;
+			}
+		}
+		/// <summary>
+		/// Override the value of WithLowLevelTests by setting this to true in inherited classes.
+		/// </summary>
+		protected bool bWithLowLevelTestsOverride;
 
 		/// <summary>
 		/// File containing the general type for this target (not including platform/group)
@@ -711,19 +740,35 @@ namespace UnrealBuildTool
 		/// Enabled for all builds that include the engine project.  Disabled only when building standalone apps that only link with Core.
 		/// </summary>
 		[RequiresUniqueBuildEnvironment]
-		public bool bCompileAgainstEngine = true;
+		public virtual bool bCompileAgainstEngine
+		{
+			get { return bCompileAgainstEnginePrivate; }
+			set { bCompileAgainstEnginePrivate = value; }
+		}
+		private bool bCompileAgainstEnginePrivate = true;
 
 		/// <summary>
 		/// Enabled for all builds that include the CoreUObject project.  Disabled only when building standalone apps that only link with Core.
 		/// </summary>
 		[RequiresUniqueBuildEnvironment]
-		public bool bCompileAgainstCoreUObject = true;
+		public virtual bool bCompileAgainstCoreUObject
+		{
+			get { return bCompileAgainstCoreUObjectPrivate; }
+			set { bCompileAgainstCoreUObjectPrivate = value; }
+		}
+		private bool bCompileAgainstCoreUObjectPrivate = true;
+
 
 		/// <summary>
 		/// Enabled for builds that need to initialize the ApplicationCore module. Command line utilities do not normally need this.
 		/// </summary>
 		[RequiresUniqueBuildEnvironment]
-		public bool bCompileAgainstApplicationCore = true;
+		public virtual bool bCompileAgainstApplicationCore
+		{
+			get { return bCompileAgainstApplicationCorePrivate; }
+			set { bCompileAgainstApplicationCorePrivate = value; }
+		}
+		private bool bCompileAgainstApplicationCorePrivate = true;
 
 		/// <summary>
 		/// Manually specified value for bCompileAgainstEditor.
@@ -735,7 +780,7 @@ namespace UnrealBuildTool
 		/// Mainly drives the value of WITH_EDITOR.
 		/// </summary>
 		[RequiresUniqueBuildEnvironment]
-		public bool bCompileAgainstEditor
+		public virtual bool bCompileAgainstEditor
 		{
 			set { bCompileAgainstEditorOverride = value; }
 			get { return bCompileAgainstEditorOverride ?? (Type == TargetType.Editor); }
@@ -1957,12 +2002,7 @@ namespace UnrealBuildTool
 		/// <returns>Array of platforms that the target supports</returns>
 		internal UnrealTargetPlatform[] GetSupportedPlatforms()
 		{
-			if (this is TestTargetRules TestTarget)
-			{
-				return TestTarget.TestedTarget.GetSupportedPlatforms();
-			}
-
-			// Otherwise take the SupportedPlatformsAttribute from the first type in the inheritance chain that supports it
+			// Take the SupportedPlatformsAttribute from the first type in the inheritance chain that supports it
 			for (Type? CurrentType = GetType(); CurrentType != null; CurrentType = CurrentType.BaseType)
 			{
 				object[] Attributes = CurrentType.GetCustomAttributes(typeof(SupportedPlatformsAttribute), false);
@@ -2610,6 +2650,7 @@ namespace UnrealBuildTool
 		{
 			get { return Inner.bForceCompilePerformanceAutomationTests; }
 		}
+
 		public bool bForceDisableAutomationTests
 		{
 			get { return Inner.bForceDisableAutomationTests; }
@@ -3174,6 +3215,16 @@ namespace UnrealBuildTool
 			get { return Inner.IsTestTarget; }
 		}
 
+		public bool ExplicitTestsTarget
+		{
+			get { return Inner.ExplicitTestsTarget; }
+		}
+
+		public bool WithLowLevelTests
+		{
+			get { return Inner.WithLowLevelTests; }
+		}
+
 		/// <summary>
 		/// Checks if current platform is part of a given platform group
 		/// </summary>
diff --git a/Engine/Source/Programs/UnrealBuildTool/Configuration/TestModuleRules.cs b/Engine/Source/Programs/UnrealBuildTool/Configuration/TestModuleRules.cs
--- a/Engine/Source/Programs/UnrealBuildTool/Configuration/TestModuleRules.cs
+++ b/Engine/Source/Programs/UnrealBuildTool/Configuration/TestModuleRules.cs
@@ -1,7 +1,12 @@
 // Copyright Epic Games, Inc. All Rights Reserved.
 
-using EpicGames.Core;
+using System.Collections.Generic;
 using System.IO;
+using System.Linq;
+using System.Xml.Linq;
+using UnrealBuildBase;
+using EpicGames.Core;
+using System;
 
 namespace UnrealBuildTool
 {
@@ -10,10 +15,38 @@ namespace UnrealBuildTool
 	/// </summary>
 	public class TestModuleRules : ModuleRules
 	{
+		private readonly XNamespace BuildGraphNamespace = XNamespace.Get("http://www.epicgames.com/BuildGraph");
+		private readonly XNamespace SchemaInstance = XNamespace.Get("http://www.w3.org/2001/XMLSchema-instance");
+		private readonly XNamespace SchemaLocation = XNamespace.Get("http://www.epicgames.com/BuildGraph ../../Build/Graph/Schema.xsd");
+
+		private bool bUsesCatch2 = true;
+
 		/// <summary>
 		/// Associated tested module of this test module.
 		/// </summary>
-		public ModuleRules TestedModule { get; private set; }
+		public ModuleRules? TestedModule { get; private set; }
+
+		/// <summary>
+		/// Constructs a TestModuleRules object as its own test module.
+		/// </summary>
+		/// <param name="Target"></param>
+		public TestModuleRules(ReadOnlyTargetRules Target) : base(Target)
+		{
+			SetupCommonProperties(Target);
+		}
+
+		/// <summary>
+		/// Constructs a TestModuleRules object as its own test module.
+		/// Sets value of bUsesCatch2.
+		/// </summary>
+		public TestModuleRules(ReadOnlyTargetRules Target, bool InUsesCatch2) : base(Target)
+		{
+			bUsesCatch2 = InUsesCatch2;
+			if (bUsesCatch2)
+			{
+				SetupCommonProperties(Target);
+			}
+		}
 
 		/// <summary>
 		/// Constructs a TestModuleRules object with an associated tested module.
@@ -22,8 +55,6 @@ namespace UnrealBuildTool
 		{
 			this.TestedModule = TestedModule;
 
-			bIsTestModuleOverride = true;
-
 			Name = TestedModule.Name + "Tests";
 			if (!string.IsNullOrEmpty(TestedModule.ShortName))
 			{
@@ -35,6 +66,28 @@ namespace UnrealBuildTool
 
 			Context = TestedModule.Context;
 
+			PrivateDependencyModuleNames.AddRange(TestedModule.PrivateDependencyModuleNames);
+
+			// Tests can refer to tested module's Public and Private paths
+			string ModulePublicDir = Path.Combine(TestedModule.ModuleDirectory, "Public");
+			if (System.IO.Directory.Exists(ModulePublicDir))
+			{
+				PublicIncludePaths.Add(ModulePublicDir);
+			}
+
+			string ModulePrivateDir = Path.Combine(TestedModule.ModuleDirectory, "Private");
+			if (System.IO.Directory.Exists(ModulePrivateDir))
+			{
+				PrivateIncludePaths.Add(ModulePrivateDir);
+			}
+
+			SetupCommonProperties(Target);
+		}
+
+		private void SetupCommonProperties(ReadOnlyTargetRules Target)
+		{
+			bIsTestModuleOverride = true;
+
 			PCHUsage = PCHUsageMode.NoPCHs;
 			PrecompileForTargets = PrecompileTargetsType.None;
 
@@ -64,32 +117,302 @@ namespace UnrealBuildTool
 				BuildOrderSettings.AddBuildOrderOverride(TeardownFile, SourceFileBuildOrder.Last);
 			}
 
-			PrivateDependencyModuleNames.AddRange(TestedModule.PrivateDependencyModuleNames);
+			string TestResourcesDir = Path.Combine(ModuleDirectory, "Resources");
+			if (System.IO.Directory.Exists(TestResourcesDir))
+			{
+				AdditionalPropertiesForReceipt.Add("ResourcesFolder", TestResourcesDir);
+			}
 
 			if (!PrivateDependencyModuleNames.Contains("LowLevelTestsRunner"))
 			{
 				PrivateDependencyModuleNames.Add("LowLevelTestsRunner");
 			}
 
-			// Tests can refer to tested module's Public and Private paths
-			string ModulePublicDir = Path.Combine(TestedModule.ModuleDirectory, "Public");
-			if (System.IO.Directory.Exists(ModulePublicDir))
+			// Platforms specific setup
+			if (Target.Platform == UnrealTargetPlatform.Android)
 			{
-				PublicIncludePaths.Add(ModulePublicDir);
+				PublicDefinitions.Add("CATCH_CONFIG_NOSTDOUT");
 			}
+		}
 
-			string ModulePrivateDir = Path.Combine(TestedModule.ModuleDirectory, "Private");
-			if (System.IO.Directory.Exists(ModulePrivateDir))
+		/// <summary>
+		/// Set test-specific resources folder relative to module directory.
+		/// This will be copied to the binaries path during deployment.
+		/// </summary>
+		protected void SetResourcesFolder(string ResourcesRelativeFolder)
+		{
+			string TestResourcesDir = Path.Combine(ModuleDirectory, ResourcesRelativeFolder);
+			if (!System.IO.Directory.Exists(TestResourcesDir))
 			{
-				PrivateIncludePaths.Add(ModulePrivateDir);
+				throw new Exception($"Resources folder not found {TestResourcesDir}");
+			}
+			else
+			{
+				ReceiptProperty? Existing = AdditionalPropertiesForReceipt.Inner.Find(Prop => Prop.Name == "ResourcesFolder");
+				if (Existing != null)
+				{
+					Existing.Value = TestResourcesDir;
+				}
+				else
+				{
+					AdditionalPropertiesForReceipt.Add("ResourcesFolder", TestResourcesDir);
+				}
 			}
-			
+		}
 
-			// Platforms specific setup
-			if (Target.Platform == UnrealTargetPlatform.Android)
+		/// <summary>
+		/// Generates or updates include file for LowLevelTests.xml containing test flags: name, short name, target name, relative binaries path, supported platforms etc.
+		/// <paramref name="TestMetadata">The test metadata specifying name, short name etc used to populate the BuildGraph properties file.</paramref>
+		/// </summary>
+		protected void UpdateBuildGraphPropertiesFile(Metadata TestMetadata, bool AddToTestNames = true)
+		{
+			bool IsPublic = false;
+			string GeneratedPropertiesScriptFile;
+
+			string RestrictedFolder = Path.Combine(Unreal.EngineDirectory.FullName, "Restricted");
+			string NotForLicenseesFolder = Path.Combine(RestrictedFolder, "NotForLicensees");
+			string NonPublicFolder = Path.Combine(NotForLicenseesFolder, "Build");
+			string NonPublicPath = Path.Combine(NonPublicFolder, "LowLevelTests_GenProps.xml");
+
+			if (IsRestrictedPath(ModuleDirectory))
 			{
-				PublicDefinitions.Add("CATCH_CONFIG_NOSTDOUT");
+				GeneratedPropertiesScriptFile = NonPublicPath;
+			}
+			else
+			{
+				IsPublic = true;
+				GeneratedPropertiesScriptFile = Path.Combine(Unreal.EngineDirectory.FullName, "Build", "LowLevelTests_GenProps.xml");
+			}
+
+			// UE-133126
+			if (System.IO.File.Exists(NonPublicPath) && System.IO.Directory.GetFileSystemEntries(NonPublicFolder).Length == 1)
+			{
+				System.IO.File.Delete(NonPublicPath);
+				System.IO.Directory.Delete(NonPublicFolder);
+				System.IO.Directory.Delete(NotForLicenseesFolder);
+				System.IO.Directory.Delete(RestrictedFolder);
+			}
+
+			if (!System.IO.File.Exists(GeneratedPropertiesScriptFile))
+			{
+				string? DirGenProps = Path.GetDirectoryName(GeneratedPropertiesScriptFile);
+				if (DirGenProps != null)
+				{
+					System.IO.Directory.CreateDirectory(DirGenProps);
+				}
+				using FileStream FileStream = System.IO.File.Create(GeneratedPropertiesScriptFile);
+				XDocument XInitFile = new XDocument(new XElement(BuildGraphNamespace + "BuildGraph", new XAttribute(XNamespace.Xmlns + "xsi", SchemaInstance), new XAttribute(SchemaInstance + "schemaLocation", SchemaLocation)));
+				XInitFile.Root?.Add(
+					new XElement(
+						BuildGraphNamespace + "Property",
+						new XAttribute("Name", "TestNames" + (!IsPublic ? "Restricted" : "")),
+						new XAttribute("Value", string.Empty)));
+
+				XInitFile.Save(FileStream);
+			}
+
+			// All relevant properties
+			string TestTargetName = Target.LaunchModuleName ?? "NoLaunchModule";
+			string TestBinariesPath = TryGetBinariesPath();
+
+			// Do not save full paths
+			if (Path.IsPathRooted(TestBinariesPath))
+			{
+				TestBinariesPath = Path.GetRelativePath(Unreal.RootDirectory.FullName, TestBinariesPath);
+			}
+
+			XDocument GenPropsDoc = XDocument.Load(GeneratedPropertiesScriptFile);
+			XElement? Root = GenPropsDoc.Root;
+			// First descendant must be TestNames
+			if (Root != null && Root.FirstNode != null)
+			{
+				XElement TestNames = (XElement)Root.FirstNode;
+				if (TestNames != null)
+				{
+					string? AllTestNames = TestNames.Attribute("Value")?.Value;
+					if (AllTestNames != null && !AllTestNames.Contains(TestMetadata.TestName) && AddToTestNames)
+					{
+						if (string.IsNullOrEmpty(AllTestNames))
+						{
+							AllTestNames += TestMetadata.TestName;
+						}
+						else if (!AllTestNames.Contains(TestMetadata.TestName))
+						{
+							AllTestNames += ";" + TestMetadata.TestName;
+						}
+					}
+					TestNames.Attribute("Value")?.SetValue(AllTestNames ?? "");
+
+
+					XElement lastUpdatedNode = TestNames;
+					InsertOrUpdateTestFlag(ref lastUpdatedNode, TestMetadata.TestName, "Disabled", TestMetadata.Disabled.ToString());
+					InsertOrUpdateTestFlag(ref lastUpdatedNode, TestMetadata.TestName, "Short", TestMetadata.TestShortName);
+					InsertOrUpdateTestFlag(ref lastUpdatedNode, TestMetadata.TestName, "Target", TestTargetName);
+					InsertOrUpdateTestFlag(ref lastUpdatedNode, TestMetadata.TestName, "BinariesRelative", TestBinariesPath);
+					InsertOrUpdateTestFlag(ref lastUpdatedNode, TestMetadata.TestName, "ReportType", TestMetadata.ReportType.ToString());
+
+					InsertOrUpdateTestOption(ref lastUpdatedNode, TestMetadata.TestName, TestMetadata.TestShortName, "Run", "Tests", false.ToString());
+
+					List<UnrealTargetPlatform> AllSupportedPlatforms = new List<UnrealTargetPlatform>();
+					var SupportedPlatforms = GetType().GetCustomAttributes(typeof(SupportedPlatformsAttribute), false);
+					// If none specified we assume all platforms are supported by default
+					if (SupportedPlatforms.Length == 0)
+					{
+						UnrealTargetPlatform[] SupportedByDefault = { UnrealTargetPlatform.Win64, UnrealTargetPlatform.Mac, UnrealTargetPlatform.Linux, UnrealTargetPlatform.Android };
+						AllSupportedPlatforms.AddRange(SupportedByDefault);
+					}
+					else
+					{
+						foreach (var Platform in SupportedPlatforms)
+						{
+							AllSupportedPlatforms.AddRange(((SupportedPlatformsAttribute)Platform).Platforms);
+						}
+					}
+
+					InsertOrUpdateTestFlag(ref lastUpdatedNode, TestMetadata.TestName, "SupportedPlatforms", AllSupportedPlatforms.Aggregate("", (current, next) => (current == "" ? next.ToString() : current + ";" + next.ToString())));
+
+					try
+					{
+						GenPropsDoc.Save(GeneratedPropertiesScriptFile);
+					}
+					catch (UnauthorizedAccessException)
+					{
+						// Expected on build machines.
+						// TODO: Ability to build for generate files and runnable tests.
+					}
+				}
+			}
+		}
+
+		private bool IsRestrictedPath(string ModuleDirectory)
+		{
+			foreach (string RestrictedFolderName in RestrictedFolder.GetNames())
+			{
+				if (ModuleDirectory.Contains(RestrictedFolderName))
+				{
+					return true;
+				}
+			}
+
+			return false;
+		}
+
+		private string TryGetBinariesPath()
+		{
+			int SourceFolderIndex = ModuleDirectory.IndexOf("Source");
+			if (SourceFolderIndex < 0)
+			{
+				throw new Exception("Could not detect source folder path for module " + GetType());
+			}
+			return ModuleDirectory.Substring(0, SourceFolderIndex) + "Binaries";
+		}
+
+#pragma warning disable 8602
+		private void InsertOrUpdateTestFlag(ref XElement ElementUpsertAfter, string TestName, string FlagSuffix, string FlagValue)
+		{
+			IEnumerable<XElement> NextChunk = ElementUpsertAfter.ElementsAfterSelf(BuildGraphNamespace + "Property")
+				.Where(prop => prop.Attribute("Name").Value.EndsWith(FlagSuffix));
+			if (NextChunk
+				.Where(prop => prop.Attribute("Name").Value == TestName + FlagSuffix)
+				.Count() == 0)
+			{
+				XElement ElementInsert = new XElement(BuildGraphNamespace + "Property");
+				ElementInsert.SetAttributeValue("Name", TestName + FlagSuffix);
+				ElementInsert.SetAttributeValue("Value", FlagValue);
+				ElementUpsertAfter.AddAfterSelf(ElementInsert);
+			}
+			else
+			{
+				NextChunk
+					.Where(prop => prop.Attribute("Name").Value == TestName + FlagSuffix).First().SetAttributeValue("Value", FlagValue);
+			}
+			ElementUpsertAfter = NextChunk.Last();
+		}
+
+		private void InsertOrUpdateTestOption(ref XElement ElementUpsertAfter, string TestName, string TestShortName, string OptionPrefix, string OptionSuffix, string DefaultValue)
+		{
+			IEnumerable<XElement> NextChunk = ElementUpsertAfter.ElementsAfterSelf(BuildGraphNamespace + "Option")
+				.Where(prop => prop.Attribute("Name").Value.StartsWith(OptionPrefix) && prop.Attribute("Name").Value.EndsWith(OptionSuffix));
+			if (NextChunk
+				.Where(prop => prop.Attribute("Name").Value == OptionPrefix + TestName + OptionSuffix)
+				.Count() == 0)
+			{
+				XElement ElementInsert = new XElement(BuildGraphNamespace + "Option");
+				ElementInsert.SetAttributeValue("Name", OptionPrefix + TestName + OptionSuffix);
+				ElementInsert.SetAttributeValue("DefaultValue", DefaultValue);
+				ElementInsert.SetAttributeValue("Description", string.Format("{0} {1} {2}", OptionPrefix, TestShortName, OptionSuffix));
+				ElementUpsertAfter.AddAfterSelf(ElementInsert);
+			}
+			else
+			{
+				XElement ElementUpdate = NextChunk
+					.Where(prop => prop.Attribute("Name").Value == OptionPrefix + TestName + OptionSuffix).First();
+				ElementUpdate.SetAttributeValue("Description", string.Format("{0} {1} {2}", OptionPrefix, TestShortName, OptionSuffix));
+				ElementUpdate.SetAttributeValue("DefaultValue", DefaultValue);
+			}
+			ElementUpsertAfter = NextChunk.Last();
+		}
+#pragma warning restore 8602
+
+		/// <summary>
+		/// Test metadata class.
+		/// </summary>
+		public class Metadata
+		{
+			/// <summary>
+			/// Catch2 report type console - prints results to stdout.
+			/// </summary>
+			public static readonly string ConsoleReportType = "console";
+
+			/// <summary>
+			/// Test long name.
+			/// </summary>
+			public string TestName { get; }
+
+			/// <summary>
+			/// Test short name used for display in build system.
+			/// </summary>
+			public string TestShortName { get; }
+
+			/// <summary>
+			/// Type of Catch2 report, defaults to console.
+			/// </summary>
+			public string ReportType { get; }
+
+			/// <summary>
+			/// Is this test disabled?
+			/// </summary>
+			public bool Disabled { get; }
+
+			/// <summary>
+			/// Ctor that sets short name to test name.
+			/// </summary>
+			public Metadata(string InTestName)
+			{
+				TestShortName = TestName = InTestName;
+				ReportType = ConsoleReportType;
+			}
+
+			/// <summary>
+			/// Ctor that sets short name and test name.
+			/// </summary>
+			public Metadata(string InTestName, string InTestShortName)
+			{
+				TestName = InTestName;
+				TestShortName = InTestShortName;
+				ReportType = ConsoleReportType;
+			}
+
+			/// <summary>
+			/// Ctor that sets short name, test name, report type and disabled status.
+			/// </summary>
+			public Metadata(string InTestName, string InTestShortName, string InReportType, bool InDisabled)
+			{
+				TestName = InTestName;
+				TestShortName = InTestShortName;
+				ReportType = InReportType;
+				Disabled = InDisabled;
 			}
 		}
 	}
-}
+}
\ No newline at end of file
diff --git a/Engine/Source/Programs/UnrealBuildTool/Configuration/TestTargetRules.cs b/Engine/Source/Programs/UnrealBuildTool/Configuration/TestTargetRules.cs
--- a/Engine/Source/Programs/UnrealBuildTool/Configuration/TestTargetRules.cs
+++ b/Engine/Source/Programs/UnrealBuildTool/Configuration/TestTargetRules.cs
@@ -1,7 +1,6 @@
 // Copyright Epic Games, Inc. All Rights Reserved.
 
 using EpicGames.Core;
-using System.Collections.Generic;
 using System.IO;
 using UnrealBuildBase;
 
@@ -15,35 +14,94 @@ namespace UnrealBuildTool
 		/// <summary>
 		/// Keeps track if low level tests executable must build with the Editor.
 		/// </summary>
-		public static bool bTestsRequireEditor = false;
+		internal static bool bTestsRequireEditor = false;
 
 		/// <summary>
 		/// Keeps track if low level tests executable must build with the Engine.
 		/// </summary>
-		public static bool bTestsRequireEngine = false;
+		internal static bool bTestsRequireEngine = false;
 
 		/// <summary>
 		/// Keeps track if low level tests executable must build with the ApplicationCore.
 		/// </summary>
-		public static bool bTestsRequireApplicationCore = false;
+		internal static bool bTestsRequireApplicationCore = false;
 
 		/// <summary>
 		/// Keeps track if low level tests executable must build with the CoreUObject.
 		/// </summary>
-		public static bool bTestsRequireCoreUObject = false;
+		internal static bool bTestsRequireCoreUObject = false;
 
 		/// <summary>
 		/// Keep track of low level test runner module instance.
 		/// </summary>
-		public static ModuleRules? LowLevelTestsRunnerModule;
+		internal static ModuleRules? LowLevelTestsRunnerModule;
 
 		/// <summary>
-		/// Associated tested target of this test target.
+		/// Associated tested target of this test target, if defined.
 		/// </summary>
-		public TargetRules TestedTarget { get; private set; }
+		protected TargetRules? TestedTarget { get; private set; }
 
 		/// <summary>
-		/// TestTargetRules is setup as a program and is linked monolithically .
+		/// Test target override for bCompileAgainstApplicationCore.
+		/// It is set to true if there is any reference to ApplicationCore in the reference chain.
+		/// </summary>
+		public override bool bCompileAgainstApplicationCore
+		{
+			get { return bTestsRequireApplicationCore; }
+			set { bTestsRequireApplicationCore = value; }
+		}
+
+		/// <summary>
+		/// Test target override for bCompileAgainstCoreUObject.
+		/// It is set to true if there is any reference to CoreUObject in the reference chain.
+		/// </summary>
+		public override bool bCompileAgainstCoreUObject
+		{
+			get { return bTestsRequireCoreUObject; }
+			set { bTestsRequireCoreUObject = value; }
+		}
+
+		/// <summary>
+		/// Test target override for bCompileAgainstEngine.
+		/// It is set to true if there is any reference to Engine in the reference chain.
+		/// </summary>
+		public override bool bCompileAgainstEngine
+		{
+			get { return bTestsRequireEngine; }
+			set { bTestsRequireEngine = value; }
+		}
+
+		/// <summary>
+		/// Test target override for bCompileAgainstEditor.
+		/// It is set to true if there is any reference to UnrealEd in the reference chain.
+		/// </summary>
+		public override bool bCompileAgainstEditor
+		{
+			get { return bTestsRequireEditor; }
+			set { bTestsRequireEditor = value; }
+		}
+
+		/// <summary>
+		/// Constructor for TestTargetRules as own target.
+		/// </summary>
+		/// <param name="Target"></param>
+		public TestTargetRules(TargetInfo Target) : base(Target)
+		{
+			SetupCommonProperties(Target);
+
+			bExplicitTestsTargetOverride = this.GetType() != typeof(TestTargetRules);
+
+			ExeBinariesSubFolder = LaunchModuleName = Name + (bExplicitTestsTargetOverride ? string.Empty : "Tests");
+
+			if (bExplicitTestsTargetOverride)
+			{
+				SolutionDirectory = "Programs/LowLevelTests";
+			}
+		}
+
+		/// <summary>
+		/// Constructor for TestTargetRules based on existing target.
+		/// TestTargetRules is setup as a program and is linked monolithically.
 		/// It removes a lot of default compilation behavior in order to produce a minimal test environment.
 		/// </summary>
 		public TestTargetRules(TargetRules TestedTarget, TargetInfo Target) : base(Target)
@@ -53,20 +111,30 @@ namespace UnrealBuildTool
 				throw new BuildException("TestedTarget can't be of type TestTargetRules.");
 			}
 
-			bIsTestTargetOverride = true;
-
 			this.TestedTarget = TestedTarget;
 
-			DefaultBuildSettings = BuildSettingsVersion.V2;
-
 			ExeBinariesSubFolder = Name = TestedTarget.Name + "Tests";
 			TargetSourceFile = File = TestedTarget.File;
-			LaunchModuleName = TestedTarget.LaunchModuleName + "Tests";
+			if (TestedTarget.LaunchModuleName != null)
+			{
+				LaunchModuleName = TestedTarget.LaunchModuleName + "Tests";
+			}
 
 			ManifestFileNames = TestedTarget.ManifestFileNames;
 
 			WindowsPlatform = TestedTarget.WindowsPlatform;
 
+			SetupCommonProperties(Target);
+		}
+
+		private void SetupCommonProperties(TargetInfo Target)
+		{
+			bIsTestTargetOverride = true;
+
+			VSTestRunSettingsFile = FileReference.Combine(Unreal.EngineDirectory, "Source", "Programs", "LowLevelTests", "vstest.runsettings");
+
+			DefaultBuildSettings = BuildSettingsVersion.V2;
+
 			Type = TargetType.Program;
 			LinkType = TargetLinkType.Monolithic;
 
@@ -116,17 +184,19 @@ namespace UnrealBuildTool
 			{
 				UndecoratedConfiguration = Target.Configuration;
 
-				string VersionScriptFile = Path.Combine(Directory.GetCurrentDirectory(), @"Developer\LowLevelTestsRunner\Private\Platform\Android\HideSymbols.ldscript");
-				AdditionalLinkerArguments = " -Wl,--version-script=" + VersionScriptFile;
-
 				GlobalDefinitions.Add("USE_ANDROID_INPUT=0");
 				GlobalDefinitions.Add("USE_ANDROID_OPENGL=0");
 				GlobalDefinitions.Add("USE_ANDROID_LAUNCH=0");
 				GlobalDefinitions.Add("USE_ANDROID_JNI=0");
 			}
-			else if (Target.Platform == UnrealTargetPlatform.IOS)
+			else if (Target.Platform == UnrealTargetPlatform.IOS) // TODO: this doesn't compile
 			{
 				GlobalDefinitions.Add("HAS_METAL=0");
+
+				bIsBuildingConsoleApplication = false;
+				// Required for IOS, but needs to fix compilation errors
+				bCompileAgainstApplicationCore = true;
+
 			}
 		}
 	}
diff --git a/Engine/Source/Programs/UnrealBuildTool/Configuration/UEBuildDeploy.cs b/Engine/Source/Programs/UnrealBuildTool/Configuration/UEBuildDeploy.cs
--- a/Engine/Source/Programs/UnrealBuildTool/Configuration/UEBuildDeploy.cs
+++ b/Engine/Source/Programs/UnrealBuildTool/Configuration/UEBuildDeploy.cs
@@ -1,11 +1,9 @@
 // Copyright Epic Games, Inc. All Rights Reserved.
 
-using System;
-using System.Collections.Generic;
+using EpicGames.Core;
 using System.Linq;
-using System.Text;
-using System.Diagnostics;
 using System.IO;
+using System;
 
 namespace UnrealBuildTool
 {
@@ -14,6 +12,39 @@ namespace UnrealBuildTool
 	/// </summary>
 	abstract class UEBuildDeploy
 	{
+		protected bool CopyFilesRecursively(string FromDir, string ToDir)
+		{
+			var FromDirInfo = new DirectoryInfo(FromDir);
+
+			if (!FromDirInfo.Exists)
+			{
+				Log.TraceError($"Cannot copy files, source directory not found: {FromDirInfo.FullName}");
+				return false;
+			}
+
+			DirectoryInfo[] DirsToCopy = FromDirInfo.GetDirectories();
+
+			Directory.CreateDirectory(ToDir);
+
+			foreach (FileInfo FromFile in FromDirInfo.GetFiles())
+			{
+				try
+				{
+					FromFile.CopyTo(Path.Combine(ToDir, FromFile.Name), true);
+				} catch (Exception ex)
+				{
+					Log.TraceError($"Error while copying {FromFile.Name} from {FromDir} to {ToDir}: {ex}");
+					return false;
+				}
+			}
+
+			foreach (DirectoryInfo SubDir in DirsToCopy)
+			{
+				CopyFilesRecursively(SubDir.FullName, Path.Combine(ToDir, SubDir.Name));
+			}
+			return true;
+		}
+
 		/// <summary>
 		/// Prepare the target for deployment
 		/// </summary>
@@ -21,6 +52,18 @@ namespace UnrealBuildTool
 		/// <returns>True if successful, false if not</returns>
 		public virtual bool PrepTargetForDeployment(TargetReceipt Receipt)
 		{
+			if (Receipt.IsTestTarget)
+			{
+				ReceiptProperty? TestResourcesFolder = Receipt.AdditionalProperties.FirstOrDefault(Prop => Prop.Name == "ResourcesFolder");
+				if (TestResourcesFolder != null && Receipt.Launch != null)
+				{
+					DirectoryInfo ResourceDirectoryInfo = new DirectoryInfo(TestResourcesFolder.Value);
+					if (ResourceDirectoryInfo.Exists)
+					{
+						return CopyFilesRecursively(TestResourcesFolder.Value, Path.Combine(Receipt.Launch.Directory.FullName, ResourceDirectoryInfo.Name));
+					}
+				}
+			}
 			return true;
 		}
 	}
diff --git a/Engine/Source/Programs/UnrealBuildTool/Configuration/UEBuildModule.cs b/Engine/Source/Programs/UnrealBuildTool/Configuration/UEBuildModule.cs
--- a/Engine/Source/Programs/UnrealBuildTool/Configuration/UEBuildModule.cs
+++ b/Engine/Source/Programs/UnrealBuildTool/Configuration/UEBuildModule.cs
@@ -1021,6 +1021,11 @@ namespace UnrealBuildTool
 
 				// Create all the dependency modules - pass through the reference stack so we can check for cycles
 				RecursivelyCreateModulesByName(Rules.PublicDependencyModuleNames, ref PublicDependencyModules, ref bDependsOnVerse, CreateModule, NextReferenceChain, ReferenceStack);
+				if (Rules.Target.IsTestTarget)
+				{
+					// Move the test runner dependency to last position to give it the opportunity to build its special dependencies such as CoreUObject, ApplicationCore etc.
+					MoveTestsRunnerDependencyToLastPosition();
+				}
 				RecursivelyCreateModulesByName(Rules.PrivateDependencyModuleNames, ref PrivateDependencyModules, ref bDependsOnVerse, CreateModule, NextReferenceChain, ReferenceStack);
 				// Dynamic loads aren't considered a reference chain so start with an empty stack
 				RecursivelyCreateModulesByName(Rules.DynamicallyLoadedModuleNames, ref DynamicallyLoadedModules, ref bDependsOnVerse, CreateModule, NextReferenceChain, new List<UEBuildModule>());
@@ -1030,6 +1035,15 @@ namespace UnrealBuildTool
 			ReferenceStack.RemoveAt(ReferenceStack.Count - 1);
 		}
 
+		private void MoveTestsRunnerDependencyToLastPosition()
+		{
+			if (Rules.PrivateDependencyModuleNames.Contains("LowLevelTestsRunner"))
+			{
+				Rules.PrivateDependencyModuleNames.Remove("LowLevelTestsRunner");
+				Rules.PrivateDependencyModuleNames.Add("LowLevelTestsRunner");
+			}
+		}
+
 		private static void LogDependencyNameList(string Title, List<string> DependencyNameList)
 		{
 			Log.TraceVeryVerbose("  " + Title);
diff --git a/Engine/Source/Programs/UnrealBuildTool/Configuration/UEBuildTarget.cs b/Engine/Source/Programs/UnrealBuildTool/Configuration/UEBuildTarget.cs
--- a/Engine/Source/Programs/UnrealBuildTool/Configuration/UEBuildTarget.cs
+++ b/Engine/Source/Programs/UnrealBuildTool/Configuration/UEBuildTarget.cs
@@ -1520,7 +1520,7 @@ namespace UnrealBuildTool
 			}
 
 			// Create the receipt
-			TargetReceipt Receipt = new TargetReceipt(ProjectFile, TargetName, TargetType, Platform, Configuration, Version, Architecture);
+			TargetReceipt Receipt = new TargetReceipt(ProjectFile, TargetName, TargetType, Platform, Configuration, Version, Architecture, Rules.IsTestTarget);
 
 			if (!Rules.bShouldCompileAsDLL)
 			{
@@ -1783,6 +1783,7 @@ namespace UnrealBuildTool
 			TargetMakefile Makefile = new TargetMakefile(ExternalMetadata, Binaries[0].OutputFilePaths[0], ReceiptFileName, ProjectIntermediateDirectory, TargetType, 
 				Rules.ConfigValueTracker, bDeployAfterCompile, bHasProjectScriptPlugin, bHasRequiredProjectScriptPlugin, UbtPlugins?.ToArray(), EnabledUbtPlugins?.ToArray(), 
 				EnabledUhtPlugins?.ToArray());
+			Makefile.IsTestTarget = Rules.IsTestTarget;
 
 			// Get diagnostic info to be printed before each build
 			TargetToolChain.GetVersionInfo(Makefile.Diagnostics);
@@ -3864,7 +3865,7 @@ namespace UnrealBuildTool
 			GlobalCompileEnvironment.Definitions.Add("WITH_DEV_AUTOMATION_TESTS=" + (bCompileDevTests ? "1" : "0"));
 			GlobalCompileEnvironment.Definitions.Add("WITH_PERF_AUTOMATION_TESTS=" + (bCompilePerfTests ? "1" : "0"));
 
-			GlobalCompileEnvironment.Definitions.Add(String.Format("WITH_LOW_LEVEL_TESTS={0}", Rules.IsTestTarget ? "1" : "0"));
+			GlobalCompileEnvironment.Definitions.Add(String.Format("WITH_LOW_LEVEL_TESTS={0}", Rules.WithLowLevelTests ? "1" : "0"));
 
 			GlobalCompileEnvironment.Definitions.Add("UNICODE");
 			GlobalCompileEnvironment.Definitions.Add("_UNICODE");
@@ -3872,12 +3873,12 @@ namespace UnrealBuildTool
 
 			GlobalCompileEnvironment.Definitions.Add(String.Format("IS_MONOLITHIC={0}", ShouldCompileMonolithic() ? "1" : "0"));
 
-			GlobalCompileEnvironment.Definitions.Add(String.Format("WITH_ENGINE={0}", Rules.bCompileAgainstEngine || TestTargetRules.bTestsRequireEngine ? "1" : "0"));
+			GlobalCompileEnvironment.Definitions.Add(String.Format("WITH_ENGINE={0}", Rules.bCompileAgainstEngine ? "1" : "0"));
 			GlobalCompileEnvironment.Definitions.Add(String.Format("WITH_UNREAL_DEVELOPER_TOOLS={0}", Rules.bBuildDeveloperTools ? "1" : "0"));
 			GlobalCompileEnvironment.Definitions.Add(String.Format("WITH_UNREAL_TARGET_DEVELOPER_TOOLS={0}", Rules.bBuildTargetDeveloperTools ? "1" : "0"));
 
 			// Set a macro to control whether to initialize ApplicationCore. Command line utilities should not generally need this.
-			if (Rules.bCompileAgainstApplicationCore || TestTargetRules.bTestsRequireApplicationCore)
+			if (Rules.bCompileAgainstApplicationCore)
 			{
 				GlobalCompileEnvironment.Definitions.Add("WITH_APPLICATION_CORE=1");
 			}
@@ -3886,7 +3887,7 @@ namespace UnrealBuildTool
 				GlobalCompileEnvironment.Definitions.Add("WITH_APPLICATION_CORE=0");
 			}
 
-			if (Rules.bCompileAgainstCoreUObject || TestTargetRules.bTestsRequireCoreUObject)
+			if (Rules.bCompileAgainstCoreUObject)
 			{
 				GlobalCompileEnvironment.Definitions.Add("WITH_COREUOBJECT=1");
 			}
@@ -3994,7 +3995,7 @@ namespace UnrealBuildTool
 				GlobalCompileEnvironment.Definitions.Add("USE_ESTIMATED_UTCNOW=0");
 			}
 
-			if ((Rules.bCompileAgainstEditor && (Rules.Type == TargetType.Editor || Rules.Type == TargetType.Program)) || TestTargetRules.bTestsRequireEditor)
+			if ((Rules.bCompileAgainstEditor && (Rules.Type == TargetType.Editor || Rules.Type == TargetType.Program)))
 			{
 				GlobalCompileEnvironment.Definitions.Add("WITH_EDITOR=1");
 				GlobalCompileEnvironment.Definitions.Add("WITH_IOSTORE_IN_EDITOR=1");
diff --git a/Engine/Source/Programs/UnrealBuildTool/Platform/HoloLens/HoloLensDeploy.cs b/Engine/Source/Programs/UnrealBuildTool/Platform/HoloLens/HoloLensDeploy.cs
--- a/Engine/Source/Programs/UnrealBuildTool/Platform/HoloLens/HoloLensDeploy.cs
+++ b/Engine/Source/Programs/UnrealBuildTool/Platform/HoloLens/HoloLensDeploy.cs
@@ -486,7 +486,7 @@ namespace UnrealBuildTool
 			FileReference ReceiptFileName = TargetReceipt.GetDefaultPath(Receipt.ProjectDir != null ? Receipt.ProjectDir : Unreal.EngineDirectory, Receipt.TargetName, Receipt.Platform, Receipt.Configuration, "Multi");
 			if (!TargetReceipt.TryRead(ReceiptFileName, Unreal.EngineDirectory, out NewReceipt))
 			{
-				NewReceipt = new TargetReceipt(Receipt.ProjectFile, Receipt.TargetName, Receipt.TargetType, Receipt.Platform, Receipt.Configuration, Receipt.Version, "Multi");
+				NewReceipt = new TargetReceipt(Receipt.ProjectFile, Receipt.TargetName, Receipt.TargetType, Receipt.Platform, Receipt.Configuration, Receipt.Version, "Multi", Receipt.IsTestTarget);
 			}
 
 			string SDK = "";
diff --git a/Engine/Source/Programs/UnrealBuildTool/Platform/Linux/UEBuildLinux.cs b/Engine/Source/Programs/UnrealBuildTool/Platform/Linux/UEBuildLinux.cs
--- a/Engine/Source/Programs/UnrealBuildTool/Platform/Linux/UEBuildLinux.cs
+++ b/Engine/Source/Programs/UnrealBuildTool/Platform/Linux/UEBuildLinux.cs
@@ -625,6 +625,15 @@ namespace UnrealBuildTool
 		/// <param name="Receipt">Receipt for the target being deployed</param>
 		public override void Deploy(TargetReceipt Receipt)
 		{
+			new UEDeployLinux().PrepTargetForDeployment(Receipt);
+		}
+	}
+
+	class UEDeployLinux: UEBuildDeploy
+	{
+		public override bool PrepTargetForDeployment(TargetReceipt Receipt)
+		{
+			return base.PrepTargetForDeployment(Receipt);
 		}
 	}
 
diff --git a/Engine/Source/Programs/UnrealBuildTool/Platform/Mac/UEDeployMac.cs b/Engine/Source/Programs/UnrealBuildTool/Platform/Mac/UEDeployMac.cs
--- a/Engine/Source/Programs/UnrealBuildTool/Platform/Mac/UEDeployMac.cs
+++ b/Engine/Source/Programs/UnrealBuildTool/Platform/Mac/UEDeployMac.cs
@@ -17,7 +17,7 @@ namespace UnrealBuildTool
 		public override bool PrepTargetForDeployment(TargetReceipt Receipt)
 		{
 			Log.TraceInformation("Deploying now!");
-			return true;
+			return base.PrepTargetForDeployment(Receipt);
 		}
 
 		public static bool GeneratePList(string ProjectDirectory, bool bIsUnrealGame, string GameName, string ProjectName, string InEngineDir, string ExeName)
diff --git a/Engine/Source/Programs/UnrealBuildTool/Platform/Windows/UEBuildWindows.cs b/Engine/Source/Programs/UnrealBuildTool/Platform/Windows/UEBuildWindows.cs
--- a/Engine/Source/Programs/UnrealBuildTool/Platform/Windows/UEBuildWindows.cs
+++ b/Engine/Source/Programs/UnrealBuildTool/Platform/Windows/UEBuildWindows.cs
@@ -1484,6 +1484,15 @@ namespace UnrealBuildTool
 		/// <param name="Receipt">Receipt for the target being deployed</param>
 		public override void Deploy(TargetReceipt Receipt)
 		{
+			new UEDeployWindows().PrepTargetForDeployment(Receipt);
+		}
+	}
+
+	class UEDeployWindows : UEBuildDeploy
+	{
+		public override bool PrepTargetForDeployment(TargetReceipt Receipt)
+		{
+			return base.PrepTargetForDeployment(Receipt);
 		}
 	}
 
diff --git a/Engine/Source/Programs/UnrealBuildTool/System/RulesAssembly.cs b/Engine/Source/Programs/UnrealBuildTool/System/RulesAssembly.cs
--- a/Engine/Source/Programs/UnrealBuildTool/System/RulesAssembly.cs
+++ b/Engine/Source/Programs/UnrealBuildTool/System/RulesAssembly.cs
@@ -388,7 +388,7 @@ namespace UnrealBuildTool
 		/// <returns>Compiled module rule info</returns>
 		public ModuleRules CreateModuleRules(string ModuleName, ReadOnlyTargetRules Target, string ReferenceChain)
 		{
-			if (Target.IsTestTarget)
+			if (Target.IsTestTarget && !Target.ExplicitTestsTarget)
 			{
 				ModuleName = TargetDescriptor.GetTestedTargetName(ModuleName);
 			}
@@ -510,13 +510,18 @@ namespace UnrealBuildTool
 				}
 				Constructor.Invoke(RulesObject, new object[] { Target });
 
-				if (Target.IsTestTarget && Target.LaunchModuleName != null && ModuleName == TargetDescriptor.GetTestedTargetName(Target.LaunchModuleName))
+				if (Target.IsTestTarget && !RulesObject.IsTestModule)
 				{
-					RulesObject = new TestModuleRules(RulesObject);
+					if (!Target.ExplicitTestsTarget)
+					{
+						if (Target.LaunchModuleName != null && ModuleName == TargetDescriptor.GetTestedTargetName(Target.LaunchModuleName))
+						{
+							RulesObject = new TestModuleRules(RulesObject);
+						}
+					}
+					RulesObject.PrepareModuleForTests();
 				}
 
-				RulesObject.PrepareModuleForTests();
-
 				return RulesObject;
 			}
 			catch (Exception Ex)
diff --git a/Engine/Source/Programs/UnrealBuildTool/System/TargetMakefile.cs b/Engine/Source/Programs/UnrealBuildTool/System/TargetMakefile.cs
--- a/Engine/Source/Programs/UnrealBuildTool/System/TargetMakefile.cs
+++ b/Engine/Source/Programs/UnrealBuildTool/System/TargetMakefile.cs
@@ -62,6 +62,11 @@ namespace UnrealBuildTool
 		/// </summary>
 		public TargetType TargetType;
 
+		/// <summary>
+		/// Is test target?
+		/// </summary>
+		public bool IsTestTarget;
+
 		/// <summary>
 		/// Map of config file keys to values. Makefile will be invalidated if any of these change.
 		/// </summary>
@@ -272,6 +277,7 @@ namespace UnrealBuildTool
 			ReceiptFile = Reader.ReadFileReference();
 			ProjectIntermediateDirectory = Reader.ReadDirectoryReferenceNotNull();
 			TargetType = (TargetType)Reader.ReadInt();
+			IsTestTarget = Reader.ReadBool();
 			ConfigValueTracker = new ConfigValueTracker(Reader);
 			bDeployAfterCompile = Reader.ReadBool();
 			bHasProjectScriptPlugin = Reader.ReadBool();
@@ -316,6 +322,7 @@ namespace UnrealBuildTool
 			Writer.WriteFileReference(ReceiptFile);
 			Writer.WriteDirectoryReference(ProjectIntermediateDirectory);
 			Writer.WriteInt((int)TargetType);
+			Writer.WriteBool(IsTestTarget);
 			ConfigValueTracker.Write(Writer);
 			Writer.WriteBool(bDeployAfterCompile);
 			Writer.WriteBool(bHasProjectScriptPlugin);
diff --git a/Engine/Source/Programs/UnrealBuildTool/System/TargetReceipt.cs b/Engine/Source/Programs/UnrealBuildTool/System/TargetReceipt.cs
--- a/Engine/Source/Programs/UnrealBuildTool/System/TargetReceipt.cs
+++ b/Engine/Source/Programs/UnrealBuildTool/System/TargetReceipt.cs
@@ -294,6 +294,11 @@ namespace UnrealBuildTool
 		/// </summary>
 		public TargetType TargetType;
 
+		/// <summary>
+		/// Whether it's a low level tests target
+		/// </summary>
+		public bool IsTestTarget;
+
 		/// <summary>
 		/// Version information for this target.
 		/// </summary>
@@ -339,7 +344,8 @@ namespace UnrealBuildTool
 		/// <param name="InConfiguration">Configuration of the target being compiled</param>
 		/// <param name="InVersion">Version information for the target</param>
 		/// <param name="InArchitecture">Architecture information for the target</param>
-		public TargetReceipt(FileReference? InProjectFile, string InTargetName, TargetType InTargetType, UnrealTargetPlatform InPlatform, UnrealTargetConfiguration InConfiguration, BuildVersion InVersion, string InArchitecture)
+		/// <param name="InIsTestTarget">Whether it's a target for low level tests</param>
+		public TargetReceipt(FileReference? InProjectFile, string InTargetName, TargetType InTargetType, UnrealTargetPlatform InPlatform, UnrealTargetConfiguration InConfiguration, BuildVersion InVersion, string InArchitecture, bool InIsTestTarget)
 		{
 			ProjectFile = InProjectFile;
 			ProjectDir = DirectoryReference.FromFile(InProjectFile);
@@ -349,6 +355,7 @@ namespace UnrealBuildTool
 			TargetType = InTargetType;
 			Version = InVersion;
 			Architecture = InArchitecture;
+			IsTestTarget = InIsTestTarget;
 		}
 
 		/// <summary>
@@ -552,9 +559,11 @@ namespace UnrealBuildTool
 				Architecture = "";
 			}
 
+			bool IsTestTarget;
+			RawObject.TryGetBoolField("IsTestTarget", out IsTestTarget);
 
 			// Create the receipt
-			TargetReceipt Receipt = new TargetReceipt(ProjectFile, TargetName, TargetType, Platform, Configuration, Version, Architecture);
+			TargetReceipt Receipt = new TargetReceipt(ProjectFile, TargetName, TargetType, Platform, Configuration, Version, Architecture, IsTestTarget);
 
 			// Get the project directory
 			DirectoryReference? ProjectDir = Receipt.ProjectDir;
@@ -719,6 +728,7 @@ namespace UnrealBuildTool
 				Writer.WriteValue("Platform", Platform.ToString());
 				Writer.WriteValue("Configuration", Configuration.ToString());
 				Writer.WriteValue("TargetType", TargetType.ToString());
+				Writer.WriteValue("IsTestTarget", IsTestTarget);
 				Writer.WriteValue("Architecture", Architecture);
 
 				if(ProjectFile != null)
diff --git a/Engine/Source/Runtime/ApplicationCore/ApplicationCore.Build.cs b/Engine/Source/Runtime/ApplicationCore/ApplicationCore.Build.cs
--- a/Engine/Source/Runtime/ApplicationCore/ApplicationCore.Build.cs
+++ b/Engine/Source/Runtime/ApplicationCore/ApplicationCore.Build.cs
@@ -92,7 +92,7 @@ public class ApplicationCore : ModuleRules
 			);
 		}
 
-		if (!Target.bCompileAgainstApplicationCore)
+		if (!Target.IsTestTarget && !Target.bCompileAgainstApplicationCore)
 		{
 			throw new System.Exception("ApplicationCore cannot be used when Target.bCompileAgainstApplicationCore = false.");
 		}
