UnrealBuildTool: Compile PCH with -fpch-instantiate-templates for Linux if compiling with clang 11+

Initial testing has shown this PCH flag improves clang compile time by up to 20%

This flag is the default for clang-cl so it already applies to windows clang builds. Other clang based platforms will be supported as additional testing is done.

https://reviews.llvm.org/D69585

Linux UnrealEditor Development:
- Before Total CPU Time: 114072.34375 s
- Before UBT Run Time: 00:26:05
- After Total CPU Time: 86165.671875 s
- After UBT Run Time: 00:20:43
Total CPU Time improvement 24.46%
UBT Run Time improvement 20.58%

Linux UnrealGame Development:
- Before Total CPU Time: 23445.140625 s
- Before UBT Run Time: 00:09:13
- After Total CPU Time: 18492.71875 s
- After UBT Run Time:  00:07:52
Total CPU Time improvement 21.12%
UBT Run Time improvement 13.56%

diff --git a/Engine/Source/Programs/UnrealBuildTool/Platform/Linux/LinuxToolChain.cs b/Engine/Source/Programs/UnrealBuildTool/Platform/Linux/LinuxToolChain.cs
--- a/Engine/Source/Programs/UnrealBuildTool/Platform/Linux/LinuxToolChain.cs
+++ b/Engine/Source/Programs/UnrealBuildTool/Platform/Linux/LinuxToolChain.cs
@@ -929,6 +929,7 @@ namespace UnrealBuildTool
 			if (CompilerVersionGreaterOrEqual(11, 0, 0))
 			{
 				Result += " -fpch-validate-input-files-content";
+				Result += " -fpch-instantiate-templates";
 			}
 			Result += GetCompilerStandardVersion_CPP(CompileEnvironment);
 			return Result;
@@ -1283,6 +1284,7 @@ namespace UnrealBuildTool
 
 			if (CompileEnvironment.PrecompiledHeaderAction == PrecompiledHeaderAction.Include)
 			{
+				// Validate PCH inputs by content if mtime check fails
 				if (CompilerVersionGreaterOrEqual(11, 0, 0))
 				{
 					PCHArguments += " -fpch-validate-input-files-content";
