UnrealBuildTool: Fix typo in HordeExecutor config loading

#rnx

diff --git a/Engine/Source/Programs/UnrealBuildTool/Executors/Experimental/HordeExecutor.cs b/Engine/Source/Programs/UnrealBuildTool/Executors/Experimental/HordeExecutor.cs
--- a/Engine/Source/Programs/UnrealBuildTool/Executors/Experimental/HordeExecutor.cs
+++ b/Engine/Source/Programs/UnrealBuildTool/Executors/Experimental/HordeExecutor.cs
@@ -353,7 +353,7 @@ namespace UnrealBuildTool
 
 				DataPathReference = DirectoryReference.FromString(DataPath.Replace("%APPSETTINGSDIR%", $"{Path.Combine(Environment.GetFolderPath(Environment.SpecialFolder.CommonApplicationData), "Epic")}{Path.DirectorySeparatorChar}"));
 
-				if (EngineConfig.GetString("Zen.AutoLaunch", "ExtraArgs", out string ExtraArgs))
+				if (!EngineConfig.GetString("Zen.AutoLaunch", "ExtraArgs", out string ExtraArgs) || string.IsNullOrEmpty(ExtraArgs))
 				{
 					Log.TraceWarning("Unable to launch zenserver: '[Zen.AutoLaunch] ExtraArgs' not set in Engine config");
 					return false;
