UnrealBuildTool: Compile PCH with -fpch-instantiate-templates for Android platform. Improves compile time by almost 15%

diff --git a/Engine/Source/Programs/UnrealBuildTool/Platform/Android/AndroidToolChain.cs b/Engine/Source/Programs/UnrealBuildTool/Platform/Android/AndroidToolChain.cs
--- a/Engine/Source/Programs/UnrealBuildTool/Platform/Android/AndroidToolChain.cs
+++ b/Engine/Source/Programs/UnrealBuildTool/Platform/Android/AndroidToolChain.cs
@@ -887,6 +887,11 @@ namespace UnrealBuildTool
 			string Result = "";
 
 			Result += " -x c++-header";
+			if (CompilerVersionGreaterOrEqual(11, 0, 0))
+			{
+				Result += " -fpch-validate-input-files-content";
+				Result += " -fpch-instantiate-templates";
+			}
 			Result += GetCppStandardCompileArgument(CompileEnvironment);
 
 			return Result;
@@ -1329,6 +1334,12 @@ namespace UnrealBuildTool
 					// Add the precompiled header file's path to the include path so Clang can find it.
 					// This needs to be before the other include paths to ensure Clang uses it instead of the source header file.
 					PCHArguments += string.Format(" -include \"{0}\"", BasePCHName);
+
+					// Validate PCH inputs by content if mtime check fails
+					if (CompilerVersionGreaterOrEqual(11, 0, 0))
+					{
+						PCHArguments += " -fpch-validate-input-files-content";
+					}
 				}
 
 				// Add include paths to the argument list (filtered by architecture)
