UnrealBuildTool: Compile PCH with -fpch-instantiate-templates for Apple platforms. Improves compile time by almost 15%

Mac UnrealEditor Development:
* Before UBT Run Time: 00:57:20
* After UBT Run Time: 00:49:44
UBT Run Time improvement 13.26%

diff --git a/Engine/Source/Programs/UnrealBuildTool/Platform/IOS/IOSToolChain.cs b/Engine/Source/Programs/UnrealBuildTool/Platform/IOS/IOSToolChain.cs
--- a/Engine/Source/Programs/UnrealBuildTool/Platform/IOS/IOSToolChain.cs
+++ b/Engine/Source/Programs/UnrealBuildTool/Platform/IOS/IOSToolChain.cs
@@ -631,6 +631,10 @@ namespace UnrealBuildTool
 			if (CompileEnvironment.PrecompiledHeaderAction == PrecompiledHeaderAction.Include)
 			{
 				PCHArguments += string.Format(" -include \"{0}\"", CompileEnvironment.PrecompiledHeaderIncludeFilename!.FullName);
+				if (GetClangVersion().Major >= 11)
+				{
+					PCHArguments += " -fpch-validate-input-files-content";
+				}
 			}
 
 			foreach (FileItem ForceIncludeFile in CompileEnvironment.ForceIncludeFiles)
@@ -665,6 +669,14 @@ namespace UnrealBuildTool
 				{
 					// Compile the file as a C++ PCH.
 					FileArguments += GetCompileArguments_PCH(CompileEnvironment);
+					if (GetClangVersion().Major >= 11)
+					{
+						FileArguments += " -fpch-validate-input-files-content";
+					}
+					if (GetClangVersion().Major >= 13) // Note this is supported for >=11 on other clang platforms
+					{
+						FileArguments += " -fpch-instantiate-templates";
+					}
 					FileArguments += GetRTTIFlag(CompileEnvironment);
 					FileArguments += GetObjCExceptionsFlag(CompileEnvironment);
 				}
diff --git a/Engine/Source/Programs/UnrealBuildTool/Platform/Mac/MacToolChain.cs b/Engine/Source/Programs/UnrealBuildTool/Platform/Mac/MacToolChain.cs
--- a/Engine/Source/Programs/UnrealBuildTool/Platform/Mac/MacToolChain.cs
+++ b/Engine/Source/Programs/UnrealBuildTool/Platform/Mac/MacToolChain.cs
@@ -518,6 +518,10 @@ namespace UnrealBuildTool
 					{
 						FileArguments += " -fpch-validate-input-files-content";
 					}
+					if (GetClangVersion().Major >= 13) // Note this is supported for >=11 on other clang platforms
+					{
+						FileArguments += " -fpch-instantiate-templates";
+					}
 					FileArguments += GetRTTIFlag(CompileEnvironment);
 				}
 				else if (Extension == ".C")
