UnrealBuildTool: add -CompactOutput command line option, supported by ParallelExecutor

Also can be set via BuildConfiguration.xml with:

<BuildConfiguration><bCompactOutput>true</bCompactOutput></BuildConfiguration>

diff --git a/Engine/Source/Programs/UnrealBuildTool/Configuration/BuildConfiguration.cs b/Engine/Source/Programs/UnrealBuildTool/Configuration/BuildConfiguration.cs
--- a/Engine/Source/Programs/UnrealBuildTool/Configuration/BuildConfiguration.cs
+++ b/Engine/Source/Programs/UnrealBuildTool/Configuration/BuildConfiguration.cs
@@ -175,5 +175,19 @@ namespace UnrealBuildTool
 		/// </summary>
 		[CommandLine("-IgnoreInvalidFiles")]
 		public bool bIgnoreInvalidFiles;
+
+		/// <summary>
+		/// Instruct the executor to write compact output e.g. only errors, if supported by the executor.
+		/// This field is used to hold the value when specified from the command line or XML
+		/// </summary>
+		[XmlConfigFile(Name = "bCompactOutput")]
+		[CommandLine("-CompactOutput")]
+		private bool bCompactOutputCommandLine = false;
+
+		/// <summary>
+		/// Instruct the executor to write compact output e.g. only errors, if supported by the executor,
+		/// and only if output is not being redirected e.g. during a build from within Visual Studio
+		/// </summary>
+		public bool bCompactOutput => bCompactOutputCommandLine && !Console.IsOutputRedirected;
 	}
 }
diff --git a/Engine/Source/Programs/UnrealBuildTool/Executors/Experimental/FASTBuild.cs b/Engine/Source/Programs/UnrealBuildTool/Executors/Experimental/FASTBuild.cs
--- a/Engine/Source/Programs/UnrealBuildTool/Executors/Experimental/FASTBuild.cs
+++ b/Engine/Source/Programs/UnrealBuildTool/Executors/Experimental/FASTBuild.cs
@@ -156,11 +156,11 @@ namespace UnrealBuildTool
 		/// <summary>
 		/// Constructor
 		/// </summary>
-		public FASTBuild(int MaxLocalActions)
+		public FASTBuild(int MaxLocalActions, bool bCompactOutput)
 		{
 			XmlConfig.ApplyTo(this);
 
-			this.LocalExecutor = new ParallelExecutor(MaxLocalActions);
+			this.LocalExecutor = new ParallelExecutor(MaxLocalActions, bCompactOutput);
 		}
 
 		public override string Name
diff --git a/Engine/Source/Programs/UnrealBuildTool/Executors/Experimental/HordeExecutor.cs b/Engine/Source/Programs/UnrealBuildTool/Executors/Experimental/HordeExecutor.cs
--- a/Engine/Source/Programs/UnrealBuildTool/Executors/Experimental/HordeExecutor.cs
+++ b/Engine/Source/Programs/UnrealBuildTool/Executors/Experimental/HordeExecutor.cs
@@ -278,7 +278,8 @@ namespace UnrealBuildTool
 		/// Constructor
 		/// </summary>
 		/// <param name="MaxLocalActions">How many local actions to execute in parallel</param>
-		public HordeExecutor(int MaxLocalActions) : base(MaxLocalActions)
+		/// <param name="bCompactOutput"></param>
+		public HordeExecutor(int MaxLocalActions, bool bCompactOutput) : base(MaxLocalActions, bCompactOutput)
 		{
 			XmlConfig.ApplyTo(this);
 		}
diff --git a/Engine/Source/Programs/UnrealBuildTool/Executors/HybridExecutor.cs b/Engine/Source/Programs/UnrealBuildTool/Executors/HybridExecutor.cs
--- a/Engine/Source/Programs/UnrealBuildTool/Executors/HybridExecutor.cs
+++ b/Engine/Source/Programs/UnrealBuildTool/Executors/HybridExecutor.cs
@@ -1,5 +1,6 @@
 // Copyright Epic Games, Inc. All Rights Reserved.
 
+using EpicGames.Core;
 using System.Collections.Generic;
 using System.Linq;
 
@@ -30,10 +31,10 @@ namespace UnrealBuildTool
 		/// <summary>
 		/// Constructor
 		/// </summary>
-		public HybridExecutor(List<TargetDescriptor> TargetDescriptors, int InMaxLocalActions, ActionExecutor? InLocalExecutor = null, ActionExecutor? InRemoteExecutor = null)
+		public HybridExecutor(List<TargetDescriptor> TargetDescriptors, int InMaxLocalActions, bool bCompactOutput, ActionExecutor? InLocalExecutor = null, ActionExecutor? InRemoteExecutor = null)
 		{
 			MaxLocalActions = InMaxLocalActions;
-			LocalExecutor = InLocalExecutor ?? new ParallelExecutor(MaxLocalActions);
+			LocalExecutor = InLocalExecutor ?? new ParallelExecutor(MaxLocalActions, bCompactOutput);
 			RemoteExecutor = InRemoteExecutor ?? (XGE.IsAvailable() ? (ActionExecutor)new XGE() : new SNDBS(TargetDescriptors));
 
 			XmlConfig.ApplyTo(this);
diff --git a/Engine/Source/Programs/UnrealBuildTool/Executors/ParallelExecutor.cs b/Engine/Source/Programs/UnrealBuildTool/Executors/ParallelExecutor.cs
--- a/Engine/Source/Programs/UnrealBuildTool/Executors/ParallelExecutor.cs
+++ b/Engine/Source/Programs/UnrealBuildTool/Executors/ParallelExecutor.cs
@@ -69,6 +69,11 @@ namespace UnrealBuildTool
 		[XmlConfigFile]
 		private static bool bPrintActionTargetNames = false;
 
+		/// <summary>
+		/// Collapse non-error output lines
+		/// </summary>
+		private bool bCompactOutput = false;
+
 		/// <summary>
 		/// How many processes that will be executed in parallel
 		/// </summary>
@@ -91,7 +96,8 @@ namespace UnrealBuildTool
 		/// Constructor
 		/// </summary>
 		/// <param name="MaxLocalActions">How many actions to execute in parallel</param>
-		public ParallelExecutor(int MaxLocalActions)
+		/// <param name="bCompactOutput">Should output be written in a compact fashion</param>
+		public ParallelExecutor(int MaxLocalActions, bool bCompactOutput)
 		{
 			XmlConfig.ApplyTo(this);
 
@@ -105,6 +111,8 @@ namespace UnrealBuildTool
 				// Figure out how many processors to use
 				NumParallelProcesses = GetDefaultNumParallelProcesses();
 			}
+
+			this.bCompactOutput = bCompactOutput;
 		}
 
 		/// <summary>
@@ -290,7 +298,9 @@ namespace UnrealBuildTool
 			return new ExecuteResults(LogLines, ExitCode, ExecutionTime, ProcessorTime);
 		}
 
-		protected static void LogCompletedAction(LinkedAction Action, Task<ExecuteResults> ExecuteTask, CancellationTokenSource CancellationTokenSource, ProgressWriter ProgressWriter, int TotalActions, ref int NumCompletedActions)
+		private static int s_previousLineLength = -1;
+
+		protected void LogCompletedAction(LinkedAction Action, Task<ExecuteResults> ExecuteTask, CancellationTokenSource CancellationTokenSource, ProgressWriter ProgressWriter, int TotalActions, ref int NumCompletedActions)
 		{
 			List<string> LogLines = new List<string>();
 			int ExitCode = int.MaxValue;
@@ -361,10 +371,37 @@ namespace UnrealBuildTool
 					}
 				}
 
-				Log.TraceInformation($"[{CompletedActions}/{TotalActions}]{TargetDetails}{CompilationTimes} {Description}");
+				string message = ($"[{CompletedActions}/{TotalActions}]{TargetDetails}{CompilationTimes} {Description}");
+				
+				if (bCompactOutput)
+				{
+					if (s_previousLineLength > 0)
+					{
+						// move the cursor to the far left position, one line back
+						Console.CursorLeft = 0;
+						Console.CursorTop -= 1;
+						// clear the line
+						Console.Write("".PadRight(s_previousLineLength));
+						// move the cursor back to the left, so output is written to the desired location
+						Console.CursorLeft = 0;
+					}
+				}
+
+				s_previousLineLength = message.Length;
+
+				Log.TraceInformation(message);
 				foreach (string Line in LogLines.Skip(Action.bShouldOutputStatusDescription ? 0 : 1))
 				{
+					// suppress library creation messages when writing compact output
+					if (bCompactOutput && Line.StartsWith("   Creating library ") && Line.EndsWith(".exp"))
+					{
+						continue;
+					}
+
 					Log.TraceInformation(Line);
+
+					// Prevent overwriting of logged lines
+					s_previousLineLength = -1;
 				}
 
 				if (ExitCode != 0)
@@ -378,6 +415,9 @@ namespace UnrealBuildTool
 					}
 					// END TEMPORARY
 
+					// prevent overrwriting of error text
+					s_previousLineLength = -1;
+
 					// Cancel all other pending tasks
 					if (bStopCompilationAfterErrors)
 					{
diff --git a/Engine/Source/Programs/UnrealBuildTool/System/ActionGraph.cs b/Engine/Source/Programs/UnrealBuildTool/System/ActionGraph.cs
--- a/Engine/Source/Programs/UnrealBuildTool/System/ActionGraph.cs
+++ b/Engine/Source/Programs/UnrealBuildTool/System/ActionGraph.cs
@@ -336,7 +336,7 @@ namespace UnrealBuildTool
 			{
 				if (BuildConfiguration.bAllowHybridExecutor && HybridExecutor.IsAvailable())
 				{
-					return new HybridExecutor(TargetDescriptors, BuildConfiguration.MaxParallelActions);
+					return new HybridExecutor(TargetDescriptors, BuildConfiguration.MaxParallelActions, BuildConfiguration.bCompactOutput);
 				}
 				else if (BuildConfiguration.bAllowXGE && XGE.IsAvailable() && ActionCount >= XGE.MinActions)
 				{
@@ -344,7 +344,7 @@ namespace UnrealBuildTool
 				}
 				else if (BuildConfiguration.bAllowFASTBuild && FASTBuild.IsAvailable())
 				{
-					return new FASTBuild(BuildConfiguration.MaxParallelActions);
+					return new FASTBuild(BuildConfiguration.MaxParallelActions, BuildConfiguration.bCompactOutput);
 				}
 				else if (BuildConfiguration.bAllowSNDBS && SNDBS.IsAvailable())
 				{
@@ -352,11 +352,11 @@ namespace UnrealBuildTool
 				}
 				else if (BuildConfiguration.bAllowHordeCompute && HordeExecutor.IsAvailable())
 				{
-					return new HordeExecutor(BuildConfiguration.MaxParallelActions);
+					return new HordeExecutor(BuildConfiguration.MaxParallelActions, BuildConfiguration.bCompactOutput);
 				}
 			}
 
-			return new ParallelExecutor(BuildConfiguration.MaxParallelActions);
+			return new ParallelExecutor(BuildConfiguration.MaxParallelActions, BuildConfiguration.bCompactOutput);
 		}
 
 		/// <summary>
