UnrealBuildTool: Delete files in parallel when cleaning (adding a few timers)

Cleaning UnrealEngine Development Win64, on a ThreadRipper
Before: 270 seconds
After:   63 seconds

#rnx

diff --git a/Engine/Source/Programs/UnrealBuildTool/Modes/CleanMode.cs b/Engine/Source/Programs/UnrealBuildTool/Modes/CleanMode.cs
--- a/Engine/Source/Programs/UnrealBuildTool/Modes/CleanMode.cs
+++ b/Engine/Source/Programs/UnrealBuildTool/Modes/CleanMode.cs
@@ -57,6 +57,7 @@ namespace UnrealBuildTool
 
 		public void Clean(List<TargetDescriptor> TargetDescriptors, BuildConfiguration BuildConfiguration)
 		{
+			using ScopedTimer CleanTimer = new ScopedTimer("CleanMode.Clean()");
 			using IScope Scope = GlobalTracer.Instance.BuildSpan("CleanMode.Clean()").StartActive();
 
 			if (TargetDescriptors.Count == 0)
@@ -105,144 +106,197 @@ namespace UnrealBuildTool
 			HashSet<FileReference> FilesToDelete = new HashSet<FileReference>();
 			HashSet<DirectoryReference> DirectoriesToDelete = new HashSet<DirectoryReference>();
 
-			for (int Idx = 0; Idx < TargetDescriptors.Count; ++Idx)
+			using (ScopedTimer GatherTimer = new ScopedTimer("Find paths to clean"))
 			{
-				TargetDescriptor TargetDescriptor = TargetDescriptors[Idx];
+				for (int Idx = 0; Idx < TargetDescriptors.Count; ++Idx)
+				{
+					TargetDescriptor TargetDescriptor = TargetDescriptors[Idx];
 
-				// Create the rules assembly
-				RulesAssembly RulesAssembly = RulesCompiler.CreateTargetRulesAssembly(TargetDescriptor.ProjectFile, TargetDescriptor.Name, bSkipRulesCompile, BuildConfiguration.bForceRulesCompile, BuildConfiguration.bUsePrecompiled, TargetDescriptor.ForeignPlugin);
+					// Create the rules assembly
+					RulesAssembly RulesAssembly = RulesCompiler.CreateTargetRulesAssembly(TargetDescriptor.ProjectFile, TargetDescriptor.Name, bSkipRulesCompile, BuildConfiguration.bForceRulesCompile, BuildConfiguration.bUsePrecompiled, TargetDescriptor.ForeignPlugin);
 
-				// Create the rules object
-				ReadOnlyTargetRules Target = new ReadOnlyTargetRules(RulesAssembly.CreateTargetRules(TargetDescriptor.Name, TargetDescriptor.Platform, TargetDescriptor.Configuration, TargetDescriptor.Architecture, TargetDescriptor.ProjectFile, TargetDescriptor.AdditionalArguments));
+					// Create the rules object
+					ReadOnlyTargetRules Target = new ReadOnlyTargetRules(RulesAssembly.CreateTargetRules(TargetDescriptor.Name, TargetDescriptor.Platform, TargetDescriptor.Configuration, TargetDescriptor.Architecture, TargetDescriptor.ProjectFile, TargetDescriptor.AdditionalArguments));
 
-				if (!bSkipPreBuildTargets && Target.PreBuildTargets.Count > 0)
-				{
-					foreach (TargetInfo PreBuildTarget in Target.PreBuildTargets)
+					if (!bSkipPreBuildTargets && Target.PreBuildTargets.Count > 0)
 					{
-						TargetDescriptor NewTarget = TargetDescriptor.FromTargetInfo(PreBuildTarget);
-						if (!TargetDescriptors.Contains(NewTarget))
+						foreach (TargetInfo PreBuildTarget in Target.PreBuildTargets)
 						{
-							TargetDescriptors.Add(NewTarget);
+							TargetDescriptor NewTarget = TargetDescriptor.FromTargetInfo(PreBuildTarget);
+							if (!TargetDescriptors.Contains(NewTarget))
+							{
+								TargetDescriptors.Add(NewTarget);
+							}
 						}
 					}
-				}
 
-				// Find the base folders that can contain binaries
-				List<DirectoryReference> BaseDirs = new List<DirectoryReference>();
-				BaseDirs.Add(Unreal.EngineDirectory);
-				foreach (FileReference Plugin in PluginsBase.EnumeratePlugins(Target.ProjectFile))
-				{
-					BaseDirs.Add(Plugin.Directory);
-				}
-				if (Target.ProjectFile != null)
-				{
-					BaseDirs.Add(Target.ProjectFile.Directory);
-				}
+					// Find the base folders that can contain binaries
+					List<DirectoryReference> BaseDirs = new List<DirectoryReference>();
+					BaseDirs.Add(Unreal.EngineDirectory);
+					foreach (FileReference Plugin in PluginsBase.EnumeratePlugins(Target.ProjectFile))
+					{
+						BaseDirs.Add(Plugin.Directory);
+					}
+					if (Target.ProjectFile != null)
+					{
+						BaseDirs.Add(Target.ProjectFile.Directory);
+					}
 
-				// If we're running a precompiled build, remove anything under the engine folder
-				BaseDirs.RemoveAll(x => RulesAssembly.IsReadOnly(x));
+					// If we're running a precompiled build, remove anything under the engine folder
+					BaseDirs.RemoveAll(x => RulesAssembly.IsReadOnly(x));
 
-				// Get all the names which can prefix build products
-				List<string> NamePrefixes = new List<string>();
-				if (Target.Type != TargetType.Program)
-				{
-					NamePrefixes.Add(UEBuildTarget.GetAppNameForTargetType(Target.Type));
-				}
-				NamePrefixes.Add(Target.Name);
+					// Get all the names which can prefix build products
+					List<string> NamePrefixes = new List<string>();
+					if (Target.Type != TargetType.Program)
+					{
+						NamePrefixes.Add(UEBuildTarget.GetAppNameForTargetType(Target.Type));
+					}
+					NamePrefixes.Add(Target.Name);
 
-				// Get the suffixes for this configuration
-				List<string> NameSuffixes = new List<string>();
-				if (Target.Configuration == Target.UndecoratedConfiguration)
-				{
-					NameSuffixes.Add("");
-				}
-				NameSuffixes.Add(String.Format("-{0}-{1}", Target.Platform.ToString(), Target.Configuration.ToString()));
-				if (!String.IsNullOrEmpty(Target.Architecture))
-				{
-					NameSuffixes.AddRange(NameSuffixes.ToArray().Select(x => x + Target.Architecture));
-				}
+					// Get the suffixes for this configuration
+					List<string> NameSuffixes = new List<string>();
+					if (Target.Configuration == Target.UndecoratedConfiguration)
+					{
+						NameSuffixes.Add("");
+					}
+					NameSuffixes.Add(String.Format("-{0}-{1}", Target.Platform.ToString(), Target.Configuration.ToString()));
+					if (!String.IsNullOrEmpty(Target.Architecture))
+					{
+						NameSuffixes.AddRange(NameSuffixes.ToArray().Select(x => x + Target.Architecture));
+					}
 
-				// Add all the makefiles and caches to be deleted
-				FilesToDelete.Add(TargetMakefile.GetLocation(Target.ProjectFile, Target.Name, Target.Platform, Target.Architecture, Target.Configuration));
-				FilesToDelete.UnionWith(SourceFileMetadataCache.GetFilesToClean(Target.ProjectFile));
+					// Add all the makefiles and caches to be deleted
+					FilesToDelete.Add(TargetMakefile.GetLocation(Target.ProjectFile, Target.Name, Target.Platform, Target.Architecture, Target.Configuration));
+					FilesToDelete.UnionWith(SourceFileMetadataCache.GetFilesToClean(Target.ProjectFile));
 
-				// Add all the intermediate folders to be deleted
-				foreach (DirectoryReference BaseDir in BaseDirs)
-				{
-					foreach (string NamePrefix in NamePrefixes)
+					// Add all the intermediate folders to be deleted
+					foreach (DirectoryReference BaseDir in BaseDirs)
 					{
-						DirectoryReference GeneratedCodeDir = DirectoryReference.Combine(BaseDir, UEBuildTarget.GetPlatformIntermediateFolder(Target.Platform, Target.Architecture, false), NamePrefix, "Inc");
-						if (DirectoryReference.Exists(GeneratedCodeDir))
+						foreach (string NamePrefix in NamePrefixes)
 						{
-							DirectoriesToDelete.Add(GeneratedCodeDir);
+							DirectoryReference GeneratedCodeDir = DirectoryReference.Combine(BaseDir, UEBuildTarget.GetPlatformIntermediateFolder(Target.Platform, Target.Architecture, false), NamePrefix, "Inc");
+							if (DirectoryReference.Exists(GeneratedCodeDir))
+							{
+								DirectoriesToDelete.Add(GeneratedCodeDir);
+							}
+
+							DirectoryReference IntermediateDir = DirectoryReference.Combine(BaseDir, UEBuildTarget.GetPlatformIntermediateFolder(Target.Platform, Target.Architecture, false), NamePrefix, Target.Configuration.ToString());
+							if (DirectoryReference.Exists(IntermediateDir))
+							{
+								DirectoriesToDelete.Add(IntermediateDir);
+							}
 						}
+					}
+
+					// todo: handle external plugin intermediates, written to the Project's Intermediate/External directory
 
-						DirectoryReference IntermediateDir = DirectoryReference.Combine(BaseDir, UEBuildTarget.GetPlatformIntermediateFolder(Target.Platform, Target.Architecture, false), NamePrefix, Target.Configuration.ToString());
-						if (DirectoryReference.Exists(IntermediateDir))
+					// List of additional files and directories to clean, specified by the target platform
+					List<FileReference> AdditionalFilesToDelete = new List<FileReference>();
+					List<DirectoryReference> AdditionalDirectoriesToDelete = new List<DirectoryReference>();
+
+					// Add all the build products from this target
+					string[] NamePrefixesArray = NamePrefixes.Distinct().ToArray();
+					string[] NameSuffixesArray = NameSuffixes.Distinct().ToArray();
+					foreach (DirectoryReference BaseDir in BaseDirs)
+					{
+						DirectoryReference BinariesDir = DirectoryReference.Combine(BaseDir, "Binaries", Target.Platform.ToString());
+						if (DirectoryReference.Exists(BinariesDir))
 						{
-							DirectoriesToDelete.Add(IntermediateDir);
+							UEBuildPlatform.GetBuildPlatform(Target.Platform).FindBuildProductsToClean(BinariesDir, NamePrefixesArray, NameSuffixesArray, AdditionalFilesToDelete, AdditionalDirectoriesToDelete);
 						}
 					}
-				}
 
-				// todo: handle external plugin intermediates, written to the Project's Intermediate/External directory
+					// Get all the additional intermediate folders created by this platform
+					UEBuildPlatform.GetBuildPlatform(Target.Platform).FindAdditionalBuildProductsToClean(Target, AdditionalFilesToDelete, AdditionalDirectoriesToDelete);
 
-				// List of additional files and directories to clean, specified by the target platform
-				List<FileReference> AdditionalFilesToDelete = new List<FileReference>();
-				List<DirectoryReference> AdditionalDirectoriesToDelete = new List<DirectoryReference>();
+					// Add the platform's files and directories to the main list
+					FilesToDelete.UnionWith(AdditionalFilesToDelete);
+					DirectoriesToDelete.UnionWith(AdditionalDirectoriesToDelete);
+				}
+			}
 
-				// Add all the build products from this target
-				string[] NamePrefixesArray = NamePrefixes.Distinct().ToArray();
-				string[] NameSuffixesArray = NameSuffixes.Distinct().ToArray();
-				foreach (DirectoryReference BaseDir in BaseDirs)
+			// Ensure no overlap between directories
+			using (ScopedTimer Timer = new ScopedTimer("Ensure no directory overlap"))
+			{
+				HashSet<DirectoryReference> SubdirectoriesToDelete = new(DirectoriesToDelete.Count);
+				foreach (DirectoryReference Directory in DirectoriesToDelete)
 				{
-					DirectoryReference BinariesDir = DirectoryReference.Combine(BaseDir, "Binaries", Target.Platform.ToString());
-					if (DirectoryReference.Exists(BinariesDir))
+					// Is this directory a subdirectory of some other directory that is being deleted?
+					for (DirectoryReference? DirectoryWalker = Directory.ParentDirectory; DirectoryWalker != null; DirectoryWalker = DirectoryWalker.ParentDirectory)
 					{
-						UEBuildPlatform.GetBuildPlatform(Target.Platform).FindBuildProductsToClean(BinariesDir, NamePrefixesArray, NameSuffixesArray, AdditionalFilesToDelete, AdditionalDirectoriesToDelete);
+						if (DirectoriesToDelete.Contains(DirectoryWalker))
+						{
+							SubdirectoriesToDelete.Add(Directory);
+							break;
+						}
 					}
 				}
-
-				// Get all the additional intermediate folders created by this platform
-				UEBuildPlatform.GetBuildPlatform(Target.Platform).FindAdditionalBuildProductsToClean(Target, AdditionalFilesToDelete, AdditionalDirectoriesToDelete);
-
-				// Add the platform's files and directories to the main list
-				FilesToDelete.UnionWith(AdditionalFilesToDelete);
-				DirectoriesToDelete.UnionWith(AdditionalDirectoriesToDelete);
+				DirectoriesToDelete.ExceptWith(SubdirectoriesToDelete);
 			}
 
-			// Delete all the directories, then all the files. By sorting the list of directories before we delete them, we avoid spamming the log if a parent directory is deleted first.
-			foreach (DirectoryReference DirectoryToDelete in DirectoriesToDelete.OrderBy(x => x.FullName))
+			// Remove any files that are contained within one of the directories
+			using (ScopedTimer Time = new ScopedTimer("Ensure no file overlap"))
 			{
-				if (DirectoryReference.Exists(DirectoryToDelete))
+				FilesToDelete.RemoveWhere(File =>
 				{
-					Log.TraceVerbose("    Deleting {0}{1}...", DirectoryToDelete, Path.DirectorySeparatorChar);
-					try
+					// Is this file in a subdirectory of some directory that is being deleted?
+					for (DirectoryReference? DirectoryWalker = File.Directory; DirectoryWalker != null; DirectoryWalker = DirectoryWalker.ParentDirectory)
 					{
-						FileUtils.ForceDeleteDirectory(DirectoryToDelete);
-					}
-					catch (Exception Ex)
-					{
-						throw new BuildException(Ex, "Unable to delete {0} ({1})", DirectoryToDelete, Ex.Message.TrimEnd());
+						if (DirectoriesToDelete.Contains(DirectoryWalker))
+						{
+							return true;
+						}
 					}
-				}
+					return false;
+				});
 			}
 
-			foreach (FileReference FileToDelete in FilesToDelete.OrderBy(x => x.FullName))
+			var DeleteDirectories = Task.Run(() =>
 			{
-				if (FileReference.Exists(FileToDelete))
+				using ScopedTimer Timer = new ScopedTimer($"Delete {DirectoriesToDelete.Count} directories");
+				Parallel.ForEach(DirectoriesToDelete, DirectoryToDelete =>
 				{
-					Log.TraceVerbose("    Deleting " + FileToDelete);
-					try
+					using ScopedTimer Timer = new ScopedTimer($"Delete directory '{DirectoryToDelete}'", bIncreaseIndent: false);
+
+					if (DirectoryReference.Exists(DirectoryToDelete))
 					{
-						FileUtils.ForceDeleteFile(FileToDelete);
+						Log.TraceVerbose($"    Deleting {DirectoryToDelete}{Path.DirectorySeparatorChar}...");
+
+						try
+						{
+							FileUtils.ForceDeleteDirectory(DirectoryToDelete);
+						}
+						catch (Exception Ex)
+						{
+							throw new BuildException(Ex, "Unable to delete {0} ({1})", DirectoryToDelete, Ex.Message.TrimEnd());
+						}
 					}
-					catch (Exception Ex)
+				});
+			});
+
+			var DeleteFiles = Task.Run(() =>
+			{
+				using ScopedTimer Timer = new ScopedTimer($"Delete {FilesToDelete.Count} files");
+				Parallel.ForEach(FilesToDelete, FileToDelete =>
+				{
+					if (FileReference.Exists(FileToDelete))
 					{
-						throw new BuildException(Ex, "Unable to delete {0} ({1})", FileToDelete, Ex.Message.TrimEnd());
+						Log.TraceVerbose($"    Deleting {FileToDelete}...");
+
+						try
+						{
+							FileUtils.ForceDeleteFile(FileToDelete);
+						}
+						catch (Exception Ex)
+						{
+							throw new BuildException(Ex, "Unable to delete {0} ({1})", FileToDelete, Ex.Message.TrimEnd());
+						}
 					}
-				}
-			}
+				});
+			});
+
+			DeleteDirectories.Wait();
+			DeleteFiles.Wait();
 
 			// Also clean all the remote targets
 			for (int Idx = 0; Idx < TargetDescriptors.Count; Idx++)
diff --git a/Engine/Source/Programs/UnrealBuildTool/System/ScopedTimer.cs b/Engine/Source/Programs/UnrealBuildTool/System/ScopedTimer.cs
--- a/Engine/Source/Programs/UnrealBuildTool/System/ScopedTimer.cs
+++ b/Engine/Source/Programs/UnrealBuildTool/System/ScopedTimer.cs
@@ -17,6 +17,7 @@ namespace UnrealBuildTool
 		DateTime StartTime;
 		string TimerName;
 		LogEventType Verbosity;
+		bool bIncreaseIndent;
 		static int Indent = 0;
 		static object IndentLock = new object();
 
@@ -25,15 +26,20 @@ namespace UnrealBuildTool
 		/// </summary>
 		/// <param name="Name">Name of the block being measured</param>
 		/// <param name="InVerbosity">Verbosity for output messages</param>
-		public ScopedTimer(string Name, LogEventType InVerbosity = LogEventType.Verbose)
+		/// <param name="bIncreaseIndent">Whether gobal indent should be increased or not; set to false when running a scope in parallel. Message will still be printed indented relative to parent scope.</param>
+		public ScopedTimer(string Name, LogEventType InVerbosity = LogEventType.Verbose, bool bIncreaseIndent = true)
 		{
 			TimerName = Name;
-			lock (IndentLock)
+			if (bIncreaseIndent)
 			{
-				Indent++;
+				lock (IndentLock)
+				{
+					Indent++;
+				}
 			}
 			Verbosity = InVerbosity;
 			StartTime = DateTime.UtcNow;
+			this.bIncreaseIndent = bIncreaseIndent;
 		}
 
 		/// <summary>
@@ -42,10 +48,13 @@ namespace UnrealBuildTool
 		public void Dispose()
 		{
 			double TotalSeconds = (DateTime.UtcNow - StartTime).TotalSeconds;
-			int LogIndent = 0;
-			lock (IndentLock)
+			int LogIndent = Indent;
+			if (bIncreaseIndent)
 			{
-				LogIndent = --Indent;
+				lock (IndentLock)
+				{
+					LogIndent = --Indent;
+				}
 			}
 			StringBuilder IndentText = new StringBuilder(LogIndent * 2);
 			IndentText.Append(' ', LogIndent * 2);
