More coding standard fixes

diff --git a/Engine/Source/Programs/UnrealBuildTool/Modes/UnrealHeaderToolMode.cs b/Engine/Source/Programs/UnrealBuildTool/Modes/UnrealHeaderToolMode.cs
--- a/Engine/Source/Programs/UnrealBuildTool/Modes/UnrealHeaderToolMode.cs
+++ b/Engine/Source/Programs/UnrealBuildTool/Modes/UnrealHeaderToolMode.cs
@@ -41,6 +41,7 @@ namespace UnrealBuildTool.Modes
 		/// <summary>
 		/// Special parsed struct names that do not require a prefix
 		/// </summary>
+		[System.Diagnostics.CodeAnalysis.SuppressMessage("CodeQuality", "IDE0052:Remove unread private members", Justification = "<Pending>")]
 		private readonly ReadOnlyHashSet<StringView> StructsWithNoPrefix;
 
 		/// <summary>
@@ -123,8 +124,7 @@ namespace UnrealBuildTool.Modes
 				throw new Exception("Attempt to redirect type identifier when the token isn't an identifier.");
 			}
 
-			StringView Redirect;
-			if (this.TypeRedirectMap.TryGetValue(Token.Value, out Redirect))
+			if (this.TypeRedirectMap.TryGetValue(Token.Value, out StringView Redirect))
 			{
 				Token.Value = Redirect;
 			}
@@ -133,8 +133,7 @@ namespace UnrealBuildTool.Modes
 		/// <inheritdoc/>
 		public bool RedirectMetaDataKey(string Key, out string NewKey)
 		{
-			string? Redirect;
-			if (this.MetaDataRedirectMap.TryGetValue(Key, out Redirect))
+			if (this.MetaDataRedirectMap.TryGetValue(Key, out string? Redirect))
 			{
 				NewKey = Redirect;
 				return Key != NewKey;
@@ -161,7 +160,6 @@ namespace UnrealBuildTool.Modes
 		/// <inheritdoc/>
 		public int FindDelegateParameterCount(StringView DelegateMacro)
 		{
-			IReadOnlyList<StringView> DelegateParameterCountStrings = this.DelegateParameterCountStrings;
 			for (int Index = 0, Count = this.DelegateParameterCountStrings.Count; Index < Count; ++Index)
 			{
 				if (DelegateMacro.Span.Contains(this.DelegateParameterCountStrings[Index].Span, StringComparison.Ordinal))
@@ -181,8 +179,7 @@ namespace UnrealBuildTool.Modes
 		/// <inheritdoc/>
 		public bool IsExporterEnabled(string Name)
 		{
-			bool Value = false;
-			this.Ini.GetBool("UnrealHeaderTool", Name, out Value);
+			this.Ini.GetBool("UnrealHeaderTool", Name, out bool Value);
 			return Value;
 		}
 		#endregion
@@ -213,21 +210,18 @@ namespace UnrealBuildTool.Modes
 
 		private bool GetBoolean(string SectionName, string KeyName, bool bDefault)
 		{
-			bool bValue;
-			if (this.Ini.TryGetValue(SectionName, KeyName, out bValue))
+			if (this.Ini.TryGetValue(SectionName, KeyName, out bool value))
 			{
-				return bValue;
+				return value;
 			}
 			return bDefault;
 		}
 
 		private UhtPointerMemberBehavior GetPointerMemberBehavior(string SectionName, string KeyName, UhtPointerMemberBehavior Default)
 		{
-			string? BehaviorStr;
-			if (this.Ini.TryGetValue(SectionName, KeyName, out BehaviorStr))
+			if (this.Ini.TryGetValue(SectionName, KeyName, out string? BehaviorStr))
 			{
-				UhtPointerMemberBehavior Value;
-				if (!Enum.TryParse(BehaviorStr, out Value))
+				if (!Enum.TryParse(BehaviorStr, out UhtPointerMemberBehavior Value))
 				{
 					throw new Exception(string.Format("Unrecognized native pointer member behavior '{0}'", BehaviorStr));
 				}
@@ -238,11 +232,9 @@ namespace UnrealBuildTool.Modes
 
 		private EGeneratedCodeVersion GetGeneratedCodeVersion(string SectionName, string KeyName, EGeneratedCodeVersion Default)
 		{
-			string? BehaviorStr;
-			if (this.Ini.TryGetValue(SectionName, KeyName, out BehaviorStr))
+			if (this.Ini.TryGetValue(SectionName, KeyName, out string? BehaviorStr))
 			{
-				EGeneratedCodeVersion Value;
-				if (!Enum.TryParse(BehaviorStr, out Value))
+				if (!Enum.TryParse(BehaviorStr, out EGeneratedCodeVersion Value))
 				{
 					throw new Exception(string.Format("Unrecognized generated code version '{0}'", BehaviorStr));
 				}
@@ -253,23 +245,19 @@ namespace UnrealBuildTool.Modes
 
 		private IReadOnlyDictionary<StringView, StringView> GetRedirectsStringView(string Section, string Key, string OldKeyName, string NewKeyName)
 		{
-			Dictionary<StringView, StringView> Redirects = new Dictionary<StringView, StringView>();
+			Dictionary<StringView, StringView> Redirects = new();
 
-			IReadOnlyList<string>? StringList;
-			if (this.Ini.TryGetValues(Section, Key, out StringList))
+			if (this.Ini.TryGetValues(Section, Key, out IReadOnlyList<string>? StringList))
 			{
 				foreach (string Line in StringList)
 				{
-					Dictionary<string, string>? Properties;
-					if (ConfigHierarchy.TryParse(Line, out Properties))
+					if (ConfigHierarchy.TryParse(Line, out Dictionary<string, string>? Properties))
 					{
-						string? OldKey;
-						if (!Properties.TryGetValue(OldKeyName, out OldKey))
+						if (!Properties.TryGetValue(OldKeyName, out string? OldKey))
 						{
 							throw new Exception(string.Format("Unable to get the {0} from the {1} value", OldKeyName, Key));
 						}
-						string? NewKey;
-						if (!Properties.TryGetValue(NewKeyName, out NewKey))
+						if (!Properties.TryGetValue(NewKeyName, out string? NewKey))
 						{
 							throw new Exception(string.Format("Unable to get the {0} from the {1} value", NewKeyName, Key));
 						}
@@ -282,23 +270,19 @@ namespace UnrealBuildTool.Modes
 
 		private IReadOnlyDictionary<string, string> GetRedirectsString(string Section, string Key, string OldKeyName, string NewKeyName)
 		{
-			Dictionary<string, string> Redirects = new Dictionary<string, string>();
+			Dictionary<string, string> Redirects = new();
 
-			IReadOnlyList<string>? StringList;
-			if (this.Ini.TryGetValues(Section, Key, out StringList))
+			if (this.Ini.TryGetValues(Section, Key, out IReadOnlyList<string>? StringList))
 			{
 				foreach (string Line in StringList)
 				{
-					Dictionary<string, string>? Properties;
-					if (ConfigHierarchy.TryParse(Line, out Properties))
+					if (ConfigHierarchy.TryParse(Line, out Dictionary<string, string>? Properties))
 					{
-						string? OldKey;
-						if (!Properties.TryGetValue(OldKeyName, out OldKey))
+						if (!Properties.TryGetValue(OldKeyName, out string? OldKey))
 						{
 							throw new Exception(string.Format("Unable to get the {0} from the {1} value", OldKeyName, Key));
 						}
-						string? NewKey;
-						if (!Properties.TryGetValue(NewKeyName, out NewKey))
+						if (!Properties.TryGetValue(NewKeyName, out string? NewKey))
 						{
 							throw new Exception(string.Format("Unable to get the {0} from the {1} value", NewKeyName, Key));
 						}
@@ -311,10 +295,9 @@ namespace UnrealBuildTool.Modes
 
 		private IReadOnlyList<StringView> GetList(string Section, string Key)
 		{
-			List<StringView> List = new List<StringView>();
+			List<StringView> List = new();
 
-			IReadOnlyList<string>? StringList;
-			if (this.Ini.TryGetValues(Section, Key, out StringList))
+			if (this.Ini.TryGetValues(Section, Key, out IReadOnlyList<string>? StringList))
 			{
 				foreach (string Value in StringList)
 				{
@@ -326,10 +309,9 @@ namespace UnrealBuildTool.Modes
 
 		private ReadOnlyHashSet<StringView> GetHashSet(string Section, string Key, StringViewComparer Comparer)
 		{
-			HashSet<StringView> Set = new HashSet<StringView>(Comparer);
+			HashSet<StringView> Set = new(Comparer);
 
-			IReadOnlyList<string>? StringList;
-			if (this.Ini.TryGetValues(Section, Key, out StringList))
+			if (this.Ini.TryGetValues(Section, Key, out IReadOnlyList<string>? StringList))
 			{
 				foreach (string Value in StringList)
 				{
@@ -429,12 +411,12 @@ namespace UnrealBuildTool.Modes
 		/// <summary>
 		/// Collection of test fragments that can be read
 		/// </summary>
-		public Dictionary<string, UhtSourceFragment> SourceFragments = new Dictionary<string, UhtSourceFragment>();
+		public Dictionary<string, UhtSourceFragment> SourceFragments = new();
 
 		/// <summary>
 		/// All output segments generated by code gen
 		/// </summary>
-		public SortedDictionary<string, string> Outputs = new SortedDictionary<string, string>();
+		public SortedDictionary<string, string> Outputs = new();
 
 		private readonly IUhtFileManager InnerManager;
 		private readonly string? RootDirectory;
@@ -546,35 +528,35 @@ namespace UnrealBuildTool.Modes
 			string InPath = Path.Combine(TestDirectory, Script);
 			string OutPath = Path.Combine(TestOutputDirectory, Script);
 
-			UhtTestFileManager TestFileManager = new UhtTestFileManager(TestDirectory);
-			UhtSession Session = new UhtSession
+			UhtTestFileManager TestFileManager = new(TestDirectory);
+			UhtSession Session = new()
 			{
 				Tables = Tables,
 				Config = Config,
 				FileManager = TestFileManager,
 				RootDirectory = TestDirectory,
-				bWarningsAsErrors = Options.bWarningsAsErrors,
-				bRelativePathInLog = true,
-				bGoWide = !Options.bNoGoWide,
-				bNoOutput = false,
-				bCullOutput = false,
-				bCacheMessages = true,
-				bIncludeDebugOutput = true,
+				WarningsAsErrors = Options.bWarningsAsErrors,
+				RelativePathInLog = true,
+				GoWide = !Options.bNoGoWide,
+				NoOutput = false,
+				CullOutput = false,
+				CacheMessages = true,
+				IncludeDebugOutput = true,
 			};
 
 			// Read the testing script
-			List<ScriptFragment> ScriptFragments = new List<ScriptFragment>();
+			List<ScriptFragment> ScriptFragments = new();
 			int ManifestIndex = -1;
 			int ConsoleIndex = -1;
-			UhtSourceFile ScriptSourceFile = new UhtSourceFile(Session, Script);
-			Dictionary<string, int> OutputFragments = new Dictionary<string, int>();
+			UhtSourceFile ScriptSourceFile = new(Session, Script);
+			Dictionary<string, int> OutputFragments = new();
 			Session.Try(ScriptSourceFile, () =>
 			{
 				ScriptSourceFile.Read();
-				UhtTokenBufferReader Reader = new UhtTokenBufferReader(ScriptSourceFile, ScriptSourceFile.Data.Memory);
+				UhtTokenBufferReader Reader = new(ScriptSourceFile, ScriptSourceFile.Data.Memory);
 
-				bool bDone = false;
-				while (!bDone)
+				bool done = false;
+				while (!done)
 				{
 
 					// Scan for the fragment header
@@ -585,40 +567,38 @@ namespace UnrealBuildTool.Modes
 					int LineNumber = 1;
 					while (true)
 					{
-						using (var SaveState = new UhtTokenSaveState(Reader))
+						using UhtTokenSaveState SaveState = new(Reader);
+						UhtToken Token = Reader.GetLine();
+						if (Token.TokenType == UhtTokenType.EndOfFile)
 						{
-							UhtToken Token = Reader.GetLine();
-							if (Token.TokenType == UhtTokenType.EndOfFile)
-							{
-								break;
-							}
-							if (Token.Value.Span.Length == 0 || (Token.Value.Span.Length > 0 && Token.Value.Span[0] != '!'))
-							{
-								break;
-							}
-							HeaderEndPos = Reader.InputPos;
-
-							int EndCommandPos = Token.Value.Span.IndexOf(' ');
-							if (EndCommandPos == -1)
-							{
-								EndCommandPos = Token.Value.Span.Length;
-							}
-							string ScriptFragmentTypeString = Token.Value.Span.Slice(1, EndCommandPos - 1).Trim().ToString();
+							break;
+						}
+						if (Token.Value.Span.Length == 0 || (Token.Value.Span.Length > 0 && Token.Value.Span[0] != '!'))
+						{
+							break;
+						}
+						HeaderEndPos = Reader.InputPos;
 
-							if (!System.Enum.TryParse<ScriptFragmentType>(ScriptFragmentTypeString, true, out Type))
-							{
-								continue;
-							}
-							if (Type == ScriptFragmentType.Unknown)
-							{
-								continue;
-							}
+						int EndCommandPos = Token.Value.Span.IndexOf(' ');
+						if (EndCommandPos == -1)
+						{
+							EndCommandPos = Token.Value.Span.Length;
+						}
+						string ScriptFragmentTypeString = Token.Value.Span[1..EndCommandPos].Trim().ToString();
 
-							Name = Token.Value.Span.Slice(EndCommandPos).Trim().ToString();
-							LineNumber = Token.InputLine;
-							SaveState.AbandonState();
-							break;
+						if (!System.Enum.TryParse<ScriptFragmentType>(ScriptFragmentTypeString, true, out Type))
+						{
+							continue;
 						}
+						if (Type == ScriptFragmentType.Unknown)
+						{
+							continue;
+						}
+
+						Name = Token.Value.Span[EndCommandPos..].Trim().ToString();
+						LineNumber = Token.InputLine;
+						SaveState.AbandonState();
+						break;
 					}
 
 					// Scan for the fragment body
@@ -626,21 +606,19 @@ namespace UnrealBuildTool.Modes
 					int BodyEndPos = BodyStartPos;
 					while (true)
 					{
-						using (var SaveState = new UhtTokenSaveState(Reader))
+						using var SaveState = new UhtTokenSaveState(Reader);
+						UhtToken Token = Reader.GetLine();
+						if (Token.TokenType == UhtTokenType.EndOfFile)
 						{
-							UhtToken Token = Reader.GetLine();
-							if (Token.TokenType == UhtTokenType.EndOfFile)
-							{
-								bDone = true;
-								break;
-							}
-							if (Token.Value.Span.Length > 0 && Token.Value.Span[0] == '!')
-							{
-								break;
-							}
-							BodyEndPos = Reader.InputPos;
-							SaveState.AbandonState();
+							done = true;
+							break;
+						}
+						if (Token.Value.Span.Length > 0 && Token.Value.Span[0] == '!')
+						{
+							break;
 						}
+						BodyEndPos = Reader.InputPos;
+						SaveState.AbandonState();
 					}
 
 					ScriptFragments.Add(new ScriptFragment
@@ -648,8 +626,8 @@ namespace UnrealBuildTool.Modes
 						Type = Type,
 						Name = Name.Replace("\\\\", "\\"), // Be kind to people cut/copy/paste escaped strings around
 						LineNumber = LineNumber,
-						Header = new StringView(ScriptSourceFile.Data.Memory.Slice(HeaderStartPos, HeaderEndPos - HeaderStartPos)),
-						Body = new StringView(ScriptSourceFile.Data.Memory.Slice(BodyStartPos, BodyEndPos - BodyStartPos)),
+						Header = new StringView(ScriptSourceFile.Data.Memory[HeaderStartPos..HeaderEndPos]),
+						Body = new StringView(ScriptSourceFile.Data.Memory[BodyStartPos..BodyEndPos]),
 						External = false,
 					});
 				}
@@ -725,7 +703,7 @@ namespace UnrealBuildTool.Modes
 			});
 
 			// Run UHT
-			if (!Session.bHasErrors)
+			if (!Session.HasErrors)
 			{
 				Session.Run(ScriptFragments[ManifestIndex].Name);
 			}
@@ -744,7 +722,7 @@ namespace UnrealBuildTool.Modes
 
 				// Generate the console block
 				List<string> ConsoleLines = Session.CollectMessages();
-				StringBuilder SBConsole = new StringBuilder();
+				StringBuilder SBConsole = new();
 				foreach (string Line in ConsoleLines)
 				{
 					SBConsole.AppendLine(Line);
@@ -782,7 +760,7 @@ namespace UnrealBuildTool.Modes
 
 				// Create the complete output.  Output includes all of the source fragments and console fragments
 				// and followed the output data sorted by file name.
-				StringBuilder SBTest = new StringBuilder();
+				StringBuilder SBTest = new();
 				for (int i = 0; i < ScriptFragments.Count; ++i)
 				{
 					if (ScriptFragments[i].Type != ScriptFragmentType.Output)
@@ -828,11 +806,6 @@ namespace UnrealBuildTool.Modes
 			return bSuccess;
 		}
 
-		private static bool RunManifestTests(UhtGlobalOptions Options, string TestDirectory, string TestOutputDirectory, List<string> Manifests)
-		{
-			return true;
-		}
-
 		private static bool RunScriptTests(UhtTables Tables, IUhtConfig Config, UhtGlobalOptions Options, string TestDirectory, string TestOutputDirectory, List<string> Scripts)
 		{
 			bool bResult = true;
@@ -858,21 +831,21 @@ namespace UnrealBuildTool.Modes
 			// Create output directory
 			Directory.CreateDirectory(TestOutputDirectory);
 
-			List<string> Scripts = new List<string>();
+			List<string> Scripts = new();
 			foreach (string Script in Directory.EnumerateFiles(TestDirectory, "*.uhttest"))
 			{
 				Scripts.Add(Path.GetFileName(Script));
 			}
 			Scripts.Sort(StringComparer.OrdinalIgnoreCase);
 
-			List<string> Directories = new List<string>();
+			List<string> Directories = new();
 			foreach (string Directory in Directory.EnumerateDirectories(TestDirectory))
 			{
 				Directories.Add(Path.GetFileName(Directory));
 			}
 			Directories.Sort(StringComparer.OrdinalIgnoreCase);
 
-			List<string> Manifests = new List<string>();
+			List<string> Manifests = new();
 			foreach (string Manifest in Directory.EnumerateFiles(TestDirectory, "*.uhtmanifest"))
 			{
 				Manifests.Add(Path.GetFileName(Manifest));
@@ -880,7 +853,6 @@ namespace UnrealBuildTool.Modes
 			Manifests.Sort(StringComparer.OrdinalIgnoreCase);
 
 			return
-				RunManifestTests(Options, TestDirectory, TestOutputDirectory, Manifests) &&
 				RunScriptTests(Tables, Config, Options, TestDirectory, TestOutputDirectory, Scripts) &&
 				RunDirectoryTests(Tables, Config, Options, TestDirectory, TestOutputDirectory, Directories);
 		}
@@ -1000,13 +972,13 @@ namespace UnrealBuildTool.Modes
 			{
 
 				// Initialize the attributes
-				UhtTables Tables = new UhtTables();
+				UhtTables Tables = new();
 
 				// Initialize the config
 				IUhtConfig Config = new UhtConfigImpl(Arguments);
 
 				// Parse the global options
-				UhtGlobalOptions Options = new UhtGlobalOptions(Arguments);
+				UhtGlobalOptions Options = new(Arguments);
 				int TargetArgumentIndex = -1;
 				if (Arguments.GetPositionalArgumentCount() == 0)
 				{
@@ -1043,7 +1015,7 @@ namespace UnrealBuildTool.Modes
 					{
 						string BaseLogFileName = FileReference.Combine(EngineProgramSavedDirectory, "UnrealBuildTool", "Log_UHT.txt").FullName;
 
-						FileReference LogFile = new FileReference(BaseLogFileName);
+						FileReference LogFile = new(BaseLogFileName);
 						foreach (string LogSuffix in Arguments.GetValues("-LogSuffix="))
 						{
 							LogFile = LogFile.ChangeExtension(null) + "_" + LogSuffix + LogFile.GetExtension();
@@ -1064,7 +1036,7 @@ namespace UnrealBuildTool.Modes
 
 				if (TargetArgumentIndex >= 0)
 				{
-					CommandLineArguments LocalArguments = new CommandLineArguments(new string[] { Arguments[TargetArgumentIndex] });
+					CommandLineArguments LocalArguments = new(new string[] { Arguments[TargetArgumentIndex] });
 					List<TargetDescriptor> TargetDescriptors = TargetDescriptor.ParseCommandLine(LocalArguments, false, false, false);
 					if (TargetDescriptors.Count == 0)
 					{
@@ -1078,25 +1050,24 @@ namespace UnrealBuildTool.Modes
 					UEBuildTarget Target = UEBuildTarget.Create(TargetDesc, false, false, false);
 
 					// Create the makefile for the target and export the module information
-					using (ISourceFileWorkingSet WorkingSet = new EmptySourceFileWorkingSet())
-					{
-						// Create the build configuration object, and read the settings
-						BuildConfiguration BuildConfiguration = new BuildConfiguration();
-						XmlConfig.ApplyTo(BuildConfiguration);
-						Arguments.ApplyTo(BuildConfiguration);
+					using ISourceFileWorkingSet WorkingSet = new EmptySourceFileWorkingSet();
 
-						// Create the makefile
-						TargetMakefile Makefile = Target.Build(BuildConfiguration, WorkingSet, TargetDesc, true);
+					// Create the build configuration object, and read the settings
+					BuildConfiguration BuildConfiguration = new();
+					XmlConfig.ApplyTo(BuildConfiguration);
+					Arguments.ApplyTo(BuildConfiguration);
 
-						FileReference ModuleInfoFileName = ExternalExecution.GetUHTModuleInfoFileName(Makefile, Target.TargetName);
-						FileReference DepsFileName = ExternalExecution.GetUHTDepsFileName(ModuleInfoFileName);
-						ManifestPath = ModuleInfoFileName.FullName;
-						ExternalExecution.WriteUHTManifest(Makefile, Target.TargetName, ModuleInfoFileName, DepsFileName);
+					// Create the makefile
+					TargetMakefile Makefile = Target.Build(BuildConfiguration, WorkingSet, TargetDesc, true);
 
-						if (Target.ProjectFile != null)
-						{
-							ProjectFile = Target.ProjectFile.FullName;
-						}
+					FileReference ModuleInfoFileName = ExternalExecution.GetUHTModuleInfoFileName(Makefile, Target.TargetName);
+					FileReference DepsFileName = ExternalExecution.GetUHTDepsFileName(ModuleInfoFileName);
+					ManifestPath = ModuleInfoFileName.FullName;
+					ExternalExecution.WriteUHTManifest(Makefile, Target.TargetName, ModuleInfoFileName, DepsFileName);
+
+					if (Target.ProjectFile != null)
+					{
+						ProjectFile = Target.ProjectFile.FullName;
 					}
 				}
 				else
@@ -1107,7 +1078,7 @@ namespace UnrealBuildTool.Modes
 
 				string? ProjectPath = ProjectFile != null ? Path.GetDirectoryName(ProjectFile) : null;
 
-				UhtSession Session = new UhtSession
+				UhtSession Session = new()
 				{
 					Tables = Tables,
 					Config = Config,
@@ -1117,12 +1088,12 @@ namespace UnrealBuildTool.Modes
 					ProjectDirectory = string.IsNullOrEmpty(ProjectPath) ? null : ProjectPath,
 					ReferenceDirectory = FileReference.Combine(EngineProgramSavedDirectory, "UnrealBuildTool", "ReferenceGeneratedCode").FullName,
 					VerifyDirectory = FileReference.Combine(EngineProgramSavedDirectory, "UnrealBuildTool", "VerifyGeneratedCode").FullName,
-					bWarningsAsErrors = Options.bWarningsAsErrors,
-					bGoWide = !Options.bNoGoWide,
-					bFailIfGeneratedCodeChanges = Options.bFailIfGeneratedCodeChanges,
-					bNoOutput = Options.bNoOutput,
-					bIncludeDebugOutput = Options.bIncludeDebugOutput,
-					bNoDefaultExporters = Options.bNoDefaultExporters,
+					WarningsAsErrors = Options.bWarningsAsErrors,
+					GoWide = !Options.bNoGoWide,
+					FailIfGeneratedCodeChanges = Options.bFailIfGeneratedCodeChanges,
+					NoOutput = Options.bNoOutput,
+					IncludeDebugOutput = Options.bIncludeDebugOutput,
+					NoDefaultExporters = Options.bNoDefaultExporters,
 				};
 
 				if (Options.bWriteRef)
@@ -1149,7 +1120,7 @@ namespace UnrealBuildTool.Modes
 				// Read and parse
 				Session.Run(ManifestPath!);
 				Session.LogMessages();
-				return (int)(Session.bHasErrors ? CompilationResult.OtherCompilationError : CompilationResult.Succeeded);
+				return (int)(Session.HasErrors ? CompilationResult.OtherCompilationError : CompilationResult.Succeeded);
 			}
 			catch (Exception Ex)
 			{
