UnrealBuildTool: Fix data-dir on mac and linux when launching zenserver for the experimental HordeExecutor

#rnx

diff --git a/Engine/Source/Programs/UnrealBuildTool/Executors/Experimental/HordeExecutor.cs b/Engine/Source/Programs/UnrealBuildTool/Executors/Experimental/HordeExecutor.cs
--- a/Engine/Source/Programs/UnrealBuildTool/Executors/Experimental/HordeExecutor.cs
+++ b/Engine/Source/Programs/UnrealBuildTool/Executors/Experimental/HordeExecutor.cs
@@ -345,21 +345,44 @@ namespace UnrealBuildTool
 					return false;
 				}
 
-				DirectoryReference? DataPathReference = null;
 				if (!EngineConfig.GetString("Zen.AutoLaunch", "DataPath", out string DataPath) || string.IsNullOrEmpty(DataPath))
 				{
 					Log.TraceWarning("Unable to launch zenserver: '[Zen.AutoLaunch] DataPath' not set in Engine config");
 					return false;
 				}
 
-				DataPathReference = DirectoryReference.FromString(DataPath.Replace("%APPSETTINGSDIR%", $"{Path.Combine(Environment.GetFolderPath(Environment.SpecialFolder.CommonApplicationData), "Epic")}{Path.DirectorySeparatorChar}"));
-
 				if (!EngineConfig.GetString("Zen.AutoLaunch", "ExtraArgs", out string ExtraArgs) || string.IsNullOrEmpty(ExtraArgs))
 				{
 					Log.TraceWarning("Unable to launch zenserver: '[Zen.AutoLaunch] ExtraArgs' not set in Engine config");
 					return false;
 				}
 
+				if (DataPath.Contains("%APPSETTINGSDIR%"))
+				{
+					DirectoryReference? CommonPath = null;
+					if (BuildHostPlatform.Current.Platform == UnrealTargetPlatform.Win64)
+					{
+						CommonPath = DirectoryReference.FromString(Path.Combine(Environment.GetFolderPath(Environment.SpecialFolder.CommonApplicationData), "Epic"));
+					}
+					else if (BuildHostPlatform.Current.Platform == UnrealTargetPlatform.Mac)
+					{
+						CommonPath = DirectoryReference.FromString(Path.Combine(Environment.GetFolderPath(Environment.SpecialFolder.UserProfile), "Library", "Application Support", "Epic"));
+					}
+					else if (BuildHostPlatform.Current.Platform == UnrealTargetPlatform.Linux || BuildHostPlatform.Current.Platform == UnrealTargetPlatform.LinuxArm64)
+					{
+						CommonPath = DirectoryReference.FromString(Path.Combine(Environment.GetFolderPath(Environment.SpecialFolder.UserProfile), ".config", "Epic"));
+					}
+
+					if (CommonPath == null)
+					{
+						Log.TraceWarning("Unable to launch zenserver: data-dir root directory not found");
+						return false;
+					}
+
+					DataPath = DataPath.Replace("%APPSETTINGSDIR%", $"{CommonPath.FullName}{Path.DirectorySeparatorChar}");
+				}
+
+				DirectoryReference? DataPathReference = DirectoryReference.FromString(DataPath);
 				if (DataPathReference == null || DataPathReference.ParentDirectory == null)
 				{
 					Log.TraceWarning("Unable to launch zenserver: data-dir not found");
@@ -402,20 +425,33 @@ namespace UnrealBuildTool
 			ManagedProcess? ZenProcess = null;
 			if (!IsReady())
 			{
+				const double TotalWaitTime = 10.0;
+				DateTime WaitStartTime = DateTime.Now;
 				if (LaunchZenServer && LaunchZen(ProcessGroup, out ZenProcess))
 				{
-					Task.Delay(1000).Wait();
+					while ((DateTime.Now - WaitStartTime).TotalSeconds < TotalWaitTime)
+					{
+						Task.Delay(100).Wait();
+						if (IsReady())
+						{
+							break;
+						}
+					}
 				}
 				if (!IsReady())
 				{
 					if (LaunchZenServer)
 					{
-						Log.TraceError("Unable to establish connection to zenserver");
+						Log.TraceError($"Unable to establish connection to zenserver after {TotalWaitTime}s");
 						return false;
 					}
 					Log.TraceInformationOnce("Unable to establish connection to zenserver, disabling HordeExecutor");
 					return base.ExecuteActions(InputActions);
 				}
+				else
+				{
+					Log.TraceInformation($"Established connection to zenserver in {(DateTime.Now - WaitStartTime).TotalSeconds}s");
+				}
 			}
 
 			using SemaphoreSlim MaxProcessSemaphore = new SemaphoreSlim(ActualNumParallelProcesses, ActualNumParallelProcesses);
