Added new deprecation message type and a setting to enable/disable deprecation messages.

#rnx

diff --git a/Engine/Source/Programs/Shared/EpicGames.UHT/Tables/UhtSpecifierTable.cs b/Engine/Source/Programs/Shared/EpicGames.UHT/Tables/UhtSpecifierTable.cs
--- a/Engine/Source/Programs/Shared/EpicGames.UHT/Tables/UhtSpecifierTable.cs
+++ b/Engine/Source/Programs/Shared/EpicGames.UHT/Tables/UhtSpecifierTable.cs
@@ -397,7 +397,7 @@ namespace EpicGames.UHT.Tables
 		{
 			if (value != null)
 			{
-				specifierContext.Scope.TokenReader.LogInfo($"Specifier '{this.Name}' has a value which is unused, future versions of UnrealHeaderTool will flag this as an error.");
+				specifierContext.Scope.TokenReader.LogDeprecation($"Specifier '{this.Name}' has a value which is unused, future versions of UnrealHeaderTool will flag this as an error.");
 			}
 			this._delegate(specifierContext);
 			return UhtSpecifierDispatchResults.Known;
diff --git a/Engine/Source/Programs/Shared/EpicGames.UHT/Utils/UhtConfig.cs b/Engine/Source/Programs/Shared/EpicGames.UHT/Utils/UhtConfig.cs
--- a/Engine/Source/Programs/Shared/EpicGames.UHT/Utils/UhtConfig.cs
+++ b/Engine/Source/Programs/Shared/EpicGames.UHT/Utils/UhtConfig.cs
@@ -59,6 +59,11 @@ namespace EpicGames.UHT.Utils
 		/// </summary>
 		public UhtPointerMemberBehavior NonEngineObjectPtrMemberBehavior { get; }
 
+		/// <summary>
+		/// If true, deprecation warnings should be shown
+		/// </summary>
+		public bool ShowDeprecations { get; }
+
 		/// <summary>
 		/// If true, UObject properties are enabled in RigVM
 		/// </summary>
diff --git a/Engine/Source/Programs/Shared/EpicGames.UHT/Utils/UhtMessage.cs b/Engine/Source/Programs/Shared/EpicGames.UHT/Utils/UhtMessage.cs
--- a/Engine/Source/Programs/Shared/EpicGames.UHT/Utils/UhtMessage.cs
+++ b/Engine/Source/Programs/Shared/EpicGames.UHT/Utils/UhtMessage.cs
@@ -34,6 +34,11 @@ namespace EpicGames.UHT.Utils
 		/// </summary>
 		Trace,
 
+		/// <summary>
+		/// The message is an informational message about deprecated patterns
+		/// </summary>
+		Deprecation,
+
 		/// <summary>
 		/// The message is an internal error and goes to the log and console
 		/// </summary>
@@ -418,15 +423,15 @@ namespace EpicGames.UHT.Utils
 		}
 
 		/// <summary>
-		/// Log information
+		/// Log a message directly to the log
 		/// </summary>
 		/// <param name="messageSite">Message site associated with the message</param>
 		/// <param name="lineNumber">Line number of the information</param>
 		/// <param name="message">Text of the information</param>
 		/// <param name="extraContext">Addition context to be appended to the error message</param>
-		public static void LogInfo(this IUhtMessageSite messageSite, int lineNumber, string message, object? extraContext = null)
+		public static void LogTrace(this IUhtMessageSite messageSite, int lineNumber, string message, object? extraContext = null)
 		{
-			LogMessage(UhtMessageType.Info, messageSite, lineNumber, message, extraContext);
+			LogMessage(UhtMessageType.Trace, messageSite, lineNumber, message, extraContext);
 		}
 
 		/// <summary>
@@ -441,15 +446,15 @@ namespace EpicGames.UHT.Utils
 		}
 
 		/// <summary>
-		/// Log a message directly to the log
+		/// Log information
 		/// </summary>
 		/// <param name="messageSite">Message site associated with the message</param>
 		/// <param name="lineNumber">Line number of the information</param>
 		/// <param name="message">Text of the information</param>
 		/// <param name="extraContext">Addition context to be appended to the error message</param>
-		public static void LogLog(this IUhtMessageSite messageSite, int lineNumber, string message, object? extraContext = null)
+		public static void LogInfo(this IUhtMessageSite messageSite, int lineNumber, string message, object? extraContext = null)
 		{
-			LogMessage(UhtMessageType.Trace, messageSite, lineNumber, message, extraContext);
+			LogMessage(UhtMessageType.Info, messageSite, lineNumber, message, extraContext);
 		}
 
 		/// <summary>
@@ -463,6 +468,29 @@ namespace EpicGames.UHT.Utils
 			LogMessage(UhtMessageType.Info, messageSite, -1, message, extraContext);
 		}
 
+		/// <summary>
+		/// Log deprecation
+		/// </summary>
+		/// <param name="messageSite">Message site associated with the message</param>
+		/// <param name="lineNumber">Line number of the information</param>
+		/// <param name="message">Text of the information</param>
+		/// <param name="extraContext">Addition context to be appended to the error message</param>
+		public static void LogDeprecation(this IUhtMessageSite messageSite, int lineNumber, string message, object? extraContext = null)
+		{
+			LogMessage(UhtMessageType.Deprecation, messageSite, lineNumber, message, extraContext);
+		}
+
+		/// <summary>
+		/// Log a deprecation
+		/// </summary>
+		/// <param name="messageSite">Message site associated with the message</param>
+		/// <param name="message">Text of the information</param>
+		/// <param name="extraContext">Addition context to be appended to the error message</param>
+		public static void LogDeprecation(this IUhtMessageSite messageSite, string message, object? extraContext = null)
+		{
+			LogMessage(UhtMessageType.Deprecation, messageSite, -1, message, extraContext);
+		}
+
 		/// <summary>
 		/// Log a message
 		/// </summary>
diff --git a/Engine/Source/Programs/Shared/EpicGames.UHT/Utils/UhtSession.cs b/Engine/Source/Programs/Shared/EpicGames.UHT/Utils/UhtSession.cs
--- a/Engine/Source/Programs/Shared/EpicGames.UHT/Utils/UhtSession.cs
+++ b/Engine/Source/Programs/Shared/EpicGames.UHT/Utils/UhtSession.cs
@@ -1459,6 +1459,11 @@ namespace EpicGames.UHT.Utils
 		/// <param name="message">Message being added</param>
 		public void AddMessage(UhtMessage message)
 		{
+			if (message.MessageType == UhtMessageType.Deprecation && !Config!.ShowDeprecations)
+			{
+				return;
+			}
+
 			lock (this._messages)
 			{
 				this._messages.Add(message);
@@ -1484,6 +1489,7 @@ namespace EpicGames.UHT.Utils
 
 				case UhtMessageType.Info:
 				case UhtMessageType.Trace:
+				case UhtMessageType.Deprecation:
 					break;
 			}
 		}
@@ -1546,6 +1552,10 @@ namespace EpicGames.UHT.Utils
 					logLevel = LogLevel.Warning;
 					break;
 
+				case UhtMessageType.Deprecation:
+					logLevel = LogLevel.Information;
+					break;
+
 				case UhtMessageType.Info:
 					logLevel = LogLevel.Information;
 					break;
@@ -1659,6 +1669,8 @@ namespace EpicGames.UHT.Utils
 					return $"{filePath}({lineNumber}){fragmentPath}: Info: {message.Message}";
 				case UhtMessageType.Trace:
 					return $"{filePath}({lineNumber}){fragmentPath}: Trace: {message.Message}";
+				case UhtMessageType.Deprecation:
+					return $"{filePath}({lineNumber}){fragmentPath}: Deprecation: {message.Message}";
 				default:
 				case UhtMessageType.Ice:
 					return $"{filePath}({lineNumber}){fragmentPath}:  Error: Internal Compiler Error - {message.Message}";
diff --git a/Engine/Source/Programs/UnrealBuildTool/Modes/UnrealHeaderToolMode.cs b/Engine/Source/Programs/UnrealBuildTool/Modes/UnrealHeaderToolMode.cs
--- a/Engine/Source/Programs/UnrealBuildTool/Modes/UnrealHeaderToolMode.cs
+++ b/Engine/Source/Programs/UnrealBuildTool/Modes/UnrealHeaderToolMode.cs
@@ -21,100 +21,110 @@ namespace UnrealBuildTool.Modes
 	/// </summary>
 	public class UhtConfigImpl : IUhtConfig
 	{
-		private readonly ConfigHierarchy Ini;
+		private readonly ConfigHierarchy _ini;
 
 		/// <summary>
 		/// Types that have been renamed, treat the old deprecated name as the new name for code generation
 		/// </summary>
-		private readonly IReadOnlyDictionary<StringView, StringView> TypeRedirectMap;
+		private readonly IReadOnlyDictionary<StringView, StringView> _typeRedirectMap;
 
 		/// <summary>
 		/// Metadata that have been renamed, treat the old deprecated name as the new name for code generation
 		/// </summary>
-		private readonly IReadOnlyDictionary<string, string> MetaDataRedirectMap;
+		private readonly IReadOnlyDictionary<string, string> _metaDataRedirectMap;
 
 		/// <summary>
 		/// Supported units in the game
 		/// </summary>
-		private readonly ReadOnlyHashSet<StringView> Units;
+		private readonly ReadOnlyHashSet<StringView> _units;
 
 		/// <summary>
 		/// Special parsed struct names that do not require a prefix
 		/// </summary>
 		[System.Diagnostics.CodeAnalysis.SuppressMessage("CodeQuality", "IDE0052:Remove unread private members", Justification = "<Pending>")]
-		private readonly ReadOnlyHashSet<StringView> StructsWithNoPrefix;
+		private readonly ReadOnlyHashSet<StringView> _structsWithNoPrefix;
 
 		/// <summary>
 		/// Special parsed struct names that have a 'T' prefix
 		/// </summary>
-		private readonly ReadOnlyHashSet<StringView> StructsWithTPrefix;
+		private readonly ReadOnlyHashSet<StringView> _structsWithTPrefix;
 
 		/// <summary>
 		/// Mapping from 'human-readable' macro substring to # of parameters for delegate declarations
 		/// Index 0 is 1 parameter, Index 1 is 2, etc...
 		/// </summary>
-		private readonly IReadOnlyList<StringView> DelegateParameterCountStrings;
+		private readonly IReadOnlyList<StringView> _delegateParameterCountStrings;
 
 		/// <summary>
 		/// Default version of generated code. Defaults to oldest possible, unless specified otherwise in config.
 		/// </summary>
-		private readonly EGeneratedCodeVersion DefaultGeneratedCodeVersionInternal = EGeneratedCodeVersion.V1;
+		private readonly EGeneratedCodeVersion _defaultGeneratedCodeVersion = EGeneratedCodeVersion.V1;
 
 		/// <summary>
 		/// Internal version of pointer warning for native pointers in the engine
 		/// </summary>
-		private readonly UhtPointerMemberBehavior EngineNativePointerMemberBehaviorInternal = UhtPointerMemberBehavior.AllowSilently;
+		private readonly UhtPointerMemberBehavior _engineNativePointerMemberBehavior = UhtPointerMemberBehavior.AllowSilently;
 
 		/// <summary>
 		/// Internal version of pointer warning for object pointers in the engine
 		/// </summary>
-		private readonly UhtPointerMemberBehavior EngineObjectPtrMemberBehaviorInternal = UhtPointerMemberBehavior.AllowSilently;
+		private readonly UhtPointerMemberBehavior _engineObjectPtrMemberBehavior = UhtPointerMemberBehavior.AllowSilently;
 
 		/// <summary>
 		/// Internal version of pointer warning for native pointers outside the engine
 		/// </summary>
-		private readonly UhtPointerMemberBehavior NonEngineNativePointerMemberBehaviorInternal = UhtPointerMemberBehavior.AllowSilently;
+		private readonly UhtPointerMemberBehavior _nonEngineNativePointerMemberBehavior = UhtPointerMemberBehavior.AllowSilently;
 
 		/// <summary>
 		/// Internal version of pointer warning for object pointers outside the engine
 		/// </summary>
-		private readonly UhtPointerMemberBehavior NonEngineObjectPtrMemberBehaviorInternal = UhtPointerMemberBehavior.AllowSilently;
+		private readonly UhtPointerMemberBehavior _nonEngineObjectPtrMemberBehavior = UhtPointerMemberBehavior.AllowSilently;
+
+		/// <summary>
+		/// If true, deprecation warnings should be shown
+		/// </summary>
+		private readonly bool _showDeprecations = true;
 
 		/// <summary>
 		/// If true, UObject properties are enabled in RigVM
 		/// </summary>
-		private readonly bool AreRigVMUObjectProeprtiesEnabledInternal = false;
+		private readonly bool _areRigVMUObjectProeprtiesEnabled = false;
 
 		/// <summary>
 		/// If true, UInterface properties are enabled in RigVM
 		/// </summary>
-		private readonly bool AreRigVMUInterfaceProeprtiesEnabledInternal = false;
+		private readonly bool _areRigVMUInterfaceProeprtiesEnabled = false;
 
 		#region IUhtConfig Implementation
 		/// <inheritdoc/>
-		public EGeneratedCodeVersion DefaultGeneratedCodeVersion => this.DefaultGeneratedCodeVersionInternal;
+		public EGeneratedCodeVersion DefaultGeneratedCodeVersion => this._defaultGeneratedCodeVersion;
 
 		/// <inheritdoc/>
-		public UhtPointerMemberBehavior EngineNativePointerMemberBehavior => this.EngineNativePointerMemberBehaviorInternal;
+		public UhtPointerMemberBehavior EngineNativePointerMemberBehavior => this._engineNativePointerMemberBehavior;
 
 		/// <inheritdoc/>
-		public UhtPointerMemberBehavior EngineObjectPtrMemberBehavior => this.EngineObjectPtrMemberBehaviorInternal;
+		public UhtPointerMemberBehavior EngineObjectPtrMemberBehavior => this._engineObjectPtrMemberBehavior;
 
 		/// <inheritdoc/>
-		public UhtPointerMemberBehavior NonEngineNativePointerMemberBehavior => this.NonEngineNativePointerMemberBehaviorInternal;
+		public UhtPointerMemberBehavior NonEngineNativePointerMemberBehavior => this._nonEngineNativePointerMemberBehavior;
 
 		/// <inheritdoc/>
-		public UhtPointerMemberBehavior NonEngineObjectPtrMemberBehavior => this.NonEngineObjectPtrMemberBehaviorInternal;
+		public UhtPointerMemberBehavior NonEngineObjectPtrMemberBehavior => this._nonEngineObjectPtrMemberBehavior;
 
 		/// <summary>
 		/// If true, UObject properties are enabled in RigVM
 		/// </summary>
-		public bool AreRigVMUObjectProeprtiesEnabled => this.AreRigVMUObjectProeprtiesEnabledInternal;
+		public bool AreRigVMUObjectProeprtiesEnabled => this._areRigVMUObjectProeprtiesEnabled;
 
 		/// <summary>
 		/// If true, UInterface properties are enabled in RigVM
 		/// </summary>
-		public bool AreRigVMUInterfaceProeprtiesEnabled => this.AreRigVMUInterfaceProeprtiesEnabledInternal;
+		public bool AreRigVMUInterfaceProeprtiesEnabled => this._areRigVMUInterfaceProeprtiesEnabled;
+
+		/// <summary>
+		/// If true, deprecation warnings should be shown
+		/// </summary>
+		public bool ShowDeprecations => this._showDeprecations;
 
 		/// <inheritdoc/>
 		public void RedirectTypeIdentifier(ref UhtToken Token)
@@ -124,7 +134,7 @@ namespace UnrealBuildTool.Modes
 				throw new Exception("Attempt to redirect type identifier when the token isn't an identifier.");
 			}
 
-			if (this.TypeRedirectMap.TryGetValue(Token.Value, out StringView Redirect))
+			if (this._typeRedirectMap.TryGetValue(Token.Value, out StringView Redirect))
 			{
 				Token.Value = Redirect;
 			}
@@ -133,7 +143,7 @@ namespace UnrealBuildTool.Modes
 		/// <inheritdoc/>
 		public bool RedirectMetaDataKey(string Key, out string NewKey)
 		{
-			if (this.MetaDataRedirectMap.TryGetValue(Key, out string? Redirect))
+			if (this._metaDataRedirectMap.TryGetValue(Key, out string? Redirect))
 			{
 				NewKey = Redirect;
 				return Key != NewKey;
@@ -148,21 +158,21 @@ namespace UnrealBuildTool.Modes
 		/// <inheritdoc/>
 		public bool IsValidUnits(StringView Units)
 		{
-			return this.Units.Contains(Units);
+			return this._units.Contains(Units);
 		}
 
 		/// <inheritdoc/>
 		public bool IsStructWithTPrefix(StringView Name)
 		{
-			return this.StructsWithTPrefix.Contains(Name);
+			return this._structsWithTPrefix.Contains(Name);
 		}
 
 		/// <inheritdoc/>
 		public int FindDelegateParameterCount(StringView DelegateMacro)
 		{
-			for (int Index = 0, Count = this.DelegateParameterCountStrings.Count; Index < Count; ++Index)
+			for (int Index = 0, Count = this._delegateParameterCountStrings.Count; Index < Count; ++Index)
 			{
-				if (DelegateMacro.Span.Contains(this.DelegateParameterCountStrings[Index].Span, StringComparison.Ordinal))
+				if (DelegateMacro.Span.Contains(this._delegateParameterCountStrings[Index].Span, StringComparison.Ordinal))
 				{
 					return Index;
 				}
@@ -173,13 +183,13 @@ namespace UnrealBuildTool.Modes
 		/// <inheritdoc/>
 		public StringView GetDelegateParameterCountString(int Index)
 		{
-			return Index >= 0 ? this.DelegateParameterCountStrings[Index] : "";
+			return Index >= 0 ? this._delegateParameterCountStrings[Index] : "";
 		}
 
 		/// <inheritdoc/>
 		public bool IsExporterEnabled(string Name)
 		{
-			this.Ini.GetBool("UnrealHeaderTool", Name, out bool Value);
+			this._ini.GetBool("UnrealHeaderTool", Name, out bool Value);
 			return Value;
 		}
 		#endregion
@@ -191,26 +201,27 @@ namespace UnrealBuildTool.Modes
 		public UhtConfigImpl(CommandLineArguments Args)
 		{
 			DirectoryReference ConfigDirectory = DirectoryReference.Combine(Unreal.EngineDirectory, "Programs/UnrealBuildTool");
-			this.Ini = ConfigCache.ReadHierarchy(ConfigHierarchyType.Engine, ConfigDirectory, BuildHostPlatform.Current.Platform, "", Args.GetRawArray());
-
-			this.TypeRedirectMap = GetRedirectsStringView("UnrealHeaderTool", "TypeRedirects", "OldType", "NewType");
-			this.MetaDataRedirectMap = GetRedirectsString("CoreUObject.Metadata", "MetadataRedirects", "OldKey", "NewKey");
-			this.StructsWithNoPrefix = GetHashSet("UnrealHeaderTool", "StructsWithNoPrefix", StringViewComparer.Ordinal);
-			this.StructsWithTPrefix = GetHashSet("UnrealHeaderTool", "StructsWithTPrefix", StringViewComparer.Ordinal);
-			this.Units = GetHashSet("UnrealHeaderTool", "Units", StringViewComparer.OrdinalIgnoreCase);
-			this.DelegateParameterCountStrings = GetList("UnrealHeaderTool", "DelegateParameterCountStrings");
-			this.DefaultGeneratedCodeVersionInternal = GetGeneratedCodeVersion("UnrealHeaderTool", "DefaultGeneratedCodeVersion", EGeneratedCodeVersion.V1);
-			this.EngineNativePointerMemberBehaviorInternal = GetPointerMemberBehavior("UnrealHeaderTool", "EngineNativePointerMemberBehavior", UhtPointerMemberBehavior.AllowSilently);
-			this.EngineObjectPtrMemberBehaviorInternal = GetPointerMemberBehavior("UnrealHeaderTool", "EngineObjectPtrMemberBehavior", UhtPointerMemberBehavior.AllowSilently);
-			this.NonEngineNativePointerMemberBehaviorInternal = GetPointerMemberBehavior("UnrealHeaderTool", "NonEngineNativePointerMemberBehavior", UhtPointerMemberBehavior.AllowSilently);
-			this.NonEngineObjectPtrMemberBehaviorInternal = GetPointerMemberBehavior("UnrealHeaderTool", "NonEngineObjectPtrMemberBehavior", UhtPointerMemberBehavior.AllowSilently);
-			this.AreRigVMUObjectProeprtiesEnabledInternal = GetBoolean("UnrealHeaderTool", "AreRigVMUObjectProeprtiesEnabled", false);
-			this.AreRigVMUInterfaceProeprtiesEnabledInternal = GetBoolean("UnrealHeaderTool", "AreRigVMUInterfaceProeprtiesEnabled", false);
+			this._ini = ConfigCache.ReadHierarchy(ConfigHierarchyType.Engine, ConfigDirectory, BuildHostPlatform.Current.Platform, "", Args.GetRawArray());
+
+			this._typeRedirectMap = GetRedirectsStringView("UnrealHeaderTool", "TypeRedirects", "OldType", "NewType");
+			this._metaDataRedirectMap = GetRedirectsString("CoreUObject.Metadata", "MetadataRedirects", "OldKey", "NewKey");
+			this._structsWithNoPrefix = GetHashSet("UnrealHeaderTool", "StructsWithNoPrefix", StringViewComparer.Ordinal);
+			this._structsWithTPrefix = GetHashSet("UnrealHeaderTool", "StructsWithTPrefix", StringViewComparer.Ordinal);
+			this._units = GetHashSet("UnrealHeaderTool", "Units", StringViewComparer.OrdinalIgnoreCase);
+			this._delegateParameterCountStrings = GetList("UnrealHeaderTool", "DelegateParameterCountStrings");
+			this._defaultGeneratedCodeVersion = GetGeneratedCodeVersion("UnrealHeaderTool", "DefaultGeneratedCodeVersion", EGeneratedCodeVersion.V1);
+			this._engineNativePointerMemberBehavior = GetPointerMemberBehavior("UnrealHeaderTool", "EngineNativePointerMemberBehavior", UhtPointerMemberBehavior.AllowSilently);
+			this._engineObjectPtrMemberBehavior = GetPointerMemberBehavior("UnrealHeaderTool", "EngineObjectPtrMemberBehavior", UhtPointerMemberBehavior.AllowSilently);
+			this._nonEngineNativePointerMemberBehavior = GetPointerMemberBehavior("UnrealHeaderTool", "NonEngineNativePointerMemberBehavior", UhtPointerMemberBehavior.AllowSilently);
+			this._nonEngineObjectPtrMemberBehavior = GetPointerMemberBehavior("UnrealHeaderTool", "NonEngineObjectPtrMemberBehavior", UhtPointerMemberBehavior.AllowSilently);
+			this._areRigVMUObjectProeprtiesEnabled = GetBoolean("UnrealHeaderTool", "AreRigVMUObjectProeprtiesEnabled", false);
+			this._areRigVMUInterfaceProeprtiesEnabled = GetBoolean("UnrealHeaderTool", "AreRigVMUInterfaceProeprtiesEnabled", false);
+			this._showDeprecations = GetBoolean("UnrealHeaderTool", "ShowDeprecations", true);
 		}
 
 		private bool GetBoolean(string SectionName, string KeyName, bool bDefault)
 		{
-			if (this.Ini.TryGetValue(SectionName, KeyName, out bool value))
+			if (this._ini.TryGetValue(SectionName, KeyName, out bool value))
 			{
 				return value;
 			}
@@ -219,7 +230,7 @@ namespace UnrealBuildTool.Modes
 
 		private UhtPointerMemberBehavior GetPointerMemberBehavior(string SectionName, string KeyName, UhtPointerMemberBehavior Default)
 		{
-			if (this.Ini.TryGetValue(SectionName, KeyName, out string? BehaviorStr))
+			if (this._ini.TryGetValue(SectionName, KeyName, out string? BehaviorStr))
 			{
 				if (!Enum.TryParse(BehaviorStr, out UhtPointerMemberBehavior Value))
 				{
@@ -232,7 +243,7 @@ namespace UnrealBuildTool.Modes
 
 		private EGeneratedCodeVersion GetGeneratedCodeVersion(string SectionName, string KeyName, EGeneratedCodeVersion Default)
 		{
-			if (this.Ini.TryGetValue(SectionName, KeyName, out string? BehaviorStr))
+			if (this._ini.TryGetValue(SectionName, KeyName, out string? BehaviorStr))
 			{
 				if (!Enum.TryParse(BehaviorStr, out EGeneratedCodeVersion Value))
 				{
@@ -247,7 +258,7 @@ namespace UnrealBuildTool.Modes
 		{
 			Dictionary<StringView, StringView> Redirects = new();
 
-			if (this.Ini.TryGetValues(Section, Key, out IReadOnlyList<string>? StringList))
+			if (this._ini.TryGetValues(Section, Key, out IReadOnlyList<string>? StringList))
 			{
 				foreach (string Line in StringList)
 				{
@@ -272,7 +283,7 @@ namespace UnrealBuildTool.Modes
 		{
 			Dictionary<string, string> Redirects = new();
 
-			if (this.Ini.TryGetValues(Section, Key, out IReadOnlyList<string>? StringList))
+			if (this._ini.TryGetValues(Section, Key, out IReadOnlyList<string>? StringList))
 			{
 				foreach (string Line in StringList)
 				{
@@ -297,7 +308,7 @@ namespace UnrealBuildTool.Modes
 		{
 			List<StringView> List = new();
 
-			if (this.Ini.TryGetValues(Section, Key, out IReadOnlyList<string>? StringList))
+			if (this._ini.TryGetValues(Section, Key, out IReadOnlyList<string>? StringList))
 			{
 				foreach (string Value in StringList)
 				{
@@ -311,7 +322,7 @@ namespace UnrealBuildTool.Modes
 		{
 			HashSet<StringView> Set = new(Comparer);
 
-			if (this.Ini.TryGetValues(Section, Key, out IReadOnlyList<string>? StringList))
+			if (this._ini.TryGetValues(Section, Key, out IReadOnlyList<string>? StringList))
 			{
 				foreach (string Value in StringList)
 				{
