Update UEBuildDeploy to remove Read-Only attribute when copying files recursively.

diff --git a/Engine/Source/Programs/UnrealBuildTool/Configuration/UEBuildDeploy.cs b/Engine/Source/Programs/UnrealBuildTool/Configuration/UEBuildDeploy.cs
--- a/Engine/Source/Programs/UnrealBuildTool/Configuration/UEBuildDeploy.cs
+++ b/Engine/Source/Programs/UnrealBuildTool/Configuration/UEBuildDeploy.cs
@@ -30,7 +30,14 @@ namespace UnrealBuildTool
 			{
 				try
 				{
-					FromFile.CopyTo(Path.Combine(ToDir, FromFile.Name), true);
+					// Ensure we can overrwrite dest file as non-read only.
+					var ToDestFile = Path.Combine(ToDir, FromFile.Name);
+					if (File.Exists(ToDestFile))
+					{
+						File.SetAttributes(ToDestFile, File.GetAttributes(ToDestFile) & ~FileAttributes.ReadOnly);
+					}
+
+					FromFile.CopyTo(ToDestFile, true);
 				} catch (Exception ex)
 				{
 					Log.TraceError($"Error while copying {FromFile.Name} from {FromDir} to {ToDir}: {ex}");
