UnrealBuildTool: Don't build a prereq action if it's not in the list of all actions to build

#rnx

diff --git a/Engine/Source/Programs/UnrealBuildTool/Executors/ParallelExecutor.cs b/Engine/Source/Programs/UnrealBuildTool/Executors/ParallelExecutor.cs
--- a/Engine/Source/Programs/UnrealBuildTool/Executors/ParallelExecutor.cs
+++ b/Engine/Source/Programs/UnrealBuildTool/Executors/ParallelExecutor.cs
@@ -186,7 +186,7 @@ namespace UnrealBuildTool
 					continue;
 				}
 
-				Task<ExecuteResults> ExecuteTask = CreateExecuteTask(Action, ExecuteTasks, ProcessGroup, MaxProcessSemaphore, CancellationToken);
+				Task<ExecuteResults> ExecuteTask = CreateExecuteTask(Action, InputActions, ExecuteTasks, ProcessGroup, MaxProcessSemaphore, CancellationToken);
 				Task LogTask = ExecuteTask.ContinueWith(antecedent => LogCompletedAction(Action, antecedent, CancellationTokenSource, ProgressWriter, TotalActions, ref NumCompletedActions), CancellationToken);
 
 				ExecuteTasks.Add(Action, ExecuteTask);
@@ -200,9 +200,10 @@ namespace UnrealBuildTool
 			return ExecuteTasks.Values.All(x => x.Result.ExitCode == 0);
 		}
 
-		protected static Task<ExecuteResults> CreateExecuteTask(LinkedAction Action, Dictionary<LinkedAction, Task<ExecuteResults>> ExecuteTasks, ManagedProcessGroup ProcessGroup, SemaphoreSlim MaxProcessSemaphore, CancellationToken CancellationToken)
+		protected static Task<ExecuteResults> CreateExecuteTask(LinkedAction Action, List<LinkedAction> InputActions, Dictionary<LinkedAction, Task<ExecuteResults>> ExecuteTasks, ManagedProcessGroup ProcessGroup, SemaphoreSlim MaxProcessSemaphore, CancellationToken CancellationToken)
 		{
-			if (Action.PrerequisiteActions.Count == 0)
+			List<LinkedAction> PrerequisiteActions = Action.PrerequisiteActions.Where(x => InputActions.Contains(x)).ToList();
+			if (PrerequisiteActions.Count == 0)
 			{
 				return Task.Factory.StartNew(
 					() => ExecuteAction(new Task<ExecuteResults>[0], Action, ProcessGroup, MaxProcessSemaphore, CancellationToken),
@@ -214,11 +215,11 @@ namespace UnrealBuildTool
 
 			// Create tasks for any preresquite actions if they don't exist already
 			List<Task<ExecuteResults>> PrerequisiteTasks = new List<Task<ExecuteResults>>();
-			foreach (var PrerequisiteAction in Action.PrerequisiteActions)
+			foreach (var PrerequisiteAction in PrerequisiteActions)
 			{
 				if (!ExecuteTasks.ContainsKey(PrerequisiteAction))
 				{
-					ExecuteTasks.Add(PrerequisiteAction, CreateExecuteTask(PrerequisiteAction, ExecuteTasks, ProcessGroup, MaxProcessSemaphore, CancellationToken));
+					ExecuteTasks.Add(PrerequisiteAction, CreateExecuteTask(PrerequisiteAction, InputActions, ExecuteTasks, ProcessGroup, MaxProcessSemaphore, CancellationToken));
 				}
 				PrerequisiteTasks.Add(ExecuteTasks[PrerequisiteAction]);
 			}
