Update to new v20 toolchain

diff --git a/Engine/Source/Programs/UnrealBuildTool/Platform/Linux/LinuxPlatformSDK.Versions.cs b/Engine/Source/Programs/UnrealBuildTool/Platform/Linux/LinuxPlatformSDK.Versions.cs
--- a/Engine/Source/Programs/UnrealBuildTool/Platform/Linux/LinuxPlatformSDK.Versions.cs
+++ b/Engine/Source/Programs/UnrealBuildTool/Platform/Linux/LinuxPlatformSDK.Versions.cs
@@ -8,14 +8,14 @@ namespace UnrealBuildTool
 	{
 		public override string GetMainVersion()
 		{
-			return "v19_clang-11.0.1-centos7";
+			return "v20_clang-13.0.1-centos7";
 		}
 
 		protected override void GetValidVersionRange(out string MinVersion, out string MaxVersion)
 		{
 			// all that matters is the number after the v, according to TryConvertVersionToInt()
 			MinVersion = "v10_clang-5.0.0-centos7";
-			MaxVersion = "v19_clang-11.0.1-centos7";
+			MaxVersion = "v20_clang-13.0.1-centos7";
 		}
 
 		protected override void GetValidSoftwareVersionRange(out string? MinVersion, out string? MaxVersion)
