UnrealBuildTool/IncludeTool PreprocessorExpressions: add support for new clang language extension macros.
#rnx

diff --git a/Engine/Source/Programs/IncludeTool/IncludeTool/PreprocessorExpression.cs b/Engine/Source/Programs/IncludeTool/IncludeTool/PreprocessorExpression.cs
--- a/Engine/Source/Programs/IncludeTool/IncludeTool/PreprocessorExpression.cs
+++ b/Engine/Source/Programs/IncludeTool/IncludeTool/PreprocessorExpression.cs
@@ -343,12 +343,30 @@ namespace IncludeTool
 				case "__has_extension":
 				case "__has_feature":
 				case "__has_warning":
-					if(Tokens[Idx + 1].Text != "(" || Tokens[Idx + 3].Text != ")")
+				case "__is_identifier":
+					if (Tokens[Idx + 1].Text != "(" || Tokens[Idx + 3].Text != ")")
 					{
 						throw new NotImplementedException();
 					}
 					Idx += 4;
 					return 0;
+				case "__has_c_attribute":
+				case "__has_attribute":
+				case "__has_declspec_attribute":
+					// the argument can either be a single identifier or a scoped identifier. 
+					if (Tokens[Idx + 1].Text == "(")
+					{
+						Idx += 3; // 3 = __has_*, (, single_identifier
+						for (; Idx < Tokens.Count; ++Idx)
+						{
+							if (Tokens[Idx].Text == ")")
+							{
+								++Idx;
+								return 0;
+							}
+						}
+					}
+					throw new NotImplementedException();
 				case "__has_include":
 				case "__has_include_next":
 					if(Tokens[Idx + 1].Text == "(")
diff --git a/Engine/Source/Programs/UnrealBuildTool/Preprocessor/Identifier.cs b/Engine/Source/Programs/UnrealBuildTool/Preprocessor/Identifier.cs
--- a/Engine/Source/Programs/UnrealBuildTool/Preprocessor/Identifier.cs
+++ b/Engine/Source/Programs/UnrealBuildTool/Preprocessor/Identifier.cs
@@ -109,11 +109,15 @@ namespace UnrealBuildTool
 		public static readonly Identifier __pragma = Identifier.FindOrAdd("__pragma");
 		public static readonly Identifier __builtin_return_address = Identifier.FindOrAdd("__builtin_return_address");
 		public static readonly Identifier __builtin_frame_address = Identifier.FindOrAdd("__builtin_frame_address");
+		public static readonly Identifier __has_attribute = Identifier.FindOrAdd("__has_attribute");
+		public static readonly Identifier __has_c_attribute = Identifier.FindOrAdd("__has_c_attribute");
 		public static readonly Identifier __has_cpp_attribute = Identifier.FindOrAdd("__has_cpp_attribute");
+		public static readonly Identifier __has_declspec_attribute = Identifier.FindOrAdd("__has_declspec_attribute");
 		public static readonly Identifier __has_keyword = Identifier.FindOrAdd("__has_keyword");
 		public static readonly Identifier __has_extension = Identifier.FindOrAdd("__has_extension");
 		public static readonly Identifier __has_include = Identifier.FindOrAdd("__has_include");
 		public static readonly Identifier __has_include_next = Identifier.FindOrAdd("__has_include_next");
+		public static readonly Identifier __is_identifier = Identifier.FindOrAdd("__is_identifier");
 		public static readonly Identifier __is_target_arch = Identifier.FindOrAdd("__is_target_arch");
 	}
 
diff --git a/Engine/Source/Programs/UnrealBuildTool/Preprocessor/PreprocessorExpression.cs b/Engine/Source/Programs/UnrealBuildTool/Preprocessor/PreprocessorExpression.cs
--- a/Engine/Source/Programs/UnrealBuildTool/Preprocessor/PreprocessorExpression.cs
+++ b/Engine/Source/Programs/UnrealBuildTool/Preprocessor/PreprocessorExpression.cs
@@ -328,7 +328,18 @@ namespace UnrealBuildTool
 						Idx += 4;
 						return 0;
 					}
-					else if (Text == Identifiers.__has_cpp_attribute || Text == Identifiers.__has_include || Text == Identifiers.__has_include_next)
+					else if (Text == Identifiers.__is_identifier)
+					{
+						if (Tokens[Idx + 1].Type != TokenType.LeftParen || Tokens[Idx + 3].Type != TokenType.RightParen)
+						{
+							throw new NotImplementedException(Text.ToString());
+						}
+						Idx += 4;
+						return 0;
+					}
+					else if (Text == Identifiers.__has_attribute || Text == Identifiers.__has_declspec_attribute
+						|| Text == Identifiers.__has_c_attribute || Text == Identifiers.__has_cpp_attribute
+						|| Text == Identifiers.__has_include || Text == Identifiers.__has_include_next)
 					{
 						Idx += Tokens.Count - Idx;
 						return 0;
