Moved UHT configuration for C# back to the original location to make installed game deployment less complicated.

#rnx

diff --git a/Engine/Programs/UnrealHeaderTool/Config/DefaultEngine.ini b/Engine/Programs/UnrealHeaderTool/Config/DefaultEngine.ini
--- a/Engine/Programs/UnrealHeaderTool/Config/DefaultEngine.ini
+++ b/Engine/Programs/UnrealHeaderTool/Config/DefaultEngine.ini
@@ -30,10 +30,136 @@ gc.MaxObjectsInProgram=262144
 +TypeRedirects=(OldType="FStringClassReference", NewType="FSoftClassPath")
 +TypeRedirects=(OldType="TAssetPtr", NewType="TSoftObjectPtr")
 +TypeRedirects=(OldType="TAssetSubclassOf", NewType="TSoftClassPtr")
++Units="Micrometers"
++Units="um"
++Units="\u00B5m"
++Units="Millimeters"
++Units="mm"
++Units="Centimeters"
++Units="cm"
++Units="Meters"
++Units="m"
++Units="Kilometers"
++Units="km"
++Units="Inches"
++Units="in"
++Units="Feet"
++Units="ft"
++Units="Yards"
++Units="yd"
++Units="Miles"
++Units="mi"
++Units="Lightyears"
++Units="ly"
++Units="Degrees"
++Units="deg"
++Units="\u00B0"
++Units="Radians"
++Units="rad"
++Units="CentimetersPerSecond"
++Units="c/s"
++Units="MetersPerSecond"
++Units="m/s"
++Units="KilometersPerHour"
++Units="km/h"
++Units="kmph"
++Units="MilesPerHour"
++Units="mi/h"
++Units="mph"
++Units="Celsius"
++Units="C"
++Units="degC"
++Units="\u00B0C"
++Units="Farenheit"
++Units="F"
++Units="degF"
++Units="\u00B0F"
++Units="Kelvin"
++Units="K"
++Units="Micrograms"
++Units="ug")
++Units="\u00B5g"
++Units="Milligrams"
++Units="mg"
++Units="Grams"
++Units="g"
++Units="Kilograms"
++Units="kg"
++Units="MetricTons"
++Units="t"
++Units="Ounces"
++Units="oz"
++Units="Pounds"
++Units="lb"
++Units="Stones"
++Units="st"
++Units="Newtons"
++Units="N"
++Units="PoundsForce"
++Units="lbf"
++Units="KilogramsForce"
++Units="kgf"
++Units="Hertz"
++Units="Hz"
++Units="Kilohertz"
++Units="KHz"
++Units="Megahertz"
++Units="MHz"
++Units="Gigahertz"
++Units="GHz"
++Units="RevolutionsPerMinute"
++Units="rpm"
++Units="Bytes"
++Units="B"
++Units="Kilobytes"
++Units="KB"
++Units="KiB"
++Units="Megabytes"
++Units="MB"
++Units="MiB"
++Units="Gigabytes"
++Units="GB"
++Units="GiB"
++Units="Terabytes"
++Units="TB"
++Units="TiB"
++Units="Lumens"
++Units="lm"
++Units="Candela"
++Units="cd"
++Units="Lux"
++Units="lx"
++Units="CandelaPerMeterSquared"
++Units="cd/m2"
++Units="CandelaPerMeter2"
++Units="Milliseconds"
++Units="ms"
++Units="Seconds"
++Units="s"
++Units="Minutes"
++Units="min"
++Units="Hours"
++Units="hrs"
++Units="Days"
++Units="dy"
++Units="Months"
++Units="mth"
++Units="Years"
++Units="yr"
++Units="ppi"
++Units="dpi"
++Units="Percent"
++Units="%"
++Units="times"
++Units="x"
++Units="multiplier"
 
 [UnrealHeaderTool]
 EngineNativePointerMemberBehavior=Disallow
 EngineObjectPtrMemberBehavior=AllowSilently
 NonEngineNativePointerMemberBehavior=AllowSilently
 NonEngineObjectPtrMemberBehavior=AllowSilently
-AutomaticSettersAndGetters=False
\ No newline at end of file
+AutomaticSettersAndGetters=False
+AreRigVMUObjectProeprtiesEnabled=False
+AreRigVMUInterfaceProeprtiesEnabled=True
+ShowDeprecations=True
\ No newline at end of file
diff --git a/Engine/Source/Programs/UnrealBuildTool/Modes/UnrealHeaderToolMode.cs b/Engine/Source/Programs/UnrealBuildTool/Modes/UnrealHeaderToolMode.cs
--- a/Engine/Source/Programs/UnrealBuildTool/Modes/UnrealHeaderToolMode.cs
+++ b/Engine/Source/Programs/UnrealBuildTool/Modes/UnrealHeaderToolMode.cs
@@ -200,7 +200,7 @@ namespace UnrealBuildTool.Modes
 		/// <param name="Args">Extra command line arguments</param>
 		public UhtConfigImpl(CommandLineArguments Args)
 		{
-			DirectoryReference ConfigDirectory = DirectoryReference.Combine(Unreal.EngineDirectory, "Programs/UnrealBuildTool");
+			DirectoryReference ConfigDirectory = DirectoryReference.Combine(Unreal.EngineDirectory, "Programs", "UnrealHeaderTool");
 			this._ini = ConfigCache.ReadHierarchy(ConfigHierarchyType.Engine, ConfigDirectory, BuildHostPlatform.Current.Platform, "", Args.GetRawArray());
 
 			this._typeRedirectMap = GetRedirectsStringView("UnrealHeaderTool", "TypeRedirects", "OldType", "NewType");
@@ -356,7 +356,7 @@ namespace UnrealBuildTool.Modes
 		/// <summary>
 		/// Specifies the path to a log file to write. Note that the default mode (eg. building, generating project files) will create a log file by default if this not specified.
 		/// </summary>
-		[CommandLine(Prefix = "-Log", Description = "Specify a log file location instead of the default Engine/Programs/UnrealBuildTool/Log_UHT.txt")]
+		[CommandLine(Prefix = "-Log", Description = "Specify a log file location instead of the default Engine/Programs/UnrealHeaderTool/Saved/Logs/UnrealHeaderTool.log")]
 		public FileReference? LogFileName = null;
 
 		/// <summary>
@@ -1024,7 +1024,7 @@ namespace UnrealBuildTool.Modes
 
 					if (!Log.HasFileWriter())
 					{
-						string BaseLogFileName = FileReference.Combine(EngineProgramSavedDirectory, "UnrealBuildTool", "Log_UHT.txt").FullName;
+						string BaseLogFileName = FileReference.Combine(EngineProgramSavedDirectory, "UnrealHeaderTool", "Saved", "Logs", "UnrealHeaderTool.log").FullName;
 
 						FileReference LogFile = new(BaseLogFileName);
 						foreach (string LogSuffix in Arguments.GetValues("-LogSuffix="))
@@ -1097,8 +1097,8 @@ namespace UnrealBuildTool.Modes
 					EngineDirectory = Unreal.EngineDirectory.FullName,
 					ProjectFile = ProjectFile,
 					ProjectDirectory = string.IsNullOrEmpty(ProjectPath) ? null : ProjectPath,
-					ReferenceDirectory = FileReference.Combine(EngineProgramSavedDirectory, "UnrealBuildTool", "ReferenceGeneratedCode").FullName,
-					VerifyDirectory = FileReference.Combine(EngineProgramSavedDirectory, "UnrealBuildTool", "VerifyGeneratedCode").FullName,
+					ReferenceDirectory = FileReference.Combine(EngineProgramSavedDirectory, "UnrealHeaderTool", "Saved", "ReferenceExports").FullName,
+					VerifyDirectory = FileReference.Combine(EngineProgramSavedDirectory, "UnrealHeaderTool", "Saved", "VerifyExports").FullName,
 					WarningsAsErrors = Options.bWarningsAsErrors,
 					GoWide = !Options.bNoGoWide,
 					FailIfGeneratedCodeChanges = Options.bFailIfGeneratedCodeChanges,
