Disable variable set but not used compiler warning occurring when using newer Android NDK.

diff --git a/Engine/Source/Programs/UnrealBuildTool/Platform/Android/AndroidToolChain.cs b/Engine/Source/Programs/UnrealBuildTool/Platform/Android/AndroidToolChain.cs
--- a/Engine/Source/Programs/UnrealBuildTool/Platform/Android/AndroidToolChain.cs
+++ b/Engine/Source/Programs/UnrealBuildTool/Platform/Android/AndroidToolChain.cs
@@ -630,6 +630,13 @@ namespace UnrealBuildTool
 			Result += " -Wno-unknown-pragmas";          // probably should kill this one, sign of another issue in PhysX?
 			Result += " -Wno-invalid-offsetof";         // needed to suppress warnings about using offsetof on non-POD types.
 			Result += " -Wno-logical-op-parentheses";   // needed for external headers we can't change
+
+			// silence warning introduced with NDK 24+
+			if (NDKVersionInt >= 240000)
+			{
+				Result += " -Wno-unused-but-set-variable";
+			}
+
 			if (BuildWithHiddenSymbolVisibility(CompileEnvironment))
 			{
 				// Result += " -fvisibility-ms-compat -fvisibility-inlines-hidden"; // This hides all symbols by default but exports all type info (vtable/rtti) for a non-monolithic setup
