UnrealBuildTool: Move include paths and defines into GetCLArguments_Global for LinuxToolChain

#rnx

diff --git a/Engine/Source/Programs/UnrealBuildTool/Platform/Linux/LinuxToolChain.cs b/Engine/Source/Programs/UnrealBuildTool/Platform/Linux/LinuxToolChain.cs
--- a/Engine/Source/Programs/UnrealBuildTool/Platform/Linux/LinuxToolChain.cs
+++ b/Engine/Source/Programs/UnrealBuildTool/Platform/Linux/LinuxToolChain.cs
@@ -805,6 +805,22 @@ namespace UnrealBuildTool
 				Result += String.Format(" --sysroot=\"{0}\"", BaseLinuxPath);
 			}
 
+			// Add include paths to the argument list.
+			foreach (DirectoryReference IncludePath in CompileEnvironment.UserIncludePaths)
+			{
+				Result += string.Format(" -I\"{0}\"", NormalizeCommandLinePath(IncludePath));
+			}
+			foreach (DirectoryReference IncludePath in CompileEnvironment.SystemIncludePaths)
+			{
+				Result += string.Format(" -I\"{0}\"", NormalizeCommandLinePath(IncludePath));
+			}
+
+			// Add preprocessor definitions to the argument list.
+			foreach (string Definition in CompileEnvironment.Definitions)
+			{
+				Result += string.Format(" -D \"{0}\"", EscapeArgument(Definition));
+			}
+
 			return Result;
 		}
 
@@ -1292,22 +1308,6 @@ namespace UnrealBuildTool
 				PCHArguments += string.Format(" -include \"{0}\"", NormalizeCommandLinePath(CompileEnvironment.PrecompiledHeaderIncludeFilename!));
 			}
 
-			// Add include paths to the argument list.
-			foreach (DirectoryReference IncludePath in CompileEnvironment.UserIncludePaths)
-			{
-				Arguments += string.Format(" -I\"{0}\"", NormalizeCommandLinePath(IncludePath));
-			}
-			foreach (DirectoryReference IncludePath in CompileEnvironment.SystemIncludePaths)
-			{
-				Arguments += string.Format(" -I\"{0}\"", NormalizeCommandLinePath(IncludePath));
-			}
-
-			// Add preprocessor definitions to the argument list.
-			foreach (string Definition in CompileEnvironment.Definitions)
-			{
-				Arguments += string.Format(" -D \"{0}\"", EscapeArgument(Definition));
-			}
-
 			// Create a compile action for each source file.
 			CPPOutput Result = new CPPOutput();
 			foreach (FileItem SourceFile in InputFiles)
