Add ShaderCompileWorker to the allowed list for Apple Silicon.

#rnx

diff --git a/Engine/Source/Programs/UnrealBuildTool/Platform/Mac/MacExports.cs b/Engine/Source/Programs/UnrealBuildTool/Platform/Mac/MacExports.cs
--- a/Engine/Source/Programs/UnrealBuildTool/Platform/Mac/MacExports.cs
+++ b/Engine/Source/Programs/UnrealBuildTool/Platform/Mac/MacExports.cs
@@ -61,6 +61,7 @@ namespace UnrealBuildTool
 					"BenchmarkTool",
 					"BlankProgram",
 					"BuildPatchTool",
+					"ShaderCompileWorker",
 					"UnrealClient",
 					"UnrealGame", 
 					"UnrealServer",
