UnrealBuildTool: Begin moving common functions from clang based toolchains to base ClangToolChain.

#rnx

* Move ClangVersion variables to base class and share function to check version. A separate change will add a virtual function for parsing the clang version for each toolchain as to not add complexity to this change.
* GetCppStandardCompileArgument: Used by all, no differences
* GetCompileArguments_CPP: Used by all. Switch & Apple are overridden
* GetCompileArguments_C: Used by all, no differences
* GetCompileArguments_MM: Used by Linux & Apple. Apple is overridden.
* GetCompileArguments_M: Used by Linux & Apple. Apple is overridden.
* GetCompileArguments_CPP: Used by all. Switch & Apple are overridden.
* GetRTTIFlag: Used by Linux & Apple, no differences.

diff --git a/Engine/Source/Programs/UnrealBuildTool/Platform/Android/AndroidToolChain.cs b/Engine/Source/Programs/UnrealBuildTool/Platform/Android/AndroidToolChain.cs
--- a/Engine/Source/Programs/UnrealBuildTool/Platform/Android/AndroidToolChain.cs
+++ b/Engine/Source/Programs/UnrealBuildTool/Platform/Android/AndroidToolChain.cs
@@ -49,7 +49,7 @@ namespace UnrealBuildTool
 		EnableMinimalUndefinedBehaviorSanitizer = 0x10,
 	}
 
-	class AndroidToolChain : ISPCToolChain, IAndroidToolChain
+	class AndroidToolChain : ClangToolChain, IAndroidToolChain
 	{
 		public static readonly string[] AllCpuSuffixes =
 		{
@@ -94,11 +94,6 @@ namespace UnrealBuildTool
 		private List<string> AdditionalGPUArches;
 		protected bool bExecuteCompilerThroughShell;
 
-		// the Clang version being used to compile
-		static int ClangVersionMajor = -1;
-		static int ClangVersionMinor = -1;
-		static int ClangVersionPatch = -1;
-
 		// Version string from the Android specific build of clang. E.g in Android (6317467 based on r365631c1) clang version 9.0.8
 		// this would be 6317467)
 		protected static string? AndroidClangBuild;
@@ -158,26 +153,6 @@ namespace UnrealBuildTool
 			return NDKVersionInt >= 220000;
 		}
 
-		/// <summary>
-		/// Checks if compiler version matches the requirements
-		/// </summary>
-		private static bool CompilerVersionGreaterOrEqual(int Major, int Minor, int Patch)
-		{
-			return ClangVersionMajor > Major ||
-				(ClangVersionMajor == Major && ClangVersionMinor > Minor) ||
-				(ClangVersionMajor == Major && ClangVersionMinor == Minor && ClangVersionPatch >= Patch);
-		}
-
-		/// <summary>
-		/// Checks if compiler version matches the requirements
-		/// </summary>
-		private static bool CompilerVersionLessThan(int Major, int Minor, int Patch)
-		{
-			return ClangVersionMajor < Major ||
-				(ClangVersionMajor == Major && ClangVersionMinor < Minor) ||
-				(ClangVersionMajor == Major && ClangVersionMinor == Minor && ClangVersionPatch < Patch);
-		}
-
 		[CommandLine("-Architectures=", ListSeparator = '+')]
 		public List<string> ArchitectureArg = new List<string>();
 
@@ -742,17 +717,7 @@ namespace UnrealBuildTool
 				Result += " -DPLATFORM_EXCEPTIONS_DISABLED=1";
 			}
 
-			// Conditionally enable (default disabled) generation of information about every class with virtual functions for use by the C++ runtime type identification features
-			// (`dynamic_cast' and `typeid'). If you don't use those parts of the language, you can save some space by using -fno-rtti.
-			// Note that exception handling uses the same information, but it will generate it as needed.
-			if (CompileEnvironment.bUseRTTI)
-			{
-				Result += " -frtti";
-			}
-			else
-			{
-				Result += " -fno-rtti";
-			}
+			Result += GetRTTIFlag(CompileEnvironment);
 
 			// Profile Guided Optimization (PGO) and Link Time Optimization (LTO)
 			if (CompileEnvironment.bPGOOptimize)
@@ -847,68 +812,6 @@ namespace UnrealBuildTool
 			return Result;
 		}
 
-		static string GetCppStandardCompileArgument(CppCompileEnvironment CompileEnvironment)
-		{
-			string Result;
-			switch (CompileEnvironment.CppStandard)
-			{
-				case CppStandardVersion.Cpp14:
-					Result = " -std=c++14";
-					break;
-				case CppStandardVersion.Latest:
-				case CppStandardVersion.Cpp17:
-					Result = " -std=c++17";
-					break;
-				case CppStandardVersion.Cpp20:
-					Result = " -std=c++20";
-					break;
-				default:
-					throw new BuildException($"Unsupported C++ standard type set: {CompileEnvironment.CppStandard}");
-			}
-
-			if (CompileEnvironment.bEnableCoroutines)
-			{
-				Result += " -fcoroutines-ts";
-				if (!CompileEnvironment.bEnableExceptions)
-				{
-					Result += " -Wno-coroutine-missing-unhandled-exception";
-				}
-			}
-
-			return Result;
-		}
-
-		static string GetCompileArguments_CPP(CppCompileEnvironment CompileEnvironment)
-		{
-			string Result = "";
-
-			Result += " -x c++";
-			Result += GetCppStandardCompileArgument(CompileEnvironment);
-
-			return Result;
-		}
-
-		static string GetCompileArguments_C()
-		{
-			string Result = " -x c";
-			return Result;
-		}
-
-		static string GetCompileArguments_PCH(CppCompileEnvironment CompileEnvironment)
-		{
-			string Result = "";
-
-			Result += " -x c++-header";
-			if (CompilerVersionGreaterOrEqual(11, 0, 0))
-			{
-				Result += " -fpch-validate-input-files-content";
-				Result += " -fpch-instantiate-templates";
-			}
-			Result += GetCppStandardCompileArgument(CompileEnvironment);
-
-			return Result;
-		}
-
 		protected virtual string GetLinkArguments(LinkEnvironment LinkEnvironment, string Architecture)
 		{
 			string Result = "";
diff --git a/Engine/Source/Programs/UnrealBuildTool/Platform/IOS/IOSToolChain.cs b/Engine/Source/Programs/UnrealBuildTool/Platform/IOS/IOSToolChain.cs
--- a/Engine/Source/Programs/UnrealBuildTool/Platform/IOS/IOSToolChain.cs
+++ b/Engine/Source/Programs/UnrealBuildTool/Platform/IOS/IOSToolChain.cs
@@ -1,4 +1,4 @@
-﻿// Copyright Epic Games, Inc. All Rights Reserved.
+// Copyright Epic Games, Inc. All Rights Reserved.
 
 
 using System;
@@ -375,70 +375,6 @@ namespace UnrealBuildTool
 			return Result;
 		}
 
-		static string GetCompileArguments_CPP(CppCompileEnvironment CompileEnvironment)
-		{
-			string Result = "";
-			Result += " -x objective-c++";
-			Result += GetCppStandardCompileArgument(CompileEnvironment);
-			Result += " -stdlib=libc++";
-
-			return Result;
-		}
-
-		static string GetCompileArguments_MM(CppCompileEnvironment CompileEnvironment)
-		{
-			string Result = "";
-			Result += " -x objective-c++";
-			Result += GetCppStandardCompileArgument(CompileEnvironment);
-			Result += " -stdlib=libc++";
-			return Result;
-		}
-
-		static string GetCompileArguments_M(CppCompileEnvironment CompileEnvironment)
-		{
-			string Result = "";
-			Result += " -x objective-c";
-			Result += GetCppStandardCompileArgument(CompileEnvironment);
-			Result += " -stdlib=libc++";
-			return Result;
-		}
-
-		static string GetCompileArguments_C()
-		{
-			string Result = "";
-			Result += " -x c";
-			return Result;
-		}
-
-		static string GetCompileArguments_PCH(CppCompileEnvironment CompileEnvironment)
-		{
-			string Result = "";
-			Result += " -x objective-c++-header";
-			Result += GetCppStandardCompileArgument(CompileEnvironment);
-			Result += " -stdlib=libc++";
-
-			return Result;
-		}
-
-		// Conditionally enable (default disabled) generation of information about every class with virtual functions for use by the C++ runtime type identification features 
-		// (`dynamic_cast' and `typeid'). If you don't use those parts of the language, you can save some space by using -fno-rtti. 
-		// Note that exception handling uses the same information, but it will generate it as needed. 
-		static string GetRTTIFlag(CppCompileEnvironment CompileEnvironment)
-		{
-			string Result = "";
-
-			if (CompileEnvironment.bUseRTTI)
-			{
-				Result = " -frtti";
-			}
-			else
-			{
-				Result = " -fno-rtti";
-			}
-
-			return Result;
-		}
-
 		// Conditionally enable (default disabled) Objective-C exceptions
 		static string GetObjCExceptionsFlag(CppCompileEnvironment CompileEnvironment)
 		{
@@ -669,14 +605,6 @@ namespace UnrealBuildTool
 				{
 					// Compile the file as a C++ PCH.
 					FileArguments += GetCompileArguments_PCH(CompileEnvironment);
-					if (GetClangVersion().Major >= 11)
-					{
-						FileArguments += " -fpch-validate-input-files-content";
-					}
-					if (GetClangVersion().Major >= 13) // Note this is supported for >=11 on other clang platforms
-					{
-						FileArguments += " -fpch-instantiate-templates";
-					}
 					FileArguments += GetRTTIFlag(CompileEnvironment);
 					FileArguments += GetObjCExceptionsFlag(CompileEnvironment);
 				}
diff --git a/Engine/Source/Programs/UnrealBuildTool/Platform/Linux/LinuxToolChain.cs b/Engine/Source/Programs/UnrealBuildTool/Platform/Linux/LinuxToolChain.cs
--- a/Engine/Source/Programs/UnrealBuildTool/Platform/Linux/LinuxToolChain.cs
+++ b/Engine/Source/Programs/UnrealBuildTool/Platform/Linux/LinuxToolChain.cs
@@ -70,7 +70,7 @@ namespace UnrealBuildTool
 		PreservePSYM = 0x100,
 	}
 
-	class LinuxToolChain : ISPCToolChain
+	class LinuxToolChain : ClangToolChain
 	{
 		/** Flavor of the current build (target triplet)*/
 		string Architecture;
@@ -181,11 +181,11 @@ namespace UnrealBuildTool
 			// refuse to use compilers that we know won't work
 
 			// prevent unknown clangs since the build is likely to fail on too old or too new compilers
-			if ((CompilerVersionMajor * 10 + CompilerVersionMinor) > 130 || (CompilerVersionMajor * 10 + CompilerVersionMinor) < 100)
+			if ((ClangVersionMajor * 10 + ClangVersionMinor) > 130 || (ClangVersionMajor * 10 + ClangVersionMinor) < 100)
 			{
 				throw new BuildException(
 					string.Format("This version of the Unreal Engine can only be compiled with clang 13.0, 12.0, 11.0, and 10.0. clang {0} may not build it - please use a different version.",
-						CompilerVersionString)
+						ClangVersionString)
 					);
 			}
 		}
@@ -226,18 +226,18 @@ namespace UnrealBuildTool
 		/// </summary>
 		private void DetermineCompilerMajMinPatchFromVersionString()
 		{
-			string[] Parts = CompilerVersionString!.Split('.');
+			string[] Parts = ClangVersionString!.Split('.');
 			if (Parts.Length >= 1)
 			{
-				CompilerVersionMajor = Convert.ToInt32(Parts[0]);
+				ClangVersionMajor = Convert.ToInt32(Parts[0]);
 			}
 			if (Parts.Length >= 2)
 			{
-				CompilerVersionMinor = Convert.ToInt32(Parts[1]);
+				ClangVersionMinor = Convert.ToInt32(Parts[1]);
 			}
 			if (Parts.Length >= 3)
 			{
-				CompilerVersionPatch = Convert.ToInt32(Parts[2]);
+				ClangVersionPatch = Convert.ToInt32(Parts[2]);
 			}
 		}
 
@@ -351,10 +351,10 @@ namespace UnrealBuildTool
 		/// </summary>
 		protected bool DetermineCompilerVersion()
 		{
-			CompilerVersionString = null;
-			CompilerVersionMajor = -1;
-			CompilerVersionMinor = -1;
-			CompilerVersionPatch = -1;
+			ClangVersionString = null;
+			ClangVersionMajor = -1;
+			ClangVersionMinor = -1;
+			ClangVersionPatch = -1;
 
 			using (Process Proc = new Process())
 			{
@@ -384,7 +384,7 @@ namespace UnrealBuildTool
 							// version match will be like "version 3.3", so remove the "version"
 							if (VersionMatch.Value.StartsWith("version "))
 							{
-								CompilerVersionString = VersionMatch.Value.Replace("version ", "");
+								ClangVersionString = VersionMatch.Value.Replace("version ", "");
 
 								DetermineCompilerMajMinPatchFromVersionString();
 							}
@@ -397,7 +397,7 @@ namespace UnrealBuildTool
 				}
 			}
 
-			return !String.IsNullOrEmpty(CompilerVersionString);
+			return !String.IsNullOrEmpty(ClangVersionString);
 		}
 
 		/// <summary>
@@ -443,16 +443,6 @@ namespace UnrealBuildTool
 			}
 		}
 
-		/// <summary>
-		/// Checks if compiler version matches the requirements
-		/// </summary>
-		private static bool CompilerVersionGreaterOrEqual(int Major, int Minor, int Patch)
-		{
-			return CompilerVersionMajor > Major ||
-				(CompilerVersionMajor == Major && CompilerVersionMinor > Minor) ||
-				(CompilerVersionMajor == Major && CompilerVersionMinor == Minor && CompilerVersionPatch >= Patch);
-		}
-
 		/// <summary>
 		/// Architecture-specific compiler switches
 		/// </summary>
@@ -853,109 +843,11 @@ namespace UnrealBuildTool
 				: string.Format("{0}={1}", myKey, myValue);
 		}
 
-		static string GetCompilerStandardVersion_CPP(CppCompileEnvironment CompileEnvironment)
-		{
-			string Result;
-			switch (CompileEnvironment.CppStandard)
-			{
-				case CppStandardVersion.Cpp14:
-					Result = " -std=c++14";
-					break;
-				case CppStandardVersion.Latest:
-				case CppStandardVersion.Cpp17:
-					Result = " -std=c++17";
-					break;
-				case CppStandardVersion.Cpp20:
-					Result = " -std=c++20";
-					break;
-				default:
-					throw new BuildException($"Unsupported C++ standard type set: {CompileEnvironment.CppStandard}");
-			}
-
-			if (CompileEnvironment.bEnableCoroutines)
-			{
-				Result += " -fcoroutines-ts";
-				if (!CompileEnvironment.bEnableExceptions)
-				{
-					Result += " -Wno-coroutine-missing-unhandled-exception";
-				}
-			}
-
-			return Result;
-		}
-
-		static string GetCompileArguments_CPP(CppCompileEnvironment CompileEnvironment)
-		{
-			string Result = "";
-			Result += " -x c++";
-			Result += GetCompilerStandardVersion_CPP(CompileEnvironment);
-			return Result;
-		}
-
-		static string GetCompileArguments_C()
-		{
-			string Result = "";
-			Result += " -x c";
-			return Result;
-		}
-
-		static string GetCompileArguments_MM(CppCompileEnvironment CompileEnvironment)
-		{
-			string Result = "";
-			Result += " -x objective-c++";
-			Result += " -fobjc-abi-version=2";
-			Result += " -fobjc-legacy-dispatch";
-			Result += GetCompilerStandardVersion_CPP(CompileEnvironment);
-			return Result;
-		}
-
-		// Conditionally enable (default disabled) generation of information about every class with virtual functions for use by the C++ runtime type identification features
-		// (`dynamic_cast' and `typeid'). If you don't use those parts of the language, you can save some space by using -fno-rtti.
-		// Note that exception handling uses the same information, but it will generate it as needed.
-		static string GetRTTIFlag(CppCompileEnvironment CompileEnvironment)
-		{
-			string Result = "";
-
-			if (CompileEnvironment.bUseRTTI)
-			{
-				Result = " -frtti";
-			}
-			else
-			{
-				Result = " -fno-rtti";
-			}
-
-			return Result;
-		}
-
-		static string GetCompileArguments_M(CppCompileEnvironment CompileEnvironment)
-		{
-			string Result = "";
-			Result += " -x objective-c";
-			Result += " -fobjc-abi-version=2";
-			Result += " -fobjc-legacy-dispatch";
-			Result += GetCompilerStandardVersion_CPP(CompileEnvironment);
-			return Result;
-		}
-
-		static string GetCompileArguments_PCH(CppCompileEnvironment CompileEnvironment)
-		{
-			string Result = "";
-			Result += " -x c++-header";
-			if (CompilerVersionGreaterOrEqual(11, 0, 0))
-			{
-				Result += " -fpch-validate-input-files-content";
-				Result += " -fpch-instantiate-templates";
-			}
-			Result += GetCompilerStandardVersion_CPP(CompileEnvironment);
-			return Result;
-		}
-
 		protected virtual string GetLinkArguments(LinkEnvironment LinkEnvironment)
 		{
 			string Result = "";
 
-			if (UsingLld(LinkEnvironment.Architecture) && (!LinkEnvironment.bIsBuildingDLL || (CompilerVersionMajor >= 9)))
+			if (UsingLld(LinkEnvironment.Architecture) && (!LinkEnvironment.bIsBuildingDLL || (ClangVersionMajor >= 9)))
 			{
 				Result += (BuildHostPlatform.Current.Platform == UnrealTargetPlatform.Win64) ? " -fuse-ld=lld.exe" : " -fuse-ld=lld";
 			}
@@ -1008,7 +900,7 @@ namespace UnrealBuildTool
 				if (CrossCompiling())
 				{
 					Result += string.Format(" -Wl,-rpath=\"{0}/lib/clang/{1}.{2}.{3}/lib/linux\"",
-							BaseLinuxPath, CompilerVersionMajor, CompilerVersionMinor, CompilerVersionPatch);
+							BaseLinuxPath, ClangVersionMajor, ClangVersionMinor, ClangVersionPatch);
 				}
 			}
 
@@ -1132,23 +1024,6 @@ namespace UnrealBuildTool
 		protected string? BreakpadEncoderPath;
 		protected string? MultiArchRoot;
 
-		/// <summary>
-		/// Version string of the current compiler, whether clang or gcc or whatever
-		/// </summary>
-		static string? CompilerVersionString;
-		/// <summary>
-		/// Major version of the current compiler, whether clang or gcc or whatever
-		/// </summary>
-		static int CompilerVersionMajor = -1;
-		/// <summary>
-		/// Minor version of the current compiler, whether clang or gcc or whatever
-		/// </summary>
-		static int CompilerVersionMinor = -1;
-		/// <summary>
-		/// Patch version of the current compiler, whether clang or gcc or whatever
-		/// </summary>
-		static int CompilerVersionPatch = -1;
-
 		/// <summary>
 		/// Whether to use old, slower way to relink circularly dependent libraries.
 		/// It makes sense to use it when cross-compiling on Windows due to race conditions between actions reading and modifying the libs.
@@ -1206,7 +1081,7 @@ namespace UnrealBuildTool
 			Log.TraceInformation("------- Build details --------");
 			Log.TraceInformation("Using {0}.", ToolchainInfo);
 			Log.TraceInformation("Using clang ({0}) version '{1}' (string), {2} (major), {3} (minor), {4} (patch)",
-				ClangPath, CompilerVersionString, CompilerVersionMajor, CompilerVersionMinor, CompilerVersionPatch);
+				ClangPath, ClangVersionString, ClangVersionMajor, ClangVersionMinor, ClangVersionPatch);
 
 			// inform the user which C++ library the engine is going to be compiled against - important for compatibility with third party code that uses STL
 			Log.TraceInformation("Using {0} standard C++ library.", ShouldUseLibcxx(CompileEnvironment.Architecture) ? "bundled libc++" : "compiler default (most likely libstdc++)");
@@ -1416,7 +1291,7 @@ namespace UnrealBuildTool
 				CompileAction.CommandArguments = string.Format(" @\"{0}\"", NormalizeCommandLinePath(CompilerResponseFileName));
 				CompileAction.PrerequisiteItems.Add(CompilerResponseFileItem);
 				CompileAction.CommandDescription = "Compile";
-				CompileAction.CommandVersion = CompilerVersionString!;
+				CompileAction.CommandVersion = ClangVersionString!;
 				CompileAction.StatusDescription = Path.GetFileName(SourceFile.AbsolutePath);
 				CompileAction.bIsGCCCompiler = true;
 
@@ -1663,7 +1538,7 @@ namespace UnrealBuildTool
 			}
 			// because the logic choosing between lld and ld is somewhat messy atm (lld fails to link .DSO due to bugs), make the name of the linker clear
 			LinkAction.CommandDescription += (LinkCommandString.Contains("-fuse-ld=lld")) ? " (lld)" : " (ld)";
-			LinkAction.CommandVersion = CompilerVersionString!;
+			LinkAction.CommandVersion = ClangVersionString!;
 			LinkAction.StatusDescription = Path.GetFileName(OutputFile.AbsolutePath);
 
 			// Add the output file to the command-line.
diff --git a/Engine/Source/Programs/UnrealBuildTool/Platform/Mac/MacToolChain.cs b/Engine/Source/Programs/UnrealBuildTool/Platform/Mac/MacToolChain.cs
--- a/Engine/Source/Programs/UnrealBuildTool/Platform/Mac/MacToolChain.cs
+++ b/Engine/Source/Programs/UnrealBuildTool/Platform/Mac/MacToolChain.cs
@@ -303,69 +303,6 @@ namespace UnrealBuildTool
 			return Result;
 		}
 		
-		static string GetCompileArguments_CPP(CppCompileEnvironment CompileEnvironment)
-		{
-			string Result = "";
-			Result += " -x objective-c++";
-			Result += GetCppStandardCompileArgument(CompileEnvironment);
-			Result += " -stdlib=libc++";
-
-			return Result;
-		}
-
-		static string GetCompileArguments_MM(CppCompileEnvironment CompileEnvironment)
-		{
-			string Result = "";
-			Result += " -x objective-c++";
-			Result += GetCppStandardCompileArgument(CompileEnvironment);
-			Result += " -stdlib=libc++";
-			return Result;
-		}
-
-		static string GetCompileArguments_M(CppCompileEnvironment CompileEnvironment)
-		{
-			string Result = "";
-			Result += " -x objective-c";
-			Result += " -stdlib=libc++";
-			return Result;
-		}
-
-		static string GetCompileArguments_C()
-		{
-			string Result = "";
-			Result += " -x c";
-			return Result;
-		}
-
-		static string GetCompileArguments_PCH(CppCompileEnvironment CompileEnvironment)
-		{
-			string Result = "";
-			Result += " -x objective-c++-header";
-			Result += GetCppStandardCompileArgument(CompileEnvironment);
-			Result += " -stdlib=libc++";
-
-			return Result;
-		}
-
-		// Conditionally enable (default disabled) generation of information about every class with virtual functions for use by the C++ runtime type identification features 
-		// (`dynamic_cast' and `typeid'). If you don't use those parts of the language, you can save some space by using -fno-rtti. 
-		// Note that exception handling uses the same information, but it will generate it as needed. 
-		static string GetRTTIFlag(CppCompileEnvironment CompileEnvironment)
-		{
-			string Result = "";
-
-			if (CompileEnvironment.bUseRTTI)
-			{
-				Result = " -frtti";
-			}
-			else
-			{
-				Result = " -fno-rtti";
-			}
-
-			return Result;
-		}
-		
 		string AddFrameworkToLinkCommand(string FrameworkName, string Arg = "-framework")
 		{
 			string Result = "";
@@ -514,14 +451,6 @@ namespace UnrealBuildTool
 				{
 					// Compile the file as a C++ PCH.
 					FileArguments += GetCompileArguments_PCH(CompileEnvironment);
-					if (GetClangVersion().Major >= 11)
-					{
-						FileArguments += " -fpch-validate-input-files-content";
-					}
-					if (GetClangVersion().Major >= 13) // Note this is supported for >=11 on other clang platforms
-					{
-						FileArguments += " -fpch-instantiate-templates";
-					}
 					FileArguments += GetRTTIFlag(CompileEnvironment);
 				}
 				else if (Extension == ".C")
diff --git a/Engine/Source/Programs/UnrealBuildTool/ToolChain/AppleToolChain.cs b/Engine/Source/Programs/UnrealBuildTool/ToolChain/AppleToolChain.cs
--- a/Engine/Source/Programs/UnrealBuildTool/ToolChain/AppleToolChain.cs
+++ b/Engine/Source/Programs/UnrealBuildTool/ToolChain/AppleToolChain.cs
@@ -145,7 +145,7 @@ namespace UnrealBuildTool
 		}
 	}
 
-	abstract class AppleToolChain : ISPCToolChain
+	abstract class AppleToolChain : ClangToolChain
 	{
 		protected FileReference? ProjectFile;
 
@@ -200,32 +200,48 @@ namespace UnrealBuildTool
 			return FullClangVersion;
 		}
 
-		protected static string GetCppStandardCompileArgument(CppCompileEnvironment CompileEnvironment)
+		protected override string GetCompileArguments_CPP(CppCompileEnvironment CompileEnvironment)
 		{
-			string Result;
-			switch (CompileEnvironment.CppStandard)
+			string Result = "";
+			Result += " -x objective-c++";
+			Result += GetCppStandardCompileArgument(CompileEnvironment);
+			Result += " -stdlib=libc++";
+
+			return Result;
+		}
+
+		protected override string GetCompileArguments_MM(CppCompileEnvironment CompileEnvironment)
+		{
+			string Result = "";
+			Result += " -x objective-c++";
+			Result += GetCppStandardCompileArgument(CompileEnvironment);
+			Result += " -stdlib=libc++";
+			return Result;
+		}
+
+		protected override string GetCompileArguments_M(CppCompileEnvironment CompileEnvironment)
+		{
+			string Result = "";
+			Result += " -x objective-c";
+			Result += GetCppStandardCompileArgument(CompileEnvironment);
+			Result += " -stdlib=libc++";
+			return Result;
+		}
+
+		protected override string GetCompileArguments_PCH(CppCompileEnvironment CompileEnvironment)
+		{
+			string Result = "";
+			Result += " -x objective-c++-header";
+			Result += GetCppStandardCompileArgument(CompileEnvironment);
+			Result += " -stdlib=libc++";
+
+			if (GetClangVersion().Major >= 11)
 			{
-				case CppStandardVersion.Cpp14:
-					Result = " -std=c++14";
-					break;
-				case CppStandardVersion.Latest:
-				case CppStandardVersion.Cpp17:
-					Result = " -std=c++17";
-					break;
-				case CppStandardVersion.Cpp20:
-					Result = " -std=c++20";
-					break;
-				default:
-					throw new BuildException($"Unsupported C++ standard type set: {CompileEnvironment.CppStandard}");
+				Result += " -fpch-validate-input-files-content";
 			}
-
-			if (CompileEnvironment.bEnableCoroutines)
+			if (GetClangVersion().Major >= 13) // Note this is supported for >=11 on other clang platforms
 			{
-				Result += " -fcoroutines-ts";
-				if (!CompileEnvironment.bEnableExceptions)
-				{
-					Result += " -Wno-coroutine-missing-unhandled-exception";
-				}
+				Result += " -fpch-instantiate-templates";
 			}
 
 			return Result;
diff --git a/Engine/Source/Programs/UnrealBuildTool/ToolChain/ClangToolChain.cs b/Engine/Source/Programs/UnrealBuildTool/ToolChain/ClangToolChain.cs
new file mode 100644
--- /dev/null
+++ b/Engine/Source/Programs/UnrealBuildTool/ToolChain/ClangToolChain.cs
@@ -0,0 +1,155 @@
+// Copyright Epic Games, Inc. All Rights Reserved.
+
+using System;
+using System.Collections;
+using System.Collections.Generic;
+using System.IO;
+using System.Diagnostics;
+using System.Text;
+using EpicGames.Core;
+using System.Text.RegularExpressions;
+using UnrealBuildBase;
+
+namespace UnrealBuildTool
+{
+	abstract class ClangToolChain : ISPCToolChain
+	{
+		// The Clang version being used to compile
+		protected string? ClangVersionString = null;
+		protected int ClangVersionMajor = -1;
+		protected int ClangVersionMinor = -1;
+		protected int ClangVersionPatch = -1;
+
+		public ClangToolChain() : base()
+		{
+		}
+
+		/// <summary>
+		/// Checks if compiler version matches the requirements
+		/// </summary>
+		protected bool CompilerVersionGreaterOrEqual(int Major, int Minor, int Patch)
+		{
+			// TODO: Temporary verification check to ensure a clang version has been set until ClangVersion is standarized across all clang-based toolchains.
+			if (ClangVersionMajor == -1)
+			{
+				throw new BuildException($"ClangVersion not found, unable to check compiler version requirements");
+			}
+
+			return ClangVersionMajor > Major ||
+				(ClangVersionMajor == Major && ClangVersionMinor > Minor) ||
+				(ClangVersionMajor == Major && ClangVersionMinor == Minor && ClangVersionPatch >= Patch);
+		}
+
+		/// <summary>
+		/// Checks if compiler version matches the requirements
+		/// </summary>
+		protected bool CompilerVersionLessThan(int Major, int Minor, int Patch)
+		{
+			// TODO: Temporary verification check until ClangVersion is standarized across all clang-based toolchains to ensure a version has been set.
+			if (ClangVersionMajor == -1)
+			{
+				throw new BuildException($"ClangVersion not found, unable to check compiler version requirements");
+			}
+
+			return ClangVersionMajor < Major ||
+				(ClangVersionMajor == Major && ClangVersionMinor < Minor) ||
+				(ClangVersionMajor == Major && ClangVersionMinor == Minor && ClangVersionPatch < Patch);
+		}
+
+		protected virtual string GetCppStandardCompileArgument(CppCompileEnvironment CompileEnvironment)
+		{
+			string Result;
+			switch (CompileEnvironment.CppStandard)
+			{
+				case CppStandardVersion.Cpp14:
+					Result = " -std=c++14";
+					break;
+				case CppStandardVersion.Latest:
+				case CppStandardVersion.Cpp17:
+					Result = " -std=c++17";
+					break;
+				case CppStandardVersion.Cpp20:
+					Result = " -std=c++20";
+					break;
+				default:
+					throw new BuildException($"Unsupported C++ standard type set: {CompileEnvironment.CppStandard}");
+			}
+
+			if (CompileEnvironment.bEnableCoroutines)
+			{
+				Result += " -fcoroutines-ts";
+				if (!CompileEnvironment.bEnableExceptions)
+				{
+					Result += " -Wno-coroutine-missing-unhandled-exception";
+				}
+			}
+
+			return Result;
+		}
+
+		protected virtual string GetCompileArguments_CPP(CppCompileEnvironment CompileEnvironment)
+		{
+			string Result = "";
+			Result += " -x c++";
+			Result += GetCppStandardCompileArgument(CompileEnvironment);
+			return Result;
+		}
+
+		protected virtual string GetCompileArguments_C()
+		{
+			string Result = "";
+			Result += " -x c";
+			return Result;
+		}
+
+		protected virtual string GetCompileArguments_MM(CppCompileEnvironment CompileEnvironment)
+		{
+			string Result = "";
+			Result += " -x objective-c++";
+			Result += " -fobjc-abi-version=2";
+			Result += " -fobjc-legacy-dispatch";
+			Result += GetCppStandardCompileArgument(CompileEnvironment);
+			return Result;
+		}
+
+		protected virtual string GetCompileArguments_M(CppCompileEnvironment CompileEnvironment)
+		{
+			string Result = "";
+			Result += " -x objective-c";
+			Result += " -fobjc-abi-version=2";
+			Result += " -fobjc-legacy-dispatch";
+			Result += GetCppStandardCompileArgument(CompileEnvironment);
+			return Result;
+		}
+
+		protected virtual string GetCompileArguments_PCH(CppCompileEnvironment CompileEnvironment)
+		{
+			string Result = "";
+			Result += " -x c++-header";
+			if (CompilerVersionGreaterOrEqual(11, 0, 0))
+			{
+				Result += " -fpch-validate-input-files-content";
+				Result += " -fpch-instantiate-templates";
+			}
+			Result += GetCppStandardCompileArgument(CompileEnvironment);
+			return Result;
+		}
+
+		// Conditionally enable (default disabled) generation of information about every class with virtual functions for use by the C++ runtime type identification features
+		// (`dynamic_cast' and `typeid'). If you don't use those parts of the language, you can save some space by using -fno-rtti.
+		// Note that exception handling uses the same information, but it will generate it as needed.
+		protected virtual string GetRTTIFlag(CppCompileEnvironment CompileEnvironment)
+		{
+			string Result;
+			if (CompileEnvironment.bUseRTTI)
+			{
+				Result = " -frtti";
+			}
+			else
+			{
+				Result = " -fno-rtti";
+			}
+			return Result;
+		}
+	}
+}
\ No newline at end of file
