UnrealBuildTool: Support new MSVC source dependency version

#rnx

diff --git a/Engine/Source/Programs/UnrealBuildTool/System/CppDependencyCache.cs b/Engine/Source/Programs/UnrealBuildTool/System/CppDependencyCache.cs
--- a/Engine/Source/Programs/UnrealBuildTool/System/CppDependencyCache.cs
+++ b/Engine/Source/Programs/UnrealBuildTool/System/CppDependencyCache.cs
@@ -170,6 +170,16 @@ namespace UnrealBuildTool
 		/// </summary>
 		static Dictionary<FileReference, CachePartition> GlobalPartitions = new Dictionary<FileReference, CachePartition>();
 
+		/// <summary>
+		/// Minimum version of a MSVC source dependency json
+		/// </summary>
+		static readonly VersionNumber MinVCDependencyVersion = new VersionNumber(1, 0);
+
+		/// <summary>
+		/// Minimum version of a MSVC source dependency json that supports additional module info
+		/// </summary>
+		static readonly VersionNumber VCDependencyAdditionalModuleInfoVersion = new VersionNumber(1, 1);
+
 		/// <summary>
 		/// Constructs a dependency cache. This method is private; call CppDependencyCache.Create() to create a cache hierarchy for a given project.
 		/// </summary>
@@ -459,14 +469,19 @@ namespace UnrealBuildTool
 
 				JsonObject Object = JsonObject.Read(InputFile.Location);
 
-				if (!Object.TryGetStringField("Version", out string? Version))
+				if (!Object.TryGetStringField("Version", out string? VersionString))
 				{
 					throw new BuildException(
 						$"Dependency file \"{InputFile.Location}\" does not have have a \"Version\" field.");
 				}
 
-				if (!String.Equals(Version, "1.1") &&
-					!String.Equals(Version, "1.0"))
+				if (!VersionNumber.TryParse(VersionString, out VersionNumber? Version) || Version == null)
+				{
+					throw new BuildException(
+						$"Dependency file \"{InputFile.Location}\" does not have have a valid \"Version\" field (\"{VersionString}\").");
+				}
+
+				if (Version < MinVCDependencyVersion)
 				{
 					throw new BuildException(
 						$"Dependency file \"{InputFile.Location}\" version (\"{Version}\") is not supported version");
@@ -482,7 +497,7 @@ namespace UnrealBuildTool
 
 				List<(string Name, string BMI)>? ImportedModules = null;
 
-				if (String.Equals(Version, "1.1") && !InputFile.HasExtension(".md.json"))
+				if (Version >= VCDependencyAdditionalModuleInfoVersion && !InputFile.HasExtension(".md.json"))
 				{
 					if (Data.TryGetObjectArrayField("ImportedModules", out JsonObject[]? ImportedModulesJson))
 					{
