UnrealBuildTool: Ban MSVC compiler in VS2022 17.2.0 & 17.3.0 Preview 1 due to https://developercommunity.visualstudio.com/t/Incorrect-instantiation-of-a-virtual-fun/10020368

#rnx

diff --git a/Engine/Source/Programs/UnrealBuildTool/Platform/Windows/MicrosoftPlatformSDK.Versions.cs b/Engine/Source/Programs/UnrealBuildTool/Platform/Windows/MicrosoftPlatformSDK.Versions.cs
--- a/Engine/Source/Programs/UnrealBuildTool/Platform/Windows/MicrosoftPlatformSDK.Versions.cs
+++ b/Engine/Source/Programs/UnrealBuildTool/Platform/Windows/MicrosoftPlatformSDK.Versions.cs
@@ -42,6 +42,9 @@ namespace UnrealBuildTool
 		/// </summary>
 		static readonly VersionNumberRange[] BannedVisualCppVersions = new VersionNumberRange[]
 		{
+			// https://developercommunity.visualstudio.com/t/Incorrect-instantiation-of-a-virtual-fun/10020368
+			VersionNumberRange.Parse("14.32.31326", "14.32.31328"), // VS2022 17.2.0
+			VersionNumberRange.Parse("14.33.31424", "14.33.31424"), // VS2022 17.3.0 Preview 1.0
 		};
 
 		static readonly VersionNumber MinimumVisualCppVersion = new VersionNumber(14, 29, 30133);
