UnrealBuildTool: Move include path and preprocessor definition arguments to ClangToolChain virtual functions

#rnx

diff --git a/Engine/Source/Programs/UnrealBuildTool/Platform/Android/AndroidToolChain.cs b/Engine/Source/Programs/UnrealBuildTool/Platform/Android/AndroidToolChain.cs
--- a/Engine/Source/Programs/UnrealBuildTool/Platform/Android/AndroidToolChain.cs
+++ b/Engine/Source/Programs/UnrealBuildTool/Platform/Android/AndroidToolChain.cs
@@ -1727,21 +1727,21 @@ namespace UnrealBuildTool
 					{
 						if (IsDirectoryForArch(IncludePath.FullName, Arch))
 						{
-							Arguments.Add(string.Format(" -I\"{0}\"", IncludePath));
+							Arguments.Add(GetUserIncludePathArgument(IncludePath));
 						}
 					}
 					foreach (DirectoryReference IncludePath in CompileEnvironment.UserIncludePaths)
 					{
 						if (IsDirectoryForArch(IncludePath.FullName, Arch))
 						{
-							Arguments.Add(string.Format(" -I\"{0}\"", IncludePath));
+							Arguments.Add(GetSystemIncludePathArgument(IncludePath));
 						}
 					}
 
 					// Preprocessor definitions.
 					foreach (string Definition in CompileEnvironment.Definitions)
 					{
-						Arguments.Add(String.Format("-D\"{0}\"", Definition));
+						Arguments.Add(GetPreprocessorDefinitionArgument(Definition));
 					}
 
 					// Consume the included header dependency list
diff --git a/Engine/Source/Programs/UnrealBuildTool/Platform/IOS/IOSToolChain.cs b/Engine/Source/Programs/UnrealBuildTool/Platform/IOS/IOSToolChain.cs
--- a/Engine/Source/Programs/UnrealBuildTool/Platform/IOS/IOSToolChain.cs
+++ b/Engine/Source/Programs/UnrealBuildTool/Platform/IOS/IOSToolChain.cs
@@ -372,6 +372,12 @@ namespace UnrealBuildTool
 				}
 			}
 
+			// Add include paths to the argument list.
+			Result += GetIncludePathArguments(CompileEnvironment);
+
+			// Add preprocessor definitions to the argument list.
+			Result += GetPreprocessorDefinitionArguments(CompileEnvironment);
+
 			return Result;
 		}
 
@@ -566,31 +572,14 @@ namespace UnrealBuildTool
 
 			if (CompileEnvironment.PrecompiledHeaderAction == PrecompiledHeaderAction.Include)
 			{
-				PCHArguments += string.Format(" -include \"{0}\"", CompileEnvironment.PrecompiledHeaderIncludeFilename!.FullName);
+				PCHArguments += GetForceIncludeFileArgument(CompileEnvironment.PrecompiledHeaderIncludeFilename!);
 				if (GetClangVersion().Major >= 11)
 				{
 					PCHArguments += " -fpch-validate-input-files-content";
 				}
 			}
 
-			foreach (FileItem ForceIncludeFile in CompileEnvironment.ForceIncludeFiles)
-			{
-				Arguments += String.Format(" -include \"{0}\"", ForceIncludeFile.Location.FullName);
-			}
-
-			// Add include paths to the argument list.
-			HashSet<DirectoryReference> AllIncludes = new HashSet<DirectoryReference>(CompileEnvironment.UserIncludePaths);
-			AllIncludes.UnionWith(CompileEnvironment.SystemIncludePaths);
-			foreach (DirectoryReference IncludePath in AllIncludes)
-			{
-				Arguments += string.Format(" -I\"{0}\"", IncludePath.FullName);
-			}
-
-			foreach (string Definition in CompileEnvironment.Definitions)
-			{
-				string DefinitionArgument = Definition.Contains("\"") ? Definition.Replace("\"", "\\\"") : Definition;
-				Arguments += string.Format(" -D \"{0}\"", DefinitionArgument);
-			}
+			Arguments += GetForceIncludeFileArguments(CompileEnvironment);
 
 			CPPOutput Result = new CPPOutput();
 			// Create a compile action for each source file.
diff --git a/Engine/Source/Programs/UnrealBuildTool/Platform/Linux/LinuxToolChain.cs b/Engine/Source/Programs/UnrealBuildTool/Platform/Linux/LinuxToolChain.cs
--- a/Engine/Source/Programs/UnrealBuildTool/Platform/Linux/LinuxToolChain.cs
+++ b/Engine/Source/Programs/UnrealBuildTool/Platform/Linux/LinuxToolChain.cs
@@ -796,51 +796,37 @@ namespace UnrealBuildTool
 			}
 
 			// Add include paths to the argument list.
-			foreach (DirectoryReference IncludePath in CompileEnvironment.UserIncludePaths)
-			{
-				Result += string.Format(" -I\"{0}\"", NormalizeCommandLinePath(IncludePath));
-			}
-			foreach (DirectoryReference IncludePath in CompileEnvironment.SystemIncludePaths)
-			{
-				Result += string.Format(" -I\"{0}\"", NormalizeCommandLinePath(IncludePath));
-			}
+			Result += GetIncludePathArguments(CompileEnvironment);
 
 			// Add preprocessor definitions to the argument list.
-			foreach (string Definition in CompileEnvironment.Definitions)
-			{
-				Result += string.Format(" -D \"{0}\"", EscapeArgument(Definition));
-			}
+			Result += GetPreprocessorDefinitionArguments(CompileEnvironment);
 
 			return Result;
 		}
 
-		/// <summary>
-		/// Sanitizes a definition argument if needed.
-		/// </summary>
-		/// <param name="definition">A string in the format "foo=bar".</param>
-		/// <returns></returns>
-		internal static string EscapeArgument(string definition)
+		/// <inheritdoc/>
+		protected override string EscapePreprocessorDefinition(string Definition)
 		{
-			string[] splitData = definition.Split('=');
-			string? myKey = splitData.ElementAtOrDefault(0);
-			string? myValue = splitData.ElementAtOrDefault(1);
+			string[] SplitData = Definition.Split('=');
+			string? Key = SplitData.ElementAtOrDefault(0);
+			string? Value = SplitData.ElementAtOrDefault(1);
 
-			if (string.IsNullOrEmpty(myKey)) { return ""; }
-			if (!string.IsNullOrEmpty(myValue))
+			if (string.IsNullOrEmpty(Key)) { return ""; }
+			if (!string.IsNullOrEmpty(Value))
 			{
-				if (!myValue.StartsWith("\"") && (myValue.Contains(" ") || myValue.Contains("$")))
+				if (!Value.StartsWith("\"") && (Value.Contains(" ") || Value.Contains("$")))
 				{
-					myValue = myValue.Trim('\"');		// trim any leading or trailing quotes
-					myValue = "\"" + myValue + "\"";	// ensure wrap string with double quotes
+					Value = Value.Trim('\"');		// trim any leading or trailing quotes
+					Value = "\"" + Value + "\"";	// ensure wrap string with double quotes
 				}
 
 				// replace double quotes to escaped double quotes if exists
-				myValue = myValue.Replace("\"", "\\\"");
+				Value = Value.Replace("\"", "\\\"");
 			}
 
-			return myValue == null
-				? string.Format("{0}", myKey)
-				: string.Format("{0}={1}", myKey, myValue);
+			return Value == null
+				? string.Format("{0}", Key)
+				: string.Format("{0}={1}", Key, Value);
 		}
 
 		protected virtual string GetLinkArguments(LinkEnvironment LinkEnvironment)
@@ -1138,22 +1124,6 @@ namespace UnrealBuildTool
 			Log.TraceInformation("------------------------------");
 		}
 
-		public static string NormalizeCommandLinePath(FileSystemReference Reference)
-		{
-			// Try to use a relative path to shorten command line length.
-			if (Reference.IsUnderDirectory(Unreal.RootDirectory))
-			{
-				return Reference.MakeRelativeTo(UnrealBuildTool.EngineSourceDirectory).Replace("\\", "/");
-			}
-
-			return Reference.FullName.Replace("\\", "/");
-		}
-
-		public static string NormalizeCommandLinePath(FileItem Item)
-		{
-			return NormalizeCommandLinePath(Item.Location);
-		}
-
 		public override CPPOutput CompileCPPFiles(CppCompileEnvironment CompileEnvironment, List<FileItem> InputFiles, DirectoryReference OutputDir, string ModuleName, IActionGraphBuilder Graph)
 		{
 			string Arguments = GetCLArguments_Global(CompileEnvironment);
@@ -1180,7 +1150,7 @@ namespace UnrealBuildTool
 				{
 					PCHArguments += " -fpch-validate-input-files-content";
 				}
-				PCHArguments += string.Format(" -include \"{0}\"", NormalizeCommandLinePath(CompileEnvironment.PrecompiledHeaderIncludeFilename!));
+				PCHArguments += GetForceIncludeFileArgument(CompileEnvironment.PrecompiledHeaderIncludeFilename!);
 			}
 
 			// Create a compile action for each source file.
@@ -1223,10 +1193,7 @@ namespace UnrealBuildTool
 					FileArguments += PCHArguments;
 				}
 
-				foreach (FileItem ForceIncludeFile in CompileEnvironment.ForceIncludeFiles)
-				{
-					FileArguments += String.Format(" -include \"{0}\"", NormalizeCommandLinePath(ForceIncludeFile));
-				}
+				FileArguments += GetForceIncludeFileArguments(CompileEnvironment);
 
 				// Add the C++ source file and its included files to the prerequisite item list.
 				CompileAction.PrerequisiteItems.Add(SourceFile);
diff --git a/Engine/Source/Programs/UnrealBuildTool/Platform/Mac/MacToolChain.cs b/Engine/Source/Programs/UnrealBuildTool/Platform/Mac/MacToolChain.cs
--- a/Engine/Source/Programs/UnrealBuildTool/Platform/Mac/MacToolChain.cs
+++ b/Engine/Source/Programs/UnrealBuildTool/Platform/Mac/MacToolChain.cs
@@ -300,6 +300,12 @@ namespace UnrealBuildTool
 				Result += " -gdwarf-2";
 			}
 
+			// Add include paths to the argument list.
+			Result += GetIncludePathArguments(CompileEnvironment);
+
+			// Add preprocessor definitions to the argument list.
+			Result += GetPreprocessorDefinitionArguments(CompileEnvironment);
+
 			return Result;
 		}
 		
@@ -369,22 +375,6 @@ namespace UnrealBuildTool
 			return Result;
 		}
 
-		public static string NormalizeCommandLinePath(FileSystemReference Reference)
-		{
-			// Try to use a relative path to shorten command line length.
-			if (Reference.IsUnderDirectory(Unreal.RootDirectory))
-			{
-				return Reference.MakeRelativeTo(UnrealBuildTool.EngineSourceDirectory);
-			}
-
-			return Reference.FullName;
-		}
-
-		public static string NormalizeCommandLinePath(FileItem Item)
-		{
-			return NormalizeCommandLinePath(Item.Location);
-		}
-
 		public override CPPOutput CompileCPPFiles(CppCompileEnvironment CompileEnvironment, List<FileItem> InputFiles, DirectoryReference OutputDir, string ModuleName, IActionGraphBuilder Graph)
 		{
 			StringBuilder Arguments = new StringBuilder();
@@ -396,33 +386,13 @@ namespace UnrealBuildTool
 			{
 				// Add the precompiled header file's path to the include path so GCC can find it.
 				// This needs to be before the other include paths to ensure GCC uses it instead of the source header file.
-				PCHArguments.Append(" -include \"");
-				PCHArguments.Append(NormalizeCommandLinePath(CompileEnvironment.PrecompiledHeaderIncludeFilename!));
-				PCHArguments.Append("\"");
+				PCHArguments.Append(GetForceIncludeFileArgument(CompileEnvironment.PrecompiledHeaderIncludeFilename!));
 				if (GetClangVersion().Major >= 11)
 				{
 					PCHArguments.Append(" -fpch-validate-input-files-content");
 				}
 			}
 
-			// Add include paths to the argument list.
-			HashSet<DirectoryReference> AllIncludes = new HashSet<DirectoryReference>(CompileEnvironment.UserIncludePaths);
-			AllIncludes.UnionWith(CompileEnvironment.SystemIncludePaths);
-			foreach (DirectoryReference IncludePath in AllIncludes)
-			{
-				Arguments.Append(" -I\"");
-				Arguments.Append(NormalizeCommandLinePath(IncludePath));
-				Arguments.Append("\"");
-			}
-
-			foreach (string Definition in CompileEnvironment.Definitions)
-			{
-				string DefinitionArgument = Definition.Contains("\"") ? Definition.Replace("\"", "\\\"") : Definition;
-				Arguments.Append(" -D\"");
-				Arguments.Append(DefinitionArgument);
-				Arguments.Append("\"");
-			}
-
 			List<string> FrameworksSearchPaths = new List<string>();
 			foreach (UEBuildFramework Framework in CompileEnvironment.AdditionalFrameworks)
 			{
@@ -479,10 +449,7 @@ namespace UnrealBuildTool
 					FileArguments += PCHArguments.ToString();
 				}
 
-				foreach (FileItem ForceIncludeFile in CompileEnvironment.ForceIncludeFiles)
-				{
-					FileArguments += String.Format(" -include \"{0}\"", NormalizeCommandLinePath(ForceIncludeFile));
-				}
+				FileArguments += GetForceIncludeFileArguments(CompileEnvironment);
 
 				// Add the C++ source file and its included files to the prerequisite item list.
 				CompileAction.PrerequisiteItems.Add(SourceFile);
@@ -900,13 +867,7 @@ namespace UnrealBuildTool
 			List<string> InputFileNames = new List<string>();
 			foreach (FileItem InputFile in LinkEnvironment.InputFiles)
 			{
-				string InputFilePath = InputFile.AbsolutePath;
-				if (InputFile.Location.IsUnderDirectory(Unreal.RootDirectory))
-				{
-					InputFilePath = InputFile.Location.MakeRelativeTo(UnrealBuildTool.EngineSourceDirectory);
-				}
-
-				InputFileNames.Add(string.Format("\"{0}\"", InputFilePath));
+				InputFileNames.Add(string.Format("\"{0}\"", NormalizeCommandLinePath(InputFile)));
 				LinkAction.PrerequisiteItems.Add(InputFile);
 			}
 
diff --git a/Engine/Source/Programs/UnrealBuildTool/ToolChain/AppleToolChain.cs b/Engine/Source/Programs/UnrealBuildTool/ToolChain/AppleToolChain.cs
--- a/Engine/Source/Programs/UnrealBuildTool/ToolChain/AppleToolChain.cs
+++ b/Engine/Source/Programs/UnrealBuildTool/ToolChain/AppleToolChain.cs
@@ -200,6 +200,12 @@ namespace UnrealBuildTool
 			return FullClangVersion;
 		}
 
+		/// <inheritdoc/>
+		protected override string EscapePreprocessorDefinition(string Definition)
+		{
+			return Definition.Contains("\"") ? Definition.Replace("\"", "\\\"") : Definition;
+		}
+
 		protected override string GetCompileArguments_CPP(CppCompileEnvironment CompileEnvironment)
 		{
 			string Result = "";
diff --git a/Engine/Source/Programs/UnrealBuildTool/ToolChain/ClangToolChain.cs b/Engine/Source/Programs/UnrealBuildTool/ToolChain/ClangToolChain.cs
--- a/Engine/Source/Programs/UnrealBuildTool/ToolChain/ClangToolChain.cs
+++ b/Engine/Source/Programs/UnrealBuildTool/ToolChain/ClangToolChain.cs
@@ -1,13 +1,8 @@
 // Copyright Epic Games, Inc. All Rights Reserved.
 
-using System;
-using System.Collections;
-using System.Collections.Generic;
-using System.IO;
-using System.Diagnostics;
-using System.Text;
 using EpicGames.Core;
-using System.Text.RegularExpressions;
+using System;
+using System.Linq;
 using UnrealBuildBase;
 
 namespace UnrealBuildTool
@@ -24,6 +19,43 @@ namespace UnrealBuildTool
 		{
 		}
 
+		/// <summary>
+		/// Normalize a path for use in a command line, making it relative to Engine/Source if under the root directory
+		/// </summary>
+		/// <param name="Reference">The FileSystemReference to normalize</param>
+		/// <returns>Normalized path as a string</returns>
+		protected static string NormalizeCommandLinePath(FileSystemReference Reference)
+		{
+			// Try to use a relative path to shorten command line length.
+			if (Reference.IsUnderDirectory(Unreal.RootDirectory))
+			{
+				return Reference.MakeRelativeTo(UnrealBuildTool.EngineSourceDirectory).Replace("\\", "/");
+			}
+
+			return Reference.FullName.Replace("\\", "/");
+		}
+
+		/// <summary>
+		/// Normalize a path for use in a command line, making it relative if under the Root Directory
+		/// </summary>
+		/// <param name="Item">The FileItem to normalize</param>
+		/// <returns>Normalized path as a string</returns>
+		protected static string NormalizeCommandLinePath(FileItem Item)
+		{
+			return NormalizeCommandLinePath(Item.Location);
+		}
+
+		/// <summary>
+		/// Sanitizes a preprocessor definition argument if needed.
+		/// </summary>
+		/// <param name="Definition">A string in the format "foo=bar" or "foo".</param>
+		/// <returns>An escaped string</returns>
+		protected virtual string EscapePreprocessorDefinition(string Definition)
+		{
+			// By default don't modify preprocessor definition, handle in platform overrides.
+			return Definition;
+		}
+
 		/// <summary>
 		/// Checks if compiler version matches the requirements
 		/// </summary>
@@ -140,14 +172,65 @@ namespace UnrealBuildTool
 		// Note that exception handling uses the same information, but it will generate it as needed.
 		protected virtual string GetRTTIFlag(CppCompileEnvironment CompileEnvironment)
 		{
-			string Result;
-			if (CompileEnvironment.bUseRTTI)
+			return CompileEnvironment.bUseRTTI ? " -frtti" : " -fno-rtti";
+		}
+
+		protected virtual string GetUserIncludePathArgument(DirectoryReference IncludePath)
+		{
+			return $" -I\"{NormalizeCommandLinePath(IncludePath)}\"";
+		}
+
+		protected virtual string GetSystemIncludePathArgument(DirectoryReference IncludePath)
+		{
+			// TODO: System include paths can be included with -isystem
+			return $" -I\"{NormalizeCommandLinePath(IncludePath)}\"";
+		}
+
+		protected virtual string GetIncludePathArguments(CppCompileEnvironment CompileEnvironment)
+		{
+			string Result = "";
+			foreach (DirectoryReference IncludePath in CompileEnvironment.UserIncludePaths)
 			{
-				Result = " -frtti";
+				Result += GetUserIncludePathArgument(IncludePath);
 			}
-			else
+			foreach (DirectoryReference IncludePath in CompileEnvironment.SystemIncludePaths)
+			{
+				Result += GetSystemIncludePathArgument(IncludePath);
+			}
+			return Result;
+		}
+
+		protected virtual string GetPreprocessorDefinitionArgument(string Definition)
+		{
+			return $" -D\"{EscapePreprocessorDefinition(Definition)}\"";
+		}
+
+		protected virtual string GetPreprocessorDefinitionArguments(CppCompileEnvironment CompileEnvironment)
+		{
+			string Result = "";
+			foreach (string Definition in CompileEnvironment.Definitions)
+			{
+				Result += GetPreprocessorDefinitionArgument(Definition);
+			}
+			return Result;
+		}
+
+		protected virtual string GetForceIncludeFileArgument(FileReference ForceIncludeFile)
+		{
+			return $" -include \"{NormalizeCommandLinePath(ForceIncludeFile)}\"";
+		}
+
+		protected virtual string GetForceIncludeFileArgument(FileItem ForceIncludeFile)
+		{
+			return GetForceIncludeFileArgument(ForceIncludeFile.Location);
+		}
+
+		protected virtual string GetForceIncludeFileArguments(CppCompileEnvironment CompileEnvironment)
+		{
+			string Result = "";
+			foreach (FileItem ForceIncludeFile in CompileEnvironment.ForceIncludeFiles)
 			{
-				Result = " -fno-rtti";
+				Result += String.Format(" -include \"{0}\"", NormalizeCommandLinePath(ForceIncludeFile));
 			}
 			return Result;
 		}
