Do not use --pic for modular builds on any Microsoft-based platforms, not just Windows.

#rnx

diff --git a/Engine/Source/Programs/UnrealBuildTool/ToolChain/ISPCToolChain.cs b/Engine/Source/Programs/UnrealBuildTool/ToolChain/ISPCToolChain.cs
--- a/Engine/Source/Programs/UnrealBuildTool/ToolChain/ISPCToolChain.cs
+++ b/Engine/Source/Programs/UnrealBuildTool/ToolChain/ISPCToolChain.cs
@@ -383,10 +383,10 @@ namespace UnrealBuildTool
 					Arguments.Add(String.Format("--cpu={0}", CpuTarget));
 				}
 
-				// PIC is needed for modular builds except on Windows
+				// PIC is needed for modular builds except on Microsoft platforms
 				if ((CompileEnvironment.bIsBuildingDLL ||
 					CompileEnvironment.bIsBuildingLibrary) &&
-					!UEBuildPlatform.IsPlatformInGroup(CompileEnvironment.Platform, UnrealPlatformGroup.Windows))
+					!UEBuildPlatform.IsPlatformInGroup(CompileEnvironment.Platform, UnrealPlatformGroup.Microsoft))
 				{
 					Arguments.Add("--pic");
 				}
@@ -588,10 +588,10 @@ namespace UnrealBuildTool
 					Arguments.Add("-O2");
 				}
 
-				// PIC is needed for modular builds except on Windows
+				// PIC is needed for modular builds except on Microsoft platforms
 				if ((CompileEnvironment.bIsBuildingDLL ||
 					CompileEnvironment.bIsBuildingLibrary) &&
-					!UEBuildPlatform.IsPlatformInGroup(CompileEnvironment.Platform, UnrealPlatformGroup.Windows))
+					!UEBuildPlatform.IsPlatformInGroup(CompileEnvironment.Platform, UnrealPlatformGroup.Microsoft))
 				{
 					Arguments.Add("--pic");
 				}
