UnrealBuildTool: Include older VS2022 17.2.0 preview toolchain in ban

#rnx

diff --git a/Engine/Source/Programs/UnrealBuildTool/Platform/Windows/MicrosoftPlatformSDK.Versions.cs b/Engine/Source/Programs/UnrealBuildTool/Platform/Windows/MicrosoftPlatformSDK.Versions.cs
--- a/Engine/Source/Programs/UnrealBuildTool/Platform/Windows/MicrosoftPlatformSDK.Versions.cs
+++ b/Engine/Source/Programs/UnrealBuildTool/Platform/Windows/MicrosoftPlatformSDK.Versions.cs
@@ -43,7 +43,7 @@ namespace UnrealBuildTool
 		static readonly VersionNumberRange[] BannedVisualCppVersions = new VersionNumberRange[]
 		{
 			// https://developercommunity.visualstudio.com/t/Incorrect-instantiation-of-a-virtual-fun/10020368
-			VersionNumberRange.Parse("14.32.31326", "14.32.31328"), // VS2022 17.2.0
+			VersionNumberRange.Parse("14.32.31302", "14.32.31328"), // VS2022 17.2.0
 			VersionNumberRange.Parse("14.33.31424", "14.33.31424"), // VS2022 17.3.0 Preview 1.0
 		};
 
