UnrealBuildTool: Add minimal GetCompileArguments_Global in ClangToolChain. Currently only sets -e, -pipe, header include paths, & preprocessor definitions

#rnx

diff --git a/Engine/Source/Programs/UnrealBuildTool/Platform/IOS/IOSToolChain.cs b/Engine/Source/Programs/UnrealBuildTool/Platform/IOS/IOSToolChain.cs
--- a/Engine/Source/Programs/UnrealBuildTool/Platform/IOS/IOSToolChain.cs
+++ b/Engine/Source/Programs/UnrealBuildTool/Platform/IOS/IOSToolChain.cs
@@ -226,13 +226,9 @@ namespace UnrealBuildTool
 		}
 
 
-		string GetCompileArguments_Global(CppCompileEnvironment CompileEnvironment)
+		protected override string GetCompileArguments_Global(CppCompileEnvironment CompileEnvironment)
 		{
-			string Result = "";
-
-			Result += " -fmessage-length=0";
-			Result += " -pipe";
-			Result += " -fpascal-strings";
+			string Result = base.GetCompileArguments_Global(CompileEnvironment);
 
 			if (RuntimePlatform.IsWindows)
 			{
@@ -312,8 +308,6 @@ namespace UnrealBuildTool
 				Result += " -fembed-bitcode";
 			}
 
-			Result += " -c";
-
 			// What architecture(s) to build for
 			Result += GetArchitectureArgument(CompileEnvironment.Configuration, CompileEnvironment.Architecture);
 
@@ -372,12 +366,6 @@ namespace UnrealBuildTool
 				}
 			}
 
-			// Add include paths to the argument list.
-			Result += GetIncludePathArguments(CompileEnvironment);
-
-			// Add preprocessor definitions to the argument list.
-			Result += GetPreprocessorDefinitionArguments(CompileEnvironment);
-
 			return Result;
 		}
 
diff --git a/Engine/Source/Programs/UnrealBuildTool/Platform/Linux/LinuxToolChain.cs b/Engine/Source/Programs/UnrealBuildTool/Platform/Linux/LinuxToolChain.cs
--- a/Engine/Source/Programs/UnrealBuildTool/Platform/Linux/LinuxToolChain.cs
+++ b/Engine/Source/Programs/UnrealBuildTool/Platform/Linux/LinuxToolChain.cs
@@ -516,11 +516,9 @@ namespace UnrealBuildTool
 
 		protected virtual string GetCLArguments_Global(CppCompileEnvironment CompileEnvironment)
 		{
-			string Result = "";
+			string Result = base.GetCompileArguments_Global(CompileEnvironment);
 
 			// build up the commandline common to C and C++
-			Result += " -c";
-			Result += " -pipe";
 
 			if (ShouldUseLibcxx(CompileEnvironment.Architecture))
 			{
@@ -795,12 +793,6 @@ namespace UnrealBuildTool
 				Result += String.Format(" --sysroot=\"{0}\"", BaseLinuxPath);
 			}
 
-			// Add include paths to the argument list.
-			Result += GetIncludePathArguments(CompileEnvironment);
-
-			// Add preprocessor definitions to the argument list.
-			Result += GetPreprocessorDefinitionArguments(CompileEnvironment);
-
 			return Result;
 		}
 
diff --git a/Engine/Source/Programs/UnrealBuildTool/Platform/Mac/MacToolChain.cs b/Engine/Source/Programs/UnrealBuildTool/Platform/Mac/MacToolChain.cs
--- a/Engine/Source/Programs/UnrealBuildTool/Platform/Mac/MacToolChain.cs
+++ b/Engine/Source/Programs/UnrealBuildTool/Platform/Mac/MacToolChain.cs
@@ -191,13 +191,10 @@ namespace UnrealBuildTool
 			return ArchArg;
 		}
 
-		string GetCompileArguments_Global(CppCompileEnvironment CompileEnvironment)
+		protected override string GetCompileArguments_Global(CppCompileEnvironment CompileEnvironment)
 		{
-			string Result = "";
+			string Result = base.GetCompileArguments_Global(CompileEnvironment);
 
-			Result += " -fmessage-length=0";
-			Result += " -pipe";
-			Result += " -fpascal-strings";
 
 			Result += " -fexceptions";
 			Result += " -DPLATFORM_EXCEPTIONS_DISABLED=0";
@@ -249,8 +246,6 @@ namespace UnrealBuildTool
 				Result += " -faligned-new"; // aligned operator new is supported only on macOS 10.14 and above
 			}
 
-			Result += " -c";
-
 			// Pass through architecture and OS info
 			Result += " " + FormatArchitectureArg(CompileEnvironment.Architecture);	
 			Result += string.Format(" -isysroot \"{0}\"", SDKPath);
@@ -300,12 +295,6 @@ namespace UnrealBuildTool
 				Result += " -gdwarf-2";
 			}
 
-			// Add include paths to the argument list.
-			Result += GetIncludePathArguments(CompileEnvironment);
-
-			// Add preprocessor definitions to the argument list.
-			Result += GetPreprocessorDefinitionArguments(CompileEnvironment);
-
 			return Result;
 		}
 		
diff --git a/Engine/Source/Programs/UnrealBuildTool/ToolChain/AppleToolChain.cs b/Engine/Source/Programs/UnrealBuildTool/ToolChain/AppleToolChain.cs
--- a/Engine/Source/Programs/UnrealBuildTool/ToolChain/AppleToolChain.cs
+++ b/Engine/Source/Programs/UnrealBuildTool/ToolChain/AppleToolChain.cs
@@ -253,6 +253,16 @@ namespace UnrealBuildTool
 			return Result;
 		}
 
+		protected override string GetCompileArguments_Global(CppCompileEnvironment CompileEnvironment)
+		{
+			string Result = base.GetCompileArguments_Global(CompileEnvironment);
+
+			Result += " -fmessage-length=0";
+			Result += " -fpascal-strings";
+
+			return Result;
+		}
+
 		protected string GetDsymutilPath(out string ExtraOptions, bool bIsForLTOBuild=false)
 		{
 			FileReference DsymutilLocation = new FileReference("/usr/bin/dsymutil");
diff --git a/Engine/Source/Programs/UnrealBuildTool/ToolChain/ClangToolChain.cs b/Engine/Source/Programs/UnrealBuildTool/ToolChain/ClangToolChain.cs
--- a/Engine/Source/Programs/UnrealBuildTool/ToolChain/ClangToolChain.cs
+++ b/Engine/Source/Programs/UnrealBuildTool/ToolChain/ClangToolChain.cs
@@ -234,5 +234,26 @@ namespace UnrealBuildTool
 			}
 			return Result;
 		}
+
+		/// <summary>
+		/// Common compile arguments for all files in a module.
+		/// Override and call base.GetCompileArguments_Global() in derived classes.
+		/// </summary>
+		protected virtual string GetCompileArguments_Global(CppCompileEnvironment CompileEnvironment)
+		{
+			string Result = "";
+
+			// build up the commandline common to C and C++
+			Result += " -c";
+			Result += " -pipe";
+
+			// Add include paths to the argument list.
+			Result += GetIncludePathArguments(CompileEnvironment);
+
+			// Add preprocessor definitions to the argument list.
+			Result += GetPreprocessorDefinitionArguments(CompileEnvironment);
+
+			return Result;
+		}
 	}
 }
\ No newline at end of file
