UnrealBuildTool: Update clang toolchains (except Android) to fill a list of string arguments rather than returning a single string

#rnx

diff --git a/Engine/Source/Programs/UnrealBuildTool/Platform/Android/AndroidToolChain.cs b/Engine/Source/Programs/UnrealBuildTool/Platform/Android/AndroidToolChain.cs
--- a/Engine/Source/Programs/UnrealBuildTool/Platform/Android/AndroidToolChain.cs
+++ b/Engine/Source/Programs/UnrealBuildTool/Platform/Android/AndroidToolChain.cs
@@ -717,7 +717,7 @@ namespace UnrealBuildTool
 				Result += " -DPLATFORM_EXCEPTIONS_DISABLED=1";
 			}
 
-			Result += GetRTTIFlag(CompileEnvironment);
+			Result += $" {GetRTTIFlag(CompileEnvironment)}";
 
 			// Profile Guided Optimization (PGO) and Link Time Optimization (LTO)
 			if (CompileEnvironment.bPGOOptimize)
@@ -1286,7 +1286,9 @@ namespace UnrealBuildTool
 					// Add C or C++ specific compiler arguments.
 					if (bIsPlainCFile)
 					{
-						FileArguments += GetCompileArguments_C();
+						List<string> CArguments = new();
+						GetCompileArguments_C(CompileEnvironment, CArguments);
+						FileArguments += $" {string.Join(' ', CArguments)}";
 
 						// remove shadow variable warnings for externally included files
 						if (!SourceFile.Location.IsUnderDirectory(Unreal.RootDirectory))
@@ -1296,11 +1298,15 @@ namespace UnrealBuildTool
 					}
 					else if (CompileEnvironment.PrecompiledHeaderAction == PrecompiledHeaderAction.Create)
 					{
-						FileArguments += GetCompileArguments_PCH(CompileEnvironment);
+						List<string> CreatePCHArguments = new();
+						GetCompileArguments_PCH(CompileEnvironment, CreatePCHArguments);
+						FileArguments += $" {string.Join(' ', CreatePCHArguments)}";
 					}
 					else
 					{
-						FileArguments += GetCompileArguments_CPP(CompileEnvironment);
+						List<string> CPPArguments = new();
+						GetCompileArguments_CPP(CompileEnvironment, CPPArguments);
+						FileArguments += $" {string.Join(' ', CPPArguments)}";
 
 						// only use PCH for .cpp files
 						FileArguments += PCHArguments;
@@ -1727,21 +1733,21 @@ namespace UnrealBuildTool
 					{
 						if (IsDirectoryForArch(IncludePath.FullName, Arch))
 						{
-							Arguments.Add(GetUserIncludePathArgument(IncludePath));
+							Arguments.Add(string.Format(" -I\"{0}\"", IncludePath));
 						}
 					}
 					foreach (DirectoryReference IncludePath in CompileEnvironment.UserIncludePaths)
 					{
 						if (IsDirectoryForArch(IncludePath.FullName, Arch))
 						{
-							Arguments.Add(GetSystemIncludePathArgument(IncludePath));
+							Arguments.Add(string.Format(" -I\"{0}\"", IncludePath));
 						}
 					}
 
 					// Preprocessor definitions.
 					foreach (string Definition in CompileEnvironment.Definitions)
 					{
-						Arguments.Add(GetPreprocessorDefinitionArgument(Definition));
+						Arguments.Add(String.Format(" -D\"{0}\"", Definition));
 					}
 
 					// Consume the included header dependency list
diff --git a/Engine/Source/Programs/UnrealBuildTool/Platform/IOS/IOSToolChain.cs b/Engine/Source/Programs/UnrealBuildTool/Platform/IOS/IOSToolChain.cs
--- a/Engine/Source/Programs/UnrealBuildTool/Platform/IOS/IOSToolChain.cs
+++ b/Engine/Source/Programs/UnrealBuildTool/Platform/IOS/IOSToolChain.cs
@@ -226,94 +226,94 @@ namespace UnrealBuildTool
 		}
 
 
-		protected override string GetCompileArguments_Global(CppCompileEnvironment CompileEnvironment)
+		protected override void GetCompileArguments_Global(CppCompileEnvironment CompileEnvironment, List<string> Arguments)
 		{
-			string Result = base.GetCompileArguments_Global(CompileEnvironment);
+			base.GetCompileArguments_Global(CompileEnvironment, Arguments);
 
 			if (RuntimePlatform.IsWindows)
 			{
-				Result += " -fdiagnostics-format=msvc";
+				Arguments.Add("-fdiagnostics-format=msvc");
 			}
 
 			// Optionally enable exception handling (off by default since it generates extra code needed to propagate exceptions)
 			if (CompileEnvironment.bEnableExceptions)
 			{
-				Result += " -fexceptions";
-				Result += " -DPLATFORM_EXCEPTIONS_DISABLED=0";
+				Arguments.Add("-fexceptions");
+				Arguments.Add("-DPLATFORM_EXCEPTIONS_DISABLED=0");
 			}
 			else
 			{
-				Result += " -fno-exceptions";
-				Result += " -DPLATFORM_EXCEPTIONS_DISABLED=1";
+				Arguments.Add("-fno-exceptions");
+				Arguments.Add("-DPLATFORM_EXCEPTIONS_DISABLED=1");
 			}
 
 			string? SanitizerMode = Environment.GetEnvironmentVariable("ENABLE_ADDRESS_SANITIZER");
 			if ((SanitizerMode != null && SanitizerMode == "YES") || (Options.HasFlag(IOSToolChainOptions.EnableAddressSanitizer)))
 			{
-				Result += " -fsanitize=address -fno-omit-frame-pointer -DFORCE_ANSI_ALLOCATOR=1";
+				Arguments.Add("-fsanitize=address -fno-omit-frame-pointer -DFORCE_ANSI_ALLOCATOR=1");
 			}
 
 			string? UndefSanitizerMode = Environment.GetEnvironmentVariable("ENABLE_UNDEFINED_BEHAVIOR_SANITIZER");
 			if ((UndefSanitizerMode != null && UndefSanitizerMode == "YES") || (Options.HasFlag(IOSToolChainOptions.EnableUndefinedBehaviorSanitizer)))
 			{
-				Result += " -fsanitize=undefined -fno-sanitize=bounds,enum,return,float-divide-by-zero";
+				Arguments.Add("-fsanitize=undefined -fno-sanitize=bounds,enum,return,float-divide-by-zero");
 			}
 
 			if (Options.HasFlag(IOSToolChainOptions.EnableThreadSanitizer))
 			{
-				Result += " -fsanitize=thread";
+				Arguments.Add("-fsanitize=thread");
 			}
 
-			Result += GetRTTIFlag(CompileEnvironment);
-			Result += " -fvisibility=hidden"; // hides the linker warnings with PhysX
+			Arguments.Add(GetRTTIFlag(CompileEnvironment));
+			Arguments.Add("-fvisibility=hidden"); // hides the linker warnings with PhysX
 
 			// use LTO if desired (like VCToolchain does)
 			if (CompileEnvironment.bAllowLTCG)
 			{
-				Result += " -flto";
+				Arguments.Add("-flto");
 			}
 
-			Result += " -Wall -Werror";
-			Result += " -Wdelete-non-virtual-dtor";
+			Arguments.Add("-Wall -Werror");
+			Arguments.Add("-Wdelete-non-virtual-dtor");
 
 			// clang 12.00 has a new warning for copies in ranged loops. Instances have all been fixed up (2020/6/26) but
 			// are likely to be reintroduced due to no equivalent on other platforms at this time so disable the warning
 			// See also MacToolChain.cs
 			if (GetClangVersion().Major >= 12)
 			{
-				Result += " -Wno-range-loop-analysis";
+				Arguments.Add("-Wno-range-loop-analysis");
 
 				// We have 'this' vs nullptr comparisons that get optimized away for newer versions of Clang, which is undesirable until we refactor these checks.
-				Result += " -fno-delete-null-pointer-checks";
+				Arguments.Add("-fno-delete-null-pointer-checks");
 			}
 
 			if (CompileEnvironment.ShadowVariableWarningLevel != WarningLevel.Off)
 			{
-				Result += " -Wshadow" + ((CompileEnvironment.ShadowVariableWarningLevel == WarningLevel.Error) ? "" : " -Wno-error=shadow");
+				Arguments.Add("-Wshadow" + ((CompileEnvironment.ShadowVariableWarningLevel == WarningLevel.Error) ? "" : " -Wno-error=shadow"));
 			}
 
 			if (CompileEnvironment.bEnableUndefinedIdentifierWarnings)
 			{
-				Result += " -Wundef" + (CompileEnvironment.bUndefinedIdentifierWarningsAsErrors ? "" : " -Wno-error=undef");
+				Arguments.Add("-Wundef" + (CompileEnvironment.bUndefinedIdentifierWarningsAsErrors ? "" : " -Wno-error=undef"));
 			}
 
 			// fix for Xcode 8.3 enabling nonportable include checks, but p4 has some invalid cases in it
 			if (Settings.Value.IOSSDKVersionFloat >= 10.3)
 			{
-				Result += " -Wno-nonportable-include-path";
+				Arguments.Add("-Wno-nonportable-include-path");
 			}
 
 			if (IsBitcodeCompilingEnabled(CompileEnvironment.Configuration))
 			{
-				Result += " -fembed-bitcode";
+				Arguments.Add("-fembed-bitcode");
 			}
 
 			// What architecture(s) to build for
-			Result += GetArchitectureArgument(CompileEnvironment.Configuration, CompileEnvironment.Architecture);
+			Arguments.Add(GetArchitectureArgument(CompileEnvironment.Configuration, CompileEnvironment.Architecture));
 
-			Result += string.Format(" -isysroot \"{0}\"", Settings.Value.GetSDKPath(CompileEnvironment.Architecture));
+			Arguments.Add(string.Format(" -isysroot \"{0}\"", Settings.Value.GetSDKPath(CompileEnvironment.Architecture)));
 
-			Result += " -m" + GetXcodeMinVersionParam() + "=" + ProjectSettings.RuntimeVersion;
+			Arguments.Add("-m" + GetXcodeMinVersionParam() + "=" + ProjectSettings.RuntimeVersion);
 
 			bool bStaticAnalysis = false;
 			string? StaticAnalysisMode = Environment.GetEnvironmentVariable("CLANG_STATIC_ANALYZER_MODE");
@@ -327,27 +327,27 @@ namespace UnrealBuildTool
 			{
 				if (CompileEnvironment.bOptimizeForSize)
 				{
-					Result += " -Oz";
+					Arguments.Add("-Oz");
 				}
 				else
 				{
-					Result += " -O3";
+					Arguments.Add("-O3");
 				}
 			}
 			else
 			{
-				Result += " -O0";
+				Arguments.Add("-O0");
 			}
 
 			if (!CompileEnvironment.bUseInlining)
 			{
-				Result += " -fno-inline-functions";
+				Arguments.Add("-fno-inline-functions");
 			}
 
 			// Create DWARF format debug info if wanted,
 			if (CompileEnvironment.bCreateDebugInfo)
 			{
-				Result += " -gdwarf-2";
+				Arguments.Add("-gdwarf-2");
 			}
 
 			// Add additional frameworks so that their headers can be found
@@ -362,11 +362,9 @@ namespace UnrealBuildTool
 					{
 						FrameworkDir = Path.GetDirectoryName(FrameworkDir)!;
 					}
-					Result += String.Format(" -F\"{0}\"", FrameworkDir);
+					Arguments.Add(String.Format(" -F\"{0}\"", FrameworkDir));
 				}
 			}
-
-			return Result;
 		}
 
 		// Conditionally enable (default disabled) Objective-C exceptions
@@ -376,11 +374,11 @@ namespace UnrealBuildTool
 
 			if (CompileEnvironment.bEnableObjCExceptions)
 			{
-				Result += " -fobjc-exceptions";
+				Result += "-fobjc-exceptions";
 			}
 			else
 			{
-				Result += " -fno-objc-exceptions";
+				Result += "-fno-objc-exceptions";
 			}
 
 			return Result;
@@ -555,63 +553,67 @@ namespace UnrealBuildTool
 
 		public override CPPOutput CompileCPPFiles(CppCompileEnvironment CompileEnvironment, List<FileItem> InputFiles, DirectoryReference OutputDir, string ModuleName, IActionGraphBuilder Graph)
 		{
-			string Arguments = GetCompileArguments_Global(CompileEnvironment);
-			string PCHArguments = "";
+			List<string> GlobalArguments = new();
+
+			GetCompileArguments_Global(CompileEnvironment, GlobalArguments);
+			List<string> PCHArguments = new();
 
 			if (CompileEnvironment.PrecompiledHeaderAction == PrecompiledHeaderAction.Include)
 			{
-				PCHArguments += GetForceIncludeFileArgument(CompileEnvironment.PrecompiledHeaderIncludeFilename!);
+				// Add the precompiled header file's path to the include path so GCC can find it.
+				// This needs to be before the other include paths to ensure GCC uses it instead of the source header file.
+				PCHArguments.Add(GetForceIncludeFileArgument(CompileEnvironment.PrecompiledHeaderIncludeFilename!));
 				if (GetClangVersion().Major >= 11)
 				{
-					PCHArguments += " -fpch-validate-input-files-content";
+					PCHArguments.Add("-fpch-validate-input-files-content");
 				}
 			}
 
-			Arguments += GetForceIncludeFileArguments(CompileEnvironment);
+			GetCompileArguments_ForceInclude(CompileEnvironment, GlobalArguments);
 
 			CPPOutput Result = new CPPOutput();
 			// Create a compile action for each source file.
 			foreach (FileItem SourceFile in InputFiles)
 			{
 				Action CompileAction = Graph.CreateAction(ActionType.Compile);
-				string FilePCHArguments = "";
-				string FileArguments = "";
+				List<string> FilePCHArguments = new();
+				List<string> FileArguments = new();
 				string Extension = Path.GetExtension(SourceFile.AbsolutePath).ToUpperInvariant();
 
 				if (CompileEnvironment.PrecompiledHeaderAction == PrecompiledHeaderAction.Create)
 				{
 					// Compile the file as a C++ PCH.
-					FileArguments += GetCompileArguments_PCH(CompileEnvironment);
-					FileArguments += GetRTTIFlag(CompileEnvironment);
-					FileArguments += GetObjCExceptionsFlag(CompileEnvironment);
+					GetCompileArguments_PCH(CompileEnvironment, FileArguments);
+					FileArguments.Add(GetRTTIFlag(CompileEnvironment));
+					FileArguments.Add(GetObjCExceptionsFlag(CompileEnvironment));
 				}
 				else if (Extension == ".C")
 				{
 					// Compile the file as C code.
-					FileArguments += GetCompileArguments_C();
+					GetCompileArguments_C(CompileEnvironment, FileArguments);
 				}
 				else if (Extension == ".MM")
 				{
 					// Compile the file as Objective-C++ code.
-					FileArguments += GetCompileArguments_MM(CompileEnvironment);
-					FileArguments += GetRTTIFlag(CompileEnvironment);
-					FileArguments += GetObjCExceptionsFlag(CompileEnvironment);
+					GetCompileArguments_MM(CompileEnvironment, FileArguments);
+					FileArguments.Add(GetRTTIFlag(CompileEnvironment));
+					FileArguments.Add(GetObjCExceptionsFlag(CompileEnvironment));
 				}
 				else if (Extension == ".M")
 				{
 					// Compile the file as Objective-C code.
-					FileArguments += GetCompileArguments_M(CompileEnvironment);
-					FileArguments += GetObjCExceptionsFlag(CompileEnvironment);
+					GetCompileArguments_M(CompileEnvironment, FileArguments);
+					FileArguments.Add(GetObjCExceptionsFlag(CompileEnvironment));
 				}
 				else
 				{
 					// Compile the file as C++ code.
-					FileArguments += GetCompileArguments_CPP(CompileEnvironment);
-					FileArguments += GetRTTIFlag(CompileEnvironment);
-					FileArguments += GetObjCExceptionsFlag(CompileEnvironment);
+					GetCompileArguments_CPP(CompileEnvironment, FileArguments);
+					FileArguments.Add(GetRTTIFlag(CompileEnvironment));
+					FileArguments.Add(GetObjCExceptionsFlag(CompileEnvironment));
 
 					// only use PCH for .cpp files
-					FilePCHArguments = PCHArguments;
+					FilePCHArguments.AddRange(PCHArguments);
 				}
 
 				// Add the C++ source file and its included files to the prerequisite item list.
@@ -639,7 +641,7 @@ namespace UnrealBuildTool
 					Result.PrecompiledHeaderFile = PrecompiledHeaderFile;
 
 					// Add the parameters needed to compile the precompiled header file to the command-line.
-					FileArguments += string.Format(" -o \"{0}\"", PrecompiledHeaderFile.AbsolutePath);
+					FileArguments.Add(string.Format("-o \"{0}\"", NormalizeCommandLinePath(PrecompiledHeaderFile)));
 				}
 				else
 				{
@@ -663,25 +665,38 @@ namespace UnrealBuildTool
 
 					CompileAction.ProducedItems.Add(ObjectFile);
 					Result.ObjectFiles.Add(ObjectFile);
-					FileArguments += string.Format(" -o \"{0}\"", ObjectFile.AbsolutePath);
+					FileArguments.Add(string.Format("-o \"{0}\"", NormalizeCommandLinePath(ObjectFile)));
 					OutputFilePath = ObjectFile.AbsolutePath;
 				}
 
 				// Add the source file path to the command-line.
-				FileArguments += string.Format(" \"{0}\"", SourceFile.AbsolutePath);
+				FileArguments.Add(string.Format("\"{0}\"", NormalizeCommandLinePath(SourceFile)));
 
 				// Generate the included header dependency list
 				if (CompileEnvironment.bGenerateDependenciesFile)
 				{
 					FileItem DependencyListFile = FileItem.GetItemByFileReference(FileReference.Combine(OutputDir, Path.GetFileName(SourceFile.AbsolutePath) + ".d"));
-					FileArguments += string.Format(" -MD -MF\"{0}\"", DependencyListFile.AbsolutePath.Replace('\\', '/'));
+					FileArguments.Add(string.Format("-MD -MF\"{0}\"", NormalizeCommandLinePath(DependencyListFile)));
 					CompileAction.DependencyListFile = DependencyListFile;
 					CompileAction.ProducedItems.Add(DependencyListFile);
 				}
 
 				string CompilerPath = Settings.Value.ToolchainDir + IOSCompiler;
 
-				string AllArgs = FilePCHArguments + Arguments + FileArguments + CompileEnvironment.AdditionalArguments;
+				// Gets the target file so we can get the correct output path.
+				FileItem TargetFile = CompileAction.ProducedItems[0];
+
+				// Creates the path to the response file using the name of the output file and creates its contents.
+				FileReference ResponseFileName = new FileReference(TargetFile.AbsolutePath + ".response");
+				List<string> ResponseFileContents = new();
+				ResponseFileContents.AddRange(FilePCHArguments);
+				ResponseFileContents.AddRange(GlobalArguments);
+				ResponseFileContents.AddRange(FileArguments);
+				if (!string.IsNullOrWhiteSpace(CompileEnvironment.AdditionalArguments))
+				{
+					ResponseFileContents.Add(CompileEnvironment.AdditionalArguments);
+				}
+				ResponseFileContents = ResponseFileContents.Select(x => Utils.ExpandVariables(x)).ToList();
 				/*				string SourceText = System.IO.File.ReadAllText(SourceFile.AbsolutePath);
 								if (CompileEnvironment.bOptimizeForSize && (SourceFile.AbsolutePath.Contains("ElementBatcher.cpp") || SourceText.Contains("ElementBatcher.cpp") || SourceFile.AbsolutePath.Contains("AnimationRuntime.cpp") || SourceText.Contains("AnimationRuntime.cpp")
 									|| SourceFile.AbsolutePath.Contains("AnimEncoding.cpp") || SourceText.Contains("AnimEncoding.cpp") || SourceFile.AbsolutePath.Contains("TextRenderComponent.cpp") || SourceText.Contains("TextRenderComponent.cpp")
@@ -693,24 +708,27 @@ namespace UnrealBuildTool
 									AllArgs = AllArgs.Replace("-Oz", "-O3");
 								}*/
 
+				// Adds the response file to the compiler input.
+				FileItem ResponseFileItem = Graph.CreateIntermediateTextFile(ResponseFileName, ResponseFileContents);
+				CompileAction.CommandArguments += string.Format("@\"{0}\"", NormalizeCommandLinePath(ResponseFileItem));
+				CompileAction.PrerequisiteItems.Add(ResponseFileItem);
+
 				// Analyze and then compile using the shell to perform the indirection
 				string? StaticAnalysisMode = Environment.GetEnvironmentVariable("CLANG_STATIC_ANALYZER_MODE");
 				if (StaticAnalysisMode != null && StaticAnalysisMode != "" && OutputFilePath != null)
 				{
-					string TempArgs = "-c \"" + CompilerPath + " " + AllArgs + " --analyze -Wno-unused-command-line-argument -Xclang -analyzer-output=html -Xclang -analyzer-config -Xclang path-diagnostics-alternate=true -Xclang -analyzer-config -Xclang report-in-main-source-file=true -Xclang -analyzer-disable-checker -Xclang deadcode.DeadStores -o " + OutputFilePath.Replace(".o", ".html") + "; " + CompilerPath + " " + AllArgs + "\"";
-					AllArgs = TempArgs;
+					CompileAction.CommandArguments = "-c \"" + CompilerPath + " " + string.Format("@\"{0}\"", ResponseFileName) + " --analyze -Wno-unused-command-line-argument -Xclang -analyzer-output=html -Xclang -analyzer-config -Xclang path-diagnostics-alternate=true -Xclang -analyzer-config -Xclang report-in-main-source-file=true -Xclang -analyzer-disable-checker -Xclang deadcode.DeadStores -o " + OutputFilePath.Replace(".o", ".html") + "; " + CompilerPath + " " + string.Format("@\"{0}\"", ResponseFileName) + "\"";
 					CompilerPath = "/bin/sh";
 				}
 
 				// RPC utility parameters are in terms of the Mac side
 				CompileAction.WorkingDirectory = GetMacDevSrcRoot();
 				CompileAction.CommandPath = new FileReference(CompilerPath);
-				CompileAction.CommandArguments = AllArgs; // Arguments + FileArguments + CompileEnvironment.AdditionalArguments;
 				CompileAction.CommandDescription = "Compile";
 				CompileAction.StatusDescription = string.Format("{0}", Path.GetFileName(SourceFile.AbsolutePath));
 				CompileAction.bIsGCCCompiler = true;
 				// We're already distributing the command by execution on Mac.
-				CompileAction.bCanExecuteRemotely = true;
+				CompileAction.bCanExecuteRemotely = CompileEnvironment.PrecompiledHeaderAction != PrecompiledHeaderAction.Create;
 				CompileAction.bShouldOutputStatusDescription = true;
 				CompileAction.CommandVersion = GetFullClangVersion();
 
diff --git a/Engine/Source/Programs/UnrealBuildTool/Platform/Linux/LinuxToolChain.cs b/Engine/Source/Programs/UnrealBuildTool/Platform/Linux/LinuxToolChain.cs
--- a/Engine/Source/Programs/UnrealBuildTool/Platform/Linux/LinuxToolChain.cs
+++ b/Engine/Source/Programs/UnrealBuildTool/Platform/Linux/LinuxToolChain.cs
@@ -452,19 +452,19 @@ namespace UnrealBuildTool
 
 			if (Architecture.StartsWith("arm") || Architecture.StartsWith("aarch64"))
 			{
-				Result += " -fsigned-char";
+				Result += "-fsigned-char";
 			}
 
 			return Result;
 		}
 
-		protected virtual string ArchitectureSpecificDefines(string Architecture)
+		static string ArchitectureSpecificDefines(string Architecture)
 		{
 			string Result = "";
 
 			if (Architecture.StartsWith("x86_64") || Architecture.StartsWith("aarch64"))
 			{
-				Result += " -D_LINUX64";
+				Result += "-D_LINUX64";
 			}
 
 			return Result;
@@ -514,37 +514,37 @@ namespace UnrealBuildTool
 			return false;
 		}
 
-		protected virtual string GetCLArguments_Global(CppCompileEnvironment CompileEnvironment)
+		protected override void GetCompileArguments_Global(CppCompileEnvironment CompileEnvironment, List<string> Arguments)
 		{
-			string Result = base.GetCompileArguments_Global(CompileEnvironment);
+			base.GetCompileArguments_Global(CompileEnvironment, Arguments);
 
 			// build up the commandline common to C and C++
 
 			if (ShouldUseLibcxx(CompileEnvironment.Architecture))
 			{
-				Result += " -nostdinc++";
-				Result += " -I" + "ThirdParty/Unix/LibCxx/include/";
-				Result += " -I" + "ThirdParty/Unix/LibCxx/include/c++/v1";
+				Arguments.Add("-nostdinc++");
+				Arguments.Add("-I" + "ThirdParty/Unix/LibCxx/include/");
+				Arguments.Add("-I" + "ThirdParty/Unix/LibCxx/include/c++/v1");
 			}
 
 			// ASan
 			if (Options.HasFlag(LinuxToolChainOptions.EnableAddressSanitizer))
 			{
 				// Force using the ANSI allocator if ASan is enabled
-				Result += " -fsanitize=address -DFORCE_ANSI_ALLOCATOR=1";
+				Arguments.Add("-fsanitize=address -DFORCE_ANSI_ALLOCATOR=1");
 			}
 
 			// TSan
 			if (Options.HasFlag(LinuxToolChainOptions.EnableThreadSanitizer))
 			{
 				// Force using the ANSI allocator if TSan is enabled
-				Result += " -fsanitize=thread -DFORCE_ANSI_ALLOCATOR=1";
+				Arguments.Add("-fsanitize=thread -DFORCE_ANSI_ALLOCATOR=1");
 			}
 
 			// UBSan
 			if (Options.HasFlag(LinuxToolChainOptions.EnableUndefinedBehaviorSanitizer))
 			{
-				Result += " -fsanitize=undefined -fno-sanitize=vptr";
+				Arguments.Add("-fsanitize=undefined -fno-sanitize=vptr");
 			}
 
 			// MSan
@@ -553,41 +553,41 @@ namespace UnrealBuildTool
 				// Force using the ANSI allocator if MSan is enabled
 				// -fsanitize-memory-track-origins adds a 1.5x-2.5x slow down ontop of MSan normal amount of overhead
 				// -fsanitize-memory-track-origins=1 is faster but collects only allocation points but not intermediate stores
-				Result += " -fsanitize=memory -fsanitize-memory-track-origins -DFORCE_ANSI_ALLOCATOR=1";
+				Arguments.Add("-fsanitize=memory -fsanitize-memory-track-origins -DFORCE_ANSI_ALLOCATOR=1");
 			}
 
-			Result += " -Wall -Werror";
+			Arguments.Add("-Wall -Werror");
 
 			if (CompilerVersionGreaterOrEqual(12, 0, 0))
 			{
-				Result += " -fbinutils-version=2.36";
+				Arguments.Add("-fbinutils-version=2.36");
 			}
 
 			if (!CompileEnvironment.Architecture.StartsWith("x86_64"))
 			{
-				Result += " -funwind-tables";               // generate unwind tables as they are needed for backtrace (on x86(64) they are generated implicitly)
+				Arguments.Add("-funwind-tables");               // generate unwind tables as they are needed for backtrace (on x86(64) they are generated implicitly)
 			}
 
-			Result += " -Wsequence-point";              // additional warning not normally included in Wall: warns if order of operations is ambigious
-			//Result += " -Wunreachable-code";            // additional warning not normally included in Wall: warns if there is code that will never be executed - not helpful due to bIsGCC and similar
-			//Result += " -Wshadow";                      // additional warning not normally included in Wall: warns if there variable/typedef shadows some other variable - not helpful because we have gobs of code that shadows variables
-			Result += " -Wdelete-non-virtual-dtor";
+			Arguments.Add("-Wsequence-point");              // additional warning not normally included in Wall: warns if order of operations is ambigious
+			//Arguments.Add("-Wunreachable-code");            // additional warning not normally included in Wall: warns if there is code that will never be executed - not helpful due to bIsGCC and similar
+			//Arguments.Add("-Wshadow");                      // additional warning not normally included in Wall: warns if there variable/typedef shadows some other variable - not helpful because we have gobs of code that shadows variables
+			Arguments.Add("-Wdelete-non-virtual-dtor");
 
-			Result += ArchitectureSpecificSwitches(CompileEnvironment.Architecture);
+			Arguments.Add(ArchitectureSpecificSwitches(CompileEnvironment.Architecture));
 
-			Result += " -fno-math-errno";               // do not assume that math ops have side effects
+			Arguments.Add("-fno-math-errno");               // do not assume that math ops have side effects
 
-			Result += GetRTTIFlag(CompileEnvironment);	// flag for run-time type info
+			Arguments.Add(GetRTTIFlag(CompileEnvironment));	// flag for run-time type info
 
 			if (CompileEnvironment.Architecture.StartsWith("x86_64"))
 			{
-				Result += " -mssse3"; // enable ssse3 by default for x86. This is default on for MSVC so lets reflect that here
+				Arguments.Add("-mssse3"); // enable ssse3 by default for x86. This is default on for MSVC so lets reflect that here
 			}
 
 			if (CompileEnvironment.bHideSymbolsByDefault)
 			{
-				Result += " -fvisibility-ms-compat";
-				Result += " -fvisibility-inlines-hidden";
+				Arguments.Add("-fvisibility-ms-compat");
+				Arguments.Add("-fvisibility-inlines-hidden");
 			}
 
 			// Clang only options
@@ -595,37 +595,37 @@ namespace UnrealBuildTool
 			{
 				if (BuildHostPlatform.Current.Platform == UnrealTargetPlatform.Win64)
 				{
-					Result += " -fdiagnostics-format=msvc";     // make diagnostics compatible with MSVC when cross-compiling
+					Arguments.Add("-fdiagnostics-format=msvc");     // make diagnostics compatible with MSVC when cross-compiling
 				}
 				else if (Log.ColorConsoleOutput)
 				{
-					Result += " -fcolor-diagnostics";
+					Arguments.Add("-fcolor-diagnostics");
 				}
 			}
 
-			Result += " -fdiagnostics-absolute-paths"; // output full paths to the files when the build fails
+			Arguments.Add("-fdiagnostics-absolute-paths"); // output full paths to the files when the build fails
 
-			Result += " -Wno-unused-private-field"; // MultichannelTcpSocket.h triggers this, possibly more
-			Result += " -Wno-tautological-compare"; // this hides the "warning : comparison of unsigned expression < 0 is always false" type warnings due to constant comparisons, which are possible with template arguments
-			Result += " -Wno-undefined-bool-conversion"; // hides checking if 'this' pointer is null
-			Result += " -Wno-unused-local-typedef";	// clang is being overly strict here? PhysX headers trigger this.
-			Result += " -Wno-inconsistent-missing-override";	// these have to be suppressed for UE 4.8, should be fixed later.
-			Result += " -Wno-undefined-var-template"; // not really a good warning to disable
-			Result += " -Wno-unused-lambda-capture";  // suppressed because capturing of compile-time constants is seemingly inconsistent. And MSVC doesn't do that.
-			Result += " -Wno-unused-variable";
+			Arguments.Add("-Wno-unused-private-field"); // MultichannelTcpSocket.h triggers this, possibly more
+			Arguments.Add("-Wno-tautological-compare"); // this hides the "warning : comparison of unsigned expression < 0 is always false" type warnings due to constant comparisons, which are possible with template arguments
+			Arguments.Add("-Wno-undefined-bool-conversion"); // hides checking if 'this' pointer is null
+			Arguments.Add("-Wno-unused-local-typedef");	// clang is being overly strict here? PhysX headers trigger this.
+			Arguments.Add("-Wno-inconsistent-missing-override");	// these have to be suppressed for UE 4.8, should be fixed later.
+			Arguments.Add("-Wno-undefined-var-template"); // not really a good warning to disable
+			Arguments.Add("-Wno-unused-lambda-capture");  // suppressed because capturing of compile-time constants is seemingly inconsistent. And MSVC doesn't do that.
+			Arguments.Add("-Wno-unused-variable");
 
 			if (CompilerVersionGreaterOrEqual(13, 0, 0))
 			{
-				Result += " -Wno-unused-but-set-variable";
-				Result += " -Wno-unused-but-set-parameter";
-				Result += " -Wno-ordered-compare-function-pointers";
+				Arguments.Add("-Wno-unused-but-set-variable");
+				Arguments.Add("-Wno-unused-but-set-parameter");
+				Arguments.Add("-Wno-ordered-compare-function-pointers");
 			}
 
-			Result += " -Wno-unused-function"; // this will hide the warnings about static functions in headers that aren't used in every single .cpp file
-			Result += " -Wno-switch"; // this hides the "enumeration value 'XXXXX' not handled in switch [-Wswitch]" warnings - we should maybe remove this at some point and add UE_LOG(, Fatal, ) to default cases
-			Result += " -Wno-unknown-pragmas";			// Slate triggers this (with its optimize on/off pragmas)
-			Result += " -Wno-invalid-offsetof"; // needed to suppress warnings about using offsetof on non-POD types.
-			Result += " -Wno-gnu-string-literal-operator-template"; // we use this feature to allow static FNames.
+			Arguments.Add("-Wno-unused-function"); // this will hide the warnings about static functions in headers that aren't used in every single .cpp file
+			Arguments.Add("-Wno-switch"); // this hides the "enumeration value 'XXXXX' not handled in switch [-Wswitch]" warnings - we should maybe remove this at some point and add UE_LOG(, Fatal, ) to default cases
+			Arguments.Add("-Wno-unknown-pragmas");			// Slate triggers this (with its optimize on/off pragmas)
+			Arguments.Add("-Wno-invalid-offsetof"); // needed to suppress warnings about using offsetof on non-POD types.
+			Arguments.Add("-Wno-gnu-string-literal-operator-template"); // we use this feature to allow static FNames.
 
 			// Profile Guided Optimization (PGO) and Link Time Optimization (LTO)
 			// Whether we actually can enable that is checked in CanUseAdvancedLinkerFeatures() earlier
@@ -637,15 +637,15 @@ namespace UnrealBuildTool
 				//
 				// Disable this warning. It's far too verbose.
 				//
-				Result += " -Wno-backend-plugin";
+				Arguments.Add("-Wno-backend-plugin");
 
 				Log.TraceInformationOnce("Enabling Profile Guided Optimization (PGO). Linking will take a while.");
-				Result += string.Format(" -fprofile-instr-use=\"{0}\"", Path.Combine(CompileEnvironment.PGODirectory!, CompileEnvironment.PGOFilenamePrefix!));
+				Arguments.Add(string.Format(" -fprofile-instr-use=\"{0}\"", Path.Combine(CompileEnvironment.PGODirectory!, CompileEnvironment.PGOFilenamePrefix!)));
 			}
 			else if (CompileEnvironment.bPGOProfile)
 			{
 				Log.TraceInformationOnce("Enabling Profile Guided Instrumentation (PGI). Linking will take a while.");
-				Result += " -fprofile-generate";
+				Arguments.Add("-fprofile-generate");
 			}
 
 			// Unlike on other platforms, allow LTO be specified independently of PGO
@@ -654,25 +654,25 @@ namespace UnrealBuildTool
 			{
 				if((Options & LinuxToolChainOptions.EnableThinLTO) != 0)
 				{
-					Result += " -flto=thin";
+					Arguments.Add("-flto=thin");
 				}
 				else
 				{
-					Result += " -flto";
+					Arguments.Add("-flto");
 				}
 			}
 
 			if (CompileEnvironment.ShadowVariableWarningLevel != WarningLevel.Off)
 			{
-				Result += " -Wshadow" + ((CompileEnvironment.ShadowVariableWarningLevel == WarningLevel.Error) ? "" : " -Wno-error=shadow");
+				Arguments.Add("-Wshadow" + ((CompileEnvironment.ShadowVariableWarningLevel == WarningLevel.Error) ? "" : " -Wno-error=shadow"));
 			}
 
 			if (CompileEnvironment.bEnableUndefinedIdentifierWarnings)
 			{
-				Result += " -Wundef" + (CompileEnvironment.bUndefinedIdentifierWarningsAsErrors ? "" : " -Wno-error=undef");
+				Arguments.Add("-Wundef" + (CompileEnvironment.bUndefinedIdentifierWarningsAsErrors ? "" : " -Wno-error=undef"));
 			}
 
-			//Result += " -DOPERATOR_NEW_INLINE=FORCENOINLINE";
+			//Arguments.Add("-DOPERATOR_NEW_INLINE=FORCENOINLINE");
 
 			bool bRetainFramePointers = CompileEnvironment.bRetainFramePointers 
 				|| Options.HasFlag(LinuxToolChainOptions.EnableAddressSanitizer) || Options.HasFlag(LinuxToolChainOptions.EnableMemorySanitizer)
@@ -681,119 +681,117 @@ namespace UnrealBuildTool
 			// shipping builds will cause this warning with "ensure", so disable only in those case
 			if (CompileEnvironment.Configuration == CppConfiguration.Shipping)
 			{
-				Result += " -Wno-unused-value";
+				Arguments.Add("-Wno-unused-value");
 				if (!bRetainFramePointers)
 				{
-					Result += " -fomit-frame-pointer";
+					Arguments.Add("-fomit-frame-pointer");
 				}
 			}
 			// switches to help debugging
 			else if (CompileEnvironment.Configuration == CppConfiguration.Debug)
 			{
-				Result += " -fno-inline";                   // disable inlining for better debuggability (e.g. callstacks, "skip file" in gdb)
-				Result += " -fstack-protector";             // detect stack smashing
-				//Result += " -fsanitize=address";            // detect address based errors (support properly and link to libasan)
+				Arguments.Add("-fno-inline");                   // disable inlining for better debuggability (e.g. callstacks, "skip file" in gdb)
+				Arguments.Add("-fstack-protector");             // detect stack smashing
+				//Arguments.Add("-fsanitize=address");            // detect address based errors (support properly and link to libasan)
 			}
 
 			if (bRetainFramePointers)
 			{
-				Result += " -fno-optimize-sibling-calls -fno-omit-frame-pointer";
+				Arguments.Add("-fno-optimize-sibling-calls -fno-omit-frame-pointer");
 			}
 
 			// debug info
 			// bCreateDebugInfo is normally set for all configurations, including Shipping - this is needed to enable callstack in Shipping builds (proper resolution: UEPLAT-205, separate files with debug info)
 			if (CompileEnvironment.bCreateDebugInfo)
 			{
-				Result += " -gdwarf-4";
+				Arguments.Add("-gdwarf-4");
 
 				if (bGdbIndexSection)
 				{
 					// Generate .debug_pubnames and .debug_pubtypes sections in a format suitable for conversion into a
 					// GDB index. This option is only useful with a linker that can produce GDB index version 7.
-					Result += " -ggnu-pubnames";
+					Arguments.Add("-ggnu-pubnames");
 				}
 
 				if (Options.HasFlag(LinuxToolChainOptions.TuneDebugInfoForLLDB))
 				{
-					Result += " -glldb";
+					Arguments.Add("-glldb");
 				}
 			}
 
 			// optimization level
 			if (!CompileEnvironment.bOptimizeCode)
 			{
-				Result += " -O0";
+				Arguments.Add("-O0");
 			}
 			else
 			{
 				// Don't over optimise if using Address/MemorySanitizer or you'll get false positive errors due to erroneous optimisation of necessary Address/MemorySanitizer instrumentation.
 				if (Options.HasFlag(LinuxToolChainOptions.EnableAddressSanitizer) || Options.HasFlag(LinuxToolChainOptions.EnableMemorySanitizer))
 				{
-					Result += " -O1 -g";
+					Arguments.Add("-O1 -g");
 
 					// This enables __asan_default_options() in UnixCommonStartup.h which disables the leak detector
-					Result += " -DDISABLE_ASAN_LEAK_DETECTOR=1";
+					Arguments.Add("-DDISABLE_ASAN_LEAK_DETECTOR=1");
 				}
 				else if (Options.HasFlag(LinuxToolChainOptions.EnableThreadSanitizer))
 				{
-					Result += " -O1 -g";
+					Arguments.Add("-O1 -g");
 				}
 				else
 				{
-					Result += " -O3";
+					Arguments.Add("-O3");
 				}
 			}
 
 			if (!CompileEnvironment.bUseInlining)
 			{
-				Result += " -fno-inline-functions";
+				Arguments.Add("-fno-inline-functions");
 			}
 
 			if (CompileEnvironment.bIsBuildingDLL)
 			{
-				Result += " -fPIC";
+				Arguments.Add("-fPIC");
 				// Use local-dynamic TLS model. This generates less efficient runtime code for __thread variables, but avoids problems of running into
 				// glibc/ld.so limit (DTV_SURPLUS) for number of dlopen()'ed DSOs with static TLS (see e.g. https://www.cygwin.com/ml/libc-help/2013-11/msg00033.html)
-				Result += " -ftls-model=local-dynamic";
+				Arguments.Add("-ftls-model=local-dynamic");
 			}
 			else
 			{
-				Result += " -ffunction-sections";
-				Result += " -fdata-sections";
+				Arguments.Add("-ffunction-sections");
+				Arguments.Add("-fdata-sections");
 			}
 
 			if (CompileEnvironment.bEnableExceptions)
 			{
-				Result += " -fexceptions";
-				Result += " -DPLATFORM_EXCEPTIONS_DISABLED=0";
+				Arguments.Add("-fexceptions");
+				Arguments.Add("-DPLATFORM_EXCEPTIONS_DISABLED=0");
 			}
 			else
 			{
-				Result += " -fno-exceptions";               // no exceptions
-				Result += " -DPLATFORM_EXCEPTIONS_DISABLED=1";
+				Arguments.Add("-fno-exceptions");               // no exceptions
+				Arguments.Add("-DPLATFORM_EXCEPTIONS_DISABLED=1");
 			}
 
 			if (bSuppressPIE && !CompileEnvironment.bIsBuildingDLL)
 			{
-				Result += " -fno-PIE";
+				Arguments.Add("-fno-PIE");
 			}
 
 			if (PlatformSDK.bVerboseCompiler)
 			{
-				Result += " -v";                            // for better error diagnosis
+				Arguments.Add("-v");                            // for better error diagnosis
 			}
 
-			Result += ArchitectureSpecificDefines(CompileEnvironment.Architecture);
+			Arguments.Add(ArchitectureSpecificDefines(CompileEnvironment.Architecture));
 			if (CrossCompiling())
 			{
 				if (!String.IsNullOrEmpty(CompileEnvironment.Architecture))
 				{
-					Result += String.Format(" -target {0}", CompileEnvironment.Architecture);        // Set target triple
+					Arguments.Add(String.Format("-target {0}", CompileEnvironment.Architecture));        // Set target triple
 				}
-				Result += String.Format(" --sysroot=\"{0}\"", BaseLinuxPath);
+				Arguments.Add(String.Format("--sysroot=\"{0}\"", BaseLinuxPath));
 			}
-
-			return Result;
 		}
 
 		/// <inheritdoc/>
@@ -1118,8 +1116,9 @@ namespace UnrealBuildTool
 
 		public override CPPOutput CompileCPPFiles(CppCompileEnvironment CompileEnvironment, List<FileItem> InputFiles, DirectoryReference OutputDir, string ModuleName, IActionGraphBuilder Graph)
 		{
-			string Arguments = GetCLArguments_Global(CompileEnvironment);
-			string PCHArguments = "";
+			List<string> GlobalArguments = new();
+			GetCompileArguments_Global(CompileEnvironment, GlobalArguments);
+			List<string> PCHArguments = new();
 
 			//var BuildPlatform = UEBuildPlatform.GetBuildPlatform(CompileEnvironment.Platform);
 
@@ -1140,9 +1139,9 @@ namespace UnrealBuildTool
 				// Validate PCH inputs by content if mtime check fails
 				if (CompilerVersionGreaterOrEqual(11, 0, 0))
 				{
-					PCHArguments += " -fpch-validate-input-files-content";
+					PCHArguments.Add("-fpch-validate-input-files-content");
 				}
-				PCHArguments += GetForceIncludeFileArgument(CompileEnvironment.PrecompiledHeaderIncludeFilename!);
+				PCHArguments.Add(GetForceIncludeFileArgument(CompileEnvironment.PrecompiledHeaderIncludeFilename!));
 			}
 
 			// Create a compile action for each source file.
@@ -1153,39 +1152,39 @@ namespace UnrealBuildTool
 				CompileAction.PrerequisiteItems.AddRange(CompileEnvironment.ForceIncludeFiles);
 				CompileAction.PrerequisiteItems.AddRange(CompileEnvironment.AdditionalPrerequisites);
 
-				string FileArguments = "";
+				List<string> FileArguments = new();
 				string Extension = Path.GetExtension(SourceFile.AbsolutePath).ToUpperInvariant();
 
 				// Add C or C++ specific compiler arguments.
 				if (CompileEnvironment.PrecompiledHeaderAction == PrecompiledHeaderAction.Create)
 				{
-					FileArguments += GetCompileArguments_PCH(CompileEnvironment);
+					GetCompileArguments_PCH(CompileEnvironment, FileArguments);
 				}
 				else if (Extension == ".C")
 				{
 					// Compile the file as C code.
-					FileArguments += GetCompileArguments_C();
+					GetCompileArguments_C(CompileEnvironment, FileArguments);
 				}
 				else if (Extension == ".MM")
 				{
 					// Compile the file as Objective-C++ code.
-					FileArguments += GetCompileArguments_MM(CompileEnvironment);
-					FileArguments += GetRTTIFlag(CompileEnvironment);
+					GetCompileArguments_MM(CompileEnvironment, FileArguments);
+					FileArguments.Add(GetRTTIFlag(CompileEnvironment));
 				}
 				else if (Extension == ".M")
 				{
 					// Compile the file as Objective-C code.
-					FileArguments += GetCompileArguments_M(CompileEnvironment);
+					GetCompileArguments_M(CompileEnvironment, FileArguments);
 				}
 				else
 				{
-					FileArguments += GetCompileArguments_CPP(CompileEnvironment);
+					GetCompileArguments_CPP(CompileEnvironment, FileArguments);
 
 					// only use PCH for .cpp files
-					FileArguments += PCHArguments;
+					FileArguments.AddRange(PCHArguments);
 				}
 
-				FileArguments += GetForceIncludeFileArguments(CompileEnvironment);
+				GetCompileArguments_ForceInclude(CompileEnvironment, FileArguments);
 
 				// Add the C++ source file and its included files to the prerequisite item list.
 				CompileAction.PrerequisiteItems.Add(SourceFile);
@@ -1199,7 +1198,7 @@ namespace UnrealBuildTool
 					Result.PrecompiledHeaderFile = PrecompiledHeaderFile;
 
 					// Add the parameters needed to compile the precompiled header file to the command-line.
-					FileArguments += string.Format(" -o \"{0}\"", NormalizeCommandLinePath(PrecompiledHeaderFile));
+					FileArguments.Add(string.Format("-o \"{0}\"", NormalizeCommandLinePath(PrecompiledHeaderFile)));
 				}
 				else
 				{
@@ -1213,17 +1212,17 @@ namespace UnrealBuildTool
 					CompileAction.ProducedItems.Add(ObjectFile);
 					Result.ObjectFiles.Add(ObjectFile);
 
-					FileArguments += string.Format(" -o \"{0}\"", NormalizeCommandLinePath(ObjectFile));
+					FileArguments.Add(string.Format("-o \"{0}\"", NormalizeCommandLinePath(ObjectFile)));
 				}
 
 				// Add the source file path to the command-line.
-				FileArguments += string.Format(" \"{0}\"", NormalizeCommandLinePath(SourceFile));
+				FileArguments.Add(string.Format("\"{0}\"", NormalizeCommandLinePath(SourceFile)));
 
 				// Generate the timing info
 				if (CompileEnvironment.bPrintTimingInfo)
 				{
 					FileItem TraceFile = FileItem.GetItemByFileReference(FileReference.Combine(OutputDir, Path.GetFileName(SourceFile.AbsolutePath) + ".json"));
-					FileArguments += " -ftime-trace";
+					FileArguments.Add("-ftime-trace");
 					CompileAction.ProducedItems.Add(TraceFile);
 				}
 
@@ -1231,7 +1230,7 @@ namespace UnrealBuildTool
 				if(!bPreprocessDepends && CompileEnvironment.bGenerateDependenciesFile)
 				{
 					FileItem DependencyListFile = FileItem.GetItemByFileReference(FileReference.Combine(OutputDir, Path.GetFileName(SourceFile.AbsolutePath) + ".d"));
-					FileArguments += string.Format(" -MD -MF\"{0}\"", NormalizeCommandLinePath(DependencyListFile));
+					FileArguments.Add(string.Format("-MD -MF\"{0}\"", NormalizeCommandLinePath(DependencyListFile)));
 					CompileAction.DependencyListFile = DependencyListFile;
 					CompileAction.ProducedItems.Add(DependencyListFile);
 				}
@@ -1239,13 +1238,17 @@ namespace UnrealBuildTool
 				CompileAction.WorkingDirectory = UnrealBuildTool.EngineSourceDirectory;
 				CompileAction.CommandPath = new FileReference(ClangPath!);
 
-				string AllArguments = (Arguments + FileArguments + CompileEnvironment.AdditionalArguments);
+				List<string> ResponseFileContents = new();
+				ResponseFileContents.AddRange(GlobalArguments);
+				ResponseFileContents.AddRange(FileArguments);
+				ResponseFileContents.Add(CompileEnvironment.AdditionalArguments);
 				// all response lines should have / instead of \, but we cannot just bulk-replace it here since some \ are used to escape quotes, e.g. Definitions.Add("FOO=TEXT(\"Bar\")");
 
+
 				Debug.Assert(CompileAction.ProducedItems.Count > 0);
 
 				FileReference CompilerResponseFileName = CompileAction.ProducedItems[0].Location + ".rsp";
-				FileItem CompilerResponseFileItem = Graph.CreateIntermediateTextFile(CompilerResponseFileName, AllArguments);
+				FileItem CompilerResponseFileItem = Graph.CreateIntermediateTextFile(CompilerResponseFileName, ResponseFileContents);
 
 				CompileAction.CommandArguments = string.Format(" @\"{0}\"", NormalizeCommandLinePath(CompilerResponseFileName));
 				CompileAction.PrerequisiteItems.Add(CompilerResponseFileItem);
@@ -1273,28 +1276,35 @@ namespace UnrealBuildTool
 					PrepassAction.CommandVersion = CompileAction.CommandVersion;
 					PrepassAction.WorkingDirectory = CompileAction.WorkingDirectory;
 
-					string PreprocessArguments = Arguments.ToString();
-					string PreprocessFileArguments = FileArguments;
-					PreprocessArguments = PreprocessArguments.Replace(" -c ", " ");
+					List<string> PreprocessGlobalArguments = new(GlobalArguments);
+					List<string> PreprocessFileArguments = new(FileArguments);
+					PreprocessGlobalArguments.Remove("-c");
 
 					FileItem DependencyListFile = FileItem.GetItemByFileReference(FileReference.Combine(OutputDir, Path.GetFileName(SourceFile.AbsolutePath) + ".d"));
-					PreprocessFileArguments += string.Format(" -M -MF\"{0}\"", NormalizeCommandLinePath(DependencyListFile));
+					PreprocessFileArguments.Add(string.Format("-M -MF\"{0}\"", NormalizeCommandLinePath(DependencyListFile)));
 					PrepassAction.DependencyListFile = DependencyListFile;
 					PrepassAction.ProducedItems.Add(DependencyListFile);
 
-					PreprocessFileArguments = PreprocessFileArguments.Replace(" -ftime-trace", string.Empty);
-					PreprocessFileArguments = PreprocessFileArguments.Replace(string.Format(" -o \"{0}\"", NormalizeCommandLinePath(CompileAction.ProducedItems.First())), String.Empty);
+					PreprocessFileArguments.Remove("-ftime-trace");
+					PreprocessFileArguments.Remove(string.Format(" -o \"{0}\"", NormalizeCommandLinePath(CompileAction.ProducedItems.First())));
 
 					PrepassAction.DeleteItems.AddRange(PrepassAction.ProducedItems);
 
-					string AllPreprocessArgs = PreprocessArguments + PreprocessFileArguments + CompileEnvironment.AdditionalArguments;
+					// Gets the target file so we can get the correct output path.
+					FileItem PreprocessTargetFile = PrepassAction.ProducedItems[0];
 
-					FileReference CompilerPrepassResponseFileName = PrepassAction.ProducedItems[0].Location + ".rsp";
-					FileItem CompilerPrepassResponseFileItem = Graph.CreateIntermediateTextFile(CompilerPrepassResponseFileName, AllPreprocessArgs);
+					// Creates the path to the response file using the name of the output file and creates its contents.
+					FileReference PreprocessResponseFileName = new FileReference(PreprocessTargetFile.AbsolutePath + ".rsp");
+					List<string> PreprocessResponseFileContents = new();
+					PreprocessResponseFileContents.AddRange(PreprocessGlobalArguments);
+					PreprocessResponseFileContents.AddRange(PreprocessFileArguments);
+					PreprocessResponseFileContents.Add(CompileEnvironment.AdditionalArguments);
 
-					PrepassAction.CommandArguments = string.Format(" @\"{0}\"", NormalizeCommandLinePath(CompilerPrepassResponseFileName));
-					PrepassAction.PrerequisiteItems.Add(CompilerResponseFileItem);
+					// Adds the response file to the compiler input.
+					FileItem PreprocessResponseFileItem = Graph.CreateIntermediateTextFile(PreprocessResponseFileName, PreprocessResponseFileContents);
+					PrepassAction.PrerequisiteItems.Add(PreprocessResponseFileItem);
 
+					PrepassAction.CommandArguments = string.Format("@\"{0}\"", PreprocessResponseFileItem);
 					CompileAction.DependencyListFile = DependencyListFile;
 					CompileAction.PrerequisiteItems.Add(DependencyListFile);
 				}
diff --git a/Engine/Source/Programs/UnrealBuildTool/Platform/Mac/MacToolChain.cs b/Engine/Source/Programs/UnrealBuildTool/Platform/Mac/MacToolChain.cs
--- a/Engine/Source/Programs/UnrealBuildTool/Platform/Mac/MacToolChain.cs
+++ b/Engine/Source/Programs/UnrealBuildTool/Platform/Mac/MacToolChain.cs
@@ -191,65 +191,65 @@ namespace UnrealBuildTool
 			return ArchArg;
 		}
 
-		protected override string GetCompileArguments_Global(CppCompileEnvironment CompileEnvironment)
+		protected override void GetCompileArguments_Global(CppCompileEnvironment CompileEnvironment, List<string> Arguments)
 		{
-			string Result = base.GetCompileArguments_Global(CompileEnvironment);
+			base.GetCompileArguments_Global(CompileEnvironment, Arguments);
 
 
-			Result += " -fexceptions";
-			Result += " -DPLATFORM_EXCEPTIONS_DISABLED=0";
+			Arguments.Add("-fexceptions");
+			Arguments.Add("-DPLATFORM_EXCEPTIONS_DISABLED=0");
 
-			Result += " -fasm-blocks";
+			Arguments.Add("-fasm-blocks");
 
 			if(CompileEnvironment.bHideSymbolsByDefault)
 			{
-				Result += " -fvisibility-ms-compat";
-				Result += " -fvisibility-inlines-hidden";
+				Arguments.Add("-fvisibility-ms-compat");
+				Arguments.Add("-fvisibility-inlines-hidden");
 			}
 			if (Options.HasFlag(MacToolChainOptions.EnableAddressSanitizer))
 			{
-				Result += " -fsanitize=address";
+				Arguments.Add("-fsanitize=address");
 			}
 			if (Options.HasFlag(MacToolChainOptions.EnableThreadSanitizer))
 			{
-				Result += " -fsanitize=thread";
+				Arguments.Add("-fsanitize=thread");
 			}
 			if (Options.HasFlag(MacToolChainOptions.EnableUndefinedBehaviorSanitizer))
 			{
-				Result += " -fsanitize=undefined";
+				Arguments.Add("-fsanitize=undefined");
 			}			
 
-			Result += " -Wall -Werror";
-			Result += " -Wdelete-non-virtual-dtor";
+			Arguments.Add("-Wall -Werror");
+			Arguments.Add("-Wdelete-non-virtual-dtor");
 
 			// clang 12.00 has a new warning for copies in ranged loops. Instances have all been fixed up (2020/6/26) but
 			// are likely to be reintroduced due to no equivalent on other platforms at this time so disable the warning
 			if (GetClangVersion().Major >= 12)
 			{
-				Result += " -Wno-range-loop-analysis ";
+				Arguments.Add("-Wno-range-loop-analysis ");
 			}			
 
-			//Result += " -Wsign-compare"; // fed up of not seeing the signed/unsigned warnings we get on Windows - lets enable them here too.
+			//Arguments.Add("-Wsign-compare"); // fed up of not seeing the signed/unsigned warnings we get on Windows - lets enable them here too.
 
 			if (CompileEnvironment.ShadowVariableWarningLevel != WarningLevel.Off)
 			{
-				Result += " -Wshadow" + ((CompileEnvironment.ShadowVariableWarningLevel == WarningLevel.Error) ? "" : " -Wno-error=shadow");
+				Arguments.Add("-Wshadow" + ((CompileEnvironment.ShadowVariableWarningLevel == WarningLevel.Error) ? "" : " -Wno-error=shadow"));
 			}
 			
 			if (CompileEnvironment.bEnableUndefinedIdentifierWarnings)
 			{
-				Result += " -Wundef" + (CompileEnvironment.bUndefinedIdentifierWarningsAsErrors ? "" : " -Wno-error=undef");
+				Arguments.Add("-Wundef" + (CompileEnvironment.bUndefinedIdentifierWarningsAsErrors ? "" : " -Wno-error=undef"));
 			}
 
 			if (CompileEnvironment.bEnableOSX109Support)
 			{
-				Result += " -faligned-new"; // aligned operator new is supported only on macOS 10.14 and above
+				Arguments.Add("-faligned-new"); // aligned operator new is supported only on macOS 10.14 and above
 			}
 
 			// Pass through architecture and OS info
-			Result += " " + FormatArchitectureArg(CompileEnvironment.Architecture);	
-			Result += string.Format(" -isysroot \"{0}\"", SDKPath);
-			Result += " -mmacosx-version-min=" + (CompileEnvironment.bEnableOSX109Support ? "10.9" : Settings.MacOSVersion);
+			Arguments.Add("" + FormatArchitectureArg(CompileEnvironment.Architecture));
+			Arguments.Add(string.Format(" -isysroot \"{0}\"", SDKPath));
+			Arguments.Add("-mmacosx-version-min=" + (CompileEnvironment.bEnableOSX109Support ? "10.9" : Settings.MacOSVersion));
 
 			bool bStaticAnalysis = false;
 			string? StaticAnalysisMode = Environment.GetEnvironmentVariable("CLANG_STATIC_ANALYZER_MODE");
@@ -264,38 +264,36 @@ namespace UnrealBuildTool
 				// Don't over optimise if using AddressSanitizer or you'll get false positive errors due to erroneous optimisation of necessary AddressSanitizer instrumentation.
 				if (Options.HasFlag(MacToolChainOptions.EnableAddressSanitizer))
 				{
-					Result += " -O1 -g -fno-optimize-sibling-calls -fno-omit-frame-pointer";
+					Arguments.Add("-O1 -g -fno-optimize-sibling-calls -fno-omit-frame-pointer");
 				}
 				else if (Options.HasFlag(MacToolChainOptions.EnableThreadSanitizer))
 				{
-					Result += " -O1 -g";
+					Arguments.Add("-O1 -g");
 				}
 				else if (CompileEnvironment.bOptimizeForSize)
 				{
-					Result += " -Oz";
+					Arguments.Add("-Oz");
 				}
 				else
 				{
-					Result += " -O3";
+					Arguments.Add("-O3");
 				}
 			}
 			else
 			{
-				Result += " -O0";
+				Arguments.Add("-O0");
 			}
 
 			if (!CompileEnvironment.bUseInlining)
 			{
-				Result += " -fno-inline-functions";
+				Arguments.Add("-fno-inline-functions");
 			}
 
 			// Create DWARF format debug info if wanted,
 			if (CompileEnvironment.bCreateDebugInfo)
 			{
-				Result += " -gdwarf-2";
+				Arguments.Add("-gdwarf-2");
 			}
-
-			return Result;
 		}
 		
 		string AddFrameworkToLinkCommand(string FrameworkName, string Arg = "-framework")
@@ -366,19 +364,19 @@ namespace UnrealBuildTool
 
 		public override CPPOutput CompileCPPFiles(CppCompileEnvironment CompileEnvironment, List<FileItem> InputFiles, DirectoryReference OutputDir, string ModuleName, IActionGraphBuilder Graph)
 		{
-			StringBuilder Arguments = new StringBuilder();
-			StringBuilder PCHArguments = new StringBuilder();
+			List<string> GlobalArguments = new();
 
-			Arguments.Append(GetCompileArguments_Global(CompileEnvironment));
+			GetCompileArguments_Global(CompileEnvironment, GlobalArguments);
+			List<string> PCHArguments = new();
 
 			if (CompileEnvironment.PrecompiledHeaderAction == PrecompiledHeaderAction.Include)
 			{
 				// Add the precompiled header file's path to the include path so GCC can find it.
 				// This needs to be before the other include paths to ensure GCC uses it instead of the source header file.
-				PCHArguments.Append(GetForceIncludeFileArgument(CompileEnvironment.PrecompiledHeaderIncludeFilename!));
+				PCHArguments.Add(GetForceIncludeFileArgument(CompileEnvironment.PrecompiledHeaderIncludeFilename!));
 				if (GetClangVersion().Major >= 11)
 				{
-					PCHArguments.Append(" -fpch-validate-input-files-content");
+					PCHArguments.Add("-fpch-validate-input-files-content");
 				}
 			}
 
@@ -388,9 +386,7 @@ namespace UnrealBuildTool
 				FileReference FrameworkPath = new FileReference(Path.GetFullPath(Framework.Name));
 				if (!FrameworksSearchPaths.Contains(FrameworkPath.Directory.FullName))
 				{
-					Arguments.Append(" -F \"");
-					Arguments.Append(NormalizeCommandLinePath(FrameworkPath.Directory));
-					Arguments.Append("\"");
+					GlobalArguments.Add($"-F \"{NormalizeCommandLinePath(FrameworkPath.Directory)}\"");
 					FrameworksSearchPaths.Add(FrameworkPath.Directory.FullName);
 				}
 			}
@@ -403,42 +399,42 @@ namespace UnrealBuildTool
 				CompileAction.PrerequisiteItems.AddRange(CompileEnvironment.ForceIncludeFiles);
 				CompileAction.PrerequisiteItems.AddRange(CompileEnvironment.AdditionalPrerequisites);
 
-				string FileArguments = "";
+				List<string> FileArguments = new();
 				string Extension = Path.GetExtension(SourceFile.AbsolutePath).ToUpperInvariant();
 
 				if (CompileEnvironment.PrecompiledHeaderAction == PrecompiledHeaderAction.Create)
 				{
 					// Compile the file as a C++ PCH.
-					FileArguments += GetCompileArguments_PCH(CompileEnvironment);
-					FileArguments += GetRTTIFlag(CompileEnvironment);
+					GetCompileArguments_PCH(CompileEnvironment, FileArguments);
+					FileArguments.Add(GetRTTIFlag(CompileEnvironment));
 				}
 				else if (Extension == ".C")
 				{
 					// Compile the file as C code.
-					FileArguments += GetCompileArguments_C();
+					GetCompileArguments_C(CompileEnvironment, FileArguments);
 				}
 				else if (Extension == ".MM")
 				{
 					// Compile the file as Objective-C++ code.
-					FileArguments += GetCompileArguments_MM(CompileEnvironment);
-					FileArguments += GetRTTIFlag(CompileEnvironment);
+					GetCompileArguments_MM(CompileEnvironment, FileArguments);
+					FileArguments.Add(GetRTTIFlag(CompileEnvironment));
 				}
 				else if (Extension == ".M")
 				{
 					// Compile the file as Objective-C++ code.
-					FileArguments += GetCompileArguments_M(CompileEnvironment);
+					GetCompileArguments_M(CompileEnvironment, FileArguments);
 				}
 				else
 				{
 					// Compile the file as C++ code.
-					FileArguments += GetCompileArguments_CPP(CompileEnvironment);
-					FileArguments += GetRTTIFlag(CompileEnvironment);
+					GetCompileArguments_CPP(CompileEnvironment, FileArguments);
+					FileArguments.Add(GetRTTIFlag(CompileEnvironment));
 
 					// only use PCH for .cpp files
-					FileArguments += PCHArguments.ToString();
+					FileArguments.AddRange(PCHArguments);
 				}
 
-				FileArguments += GetForceIncludeFileArguments(CompileEnvironment);
+				GetCompileArguments_ForceInclude(CompileEnvironment, FileArguments);
 
 				// Add the C++ source file and its included files to the prerequisite item list.
 				CompileAction.PrerequisiteItems.Add(SourceFile);
@@ -452,7 +448,7 @@ namespace UnrealBuildTool
 					Result.PrecompiledHeaderFile = PrecompiledHeaderFile;
 
 					// Add the parameters needed to compile the precompiled header file to the command-line.
-					FileArguments += string.Format(" -o \"{0}\"", NormalizeCommandLinePath(PrecompiledHeaderFile));
+					FileArguments.Add(string.Format("-o \"{0}\"", NormalizeCommandLinePath(PrecompiledHeaderFile)));
 				}
 				else
 				{
@@ -466,18 +462,18 @@ namespace UnrealBuildTool
 
 					CompileAction.ProducedItems.Add(ObjectFile);
 					Result.ObjectFiles.Add(ObjectFile);
-					FileArguments += string.Format(" -o \"{0}\"", NormalizeCommandLinePath(ObjectFile));
+					FileArguments.Add(string.Format("-o \"{0}\"", NormalizeCommandLinePath(ObjectFile)));
 					OutputFilePath = ObjectFile.AbsolutePath;
 				}
 
 				// Add the source file path to the command-line.
-				FileArguments += string.Format(" \"{0}\"", NormalizeCommandLinePath(SourceFile));
+				FileArguments.Add(string.Format("\"{0}\"", NormalizeCommandLinePath(SourceFile)));
 
 				// Generate the timing info
 				if (CompileEnvironment.bPrintTimingInfo)
 				{
 					FileItem TraceFile = FileItem.GetItemByFileReference(FileReference.Combine(OutputDir, Path.GetFileName(SourceFile.AbsolutePath) + ".json"));
-					FileArguments += " -ftime-trace";
+					FileArguments.Add("-ftime-trace");
 					CompileAction.ProducedItems.Add(TraceFile);
 				}
 
@@ -485,7 +481,7 @@ namespace UnrealBuildTool
 				if(!bPreprocessDepends && CompileEnvironment.bGenerateDependenciesFile)
 				{
 					FileItem DependencyListFile = FileItem.GetItemByFileReference(FileReference.Combine(OutputDir, Path.GetFileName(SourceFile.AbsolutePath) + ".d"));
-					FileArguments += string.Format(" -MD -MF\"{0}\"", NormalizeCommandLinePath(DependencyListFile));
+					FileArguments.Add(string.Format("-MD -MF\"{0}\"", NormalizeCommandLinePath(DependencyListFile)));
 					CompileAction.DependencyListFile = DependencyListFile;
 					CompileAction.ProducedItems.Add(DependencyListFile);
 				}
@@ -507,22 +503,34 @@ namespace UnrealBuildTool
 					}
 				}
 
-				string AllArgs = Arguments + FileArguments + EscapedAdditionalArgs;
+				// Gets the target file so we can get the correct output path.
+				FileItem TargetFile = CompileAction.ProducedItems[0];
+
+				// Creates the path to the response file using the name of the output file and creates its contents.
+				FileReference ResponseFileName = new FileReference(TargetFile.AbsolutePath + ".response");
+				List<string> ResponseFileContents = new();
+				ResponseFileContents.AddRange(GlobalArguments);
+				ResponseFileContents.AddRange(FileArguments);
+				ResponseFileContents.Add(EscapedAdditionalArgs);
+
+				// Adds the response file to the compiler input.
+				FileItem ResponseFileItem = Graph.CreateIntermediateTextFile(ResponseFileName, ResponseFileContents);
+				CompileAction.CommandArguments += string.Format("@\"{0}\"", NormalizeCommandLinePath(ResponseFileItem));
+				CompileAction.PrerequisiteItems.Add(ResponseFileItem);
 
 				string CompilerPath = Settings.ToolchainDir + MacCompiler;
 				
 				// Analyze and then compile using the shell to perform the indirection
 				string? StaticAnalysisMode = Environment.GetEnvironmentVariable("CLANG_STATIC_ANALYZER_MODE");
-				if(StaticAnalysisMode != null && StaticAnalysisMode != "" && OutputFilePath != null)
+				if (StaticAnalysisMode != null && StaticAnalysisMode != "" && OutputFilePath != null)
 				{
-					string TempArgs = "-c \"" + CompilerPath + " " + AllArgs + " --analyze -Wno-unused-command-line-argument -Xclang -analyzer-output=html -Xclang -analyzer-config -Xclang path-diagnostics-alternate=true -Xclang -analyzer-config -Xclang report-in-main-source-file=true -Xclang -analyzer-disable-checker -Xclang deadcode.DeadStores -o " + OutputFilePath.Replace(".o", ".html") + "; " + CompilerPath + " " + AllArgs + "\"";
-					AllArgs = TempArgs;
+					CompileAction.CommandArguments = "-c \"" + CompilerPath + " " + string.Format("@\"{0}\"", ResponseFileName) + " --analyze -Wno-unused-command-line-argument -Xclang -analyzer-output=html -Xclang -analyzer-config -Xclang path-diagnostics-alternate=true -Xclang -analyzer-config -Xclang report-in-main-source-file=true -Xclang -analyzer-disable-checker -Xclang deadcode.DeadStores -o " + OutputFilePath.Replace(".o", ".html") + "; " + CompilerPath + " " + string.Format("@\"{0}\"", ResponseFileName) + "\"";
 					CompilerPath = "/bin/sh";
 				}
 
+				// RPC utility parameters are in terms of the Mac side
 				CompileAction.WorkingDirectory = GetMacDevSrcRoot();
 				CompileAction.CommandPath = new FileReference(CompilerPath);
-				CompileAction.CommandArguments = AllArgs;
 
 				// For compilation we delete everything we produce
 				CompileAction.DeleteItems.AddRange(CompileAction.ProducedItems);
@@ -546,24 +554,36 @@ namespace UnrealBuildTool
 					PrepassAction.CommandVersion = CompileAction.CommandVersion;
 					PrepassAction.WorkingDirectory = CompileAction.WorkingDirectory;
 
-					string PreprocessArguments = Arguments.ToString();
-					string PreprocessFileArguments = FileArguments;
-					PreprocessArguments = PreprocessArguments.Replace(" -c ", " ");
+					List<string> PreprocessGlobalArguments = new(GlobalArguments);
+					List<string> PreprocessFileArguments = new(FileArguments);
+					PreprocessGlobalArguments.Remove("-c");
 
 					FileItem DependencyListFile = FileItem.GetItemByFileReference(FileReference.Combine(OutputDir, Path.GetFileName(SourceFile.AbsolutePath) + ".d"));
-					PreprocessFileArguments += string.Format(" -M -MF\"{0}\"", NormalizeCommandLinePath(DependencyListFile));
+					PreprocessFileArguments.Add(string.Format("-M -MF\"{0}\"", NormalizeCommandLinePath(DependencyListFile)));
 					PrepassAction.DependencyListFile = DependencyListFile;
 					PrepassAction.ProducedItems.Add(DependencyListFile);
 
-					PreprocessFileArguments = PreprocessFileArguments.Replace(" -ftime-trace", string.Empty);
-					PreprocessFileArguments = PreprocessFileArguments.Replace(string.Format(" -o \"{0}\"", NormalizeCommandLinePath(CompileAction.ProducedItems.First())), String.Empty);
+					PreprocessFileArguments.Remove("-ftime-trace");
+					PreprocessFileArguments.Remove(string.Format(" -o \"{0}\"", NormalizeCommandLinePath(CompileAction.ProducedItems.First())));
 
 					PrepassAction.DeleteItems.AddRange(PrepassAction.ProducedItems);
 
-					string AllPreprocessArgs = PreprocessArguments + PreprocessFileArguments + EscapedAdditionalArgs;
+					// Gets the target file so we can get the correct output path.
+					FileItem PreprocessTargetFile = PrepassAction.ProducedItems[0];
+
+					// Creates the path to the response file using the name of the output file and creates its contents.
+					FileReference PreprocessResponseFileName = new FileReference(PreprocessTargetFile.AbsolutePath + ".response");
+					List<string> PreprocessResponseFileContents = new();
+					PreprocessResponseFileContents.AddRange(PreprocessGlobalArguments);
+					PreprocessResponseFileContents.AddRange(PreprocessFileArguments);
+					PreprocessResponseFileContents.Add(EscapedAdditionalArgs);
+
+					// Adds the response file to the compiler input.
+					FileItem PreprocessResponseFileItem = Graph.CreateIntermediateTextFile(PreprocessResponseFileName, PreprocessResponseFileContents);
+					PrepassAction.PrerequisiteItems.Add(PreprocessResponseFileItem);
 
 					PrepassAction.CommandPath = new FileReference(CompilerPath);
-					PrepassAction.CommandArguments = AllPreprocessArgs;
+					PrepassAction.CommandArguments = string.Format("@\"{0}\"", PreprocessResponseFileItem);
 					CompileAction.DependencyListFile = DependencyListFile;
 					CompileAction.PrerequisiteItems.Add(DependencyListFile);
 				}
diff --git a/Engine/Source/Programs/UnrealBuildTool/ToolChain/AppleToolChain.cs b/Engine/Source/Programs/UnrealBuildTool/ToolChain/AppleToolChain.cs
--- a/Engine/Source/Programs/UnrealBuildTool/ToolChain/AppleToolChain.cs
+++ b/Engine/Source/Programs/UnrealBuildTool/ToolChain/AppleToolChain.cs
@@ -206,61 +206,49 @@ namespace UnrealBuildTool
 			return Definition.Contains("\"") ? Definition.Replace("\"", "\\\"") : Definition;
 		}
 
-		protected override string GetCompileArguments_CPP(CppCompileEnvironment CompileEnvironment)
+		protected override void GetCompileArguments_CPP(CppCompileEnvironment CompileEnvironment, List<string> Arguments)
 		{
-			string Result = "";
-			Result += " -x objective-c++";
-			Result += GetCppStandardCompileArgument(CompileEnvironment);
-			Result += " -stdlib=libc++";
-
-			return Result;
+			Arguments.Add("-x objective-c++");
+			GetCppStandardCompileArgument(CompileEnvironment, Arguments);
+			Arguments.Add("-stdlib=libc++");
 		}
 
-		protected override string GetCompileArguments_MM(CppCompileEnvironment CompileEnvironment)
+		protected override void GetCompileArguments_MM(CppCompileEnvironment CompileEnvironment, List<string> Arguments)
 		{
-			string Result = "";
-			Result += " -x objective-c++";
-			Result += GetCppStandardCompileArgument(CompileEnvironment);
-			Result += " -stdlib=libc++";
-			return Result;
+			Arguments.Add("-x objective-c++");
+			GetCppStandardCompileArgument(CompileEnvironment, Arguments);
+			Arguments.Add("-stdlib=libc++");
 		}
 
-		protected override string GetCompileArguments_M(CppCompileEnvironment CompileEnvironment)
+		protected override void GetCompileArguments_M(CppCompileEnvironment CompileEnvironment, List<string> Arguments)
 		{
-			string Result = "";
-			Result += " -x objective-c";
-			Result += GetCppStandardCompileArgument(CompileEnvironment);
-			Result += " -stdlib=libc++";
-			return Result;
+			Arguments.Add("-x objective-c");
+			GetCppStandardCompileArgument(CompileEnvironment, Arguments);
+			Arguments.Add("-stdlib=libc++");
 		}
 
-		protected override string GetCompileArguments_PCH(CppCompileEnvironment CompileEnvironment)
+		protected override void GetCompileArguments_PCH(CppCompileEnvironment CompileEnvironment, List<string> Arguments)
 		{
-			string Result = "";
-			Result += " -x objective-c++-header";
-			Result += GetCppStandardCompileArgument(CompileEnvironment);
-			Result += " -stdlib=libc++";
+			Arguments.Add("-x objective-c++-header");
+			GetCppStandardCompileArgument(CompileEnvironment, Arguments);
+			Arguments.Add("-stdlib=libc++");
 
 			if (GetClangVersion().Major >= 11)
 			{
-				Result += " -fpch-validate-input-files-content";
+				Arguments.Add("-fpch-validate-input-files-content");
 			}
 			if (GetClangVersion().Major >= 13) // Note this is supported for >=11 on other clang platforms
 			{
-				Result += " -fpch-instantiate-templates";
+				Arguments.Add("-fpch-instantiate-templates");
 			}
-
-			return Result;
 		}
 
-		protected override string GetCompileArguments_Global(CppCompileEnvironment CompileEnvironment)
+		protected override void GetCompileArguments_Global(CppCompileEnvironment CompileEnvironment, List<string> Arguments)
 		{
-			string Result = base.GetCompileArguments_Global(CompileEnvironment);
-
-			Result += " -fmessage-length=0";
-			Result += " -fpascal-strings";
+			base.GetCompileArguments_Global(CompileEnvironment, Arguments);
 
-			return Result;
+			Arguments.Add("-fmessage-length=0");
+			Arguments.Add("-fpascal-strings");
 		}
 
 		protected string GetDsymutilPath(out string ExtraOptions, bool bIsForLTOBuild=false)
diff --git a/Engine/Source/Programs/UnrealBuildTool/ToolChain/ClangToolChain.cs b/Engine/Source/Programs/UnrealBuildTool/ToolChain/ClangToolChain.cs
--- a/Engine/Source/Programs/UnrealBuildTool/ToolChain/ClangToolChain.cs
+++ b/Engine/Source/Programs/UnrealBuildTool/ToolChain/ClangToolChain.cs
@@ -2,6 +2,7 @@
 
 using EpicGames.Core;
 using System;
+using System.Collections.Generic;
 using System.Linq;
 using UnrealBuildBase;
 
@@ -88,20 +89,19 @@ namespace UnrealBuildTool
 				(ClangVersionMajor == Major && ClangVersionMinor == Minor && ClangVersionPatch < Patch);
 		}
 
-		protected virtual string GetCppStandardCompileArgument(CppCompileEnvironment CompileEnvironment)
+		protected virtual void GetCppStandardCompileArgument(CppCompileEnvironment CompileEnvironment, List<string> Arguments)
 		{
-			string Result;
 			switch (CompileEnvironment.CppStandard)
 			{
 				case CppStandardVersion.Cpp14:
-					Result = " -std=c++14";
+					Arguments.Add("-std=c++14");
 					break;
 				case CppStandardVersion.Latest:
 				case CppStandardVersion.Cpp17:
-					Result = " -std=c++17";
+					Arguments.Add("-std=c++17");
 					break;
 				case CppStandardVersion.Cpp20:
-					Result = " -std=c++20";
+					Arguments.Add("-std=c++20");
 					break;
 				default:
 					throw new BuildException($"Unsupported C++ standard type set: {CompileEnvironment.CppStandard}");
@@ -109,62 +109,50 @@ namespace UnrealBuildTool
 
 			if (CompileEnvironment.bEnableCoroutines)
 			{
-				Result += " -fcoroutines-ts";
+				Arguments.Add("-fcoroutines-ts");
 				if (!CompileEnvironment.bEnableExceptions)
 				{
-					Result += " -Wno-coroutine-missing-unhandled-exception";
+					Arguments.Add("-Wno-coroutine-missing-unhandled-exception");
 				}
 			}
-
-			return Result;
 		}
 
-		protected virtual string GetCompileArguments_CPP(CppCompileEnvironment CompileEnvironment)
+		protected virtual void GetCompileArguments_CPP(CppCompileEnvironment CompileEnvironment, List<string> Arguments)
 		{
-			string Result = "";
-			Result += " -x c++";
-			Result += GetCppStandardCompileArgument(CompileEnvironment);
-			return Result;
+			Arguments.Add("-x c++");
+			GetCppStandardCompileArgument(CompileEnvironment, Arguments);
 		}
 
-		protected virtual string GetCompileArguments_C()
+		protected virtual void GetCompileArguments_C(CppCompileEnvironment CompileEnvironment, List<string> Arguments)
 		{
-			string Result = "";
-			Result += " -x c";
-			return Result;
+			Arguments.Add("-x c");
 		}
 
-		protected virtual string GetCompileArguments_MM(CppCompileEnvironment CompileEnvironment)
+		protected virtual void GetCompileArguments_MM(CppCompileEnvironment CompileEnvironment, List<string> Arguments)
 		{
-			string Result = "";
-			Result += " -x objective-c++";
-			Result += " -fobjc-abi-version=2";
-			Result += " -fobjc-legacy-dispatch";
-			Result += GetCppStandardCompileArgument(CompileEnvironment);
-			return Result;
+			Arguments.Add("-x objective-c++");
+			Arguments.Add("-fobjc-abi-version=2");
+			Arguments.Add("-fobjc-legacy-dispatch");
+			GetCppStandardCompileArgument(CompileEnvironment, Arguments);
 		}
 
-		protected virtual string GetCompileArguments_M(CppCompileEnvironment CompileEnvironment)
+		protected virtual void GetCompileArguments_M(CppCompileEnvironment CompileEnvironment, List<string> Arguments)
 		{
-			string Result = "";
-			Result += " -x objective-c";
-			Result += " -fobjc-abi-version=2";
-			Result += " -fobjc-legacy-dispatch";
-			Result += GetCppStandardCompileArgument(CompileEnvironment);
-			return Result;
+			Arguments.Add("-x objective-c");
+			Arguments.Add("-fobjc-abi-version=2");
+			Arguments.Add("-fobjc-legacy-dispatch");
+			GetCppStandardCompileArgument(CompileEnvironment, Arguments);
 		}
 
-		protected virtual string GetCompileArguments_PCH(CppCompileEnvironment CompileEnvironment)
+		protected virtual void GetCompileArguments_PCH(CppCompileEnvironment CompileEnvironment, List<string> Arguments)
 		{
-			string Result = "";
-			Result += " -x c++-header";
+			Arguments.Add("-x c++-header");
 			if (CompilerVersionGreaterOrEqual(11, 0, 0))
 			{
-				Result += " -fpch-validate-input-files-content";
-				Result += " -fpch-instantiate-templates";
+				Arguments.Add("-fpch-validate-input-files-content");
+				Arguments.Add("-fpch-instantiate-templates");
 			}
-			Result += GetCppStandardCompileArgument(CompileEnvironment);
-			return Result;
+			GetCppStandardCompileArgument(CompileEnvironment, Arguments);
 		}
 
 		// Conditionally enable (default disabled) generation of information about every class with virtual functions for use by the C++ runtime type identification features
@@ -172,52 +160,39 @@ namespace UnrealBuildTool
 		// Note that exception handling uses the same information, but it will generate it as needed.
 		protected virtual string GetRTTIFlag(CppCompileEnvironment CompileEnvironment)
 		{
-			return CompileEnvironment.bUseRTTI ? " -frtti" : " -fno-rtti";
+			return CompileEnvironment.bUseRTTI ? "-frtti" : "-fno-rtti";
 		}
 
 		protected virtual string GetUserIncludePathArgument(DirectoryReference IncludePath)
 		{
-			return $" -I\"{NormalizeCommandLinePath(IncludePath)}\"";
+			return $"-I\"{NormalizeCommandLinePath(IncludePath)}\"";
 		}
 
 		protected virtual string GetSystemIncludePathArgument(DirectoryReference IncludePath)
 		{
 			// TODO: System include paths can be included with -isystem
-			return $" -I\"{NormalizeCommandLinePath(IncludePath)}\"";
+			return $"-I\"{NormalizeCommandLinePath(IncludePath)}\"";
 		}
 
-		protected virtual string GetIncludePathArguments(CppCompileEnvironment CompileEnvironment)
+		protected virtual void GetCompileArguments_IncludePaths(CppCompileEnvironment CompileEnvironment, List<string> Arguments)
 		{
-			string Result = "";
-			foreach (DirectoryReference IncludePath in CompileEnvironment.UserIncludePaths)
-			{
-				Result += GetUserIncludePathArgument(IncludePath);
-			}
-			foreach (DirectoryReference IncludePath in CompileEnvironment.SystemIncludePaths)
-			{
-				Result += GetSystemIncludePathArgument(IncludePath);
-			}
-			return Result;
+			Arguments.AddRange(CompileEnvironment.UserIncludePaths.Select(IncludePath => GetUserIncludePathArgument(IncludePath)));
+			Arguments.AddRange(CompileEnvironment.SystemIncludePaths.Select(IncludePath => GetSystemIncludePathArgument(IncludePath)));
 		}
 
 		protected virtual string GetPreprocessorDefinitionArgument(string Definition)
 		{
-			return $" -D\"{EscapePreprocessorDefinition(Definition)}\"";
+			return $"-D\"{EscapePreprocessorDefinition(Definition)}\"";
 		}
 
-		protected virtual string GetPreprocessorDefinitionArguments(CppCompileEnvironment CompileEnvironment)
+		protected virtual void GetCompileArguments_PreprocessorDefinitions(CppCompileEnvironment CompileEnvironment, List<string> Arguments)
 		{
-			string Result = "";
-			foreach (string Definition in CompileEnvironment.Definitions)
-			{
-				Result += GetPreprocessorDefinitionArgument(Definition);
-			}
-			return Result;
+			Arguments.AddRange(CompileEnvironment.Definitions.Select(Definition => GetPreprocessorDefinitionArgument(Definition)));
 		}
 
 		protected virtual string GetForceIncludeFileArgument(FileReference ForceIncludeFile)
 		{
-			return $" -include \"{NormalizeCommandLinePath(ForceIncludeFile)}\"";
+			return $"-include \"{NormalizeCommandLinePath(ForceIncludeFile)}\"";
 		}
 
 		protected virtual string GetForceIncludeFileArgument(FileItem ForceIncludeFile)
@@ -225,35 +200,30 @@ namespace UnrealBuildTool
 			return GetForceIncludeFileArgument(ForceIncludeFile.Location);
 		}
 
-		protected virtual string GetForceIncludeFileArguments(CppCompileEnvironment CompileEnvironment)
+		protected virtual void GetCompileArguments_ForceInclude(CppCompileEnvironment CompileEnvironment, List<string> Arguments)
 		{
-			string Result = "";
-			foreach (FileItem ForceIncludeFile in CompileEnvironment.ForceIncludeFiles)
-			{
-				Result += String.Format(" -include \"{0}\"", NormalizeCommandLinePath(ForceIncludeFile));
-			}
-			return Result;
+			Arguments.AddRange(CompileEnvironment.ForceIncludeFiles.Select(ForceIncludeFile => GetForceIncludeFileArgument(ForceIncludeFile)));
 		}
 
+
 		/// <summary>
 		/// Common compile arguments for all files in a module.
 		/// Override and call base.GetCompileArguments_Global() in derived classes.
 		/// </summary>
-		protected virtual string GetCompileArguments_Global(CppCompileEnvironment CompileEnvironment)
+		///
+		/// <param name="CompileEnvironment"></param>
+		/// <param name="Arguments"></param>
+		protected virtual void GetCompileArguments_Global(CppCompileEnvironment CompileEnvironment, List<string> Arguments)
 		{
-			string Result = "";
-
 			// build up the commandline common to C and C++
-			Result += " -c";
-			Result += " -pipe";
+			Arguments.Add("-c");
+			Arguments.Add("-pipe");
 
 			// Add include paths to the argument list.
-			Result += GetIncludePathArguments(CompileEnvironment);
+			GetCompileArguments_IncludePaths(CompileEnvironment, Arguments);
 
 			// Add preprocessor definitions to the argument list.
-			Result += GetPreprocessorDefinitionArguments(CompileEnvironment);
-
-			return Result;
+			GetCompileArguments_PreprocessorDefinitions(CompileEnvironment, Arguments);
 		}
 	}
 }
\ No newline at end of file
