UBT: ConfigFile attribute: Added support for properties.

[ConfigFile(...)] can be added to both fields and properties. The latter can be useful e.g. with derived properties.

diff --git a/Engine/Source/Programs/UnrealBuildTool/System/ConfigCache.cs b/Engine/Source/Programs/UnrealBuildTool/System/ConfigCache.cs
--- a/Engine/Source/Programs/UnrealBuildTool/System/ConfigCache.cs
+++ b/Engine/Source/Programs/UnrealBuildTool/System/ConfigCache.cs
@@ -26,15 +26,10 @@ namespace UnrealBuildTool
 		delegate void AddElementDelegate(object TargetObject, object? ValueObject); 
 
 		/// <summary>
-		/// Caches information about a field with a [ConfigFile] attribute in a type
+		/// Caches information about a member with a [ConfigFile] attribute in a type
 		/// </summary>
-		class ConfigField
+		abstract class ConfigMember
 		{
-			/// <summary>
-			/// The field with the config attribute
-			/// </summary>
-			public FieldInfo FieldInfo;
-
 			/// <summary>
 			/// The attribute instance
 			/// </summary>
@@ -50,16 +45,88 @@ namespace UnrealBuildTool
 			/// </summary>
 			public AddElementDelegate? AddElement;
 
+			/// <summary>
+			/// Constructor
+			/// </summary>
+			/// <param name="Attribute"></param>
+			public ConfigMember(ConfigFileAttribute Attribute)
+			{
+				this.Attribute = Attribute;
+			}
+
+			/// <summary>
+			/// Returns Reflection.MemberInfo describing the target class member.
+			/// </summary>
+			public abstract MemberInfo               MemberInfo { get; }
+
+			/// <summary>
+			/// Returns Reflection.Type of the target class member.
+			/// </summary>
+			public abstract Type                     Type       { get; }
+
+			/// <summary>
+			/// Returns the value setter of the target class member.
+			/// </summary>
+			public abstract Action<object?, object?> SetValue   { get; }
+
+			/// <summary>
+			/// Returns the value getter of the target class member.
+			/// </summary>
+			public abstract Func<object?, object?>   GetValue   { get; }
+		}
+
+		/// <summary>
+		/// Caches information about a field with a [ConfigFile] attribute in a type
+		/// </summary>
+		class ConfigField : ConfigMember
+		{
+			/// <summary>
+			/// Reflection description of the field with the config attribute.
+			/// </summary>
+			private FieldInfo FieldInfo;
+
 			/// <summary>
 			/// Constructor
 			/// </summary>
 			/// <param name="FieldInfo"></param>
 			/// <param name="Attribute"></param>
 			public ConfigField(FieldInfo FieldInfo, ConfigFileAttribute Attribute)
+				: base(Attribute)
 			{
 				this.FieldInfo = FieldInfo;
-				this.Attribute = Attribute;
 			}
+
+			public override MemberInfo               MemberInfo => FieldInfo;
+			public override Type                     Type       => FieldInfo.FieldType;
+			public override Action<object?, object?> SetValue   => FieldInfo.SetValue;
+			public override Func<object?, object?>   GetValue   => FieldInfo.GetValue;
+		}
+
+		/// <summary>
+		/// Caches information about a property with a [ConfigFile] attribute in a type
+		/// </summary>
+		class ConfigProperty : ConfigMember
+		{
+			/// <summary>
+			/// Reflection description of the property with the config attribute.
+			/// </summary>
+			private PropertyInfo PropertyInfo;
+
+			/// <summary>
+			/// Constructor
+			/// </summary>
+			/// <param name="PropertyInfo"></param>
+			/// <param name="Attribute"></param>
+			public ConfigProperty(PropertyInfo PropertyInfo, ConfigFileAttribute Attribute)
+				: base(Attribute)
+			{
+				this.PropertyInfo = PropertyInfo;
+			}
+
+			public override MemberInfo               MemberInfo => PropertyInfo;
+			public override Type                     Type       => PropertyInfo.PropertyType;
+			public override Action<object?, object?> SetValue   => PropertyInfo.SetValue;
+			public override Func<object?, object?>   GetValue   => PropertyInfo.GetValue;
 		}
 
 		/// <summary>
@@ -151,7 +218,7 @@ namespace UnrealBuildTool
 		/// <summary>
 		/// Cache of config fields by type
 		/// </summary>
-		static Dictionary<Type, List<ConfigField>> TypeToConfigFields = new Dictionary<Type, List<ConfigField>>();
+		static Dictionary<Type, List<ConfigMember>> TypeToConfigMembers = new Dictionary<Type, List<ConfigMember>>();
 
 		/// <summary>
 		/// Attempts to read a config file (or retrieve it from the cache)
@@ -256,46 +323,56 @@ namespace UnrealBuildTool
 		/// </summary>
 		/// <param name="TargetObjectType">Type to get configurable fields for</param>
 		/// <returns>List of config fields for the given type</returns>
-		static List<ConfigField> FindConfigFieldsForType(Type TargetObjectType)
+		static List<ConfigMember> FindConfigMembersForType(Type TargetObjectType)
 		{
-			List<ConfigField>? Fields;
-			lock(TypeToConfigFields)
+			List<ConfigMember>? Members;
+			lock(TypeToConfigMembers)
 			{
-				if (!TypeToConfigFields.TryGetValue(TargetObjectType, out Fields))
+				if (!TypeToConfigMembers.TryGetValue(TargetObjectType, out Members))
 				{
-					Fields = new List<ConfigField>();
+					Members = new List<ConfigMember>();
 					if(TargetObjectType.BaseType != null)
 					{
-						Fields.AddRange(FindConfigFieldsForType(TargetObjectType.BaseType));
+						Members.AddRange(FindConfigMembersForType(TargetObjectType.BaseType));
 					}
 					foreach (FieldInfo FieldInfo in TargetObjectType.GetFields(BindingFlags.Instance | BindingFlags.GetField | BindingFlags.GetProperty | BindingFlags.Public | BindingFlags.NonPublic | BindingFlags.DeclaredOnly))
 					{
-						IEnumerable<ConfigFileAttribute> Attributes = FieldInfo.GetCustomAttributes<ConfigFileAttribute>();
-						foreach (ConfigFileAttribute Attribute in Attributes)
-						{
-							// Copy the field 
-							ConfigField Setter = new ConfigField(FieldInfo, Attribute);
+						ProcessConfigTypeMember<FieldInfo>(TargetObjectType, FieldInfo, Members, (FieldInfo, Attribute) => new ConfigField(FieldInfo, Attribute));
+					}
+					foreach (PropertyInfo PropertyInfo in TargetObjectType.GetProperties(BindingFlags.Instance | BindingFlags.GetField | BindingFlags.GetProperty | BindingFlags.Public | BindingFlags.NonPublic | BindingFlags.DeclaredOnly))
+					{
+						ProcessConfigTypeMember<PropertyInfo>(TargetObjectType, PropertyInfo, Members, (PropertyInfo, Attribute) => new ConfigProperty(PropertyInfo, Attribute));
+					}
+					TypeToConfigMembers.Add(TargetObjectType, Members);
+				}
+			}
+			return Members;
+		}
 
-							// Check if the field type implements ICollection<>. If so, we can take multiple values.
-							foreach (Type InterfaceType in FieldInfo.FieldType.GetInterfaces())
-							{
-								if (InterfaceType.IsGenericType && InterfaceType.GetGenericTypeDefinition() == typeof(ICollection<>))
-								{
-									MethodInfo MethodInfo = InterfaceType.GetRuntimeMethod("Add", new Type[] { InterfaceType.GenericTypeArguments[0] })!;
-									Setter.AddElement = (Target, Value) => { MethodInfo.Invoke(Setter.FieldInfo.GetValue(Target), new object?[] { Value }); };
-									Setter.ElementType = InterfaceType.GenericTypeArguments[0];
-									break;
-								}
-							}
+		static void ProcessConfigTypeMember<MEMBER>(Type TargetType, MEMBER MemberInfo, List<ConfigMember> Members, Func<MEMBER, ConfigFileAttribute, ConfigMember> CreateConfigMember)
+			where MEMBER : System.Reflection.MemberInfo
+		{
+			IEnumerable<ConfigFileAttribute> Attributes = MemberInfo.GetCustomAttributes<ConfigFileAttribute>();
+			foreach (ConfigFileAttribute Attribute in Attributes)
+			{
+				// Copy the field 
+				ConfigMember Setter = CreateConfigMember(MemberInfo, Attribute);
 
-							// Add it to the output list
-							Fields.Add(Setter);
-						}
+				// Check if the field type implements ICollection<>. If so, we can take multiple values.
+				foreach (Type InterfaceType in Setter.Type.GetInterfaces())
+				{
+					if (InterfaceType.IsGenericType && InterfaceType.GetGenericTypeDefinition() == typeof(ICollection<>))
+					{
+						MethodInfo MethodInfo = InterfaceType.GetRuntimeMethod("Add", new Type[] { InterfaceType.GenericTypeArguments[0] })!;
+						Setter.AddElement = (Target, Value) => { MethodInfo.Invoke(Setter.GetValue(Target), new object?[] { Value }); };
+						Setter.ElementType = InterfaceType.GenericTypeArguments[0];
+						break;
 					}
-					TypeToConfigFields.Add(TargetObjectType, Fields);
 				}
+
+				// Add it to the output list
+				Members.Add(Setter);
 			}
-			return Fields;
 		}
 
 		/// <summary>
@@ -318,28 +395,28 @@ namespace UnrealBuildTool
 		/// <param name="ConfigValues">Will be populated with config values that were retrieved. May be null.</param>
 		internal static void ReadSettings(DirectoryReference? ProjectDir, UnrealTargetPlatform Platform, object TargetObject, Dictionary<ConfigDependencyKey, IReadOnlyList<string>?>? ConfigValues)
 		{
-			List<ConfigField> Fields = FindConfigFieldsForType(TargetObject.GetType());
-			foreach(ConfigField Field in Fields)
+			List<ConfigMember> Members = FindConfigMembersForType(TargetObject.GetType());
+			foreach(ConfigMember Member in Members)
 			{
 				// Read the hierarchy listed
-				ConfigHierarchy Hierarchy = ReadHierarchy(Field.Attribute.ConfigType, ProjectDir, Platform);
+				ConfigHierarchy Hierarchy = ReadHierarchy(Member.Attribute.ConfigType, ProjectDir, Platform);
 
 				// Get the key name
-				string KeyName = Field.Attribute.KeyName ?? Field.FieldInfo.Name;
+				string KeyName = Member.Attribute.KeyName ?? Member.MemberInfo.Name;
 
 				// Get the value(s) associated with this key
 				IReadOnlyList<string>? Values;
-				Hierarchy.TryGetValues(Field.Attribute.SectionName, KeyName, out Values);
+				Hierarchy.TryGetValues(Member.Attribute.SectionName, KeyName, out Values);
 
 				// Parse the values from the config files and update the target object
-				if (Field.AddElement == null)
+				if (Member.AddElement == null)
 				{
 					if(Values != null && Values.Count == 1)
 					{
 						object? Value;
-						if(TryParseValue(Values[0], Field.FieldInfo.FieldType, out Value))
+						if(TryParseValue(Values[0], Member.Type, out Value))
 						{
-							Field.FieldInfo.SetValue(TargetObject, Value);
+							Member.SetValue(TargetObject, Value);
 						}
 					}
 				}
@@ -350,9 +427,9 @@ namespace UnrealBuildTool
 						foreach(string Item in Values)
 						{
 							object? Value;
-							if(TryParseValue(Item, Field.ElementType!, out Value))
+							if(TryParseValue(Item, Member.ElementType!, out Value))
 							{
-								Field.AddElement(TargetObject, Value);
+								Member.AddElement(TargetObject, Value);
 							}
 						}
 					}
@@ -361,7 +438,7 @@ namespace UnrealBuildTool
 				// Save the dependency
 				if (ConfigValues != null)
 				{
-					ConfigDependencyKey Key = new ConfigDependencyKey(Field.Attribute.ConfigType, ProjectDir, Platform, Field.Attribute.SectionName, KeyName);
+					ConfigDependencyKey Key = new ConfigDependencyKey(Member.Attribute.ConfigType, ProjectDir, Platform, Member.Attribute.SectionName, KeyName);
 					ConfigValues[Key] = Values;
 				}
 			}
diff --git a/Engine/Source/Programs/UnrealBuildTool/System/ConfigFileAttribute.cs b/Engine/Source/Programs/UnrealBuildTool/System/ConfigFileAttribute.cs
--- a/Engine/Source/Programs/UnrealBuildTool/System/ConfigFileAttribute.cs
+++ b/Engine/Source/Programs/UnrealBuildTool/System/ConfigFileAttribute.cs
@@ -11,7 +11,7 @@ namespace UnrealBuildTool
 	/// <summary>
 	/// Attribute indicating a value which should be populated from a UE .ini config file
 	/// </summary>
-	[AttributeUsage(AttributeTargets.Field)]
+	[AttributeUsage(AttributeTargets.Field | AttributeTargets.Property)]
 	public class ConfigFileAttribute : Attribute
 	{
 		/// <summary>
