UBT: XmlConfigFile attribute: Added support for properties.

[XmlConfigFile(...)] can be added to both fields and properties. The latter can be useful e.g. with derived properties.

diff --git a/Engine/Source/Programs/UnrealBuildTool/System/RulesDocumentation.cs b/Engine/Source/Programs/UnrealBuildTool/System/RulesDocumentation.cs
--- a/Engine/Source/Programs/UnrealBuildTool/System/RulesDocumentation.cs
+++ b/Engine/Source/Programs/UnrealBuildTool/System/RulesDocumentation.cs
@@ -175,6 +175,29 @@ namespace UnrealBuildTool
 			return TryGetXmlComment(Documentation, $"P:{Property.DeclaringType}.{Property.Name}", out Lines);
 		}
 
+		/// <summary>
+		/// Gets the XML comment for a particular field
+		/// </summary>
+		/// <param name="Documentation">The XML documentation</param>
+		/// <param name="Member">The member to search for (field or property)</param>
+		/// <param name="Lines">Receives the description for the requested member</param>
+		/// <returns>True if a comment was found for the member</returns>
+		public static bool TryGetXmlComment(XmlDocument Documentation, MemberInfo Member, [NotNullWhen(true)] out List<string>? Lines)
+		{
+			if (Member is FieldInfo)
+			{
+				return TryGetXmlComment(Documentation, (Member as FieldInfo)!, out Lines);
+			}
+			else if (Member is PropertyInfo)
+			{
+				return TryGetXmlComment(Documentation, (Member as PropertyInfo)!, out Lines);
+			}
+
+			Lines = null;
+
+			return false;
+		}
+
 		/// <summary>
 		/// Gets the XML comment for a particular field
 		/// </summary>
diff --git a/Engine/Source/Programs/UnrealBuildTool/System/XmlConfig.cs b/Engine/Source/Programs/UnrealBuildTool/System/XmlConfig.cs
--- a/Engine/Source/Programs/UnrealBuildTool/System/XmlConfig.cs
+++ b/Engine/Source/Programs/UnrealBuildTool/System/XmlConfig.cs
@@ -117,7 +117,7 @@ namespace UnrealBuildTool
 				if(Values == null)
 				{
 					// Find all the configurable fields from the given types
-					Dictionary<string, Dictionary<string, FieldInfo>> CategoryToFields = new Dictionary<string, Dictionary<string, FieldInfo>>();
+					Dictionary<string, Dictionary<string, XmlConfigData.TargetMember>> CategoryToFields = new Dictionary<string, Dictionary<string, XmlConfigData.TargetMember>>();
 					FindConfigurableFields(ConfigTypes, CategoryToFields);
 
 					// Create a schema for the config files
@@ -128,8 +128,8 @@ namespace UnrealBuildTool
 					}
 
 					// Read all the XML files and validate them against the schema
-					Dictionary<Type, Dictionary<FieldInfo, XmlConfigData.ValueInfo>> TypeToValues =
-						new Dictionary<Type, Dictionary<FieldInfo, XmlConfigData.ValueInfo>>();
+					Dictionary<Type, Dictionary<XmlConfigData.TargetMember, XmlConfigData.ValueInfo>> TypeToValues =
+						new Dictionary<Type, Dictionary<XmlConfigData.TargetMember, XmlConfigData.ValueInfo>>();
 					foreach(FileReference InputFile in InputFileLocations)
 					{
 						if(!TryReadFile(InputFile, CategoryToFields, TypeToValues, Schema))
@@ -152,12 +152,12 @@ namespace UnrealBuildTool
 			// Apply all the static field values
 			foreach(KeyValuePair<Type, XmlConfigData.ValueInfo[]> TypeValuesPair in Values.TypeToValues)
 			{
-				foreach(XmlConfigData.ValueInfo FieldValue in TypeValuesPair.Value)
+				foreach(XmlConfigData.ValueInfo MemberValue in TypeValuesPair.Value)
 				{
-					if(FieldValue.FieldInfo.IsStatic)
+					if(MemberValue.Target.IsStatic)
 					{
-						object Value = InstanceValue(FieldValue.Value, FieldValue.FieldInfo.FieldType);
-						FieldValue.FieldInfo.SetValue(null, Value);
+						object Value = InstanceValue(MemberValue.Value, MemberValue.Target.Type);
+						MemberValue.Target.SetValue(null, Value);
 					}
 				}
 			}
@@ -205,6 +205,16 @@ namespace UnrealBuildTool
 					}
 				}
 			}
+			foreach (PropertyInfo Property in Type.GetProperties(BindingFlags.Instance | BindingFlags.Static | BindingFlags.Public | BindingFlags.NonPublic))
+			{
+				foreach (CustomAttributeData CustomAttribute in Property.CustomAttributes)
+				{
+					if (CustomAttribute.AttributeType == typeof(XmlConfigFileAttribute))
+					{
+						return true;
+					}
+				}
+			}
 			return false;
 		}
 
@@ -319,45 +329,66 @@ namespace UnrealBuildTool
 		/// <param name="TargetObject">The object instance to be configured</param>
 		public static void ApplyTo(object TargetObject)
 		{
-			for(Type? TargetType = TargetObject.GetType(); TargetType != null; TargetType = TargetType.BaseType)
+			for (Type? TargetType = TargetObject.GetType(); TargetType != null; TargetType = TargetType.BaseType)
 			{ 
 				XmlConfigData.ValueInfo[]? FieldValues;
 				if(Values!.TypeToValues.TryGetValue(TargetType, out FieldValues))
 				{
 					foreach(XmlConfigData.ValueInfo FieldValue in FieldValues)
 					{
-						if (!FieldValue.FieldInfo.IsStatic)
+						if (!FieldValue.Target.IsStatic)
 						{
-							FieldInfo FieldInfoToWrite = FieldValue.FieldInfo;
+							XmlConfigData.TargetMember TargetToWrite = FieldValue.Target;
 							
 							// Check if setting has been deprecated
 							if (FieldValue.XmlConfigAttribute.Deprecated)
 							{
 								Log.TraceWarning($"Deprecated setting found in \"{FieldValue.SourceFile}\":");
 								Log.TraceWarning(
-									$"The setting \"{FieldValue.FieldInfo.Name}\" is deprecated. Support for this setting will be removed in a future version of Unreal Engine.");
+									$"The setting \"{FieldValue.Target.MemberInfo.Name}\" is deprecated. Support for this setting will be removed in a future version of Unreal Engine.");
 
 								if (FieldValue.XmlConfigAttribute.NewAttributeName != null)
 								{
 									Log.TraceWarning(
-										$"Use \"{FieldValue.XmlConfigAttribute.NewAttributeName}\" in place of \"{FieldValue.FieldInfo.Name}\"");
+										$"Use \"{FieldValue.XmlConfigAttribute.NewAttributeName}\" in place of \"{FieldValue.Target.MemberInfo.Name}\"");
 									Log.TraceInformation(
-										$"The value provided for deprecated setting \"{FieldValue.FieldInfo.Name}\" will be applied to \"{FieldValue.XmlConfigAttribute.NewAttributeName}\"");
+										$"The value provided for deprecated setting \"{FieldValue.Target.MemberInfo.Name}\" will be applied to \"{FieldValue.XmlConfigAttribute.NewAttributeName}\"");
 
-									// Find the new field to which the setting should be actually applied
-									FieldInfoToWrite = TargetType.GetRuntimeFields()
-										.First(x => x.Name == FieldValue.XmlConfigAttribute.NewAttributeName);
+									TargetToWrite = GetTargetMember(TargetType, FieldValue.XmlConfigAttribute.NewAttributeName) ?? TargetToWrite;
 								}
 							}
 
-							object ValueInstance = InstanceValue(FieldValue.Value, FieldValue.FieldInfo.FieldType);
-							FieldInfoToWrite.SetValue(TargetObject, ValueInstance);
+							object ValueInstance = InstanceValue(FieldValue.Value, FieldValue.Target.Type);
+							TargetToWrite.SetValue(TargetObject, ValueInstance);
 						}
 					}
 				}
 			}
 		}
 
+		private static XmlConfigData.TargetMember? GetTargetMember(Type TargetType, string MemberName)
+		{
+			// First, try to find the new field to which the setting should be actually applied.
+			FieldInfo? FieldInfo = TargetType.GetRuntimeFields()
+				.First(x => x.Name == MemberName);
+
+			if (FieldInfo != null)
+			{
+				return new XmlConfigData.TargetField(FieldInfo);
+			}
+
+			// If not found, try to find the new property to which the setting should be actually applied.
+			PropertyInfo? PropertyInfo = TargetType.GetRuntimeProperties()
+				.First(x => x.Name == MemberName);
+
+			if (PropertyInfo != null)
+			{
+				return new XmlConfigData.TargetProperty(PropertyInfo);
+			}
+
+			return null;
+		}
+
 		/// <summary>
 		/// Instances a value for assignment to a target object
 		/// </summary>
@@ -396,7 +427,7 @@ namespace UnrealBuildTool
 			// Find the value with the matching name
 			foreach(XmlConfigData.ValueInfo FieldValue in FieldValues)
 			{
-				if(FieldValue.FieldInfo.Name == Name)
+				if(FieldValue.Target.MemberInfo.Name == Name)
 				{
 					Value = FieldValue.Value;
 					return true;
@@ -413,27 +444,37 @@ namespace UnrealBuildTool
 		/// </summary>
 		/// <param name="ConfigTypes">Array of types to search</param>
 		/// <param name="CategoryToFields">Dictionaries populated with category -> name -> field mappings on return</param>
-		static void FindConfigurableFields(IEnumerable<Type> ConfigTypes, Dictionary<string, Dictionary<string, FieldInfo>> CategoryToFields)
+		static void FindConfigurableFields(IEnumerable<Type> ConfigTypes, Dictionary<string, Dictionary<string, XmlConfigData.TargetMember>> CategoryToFields)
 		{
 			foreach(Type ConfigType in ConfigTypes)
 			{
 				foreach(FieldInfo FieldInfo in ConfigType.GetFields(BindingFlags.Instance | BindingFlags.Static | BindingFlags.GetField | BindingFlags.Public | BindingFlags.NonPublic))
 				{
-					IEnumerable<XmlConfigFileAttribute> Attributes = FieldInfo.GetCustomAttributes<XmlConfigFileAttribute>();
-					foreach(XmlConfigFileAttribute Attribute in Attributes)
-					{
-						string CategoryName = Attribute.Category ?? ConfigType.Name;
+					ProcessConfigurableMember<FieldInfo>(ConfigType, FieldInfo, CategoryToFields, FieldInfo => new XmlConfigData.TargetField(FieldInfo));
+				}
+				foreach (PropertyInfo PropertyInfo in ConfigType.GetProperties(BindingFlags.Instance | BindingFlags.Static | BindingFlags.GetField | BindingFlags.Public | BindingFlags.NonPublic))
+				{
+					ProcessConfigurableMember<PropertyInfo>(ConfigType, PropertyInfo, CategoryToFields, PropertyInfo => new XmlConfigData.TargetProperty(PropertyInfo));
+				}
+			}
+		}
 
-						Dictionary<string, FieldInfo>? NameToField;
-						if(!CategoryToFields.TryGetValue(CategoryName, out NameToField))
-						{
-							NameToField = new Dictionary<string, FieldInfo>();
-							CategoryToFields.Add(CategoryName, NameToField);
-						}
+		private static void ProcessConfigurableMember<MEMBER>(Type Type, MEMBER MemberInfo, Dictionary<string, Dictionary<string, XmlConfigData.TargetMember>> CategoryToFields, Func<MEMBER, XmlConfigData.TargetMember> CreateTarget)
+			where MEMBER : System.Reflection.MemberInfo
+		{
+			IEnumerable<XmlConfigFileAttribute> Attributes = MemberInfo.GetCustomAttributes<XmlConfigFileAttribute>();
+			foreach (XmlConfigFileAttribute Attribute in Attributes)
+			{
+				string CategoryName = Attribute.Category ?? Type.Name;
 
-						NameToField[Attribute.Name ?? FieldInfo.Name] = FieldInfo;
-					}
+				Dictionary<string, XmlConfigData.TargetMember>? NameToTarget;
+				if (!CategoryToFields.TryGetValue(CategoryName, out NameToTarget))
+				{
+					NameToTarget = new Dictionary<string, XmlConfigData.TargetMember>();
+					CategoryToFields.Add(CategoryName, NameToTarget);
 				}
+
+				NameToTarget[Attribute.Name ?? MemberInfo.Name] = CreateTarget(MemberInfo);
 			}
 		}
 
@@ -442,18 +483,18 @@ namespace UnrealBuildTool
 		/// </summary>
 		/// <param name="CategoryToFields">Lookup for all field settings</param>
 		/// <returns>New schema instance</returns>
-		static XmlSchema CreateSchema(Dictionary<string, Dictionary<string, FieldInfo>> CategoryToFields)
+		static XmlSchema CreateSchema(Dictionary<string, Dictionary<string, XmlConfigData.TargetMember>> CategoryToFields)
 		{
 			// Create elements for all the categories
 			XmlSchemaAll RootAll = new XmlSchemaAll();
-			foreach(KeyValuePair<string, Dictionary<string, FieldInfo>> CategoryPair in CategoryToFields)
+			foreach(KeyValuePair<string, Dictionary<string, XmlConfigData.TargetMember>> CategoryPair in CategoryToFields)
 			{
 				string CategoryName = CategoryPair.Key;
 
 				XmlSchemaAll CategoryAll = new XmlSchemaAll();
-				foreach (KeyValuePair<string, FieldInfo> FieldPair in CategoryPair.Value)
+				foreach (KeyValuePair<string, XmlConfigData.TargetMember> FieldPair in CategoryPair.Value)
 				{
-					XmlSchemaElement Element = CreateSchemaFieldElement(FieldPair.Key, FieldPair.Value.FieldType);
+					XmlSchemaElement Element = CreateSchemaFieldElement(FieldPair.Key, FieldPair.Value.Type);
 					CategoryAll.Items.Add(Element);
 				}
 
@@ -629,8 +670,8 @@ namespace UnrealBuildTool
 		/// <param name="TypeToValues">Map of types to fields and their associated values</param>
 		/// <param name="Schema">Schema to validate against</param>
 		/// <returns>True if the file was read successfully</returns>
-		static bool TryReadFile(FileReference Location, Dictionary<string, Dictionary<string, FieldInfo>> CategoryToFields,
-			Dictionary<Type, Dictionary<FieldInfo, XmlConfigData.ValueInfo>> TypeToValues, XmlSchema Schema)
+		static bool TryReadFile(FileReference Location, Dictionary<string, Dictionary<string, XmlConfigData.TargetMember>> CategoryToFields,
+			Dictionary<Type, Dictionary<XmlConfigData.TargetMember, XmlConfigData.ValueInfo>> TypeToValues, XmlSchema Schema)
 		{
 			// Read the XML file, and validate it against the schema
 			XmlConfigFile? ConfigFile;
@@ -642,39 +683,39 @@ namespace UnrealBuildTool
 			// Parse the document
 			foreach(XmlElement CategoryElement in ConfigFile.DocumentElement!.ChildNodes.OfType<XmlElement>())
 			{
-				Dictionary<string, FieldInfo>? NameToField;
+				Dictionary<string, XmlConfigData.TargetMember>? NameToField;
 				if(CategoryToFields.TryGetValue(CategoryElement.Name, out NameToField))
 				{
 					foreach(XmlElement KeyElement in CategoryElement.ChildNodes.OfType<XmlElement>())
 					{
-						FieldInfo? Field;
+						XmlConfigData.TargetMember? Field;
 						object Value;
 						if(NameToField.TryGetValue(KeyElement.Name, out Field))
 						{
-							if (Field.FieldType == typeof(string[]))
+							if (Field.Type == typeof(string[]))
 							{
 								Value = KeyElement.ChildNodes.OfType<XmlElement>().Where(x => x.Name == "Item").Select(x => x.InnerText).ToArray();
 							}
-							else if (TryGetNullableStructType(Field.FieldType, out Type? StructType))
+							else if (TryGetNullableStructType(Field.Type, out Type? StructType))
 							{
 								Value = ParseValue(StructType, KeyElement.InnerText);
 							}
 							else
 							{
-								Value = ParseValue(Field.FieldType, KeyElement.InnerText);
+								Value = ParseValue(Field.Type, KeyElement.InnerText);
 							}
 
 							// Add it to the set of values for the type containing this field
-							Dictionary<FieldInfo, XmlConfigData.ValueInfo>? FieldToValue;
-							if(!TypeToValues.TryGetValue(Field.DeclaringType!, out FieldToValue))
+							Dictionary<XmlConfigData.TargetMember, XmlConfigData.ValueInfo>? FieldToValue;
+							if(!TypeToValues.TryGetValue(Field.MemberInfo.DeclaringType!, out FieldToValue))
 							{
-								FieldToValue = new Dictionary<FieldInfo, XmlConfigData.ValueInfo>();
-								TypeToValues.Add(Field.DeclaringType!, FieldToValue);
+								FieldToValue = new Dictionary<XmlConfigData.TargetMember, XmlConfigData.ValueInfo>();
+								TypeToValues.Add(Field.MemberInfo.DeclaringType!, FieldToValue);
 							}
 							
 							// Parse the corresponding value
 							XmlConfigData.ValueInfo FieldValue = new XmlConfigData.ValueInfo(Field, Value, Location,
-								Field.GetCustomAttribute<XmlConfigFileAttribute>()!);
+								Field.MemberInfo.GetCustomAttribute<XmlConfigFileAttribute>()!);
 							
 							FieldToValue[Field] = FieldValue;
 						}
@@ -786,7 +827,7 @@ namespace UnrealBuildTool
 			List<Type> ConfigTypes = FindConfigurableTypes();
 
 			// Find all the configurable fields from the given types
-			Dictionary<string, Dictionary<string, FieldInfo>> CategoryToFields = new Dictionary<string, Dictionary<string, FieldInfo>>();
+			Dictionary<string, Dictionary<string, XmlConfigData.TargetMember>> CategoryToFields = new Dictionary<string, Dictionary<string, XmlConfigData.TargetMember>>();
 			FindConfigurableFields(ConfigTypes, CategoryToFields);
 			CategoryToFields = CategoryToFields.Where(x => x.Value.Count > 0).ToDictionary(x => x.Key, x => x.Value);
 
@@ -835,7 +876,7 @@ namespace UnrealBuildTool
 		/// <param name="OutputFile">The output file</param>
 		/// <param name="InputDocumentation">The XML documentation for this assembly</param>
 		/// <param name="CategoryToFields">Map of string to types to fields</param>
-		private static void WriteDocumentationUDN(FileReference OutputFile, XmlDocument InputDocumentation, Dictionary<string, Dictionary<string, FieldInfo>> CategoryToFields)
+		private static void WriteDocumentationUDN(FileReference OutputFile, XmlDocument InputDocumentation, Dictionary<string, Dictionary<string, XmlConfigData.TargetMember>> CategoryToFields)
 		{
 			using (StreamWriter Writer = new StreamWriter(OutputFile.FullName))
 			{
@@ -846,18 +887,18 @@ namespace UnrealBuildTool
 				Writer.WriteLine("Version: {0}.{1}", ReadOnlyBuildVersion.Current.MajorVersion, ReadOnlyBuildVersion.Current.MinorVersion);
 				Writer.WriteLine("");
 
-				foreach (KeyValuePair<string, Dictionary<string, FieldInfo>> CategoryPair in CategoryToFields)
+				foreach (KeyValuePair<string, Dictionary<string, XmlConfigData.TargetMember>> CategoryPair in CategoryToFields)
 				{
 					string CategoryName = CategoryPair.Key;
 					Writer.WriteLine("### {0}", CategoryName);
 					Writer.WriteLine();
 
-					Dictionary<string, FieldInfo> Fields = CategoryPair.Value;
-					foreach (KeyValuePair<string, FieldInfo> FieldPair in Fields)
+					Dictionary<string, XmlConfigData.TargetMember> Fields = CategoryPair.Value;
+					foreach (KeyValuePair<string, XmlConfigData.TargetMember> FieldPair in Fields)
 					{
 						// Get the XML comment for this field
 						List<string>? Lines;
-						if(!RulesDocumentation.TryGetXmlComment(InputDocumentation, FieldPair.Value, out Lines) || Lines.Count == 0)
+						if(!RulesDocumentation.TryGetXmlComment(InputDocumentation, FieldPair.Value.MemberInfo, out Lines) || Lines.Count == 0)
 						{
 							continue;
 						}
@@ -887,25 +928,25 @@ namespace UnrealBuildTool
 		/// <param name="OutputFile">The output file</param>
 		/// <param name="InputDocumentation">The XML documentation for this assembly</param>
 		/// <param name="CategoryToFields">Map of string to types to fields</param>
-		private static void WriteDocumentationHTML(FileReference OutputFile, XmlDocument InputDocumentation, Dictionary<string, Dictionary<string, FieldInfo>> CategoryToFields)
+		private static void WriteDocumentationHTML(FileReference OutputFile, XmlDocument InputDocumentation, Dictionary<string, Dictionary<string, XmlConfigData.TargetMember>> CategoryToFields)
 		{
 			using (StreamWriter Writer = new StreamWriter(OutputFile.FullName))
 			{
 				Writer.WriteLine("<html>");
 				Writer.WriteLine("  <body>");
 				Writer.WriteLine("  <h2>BuildConfiguration Properties</h2>");
-				foreach (KeyValuePair<string, Dictionary<string, FieldInfo>> CategoryPair in CategoryToFields)
+				foreach (KeyValuePair<string, Dictionary<string, XmlConfigData.TargetMember>> CategoryPair in CategoryToFields)
 				{
 					string CategoryName = CategoryPair.Key;
 					Writer.WriteLine("    <h3>{0}</h3>", CategoryName);
 					Writer.WriteLine("    <dl>");
 
-					Dictionary<string, FieldInfo> Fields = CategoryPair.Value;
-					foreach (KeyValuePair<string, FieldInfo> FieldPair in Fields)
+					Dictionary<string, XmlConfigData.TargetMember> Fields = CategoryPair.Value;
+					foreach (KeyValuePair<string, XmlConfigData.TargetMember> FieldPair in Fields)
 					{
 						// Get the XML comment for this field
 						List<string>? Lines;
-						if (!RulesDocumentation.TryGetXmlComment(InputDocumentation, FieldPair.Value, out Lines) || Lines.Count == 0)
+						if (!RulesDocumentation.TryGetXmlComment(InputDocumentation, FieldPair.Value.MemberInfo, out Lines) || Lines.Count == 0)
 						{
 							continue;
 						}
diff --git a/Engine/Source/Programs/UnrealBuildTool/System/XmlConfigData.cs b/Engine/Source/Programs/UnrealBuildTool/System/XmlConfigData.cs
--- a/Engine/Source/Programs/UnrealBuildTool/System/XmlConfigData.cs
+++ b/Engine/Source/Programs/UnrealBuildTool/System/XmlConfigData.cs
@@ -28,16 +28,101 @@ namespace UnrealBuildTool
 		/// </summary>
 		public FileReference[] InputFiles;
 
+		/// <summary>
+		/// Abstract description of a target data member.
+		/// </summary>
+		public abstract class TargetMember
+		{
+			/// <summary>
+			/// Returns Reflection.MemberInfo describing the target class member.
+			/// </summary>
+			public abstract MemberInfo               MemberInfo { get; }
+
+			/// <summary>
+			/// Returns Reflection.Type of the target class member.
+			/// </summary>
+			public abstract Type                     Type       { get; }
+
+			/// <summary>
+			/// Indicates whether the target class member is static or not.
+			/// </summary>
+			public abstract bool                     IsStatic   { get; }
+
+			/// <summary>
+			/// Returns the value setter of the target class member.
+			/// </summary>
+			public abstract Action<object?, object?> SetValue   { get; }
+
+			/// <summary>
+			/// Returns the value getter of the target class member.
+			/// </summary>
+			public abstract Func<object?, object?>   GetValue   { get; }
+		}
+
+		/// <summary>
+		/// Description of a field member.
+		/// </summary>
+		public class TargetField : TargetMember
+		{
+			public override MemberInfo               MemberInfo => FieldInfo;
+			public override Type                     Type       => FieldInfo.FieldType;
+			public override bool                     IsStatic   => FieldInfo.IsStatic;
+			public override Action<object?, object?> SetValue   => FieldInfo.SetValue;
+			public override Func<object?, object?>   GetValue   => FieldInfo.GetValue;
+
+			private FieldInfo FieldInfo;
+
+			public TargetField(FieldInfo FieldInfo)
+			{
+				this.FieldInfo = FieldInfo;
+			}
+		}
+
+		/// <summary>
+		/// Description of a property member.
+		/// </summary>
+		public class TargetProperty : TargetMember
+		{
+			public override MemberInfo               MemberInfo => PropertyInfo;
+			public override Type                     Type       => PropertyInfo.PropertyType;
+			public override bool                     IsStatic   => PropertyInfo.GetGetMethod()!.IsStatic;
+			public override Action<object?, object?> SetValue   => PropertyInfo.SetValue;
+			public override Func<object?, object?>   GetValue   => PropertyInfo.GetValue;
+
+			private PropertyInfo PropertyInfo;
+
+			public TargetProperty(PropertyInfo PropertyInfo)
+			{
+				this.PropertyInfo = PropertyInfo;
+			}
+		}
+
 		public class ValueInfo
 		{
-			public FieldInfo FieldInfo;
+			public TargetMember Target;
 			public object Value;
 			public FileReference SourceFile;
 			public XmlConfigFileAttribute XmlConfigAttribute;
 
 			public ValueInfo(FieldInfo FieldInfo, object Value, FileReference SourceFile, XmlConfigFileAttribute XmlConfigAttribute)
 			{
-				this.FieldInfo = FieldInfo;
+				this.Target = new TargetField(FieldInfo);
+				this.Value = Value;
+				this.SourceFile = SourceFile;
+				this.XmlConfigAttribute = XmlConfigAttribute;
+			}
+
+			public ValueInfo(PropertyInfo PropertyInfo, object Value, FileReference SourceFile, XmlConfigFileAttribute XmlConfigAttribute)
+			{
+				this.Target = new TargetProperty(PropertyInfo);
+				this.Value = Value;
+				this.SourceFile = SourceFile;
+				this.XmlConfigAttribute = XmlConfigAttribute;
+			}
+
+			public ValueInfo(TargetMember Target, object Value, FileReference SourceFile, XmlConfigFileAttribute XmlConfigAttribute)
+			{
+				this.Target = Target;
 				this.Value = Value;
 				this.SourceFile = SourceFile;
 				this.XmlConfigAttribute = XmlConfigAttribute;
@@ -45,7 +130,7 @@ namespace UnrealBuildTool
 		}
 		
 		/// <summary>
-		/// Stores a mapping from type -> field -> value, with all the config values for configurable fields.
+		/// Stores a mapping from type -> member -> value, with all the config values for configurable fields.
 		/// </summary>
 		public Dictionary<Type, ValueInfo[]> TypeToValues;
 
@@ -70,7 +155,7 @@ namespace UnrealBuildTool
 		public static bool TryRead(FileReference Location, IEnumerable<Type> Types, [NotNullWhen(true)] out XmlConfigData? Data)
 		{
 			// Check the file exists first
-			if(!FileReference.Exists(Location))
+			if (!FileReference.Exists(Location))
 			{
 				Data = null;
 				return false;
@@ -107,26 +192,30 @@ namespace UnrealBuildTool
 
 					// Read all the values
 					ValueInfo[] Values = new ValueInfo[Reader.ReadInt32()];
-					for(int ValueIdx = 0; ValueIdx < Values.Length; ValueIdx++)
+					for (int ValueIdx = 0; ValueIdx < Values.Length; ValueIdx++)
 					{
-						string FieldName = Reader.ReadString();
+						string MemberName = Reader.ReadString();
+
+						XmlConfigData.TargetMember? TargetMember = GetTargetMemberWithAttribute<XmlConfigFileAttribute>(Type, MemberName);
+
+						if (TargetMember != null)
+						{
+							// If TargetMember is not null, we know it has our attribute.
+							XmlConfigFileAttribute XmlConfigAttribute = TargetMember!.MemberInfo.GetCustomAttribute<XmlConfigFileAttribute>()!;
 
-						// Find the matching field on the output type
-						FieldInfo? Field = Type.GetField(FieldName, BindingFlags.Public | BindingFlags.NonPublic | BindingFlags.Instance | BindingFlags.Static);
-						XmlConfigFileAttribute? XmlConfigAttribute = Field?.GetCustomAttribute<XmlConfigFileAttribute>();
-						if(Field == null || XmlConfigAttribute == null)
+							// Try to parse the value and add it to the output array
+							object Value = Reader.ReadObject(TargetMember.Type)!;
+
+							// Read the path of the config file that provided this setting
+							FileReference SourceFile = Reader.ReadFileReference();
+
+							Values[ValueIdx] = new ValueInfo(TargetMember, Value, SourceFile, XmlConfigAttribute);
+						}
+						else
 						{
 							Data = null;
 							return false;
 						}
-
-						// Try to parse the value and add it to the output array
-						object Value = Reader.ReadObject(Field.FieldType)!;
-						
-						// Read the path of the config file that provided this setting
-						FileReference SourceFile = Reader.ReadFileReference();
-
-						Values[ValueIdx] = new ValueInfo(Field, Value, SourceFile, XmlConfigAttribute);
 					}
 
 					// Add it to the type map
@@ -139,6 +228,37 @@ namespace UnrealBuildTool
 			}
 		}
 
+		/// <summary>
+		/// Find a data member (field or property) with the given name and attribute and returns TargetMember wrapper created for it.
+		/// </summary>
+		/// <typeparam name="T">Attribute a member has to have to be considered.</typeparam>
+		/// <param name="Type">Type which members are to be searched</param>
+		/// <param name="MemberName">Name of a member (field or property) to find.</param>
+		/// <returns>TargetMember wrapper or null if no member has been found.</returns>
+		private static XmlConfigData.TargetMember? GetTargetMemberWithAttribute<T>(Type Type, string MemberName)
+			where T : System.Attribute
+		{
+			T? XmlConfigAttribute = null;
+
+			FieldInfo? Field = Type.GetField(MemberName, BindingFlags.Public | BindingFlags.NonPublic | BindingFlags.Instance | BindingFlags.Static);
+			XmlConfigAttribute = Field?.GetCustomAttribute<T>();
+
+			if (Field != null && XmlConfigAttribute != null)
+			{
+				return new XmlConfigData.TargetField(Field);
+			}
+
+			PropertyInfo? Property = Type.GetProperty(MemberName, BindingFlags.Public | BindingFlags.NonPublic | BindingFlags.Instance | BindingFlags.Static);
+			XmlConfigAttribute = Property?.GetCustomAttribute<T>();
+
+			if (Property != null && XmlConfigAttribute != null)
+			{
+				return new XmlConfigData.TargetProperty(Property);
+			}
+
+			return null;
+		}
+
 		/// <summary>
 		/// Writes the coalesced config hierarchy to disk
 		/// </summary>
@@ -159,11 +279,11 @@ namespace UnrealBuildTool
 				{
 					Writer.Write(TypePair.Key.Name);
 					Writer.Write(TypePair.Value.Length);
-					foreach(ValueInfo FieldPair in TypePair.Value)
+					foreach(ValueInfo MemberPair in TypePair.Value)
 					{
-						Writer.Write(FieldPair.FieldInfo.Name);
-						Writer.Write(FieldPair.FieldInfo.FieldType, FieldPair.Value);
-						Writer.Write(FieldPair.SourceFile);
+						Writer.Write(MemberPair.Target.MemberInfo.Name);
+						Writer.Write(MemberPair.Target.Type, MemberPair.Value);
+						Writer.Write(MemberPair.SourceFile);
 					}
 				}
 			}
diff --git a/Engine/Source/Programs/UnrealBuildTool/System/XmlConfigFileAttribute.cs b/Engine/Source/Programs/UnrealBuildTool/System/XmlConfigFileAttribute.cs
--- a/Engine/Source/Programs/UnrealBuildTool/System/XmlConfigFileAttribute.cs
+++ b/Engine/Source/Programs/UnrealBuildTool/System/XmlConfigFileAttribute.cs
@@ -11,7 +11,7 @@ namespace UnrealBuildTool
 	/// <summary>
 	/// Marks a field as being serializable from a config file
 	/// </summary>
-	[AttributeUsage(AttributeTargets.Field, AllowMultiple = true)]
+	[AttributeUsage(AttributeTargets.Field | AttributeTargets.Property, AllowMultiple = true)]
 	class XmlConfigFileAttribute : Attribute
 	{
 		/// <summary>
