Fix builds where the HoloLens platform isn't synced.

 - temporarily disabled hololens native code analysis and architecture override

#rnx

diff --git a/Engine/Platforms/Hololens/Source/Programs/UnrealBuildTool/UEBuildHoloLens.cs b/Engine/Platforms/Hololens/Source/Programs/UnrealBuildTool/UEBuildHoloLens.cs
--- a/Engine/Platforms/Hololens/Source/Programs/UnrealBuildTool/UEBuildHoloLens.cs
+++ b/Engine/Platforms/Hololens/Source/Programs/UnrealBuildTool/UEBuildHoloLens.cs
@@ -83,7 +83,15 @@ namespace UnrealBuildTool
 				Architecture = (WindowsArchitecture)Enum.Parse(typeof(WindowsArchitecture), Info.Architecture, true);
 			}
 		}
-    }
+
+		/// <summary>
+		/// Constructor - temporary #jira UE-149794
+		/// </summary>
+		public HoloLensTargetRules()
+		{
+			//#FIXME: need to set Architecture
+		}
+	}
 
 	/// <summary>
 	/// Read-only wrapper for HoloLens-specific target settings
@@ -150,6 +158,37 @@ namespace UnrealBuildTool
 		#endregion
 	}
 
+
+	abstract partial class TargetRules
+	{
+		/// <summary>
+		/// HoloLens-specific target settings.
+		/// </summary>
+		[ConfigSubObject]
+		public HoloLensTargetRules HoloLensPlatform = new HoloLensTargetRules();
+
+	}
+
+
+	public partial class ReadOnlyTargetRules
+	{
+		private ReadOnlyHoloLensTargetRules? _HoloLensPlatform = null;
+		/// <summary>
+		/// Singleton accessor for HoloLensPlatform
+		/// </summary>
+		public ReadOnlyHoloLensTargetRules HoloLensPlatform
+		{
+			get
+			{
+				if (_HoloLensPlatform == null)
+				{
+					_HoloLensPlatform = new ReadOnlyHoloLensTargetRules(Inner.HoloLensPlatform);
+				}
+				return _HoloLensPlatform;
+			}
+		}
+	}
+
 	[SupportedOSPlatform("windows")]
 	class HoloLensPlatform : WindowsPlatform
 	{
diff --git a/Engine/Source/Programs/UnrealBuildTool/Configuration/TargetRules.cs b/Engine/Source/Programs/UnrealBuildTool/Configuration/TargetRules.cs
--- a/Engine/Source/Programs/UnrealBuildTool/Configuration/TargetRules.cs
+++ b/Engine/Source/Programs/UnrealBuildTool/Configuration/TargetRules.cs
@@ -1814,12 +1814,6 @@ namespace UnrealBuildTool
 		[ConfigSubObject]
 		public WindowsTargetRules WindowsPlatform; // Requires 'this' parameter; initialized in constructor
 
-		/// <summary>
-		/// HoloLens-specific target settings.
-		/// </summary>
-		[ConfigSubObject]
-		public HoloLensTargetRules HoloLensPlatform;
-
 		/// <summary>
 		/// Constructor.
 		/// </summary>
@@ -1833,7 +1827,6 @@ namespace UnrealBuildTool
 			this.ProjectFile = Target.ProjectFile;
 			this.Version = Target.Version;
 			this.WindowsPlatform = new WindowsTargetRules(this);
-			this.HoloLensPlatform = new HoloLensTargetRules(Target);
 
 			// Read settings from config files
 			Dictionary<ConfigDependencyKey, IReadOnlyList<string>?> ConfigValues = new Dictionary<ConfigDependencyKey, IReadOnlyList<string>?>();
@@ -2179,7 +2172,6 @@ namespace UnrealBuildTool
 			LinuxPlatform = new ReadOnlyLinuxTargetRules(Inner.LinuxPlatform);
 			MacPlatform = new ReadOnlyMacTargetRules(Inner.MacPlatform);
 			WindowsPlatform = new ReadOnlyWindowsTargetRules(Inner.WindowsPlatform);
-			HoloLensPlatform = new ReadOnlyHoloLensTargetRules(Inner.HoloLensPlatform);
 		}
 
 		/// <summary>
@@ -3141,12 +3133,6 @@ namespace UnrealBuildTool
 			private set;
 		}
 
-		public ReadOnlyHoloLensTargetRules HoloLensPlatform
-		{
-			get;
-			private set;
-		}
-
 		public bool bShouldCompileAsDLL
 		{
 			get { return Inner.bShouldCompileAsDLL; }
diff --git a/Engine/Source/Programs/UnrealBuildTool/Platform/Windows/VCToolChain.cs b/Engine/Source/Programs/UnrealBuildTool/Platform/Windows/VCToolChain.cs
--- a/Engine/Source/Programs/UnrealBuildTool/Platform/Windows/VCToolChain.cs
+++ b/Engine/Source/Programs/UnrealBuildTool/Platform/Windows/VCToolChain.cs
@@ -1353,6 +1353,7 @@ namespace UnrealBuildTool
 					CompileAction.Arguments.Add(CompileEnvironment.AdditionalArguments);
 				}
 
+				/* FIXME: need to fix this #jira UE-149794
 				if (CompileEnvironment.Platform.IsInGroup(UnrealPlatformGroup.HoloLens) && Target.HoloLensPlatform.bRunNativeCodeAnalysis)
 				{
 					// Add the analysis log to the produced item list.
@@ -1377,7 +1378,7 @@ namespace UnrealBuildTool
 						// A non default ruleset was specified
 						CompileAction.Arguments.AddFormat("/analyze:ruleset \"{0}\"", rulesetFile);
 					}
-				}
+				}*/
 
 				if (SourceFile.HasExtension(".ixx"))
 				{
