UnrealBuildTool: Merging common clang warnings to ClangToolChain base class. Moved all warnings to overridden functions out of GetCompileArguments_Global

Removed -Wsequence-point (Synonym for -Wunsequenced) from Linux as it is enabled by default
#rnx

diff --git a/Engine/Source/Programs/UnrealBuildTool/Platform/IOS/IOSToolChain.cs b/Engine/Source/Programs/UnrealBuildTool/Platform/IOS/IOSToolChain.cs
--- a/Engine/Source/Programs/UnrealBuildTool/Platform/IOS/IOSToolChain.cs
+++ b/Engine/Source/Programs/UnrealBuildTool/Platform/IOS/IOSToolChain.cs
@@ -225,7 +225,19 @@ namespace UnrealBuildTool
 			return OutputFile.ChangeExtension(DebugExtension);
 		}
 
+		/// <inheritdoc/>
+		protected override void GetCompileArguments_WarningsAndErrors(CppCompileEnvironment CompileEnvironment, List<string> Arguments)
+		{
+			base.GetCompileArguments_WarningsAndErrors(CompileEnvironment, Arguments);
 
+			// fix for Xcode 8.3 enabling nonportable include checks, but p4 has some invalid cases in it
+			if (Settings.Value.IOSSDKVersionFloat >= 10.3)
+			{
+				Arguments.Add("-Wno-nonportable-include-path");
+			}
+		}
+
+		/// <inheritdoc/>
 		protected override void GetCompileArguments_Global(CppCompileEnvironment CompileEnvironment, List<string> Arguments)
 		{
 			base.GetCompileArguments_Global(CompileEnvironment, Arguments);
@@ -273,36 +285,12 @@ namespace UnrealBuildTool
 				Arguments.Add("-flto");
 			}
 
-			Arguments.Add("-Wall -Werror");
-			Arguments.Add("-Wdelete-non-virtual-dtor");
-
-			// clang 12.00 has a new warning for copies in ranged loops. Instances have all been fixed up (2020/6/26) but
-			// are likely to be reintroduced due to no equivalent on other platforms at this time so disable the warning
-			// See also MacToolChain.cs
 			if (GetClangVersion().Major >= 12)
 			{
-				Arguments.Add("-Wno-range-loop-analysis");
-
 				// We have 'this' vs nullptr comparisons that get optimized away for newer versions of Clang, which is undesirable until we refactor these checks.
 				Arguments.Add("-fno-delete-null-pointer-checks");
 			}
 
-			if (CompileEnvironment.ShadowVariableWarningLevel != WarningLevel.Off)
-			{
-				Arguments.Add("-Wshadow" + ((CompileEnvironment.ShadowVariableWarningLevel == WarningLevel.Error) ? "" : " -Wno-error=shadow"));
-			}
-
-			if (CompileEnvironment.bEnableUndefinedIdentifierWarnings)
-			{
-				Arguments.Add("-Wundef" + (CompileEnvironment.bUndefinedIdentifierWarningsAsErrors ? "" : " -Wno-error=undef"));
-			}
-
-			// fix for Xcode 8.3 enabling nonportable include checks, but p4 has some invalid cases in it
-			if (Settings.Value.IOSSDKVersionFloat >= 10.3)
-			{
-				Arguments.Add("-Wno-nonportable-include-path");
-			}
-
 			if (IsBitcodeCompilingEnabled(CompileEnvironment.Configuration))
 			{
 				Arguments.Add("-fembed-bitcode");
diff --git a/Engine/Source/Programs/UnrealBuildTool/Platform/Linux/LinuxToolChain.cs b/Engine/Source/Programs/UnrealBuildTool/Platform/Linux/LinuxToolChain.cs
--- a/Engine/Source/Programs/UnrealBuildTool/Platform/Linux/LinuxToolChain.cs
+++ b/Engine/Source/Programs/UnrealBuildTool/Platform/Linux/LinuxToolChain.cs
@@ -141,7 +141,7 @@ namespace UnrealBuildTool
 				{
 					throw new BuildException("LINUX_MULTIARCH_ROOT environment variable is not set; cannot instantiate Linux toolchain");
 				}
-				if (String.IsNullOrEmpty(MultiArchRoot)) 
+				if (String.IsNullOrEmpty(MultiArchRoot))
 				{
 					MultiArchRoot = BaseLinuxPath;
 					Log.TraceInformation("Using LINUX_ROOT (deprecated, consider LINUX_MULTIARCH_ROOT)");
@@ -151,12 +151,12 @@ namespace UnrealBuildTool
 				ToolchainInfo = String.Format("toolchain located at '{0}'", BaseLinuxPath);
 
 				// set up the path to our toolchain
-				ClangPath   = Path.Combine(BaseLinuxPath, @"bin", "clang++" + GetHostPlatformBinarySuffix());
-				ArPath      = Path.Combine(Path.Combine(BaseLinuxPath, String.Format("bin/{0}-{1}", Architecture, "ar" + GetHostPlatformBinarySuffix())));
-				LlvmArPath  = Path.Combine(Path.Combine(BaseLinuxPath, String.Format("bin/{0}", "llvm-ar" + GetHostPlatformBinarySuffix())));
-				RanlibPath  = Path.Combine(Path.Combine(BaseLinuxPath, String.Format("bin/{0}-{1}", Architecture, "ranlib" + GetHostPlatformBinarySuffix())));
+				ClangPath = Path.Combine(BaseLinuxPath, @"bin", "clang++" + GetHostPlatformBinarySuffix());
+				ArPath = Path.Combine(Path.Combine(BaseLinuxPath, String.Format("bin/{0}-{1}", Architecture, "ar" + GetHostPlatformBinarySuffix())));
+				LlvmArPath = Path.Combine(Path.Combine(BaseLinuxPath, String.Format("bin/{0}", "llvm-ar" + GetHostPlatformBinarySuffix())));
+				RanlibPath = Path.Combine(Path.Combine(BaseLinuxPath, String.Format("bin/{0}-{1}", Architecture, "ranlib" + GetHostPlatformBinarySuffix())));
 				ObjcopyPath = Path.Combine(Path.Combine(BaseLinuxPath, String.Format("bin/{0}", "llvm-objcopy" + GetHostPlatformBinarySuffix())));
-				StripPath   = ObjcopyPath;
+				StripPath = ObjcopyPath;
 
 				// When cross-compiling on Windows, use old FixDeps. It is slow, but it does not have timing issues
 				bUseFixdeps = BuildHostPlatform.Current.Platform == UnrealTargetPlatform.Win64;
@@ -246,7 +246,7 @@ namespace UnrealBuildTool
 			bool bUseCmdExe = BuildHostPlatform.Current.Platform == UnrealTargetPlatform.Win64;
 			string DumpCommand = bUseCmdExe ? "\"{0}\" \"{1}\" \"{2}\" 2>NUL" : "\"{0}\" -c -o \"{2}\" \"{1}\"";
 			FileItem EncodedBinarySymbolsFile = FileItem.GetItemByPath(Path.Combine(LinkEnvironment.OutputDirectory!.FullName, OutputFile.Location.GetFileNameWithoutExtension() + ".sym"));
-			FileItem SymbolsFile  = FileItem.GetItemByPath(Path.Combine(LinkEnvironment.LocalShadowDirectory!.FullName, OutputFile.Location.GetFileName() + ".psym"));
+			FileItem SymbolsFile = FileItem.GetItemByPath(Path.Combine(LinkEnvironment.LocalShadowDirectory!.FullName, OutputFile.Location.GetFileName() + ".psym"));
 			FileItem StrippedFile = FileItem.GetItemByPath(Path.Combine(LinkEnvironment.LocalShadowDirectory.FullName, OutputFile.Location.GetFileName() + "_nodebug"));
 			FileItem DebugFile = FileItem.GetItemByPath(Path.Combine(LinkEnvironment.OutputDirectory.FullName, OutputFile.Location.GetFileNameWithoutExtension() + ".debug"));
 
@@ -514,6 +514,47 @@ namespace UnrealBuildTool
 			return false;
 		}
 
+		/// <inheritdoc/>
+		protected override void GetCompileArguments_WarningsAndErrors(CppCompileEnvironment CompileEnvironment, List<string> Arguments)
+		{
+			base.GetCompileArguments_WarningsAndErrors(CompileEnvironment, Arguments);
+
+			//Arguments.Add("-Wunreachable-code");            // additional warning not normally included in Wall: warns if there is code that will never be executed - not helpful due to bIsGCC and similar
+
+			Arguments.Add("-Wno-unused-private-field"); // MultichannelTcpSocket.h triggers this, possibly more
+			Arguments.Add("-Wno-tautological-compare"); // this hides the "warning : comparison of unsigned expression < 0 is always false" type warnings due to constant comparisons, which are possible with template arguments
+			Arguments.Add("-Wno-undefined-bool-conversion"); // hides checking if 'this' pointer is null
+			Arguments.Add("-Wno-unused-local-typedef"); // clang is being overly strict here? PhysX headers trigger this.
+			Arguments.Add("-Wno-inconsistent-missing-override");    // these have to be suppressed for UE 4.8, should be fixed later.
+			Arguments.Add("-Wno-undefined-var-template"); // not really a good warning to disable
+			Arguments.Add("-Wno-unused-lambda-capture");  // suppressed because capturing of compile-time constants is seemingly inconsistent. And MSVC doesn't do that.
+			Arguments.Add("-Wno-unused-variable");
+
+			Arguments.Add("-Wno-unused-function"); // this will hide the warnings about static functions in headers that aren't used in every single .cpp file
+			Arguments.Add("-Wno-switch"); // this hides the "enumeration value 'XXXXX' not handled in switch [-Wswitch]" warnings - we should maybe remove this at some point and add UE_LOG(, Fatal, ) to default cases
+			Arguments.Add("-Wno-unknown-pragmas");          // Slate triggers this (with its optimize on/off pragmas)
+			Arguments.Add("-Wno-invalid-offsetof"); // needed to suppress warnings about using offsetof on non-POD types.
+			Arguments.Add("-Wno-gnu-string-literal-operator-template"); // we use this feature to allow static FNames.
+
+			if (CompileEnvironment.bPGOOptimize)
+			{
+				//
+				// Clang emits a warning for each compiled function that doesn't have a matching entry in the profile data.
+				// This can happen when the profile data is older than the binaries we're compiling.
+				//
+				// Disable this warning. It's far too verbose.
+				//
+				Arguments.Add("-Wno-backend-plugin");
+			}
+
+			// shipping builds will cause this warning with "ensure", so disable only in those case
+			if (CompileEnvironment.Configuration == CppConfiguration.Shipping)
+			{
+				Arguments.Add("-Wno-unused-value");
+			}
+		}
+
+		/// <inheritdoc/>
 		protected override void GetCompileArguments_Global(CppCompileEnvironment CompileEnvironment, List<string> Arguments)
 		{
 			base.GetCompileArguments_Global(CompileEnvironment, Arguments);
@@ -556,8 +597,6 @@ namespace UnrealBuildTool
 				Arguments.Add("-fsanitize=memory -fsanitize-memory-track-origins -DFORCE_ANSI_ALLOCATOR=1");
 			}
 
-			Arguments.Add("-Wall -Werror");
-
 			if (CompilerVersionGreaterOrEqual(12, 0, 0))
 			{
 				Arguments.Add("-fbinutils-version=2.36");
@@ -568,16 +607,11 @@ namespace UnrealBuildTool
 				Arguments.Add("-funwind-tables");               // generate unwind tables as they are needed for backtrace (on x86(64) they are generated implicitly)
 			}
 
-			Arguments.Add("-Wsequence-point");              // additional warning not normally included in Wall: warns if order of operations is ambigious
-			//Arguments.Add("-Wunreachable-code");            // additional warning not normally included in Wall: warns if there is code that will never be executed - not helpful due to bIsGCC and similar
-			//Arguments.Add("-Wshadow");                      // additional warning not normally included in Wall: warns if there variable/typedef shadows some other variable - not helpful because we have gobs of code that shadows variables
-			Arguments.Add("-Wdelete-non-virtual-dtor");
-
 			Arguments.Add(ArchitectureSpecificSwitches(CompileEnvironment.Architecture));
 
 			Arguments.Add("-fno-math-errno");               // do not assume that math ops have side effects
 
-			Arguments.Add(GetRTTIFlag(CompileEnvironment));	// flag for run-time type info
+			Arguments.Add(GetRTTIFlag(CompileEnvironment)); // flag for run-time type info
 
 			if (CompileEnvironment.Architecture.StartsWith("x86_64"))
 			{
@@ -605,40 +639,10 @@ namespace UnrealBuildTool
 
 			Arguments.Add("-fdiagnostics-absolute-paths"); // output full paths to the files when the build fails
 
-			Arguments.Add("-Wno-unused-private-field"); // MultichannelTcpSocket.h triggers this, possibly more
-			Arguments.Add("-Wno-tautological-compare"); // this hides the "warning : comparison of unsigned expression < 0 is always false" type warnings due to constant comparisons, which are possible with template arguments
-			Arguments.Add("-Wno-undefined-bool-conversion"); // hides checking if 'this' pointer is null
-			Arguments.Add("-Wno-unused-local-typedef");	// clang is being overly strict here? PhysX headers trigger this.
-			Arguments.Add("-Wno-inconsistent-missing-override");	// these have to be suppressed for UE 4.8, should be fixed later.
-			Arguments.Add("-Wno-undefined-var-template"); // not really a good warning to disable
-			Arguments.Add("-Wno-unused-lambda-capture");  // suppressed because capturing of compile-time constants is seemingly inconsistent. And MSVC doesn't do that.
-			Arguments.Add("-Wno-unused-variable");
-
-			if (CompilerVersionGreaterOrEqual(13, 0, 0))
-			{
-				Arguments.Add("-Wno-unused-but-set-variable");
-				Arguments.Add("-Wno-unused-but-set-parameter");
-				Arguments.Add("-Wno-ordered-compare-function-pointers");
-			}
-
-			Arguments.Add("-Wno-unused-function"); // this will hide the warnings about static functions in headers that aren't used in every single .cpp file
-			Arguments.Add("-Wno-switch"); // this hides the "enumeration value 'XXXXX' not handled in switch [-Wswitch]" warnings - we should maybe remove this at some point and add UE_LOG(, Fatal, ) to default cases
-			Arguments.Add("-Wno-unknown-pragmas");			// Slate triggers this (with its optimize on/off pragmas)
-			Arguments.Add("-Wno-invalid-offsetof"); // needed to suppress warnings about using offsetof on non-POD types.
-			Arguments.Add("-Wno-gnu-string-literal-operator-template"); // we use this feature to allow static FNames.
-
 			// Profile Guided Optimization (PGO) and Link Time Optimization (LTO)
 			// Whether we actually can enable that is checked in CanUseAdvancedLinkerFeatures() earlier
 			if (CompileEnvironment.bPGOOptimize)
 			{
-				//
-				// Clang emits a warning for each compiled function that doesn't have a matching entry in the profile data.
-				// This can happen when the profile data is older than the binaries we're compiling.
-				//
-				// Disable this warning. It's far too verbose.
-				//
-				Arguments.Add("-Wno-backend-plugin");
-
 				Log.TraceInformationOnce("Enabling Profile Guided Optimization (PGO). Linking will take a while.");
 				Arguments.Add(string.Format(" -fprofile-instr-use=\"{0}\"", Path.Combine(CompileEnvironment.PGODirectory!, CompileEnvironment.PGOFilenamePrefix!)));
 			}
@@ -652,7 +656,7 @@ namespace UnrealBuildTool
 			// Whether we actually can enable that is checked in CanUseAdvancedLinkerFeatures() earlier
 			if (CompileEnvironment.bAllowLTCG)
 			{
-				if((Options & LinuxToolChainOptions.EnableThinLTO) != 0)
+				if ((Options & LinuxToolChainOptions.EnableThinLTO) != 0)
 				{
 					Arguments.Add("-flto=thin");
 				}
@@ -662,26 +666,14 @@ namespace UnrealBuildTool
 				}
 			}
 
-			if (CompileEnvironment.ShadowVariableWarningLevel != WarningLevel.Off)
-			{
-				Arguments.Add("-Wshadow" + ((CompileEnvironment.ShadowVariableWarningLevel == WarningLevel.Error) ? "" : " -Wno-error=shadow"));
-			}
-
-			if (CompileEnvironment.bEnableUndefinedIdentifierWarnings)
-			{
-				Arguments.Add("-Wundef" + (CompileEnvironment.bUndefinedIdentifierWarningsAsErrors ? "" : " -Wno-error=undef"));
-			}
-
 			//Arguments.Add("-DOPERATOR_NEW_INLINE=FORCENOINLINE");
 
-			bool bRetainFramePointers = CompileEnvironment.bRetainFramePointers 
+			bool bRetainFramePointers = CompileEnvironment.bRetainFramePointers
 				|| Options.HasFlag(LinuxToolChainOptions.EnableAddressSanitizer) || Options.HasFlag(LinuxToolChainOptions.EnableMemorySanitizer)
 				|| CompileEnvironment.Configuration == CppConfiguration.Debug;
 
-			// shipping builds will cause this warning with "ensure", so disable only in those case
 			if (CompileEnvironment.Configuration == CppConfiguration.Shipping)
 			{
-				Arguments.Add("-Wno-unused-value");
 				if (!bRetainFramePointers)
 				{
 					Arguments.Add("-fomit-frame-pointer");
@@ -692,7 +684,7 @@ namespace UnrealBuildTool
 			{
 				Arguments.Add("-fno-inline");                   // disable inlining for better debuggability (e.g. callstacks, "skip file" in gdb)
 				Arguments.Add("-fstack-protector");             // detect stack smashing
-				//Arguments.Add("-fsanitize=address");            // detect address based errors (support properly and link to libasan)
+																//Arguments.Add("-fsanitize=address");            // detect address based errors (support properly and link to libasan)
 			}
 
 			if (bRetainFramePointers)
@@ -806,8 +798,8 @@ namespace UnrealBuildTool
 			{
 				if (!Value.StartsWith("\"") && (Value.Contains(" ") || Value.Contains("$")))
 				{
-					Value = Value.Trim('\"');		// trim any leading or trailing quotes
-					Value = "\"" + Value + "\"";	// ensure wrap string with double quotes
+					Value = Value.Trim('\"');       // trim any leading or trailing quotes
+					Value = "\"" + Value + "\"";    // ensure wrap string with double quotes
 				}
 
 				// replace double quotes to escaped double quotes if exists
@@ -890,7 +882,7 @@ namespace UnrealBuildTool
 			// RPATH for third party libs
 			Result += " -Wl,-rpath=${ORIGIN}";
 			Result += " -Wl,-rpath-link=${ORIGIN}";
-			Result += " -Wl,-rpath=${ORIGIN}/..";	// for modules that are in sub-folders of the main Engine/Binary/Linux folder
+			Result += " -Wl,-rpath=${ORIGIN}/..";   // for modules that are in sub-folders of the main Engine/Binary/Linux folder
 			if (LinkEnvironment.Architecture.StartsWith("x86_64"))
 			{
 				Result += " -Wl,-rpath=${ORIGIN}/../../../Engine/Binaries/ThirdParty/Qualcomm/Linux";
@@ -953,7 +945,7 @@ namespace UnrealBuildTool
 			// whether we actually can do that is checked in CanUseAdvancedLinkerFeatures() earlier
 			if (LinkEnvironment.bAllowLTCG)
 			{
-				if((Options & LinuxToolChainOptions.EnableThinLTO) != 0)
+				if ((Options & LinuxToolChainOptions.EnableThinLTO) != 0)
 				{
 					Result += String.Format(" -flto=thin -Wl,--thinlto-jobs={0}", Utils.GetPhysicalProcessorCount());
 				}
@@ -1227,7 +1219,7 @@ namespace UnrealBuildTool
 				}
 
 				// Generate the included header dependency list
-				if(!bPreprocessDepends && CompileEnvironment.bGenerateDependenciesFile)
+				if (!bPreprocessDepends && CompileEnvironment.bGenerateDependenciesFile)
 				{
 					FileItem DependencyListFile = FileItem.GetItemByFileReference(FileReference.Combine(OutputDir, Path.GetFileName(SourceFile.AbsolutePath) + ".d"));
 					FileArguments.Add(string.Format("-MD -MF\"{0}\"", NormalizeCommandLinePath(DependencyListFile)));
@@ -1577,11 +1569,11 @@ namespace UnrealBuildTool
 				}
 			}
 
-			foreach(string RuntimeLibaryPath in LinkEnvironment.RuntimeLibraryPaths)
+			foreach (string RuntimeLibaryPath in LinkEnvironment.RuntimeLibraryPaths)
 			{
 				string RelativePath = RuntimeLibaryPath;
 
-				if(!RelativePath.StartsWith("$"))
+				if (!RelativePath.StartsWith("$"))
 				{
 					if (LinkEnvironment.bIsBuildingDLL)
 					{
@@ -1769,7 +1761,7 @@ namespace UnrealBuildTool
 			{
 				LinkCommandString = LinkCommandString.Replace("{", "'{");
 				LinkCommandString = LinkCommandString.Replace("}", "}'");
-				LinkCommandString = LinkCommandString.Replace("$'{", "'${");	// fixing $'{ORIGIN}' to be '${ORIGIN}'
+				LinkCommandString = LinkCommandString.Replace("$'{", "'${");    // fixing $'{ORIGIN}' to be '${ORIGIN}'
 			}
 
 			string LinkScriptName = string.Format((bUseCmdExe ? "Link-{0}.link.bat" : "Link-{0}.link.sh"), OutputFile.Location.GetFileName());
@@ -1848,23 +1840,23 @@ namespace UnrealBuildTool
 						Directory.CreateDirectory(Path.GetDirectoryName(FixDepsScriptPath)!);
 						using (StreamWriter Writer = File.CreateText(FixDepsScriptPath))
 						{
-						if (bUseCmdExe)
-						{
-							Writer.NewLine = "\r\n";
-							Writer.WriteLine("@echo off");
-							Writer.WriteLine("rem Automatically generated by UnrealBuildTool");
-							Writer.WriteLine("rem *DO NOT EDIT*");
-							Writer.WriteLine();
-						}
-						else
-						{
-							Writer.NewLine = "\n";
-							Writer.WriteLine("#!/bin/sh");
-							Writer.WriteLine("# Automatically generated by UnrealBuildTool");
-							Writer.WriteLine("# *DO NOT EDIT*");
-							Writer.WriteLine();
-							Writer.WriteLine("set -o errexit");
-						}
+							if (bUseCmdExe)
+							{
+								Writer.NewLine = "\r\n";
+								Writer.WriteLine("@echo off");
+								Writer.WriteLine("rem Automatically generated by UnrealBuildTool");
+								Writer.WriteLine("rem *DO NOT EDIT*");
+								Writer.WriteLine();
+							}
+							else
+							{
+								Writer.NewLine = "\n";
+								Writer.WriteLine("#!/bin/sh");
+								Writer.WriteLine("# Automatically generated by UnrealBuildTool");
+								Writer.WriteLine("# *DO NOT EDIT*");
+								Writer.WriteLine();
+								Writer.WriteLine("set -o errexit");
+							}
 						}
 					}
 
@@ -1939,58 +1931,58 @@ namespace UnrealBuildTool
 					Directory.CreateDirectory(Path.GetDirectoryName(RelinkScriptFullPath)!);
 					using (StreamWriter RelinkWriter = File.CreateText(RelinkScriptFullPath))
 					{
-					string RelinkInvocation = LinkCommandString;
-					string Replace = "-Wl,--allow-shlib-undefined";
-					RelinkInvocation = RelinkInvocation.Replace(Replace, EngineAndGameLibrariesString);
+						string RelinkInvocation = LinkCommandString;
+						string Replace = "-Wl,--allow-shlib-undefined";
+						RelinkInvocation = RelinkInvocation.Replace(Replace, EngineAndGameLibrariesString);
 
-					// should be the same as RelinkedFileRef
-					RelinkInvocation = RelinkInvocation.Replace(LinkOutputFileForwardSlashes, RelinkedFileForwardSlashes);
-					RelinkInvocation = RelinkInvocation.Replace("$", "\\$");
+						// should be the same as RelinkedFileRef
+						RelinkInvocation = RelinkInvocation.Replace(LinkOutputFileForwardSlashes, RelinkedFileForwardSlashes);
+						RelinkInvocation = RelinkInvocation.Replace("$", "\\$");
 
-					if (bUseCmdExe)
-					{
-						RelinkWriter.WriteLine("@echo off");
-						RelinkWriter.WriteLine("rem Automatically generated by UnrealBuildTool");
-						RelinkWriter.WriteLine("rem *DO NOT EDIT*");
-						RelinkWriter.WriteLine();
-						RelinkWriter.WriteLine("set Retries=0");
-						RelinkWriter.WriteLine(":relinkloop");
-						RelinkWriter.WriteLine("if %Retries% GEQ 10 goto failedtorelink");
-						RelinkWriter.WriteLine(RelinkInvocation);
-						RelinkWriter.WriteLine("if %errorlevel% neq 0 goto sleepandretry");
-						RelinkWriter.WriteLine("copy /B \"{0}\" \"{1}.temp\" >NUL 2>NUL", RelinkedFileForwardSlashes, OutputFile.AbsolutePath);
-						RelinkWriter.WriteLine("if %errorlevel% neq 0 goto sleepandretry");
-						RelinkWriter.WriteLine("move /Y \"{0}.temp\" \"{1}\" >NUL 2>NUL", OutputFile.AbsolutePath, OutputFile.AbsolutePath);
-						RelinkWriter.WriteLine("if %errorlevel% neq 0 goto sleepandretry");
-						RelinkWriter.WriteLine(GetDumpEncodeDebugCommand(LinkEnvironment, OutputFile));
-						RelinkWriter.WriteLine("echo \"Dummy\" >> \"{0}\" && copy /b \"{0}\" +,,", RelinkActionDummyProductRef.FullName);
-						RelinkWriter.WriteLine("echo Relinked {0} successfully after %Retries% retries", OutputFile.AbsolutePath);
-						RelinkWriter.WriteLine("exit 0");
-						RelinkWriter.WriteLine(":sleepandretry");
-						RelinkWriter.WriteLine("ping 127.0.0.1 -n 1 -w 5000 >NUL 2>NUL");     // timeout complains about lack of redirection
-						RelinkWriter.WriteLine("set /a Retries+=1");
-						RelinkWriter.WriteLine("goto relinkloop");
-						RelinkWriter.WriteLine(":failedtorelink");
-						RelinkWriter.WriteLine("echo Failed to relink {0} after %Retries% retries", OutputFile.AbsolutePath);
-						RelinkWriter.WriteLine("exit 1");
-					}
-					else
-					{
-						RelinkWriter.NewLine = "\n";
-						RelinkWriter.WriteLine("#!/bin/sh");
-						RelinkWriter.WriteLine("# Automatically generated by UnrealBuildTool");
-						RelinkWriter.WriteLine("# *DO NOT EDIT*");
-						RelinkWriter.WriteLine();
-						RelinkWriter.WriteLine("set -o errexit");
-						RelinkWriter.WriteLine(RelinkInvocation);
-						RelinkWriter.WriteLine("TIMESTAMP=`stat --format %y \"{0}\"`", OutputFile.AbsolutePath);
-						RelinkWriter.WriteLine("cp \"{0}\" \"{1}.temp\"", RelinkedFileForwardSlashes, OutputFile.AbsolutePath);
-						RelinkWriter.WriteLine("mv \"{0}.temp\" \"{1}\"", OutputFile.AbsolutePath, OutputFile.AbsolutePath);
-						RelinkWriter.WriteLine(GetDumpEncodeDebugCommand(LinkEnvironment, OutputFile));
-						RelinkWriter.WriteLine("touch -d \"$TIMESTAMP\" \"{0}\"", OutputFile.AbsolutePath);
-						RelinkWriter.WriteLine();
-						RelinkWriter.WriteLine("echo \"Dummy\" >> \"{0}\"", RelinkActionDummyProductRef.FullName);
-					}
+						if (bUseCmdExe)
+						{
+							RelinkWriter.WriteLine("@echo off");
+							RelinkWriter.WriteLine("rem Automatically generated by UnrealBuildTool");
+							RelinkWriter.WriteLine("rem *DO NOT EDIT*");
+							RelinkWriter.WriteLine();
+							RelinkWriter.WriteLine("set Retries=0");
+							RelinkWriter.WriteLine(":relinkloop");
+							RelinkWriter.WriteLine("if %Retries% GEQ 10 goto failedtorelink");
+							RelinkWriter.WriteLine(RelinkInvocation);
+							RelinkWriter.WriteLine("if %errorlevel% neq 0 goto sleepandretry");
+							RelinkWriter.WriteLine("copy /B \"{0}\" \"{1}.temp\" >NUL 2>NUL", RelinkedFileForwardSlashes, OutputFile.AbsolutePath);
+							RelinkWriter.WriteLine("if %errorlevel% neq 0 goto sleepandretry");
+							RelinkWriter.WriteLine("move /Y \"{0}.temp\" \"{1}\" >NUL 2>NUL", OutputFile.AbsolutePath, OutputFile.AbsolutePath);
+							RelinkWriter.WriteLine("if %errorlevel% neq 0 goto sleepandretry");
+							RelinkWriter.WriteLine(GetDumpEncodeDebugCommand(LinkEnvironment, OutputFile));
+							RelinkWriter.WriteLine("echo \"Dummy\" >> \"{0}\" && copy /b \"{0}\" +,,", RelinkActionDummyProductRef.FullName);
+							RelinkWriter.WriteLine("echo Relinked {0} successfully after %Retries% retries", OutputFile.AbsolutePath);
+							RelinkWriter.WriteLine("exit 0");
+							RelinkWriter.WriteLine(":sleepandretry");
+							RelinkWriter.WriteLine("ping 127.0.0.1 -n 1 -w 5000 >NUL 2>NUL");     // timeout complains about lack of redirection
+							RelinkWriter.WriteLine("set /a Retries+=1");
+							RelinkWriter.WriteLine("goto relinkloop");
+							RelinkWriter.WriteLine(":failedtorelink");
+							RelinkWriter.WriteLine("echo Failed to relink {0} after %Retries% retries", OutputFile.AbsolutePath);
+							RelinkWriter.WriteLine("exit 1");
+						}
+						else
+						{
+							RelinkWriter.NewLine = "\n";
+							RelinkWriter.WriteLine("#!/bin/sh");
+							RelinkWriter.WriteLine("# Automatically generated by UnrealBuildTool");
+							RelinkWriter.WriteLine("# *DO NOT EDIT*");
+							RelinkWriter.WriteLine();
+							RelinkWriter.WriteLine("set -o errexit");
+							RelinkWriter.WriteLine(RelinkInvocation);
+							RelinkWriter.WriteLine("TIMESTAMP=`stat --format %y \"{0}\"`", OutputFile.AbsolutePath);
+							RelinkWriter.WriteLine("cp \"{0}\" \"{1}.temp\"", RelinkedFileForwardSlashes, OutputFile.AbsolutePath);
+							RelinkWriter.WriteLine("mv \"{0}.temp\" \"{1}\"", OutputFile.AbsolutePath, OutputFile.AbsolutePath);
+							RelinkWriter.WriteLine(GetDumpEncodeDebugCommand(LinkEnvironment, OutputFile));
+							RelinkWriter.WriteLine("touch -d \"$TIMESTAMP\" \"{0}\"", OutputFile.AbsolutePath);
+							RelinkWriter.WriteLine();
+							RelinkWriter.WriteLine("echo \"Dummy\" >> \"{0}\"", RelinkActionDummyProductRef.FullName);
+						}
 					}
 
 					RelinkAction.CommandPath = ShellBinary;
diff --git a/Engine/Source/Programs/UnrealBuildTool/Platform/Mac/MacToolChain.cs b/Engine/Source/Programs/UnrealBuildTool/Platform/Mac/MacToolChain.cs
--- a/Engine/Source/Programs/UnrealBuildTool/Platform/Mac/MacToolChain.cs
+++ b/Engine/Source/Programs/UnrealBuildTool/Platform/Mac/MacToolChain.cs
@@ -191,11 +191,19 @@ namespace UnrealBuildTool
 			return ArchArg;
 		}
 
+		/// <inheritdoc/>
+		protected override void GetCompileArguments_WarningsAndErrors(CppCompileEnvironment CompileEnvironment, List<string> Arguments)
+		{
+			base.GetCompileArguments_WarningsAndErrors(CompileEnvironment, Arguments);
+
+			//Arguments.Add("-Wsign-compare"); // fed up of not seeing the signed/unsigned warnings we get on Windows - lets enable them here too.
+		}
+
+		/// <inheritdoc/>
 		protected override void GetCompileArguments_Global(CppCompileEnvironment CompileEnvironment, List<string> Arguments)
 		{
 			base.GetCompileArguments_Global(CompileEnvironment, Arguments);
 
-
 			Arguments.Add("-fexceptions");
 			Arguments.Add("-DPLATFORM_EXCEPTIONS_DISABLED=0");
 
@@ -217,29 +225,7 @@ namespace UnrealBuildTool
 			if (Options.HasFlag(MacToolChainOptions.EnableUndefinedBehaviorSanitizer))
 			{
 				Arguments.Add("-fsanitize=undefined");
-			}			
-
-			Arguments.Add("-Wall -Werror");
-			Arguments.Add("-Wdelete-non-virtual-dtor");
-
-			// clang 12.00 has a new warning for copies in ranged loops. Instances have all been fixed up (2020/6/26) but
-			// are likely to be reintroduced due to no equivalent on other platforms at this time so disable the warning
-			if (GetClangVersion().Major >= 12)
-			{
-				Arguments.Add("-Wno-range-loop-analysis ");
-			}			
-
-			//Arguments.Add("-Wsign-compare"); // fed up of not seeing the signed/unsigned warnings we get on Windows - lets enable them here too.
-
-			if (CompileEnvironment.ShadowVariableWarningLevel != WarningLevel.Off)
-			{
-				Arguments.Add("-Wshadow" + ((CompileEnvironment.ShadowVariableWarningLevel == WarningLevel.Error) ? "" : " -Wno-error=shadow"));
-			}
-			
-			if (CompileEnvironment.bEnableUndefinedIdentifierWarnings)
-			{
-				Arguments.Add("-Wundef" + (CompileEnvironment.bUndefinedIdentifierWarningsAsErrors ? "" : " -Wno-error=undef"));
-			}
+			}					
 
 			if (CompileEnvironment.bEnableOSX109Support)
 			{
diff --git a/Engine/Source/Programs/UnrealBuildTool/ToolChain/AppleToolChain.cs b/Engine/Source/Programs/UnrealBuildTool/ToolChain/AppleToolChain.cs
--- a/Engine/Source/Programs/UnrealBuildTool/ToolChain/AppleToolChain.cs
+++ b/Engine/Source/Programs/UnrealBuildTool/ToolChain/AppleToolChain.cs
@@ -183,7 +183,13 @@ namespace UnrealBuildTool
 			if (ClangVersion == null)
 			{
 				FileReference ClangLocation = new FileReference("/usr/bin/clang");
-				ClangVersion = RunToolAndCaptureVersion(ClangLocation, "--version");			
+				ClangVersion = RunToolAndCaptureVersion(ClangLocation, "--version");
+			}
+			if (ClangVersionMajor == -1 || ClangVersionMinor == -1 || ClangVersionPatch == -1)
+			{
+				ClangVersionMajor = ClangVersion.Major;
+				ClangVersionMinor = ClangVersion.Minor;
+				ClangVersionPatch = ClangVersion.Build;
 			}
 			return ClangVersion;
 		}
@@ -243,8 +249,25 @@ namespace UnrealBuildTool
 			}
 		}
 
+		/// <inheritdoc/>
+		protected override void GetCompileArguments_WarningsAndErrors(CppCompileEnvironment CompileEnvironment, List<string> Arguments)
+		{
+			base.GetCompileArguments_WarningsAndErrors(CompileEnvironment, Arguments);
+
+			// clang 12.00 has a new warning for copies in ranged loops. Instances have all been fixed up (2020/6/26) but
+			// are likely to be reintroduced due to no equivalent on other platforms at this time so disable the warning
+			if (GetClangVersion().Major >= 12)
+			{
+				Arguments.Add("-Wno-range-loop-analysis");
+			}
+		}
+
+		/// <inheritdoc/>
 		protected override void GetCompileArguments_Global(CppCompileEnvironment CompileEnvironment, List<string> Arguments)
 		{
+			// Ensure Clang version is set for base.GetCompileArguments_Global()
+			GetClangVersion();
+
 			base.GetCompileArguments_Global(CompileEnvironment, Arguments);
 
 			Arguments.Add("-fmessage-length=0");
diff --git a/Engine/Source/Programs/UnrealBuildTool/ToolChain/ClangToolChain.cs b/Engine/Source/Programs/UnrealBuildTool/ToolChain/ClangToolChain.cs
--- a/Engine/Source/Programs/UnrealBuildTool/ToolChain/ClangToolChain.cs
+++ b/Engine/Source/Programs/UnrealBuildTool/ToolChain/ClangToolChain.cs
@@ -3,6 +3,7 @@
 using EpicGames.Core;
 using System;
 using System.Collections.Generic;
+using System.Diagnostics;
 using System.Linq;
 using UnrealBuildBase;
 
@@ -63,14 +64,12 @@ namespace UnrealBuildTool
 		protected bool CompilerVersionGreaterOrEqual(int Major, int Minor, int Patch)
 		{
 			// TODO: Temporary verification check to ensure a clang version has been set until ClangVersion is standarized across all clang-based toolchains.
-			if (ClangVersionMajor == -1)
+			if (ClangVersionMajor == -1 || ClangVersionMinor == -1 || ClangVersionPatch == -1)
 			{
-				throw new BuildException($"ClangVersion not found, unable to check compiler version requirements");
+				throw new BuildException($"ClangVersion not valid ({ClangVersionMajor}.{ClangVersionMinor}.{ClangVersionPatch}), unable to check compiler version requirements");
 			}
 
-			return ClangVersionMajor > Major ||
-				(ClangVersionMajor == Major && ClangVersionMinor > Minor) ||
-				(ClangVersionMajor == Major && ClangVersionMinor == Minor && ClangVersionPatch >= Patch);
+			return new Version(ClangVersionMajor, ClangVersionMinor, ClangVersionPatch) >= new Version(Major, Minor, Patch);
 		}
 
 		/// <summary>
@@ -79,14 +78,12 @@ namespace UnrealBuildTool
 		protected bool CompilerVersionLessThan(int Major, int Minor, int Patch)
 		{
 			// TODO: Temporary verification check until ClangVersion is standarized across all clang-based toolchains to ensure a version has been set.
-			if (ClangVersionMajor == -1)
+			if (ClangVersionMajor == -1 || ClangVersionMinor == -1 || ClangVersionPatch == -1)
 			{
-				throw new BuildException($"ClangVersion not found, unable to check compiler version requirements");
+				throw new BuildException($"ClangVersion not valid ({ClangVersionMajor}.{ClangVersionMinor}.{ClangVersionPatch}), unable to check compiler version requirements");
 			}
 
-			return ClangVersionMajor < Major ||
-				(ClangVersionMajor == Major && ClangVersionMinor < Minor) ||
-				(ClangVersionMajor == Major && ClangVersionMinor == Minor && ClangVersionPatch < Patch);
+			return new Version(ClangVersionMajor, ClangVersionMinor, ClangVersionPatch) < new Version(Major, Minor, Patch);
 		}
 
 		protected virtual void GetCppStandardCompileArgument(CppCompileEnvironment CompileEnvironment, List<string> Arguments)
@@ -205,6 +202,39 @@ namespace UnrealBuildTool
 			Arguments.AddRange(CompileEnvironment.ForceIncludeFiles.Select(ForceIncludeFile => GetForceIncludeFileArgument(ForceIncludeFile)));
 		}
 
+		/// <summary>
+		/// Common compile arguments that control which warnings are enabled.
+		/// https://clang.llvm.org/docs/DiagnosticsReference.html
+		/// </summary>
+		/// <param name="CompileEnvironment"></param>
+		/// <param name="Arguments"></param>
+		protected virtual void GetCompileArguments_WarningsAndErrors(CppCompileEnvironment CompileEnvironment, List<string> Arguments)
+		{
+			Arguments.Add("-Wall");                     // https://clang.llvm.org/docs/DiagnosticsReference.html#wall
+			Arguments.Add("-Werror");                   // https://clang.llvm.org/docs/UsersManual.html#cmdoption-werror
+
+			Arguments.Add("-Wdelete-non-virtual-dtor"); // https://clang.llvm.org/docs/DiagnosticsReference.html#wdelete-non-virtual-dtor
+
+			// https://clang.llvm.org/docs/DiagnosticsReference.html#wshadow
+			if (CompileEnvironment.ShadowVariableWarningLevel != WarningLevel.Off)
+			{
+				Arguments.Add("-Wshadow" + ((CompileEnvironment.ShadowVariableWarningLevel == WarningLevel.Error) ? "" : " -Wno-error=shadow"));
+			}
+
+			// https://clang.llvm.org/docs/DiagnosticsReference.html#wundef
+			if (CompileEnvironment.bEnableUndefinedIdentifierWarnings)
+			{
+				Arguments.Add("-Wundef" + (CompileEnvironment.bUndefinedIdentifierWarningsAsErrors ? "" : " -Wno-error=undef"));
+			}
+
+			if (CompilerVersionGreaterOrEqual(13, 0, 0))
+			{
+				Arguments.Add("-Wno-unused-but-set-variable");           // https://clang.llvm.org/docs/DiagnosticsReference.html#wunused-but-set-variable
+				Arguments.Add("-Wno-unused-but-set-parameter");          // https://clang.llvm.org/docs/DiagnosticsReference.html#wunused-but-set-parameter
+				Arguments.Add("-Wno-ordered-compare-function-pointers"); // https://clang.llvm.org/docs/DiagnosticsReference.html#wordered-compare-function-pointers
+			}
+		}
+
 
 		/// <summary>
 		/// Common compile arguments for all files in a module.
@@ -224,6 +254,9 @@ namespace UnrealBuildTool
 
 			// Add preprocessor definitions to the argument list.
 			GetCompileArguments_PreprocessorDefinitions(CompileEnvironment, Arguments);
+
+			// Add warning and error flags to the argument list.
+			GetCompileArguments_WarningsAndErrors(CompileEnvironment, Arguments);
 		}
 	}
 }
\ No newline at end of file
