UnrealBuildTool: Pass /Zc:lambda- for VS2022 17.2 toolchain if strict conformance mode is enabled

#rnx

diff --git a/Engine/Source/Programs/UnrealBuildTool/Platform/Windows/VCToolChain.cs b/Engine/Source/Programs/UnrealBuildTool/Platform/Windows/VCToolChain.cs
--- a/Engine/Source/Programs/UnrealBuildTool/Platform/Windows/VCToolChain.cs
+++ b/Engine/Source/Programs/UnrealBuildTool/Platform/Windows/VCToolChain.cs
@@ -411,6 +411,10 @@ namespace UnrealBuildTool
 				
 				Arguments.Add("/permissive-");
 				Arguments.Add("/Zc:strictStrings-"); // Have to disable strict const char* semantics due to Windows headers not being compliant.
+				if (EnvVars.CompilerVersion >= new VersionNumber(14, 32))
+				{
+					Arguments.Add("/Zc:lambda-");
+				}
 			}
 
 			// @todo HoloLens: UE is non-compliant when it comes to use of %s and %S
