UnrealBuildTool: Remove warnings from base ClangToolChain (missed in previous cl)

#rnx

diff --git a/Engine/Source/Programs/UnrealBuildTool/ToolChain/ClangToolChain.cs b/Engine/Source/Programs/UnrealBuildTool/ToolChain/ClangToolChain.cs
--- a/Engine/Source/Programs/UnrealBuildTool/ToolChain/ClangToolChain.cs
+++ b/Engine/Source/Programs/UnrealBuildTool/ToolChain/ClangToolChain.cs
@@ -225,13 +225,6 @@ namespace UnrealBuildTool
 			{
 				Arguments.Add("-Wundef" + (CompileEnvironment.bUndefinedIdentifierWarningsAsErrors ? "" : " -Wno-error=undef"));
 			}
-
-			if (CompilerVersionGreaterOrEqual(13, 0, 0))
-			{
-				Arguments.Add("-Wno-unused-but-set-variable");           // https://clang.llvm.org/docs/DiagnosticsReference.html#wunused-but-set-variable
-				Arguments.Add("-Wno-unused-but-set-parameter");          // https://clang.llvm.org/docs/DiagnosticsReference.html#wunused-but-set-parameter
-				Arguments.Add("-Wno-ordered-compare-function-pointers"); // https://clang.llvm.org/docs/DiagnosticsReference.html#wordered-compare-function-pointers
-			}
 		}
 
 
