UnrealBuildTool: Move more common warnings to ClangToolChain

#rnx

diff --git a/Engine/Source/Programs/UnrealBuildTool/ToolChain/ClangToolChain.cs b/Engine/Source/Programs/UnrealBuildTool/ToolChain/ClangToolChain.cs
--- a/Engine/Source/Programs/UnrealBuildTool/ToolChain/ClangToolChain.cs
+++ b/Engine/Source/Programs/UnrealBuildTool/ToolChain/ClangToolChain.cs
@@ -209,10 +209,15 @@ namespace UnrealBuildTool
 		/// <param name="Arguments"></param>
 		protected virtual void GetCompileArguments_WarningsAndErrors(CppCompileEnvironment CompileEnvironment, List<string> Arguments)
 		{
-			Arguments.Add("-Wall");                     // https://clang.llvm.org/docs/DiagnosticsReference.html#wall
-			Arguments.Add("-Werror");                   // https://clang.llvm.org/docs/UsersManual.html#cmdoption-werror
+			Arguments.Add("-Wall");     // https://clang.llvm.org/docs/DiagnosticsReference.html#wall
+			Arguments.Add("-Werror");   // https://clang.llvm.org/docs/UsersManual.html#cmdoption-werror
 
-			Arguments.Add("-Wdelete-non-virtual-dtor"); // https://clang.llvm.org/docs/DiagnosticsReference.html#wdelete-non-virtual-dtor
+			Arguments.Add("-Wdelete-non-virtual-dtor");     // https://clang.llvm.org/docs/DiagnosticsReference.html#wdelete-non-virtual-dtor
+			Arguments.Add("-Wenum-conversion");             // https://clang.llvm.org/docs/DiagnosticsReference.html#wenum-conversion
+			Arguments.Add("-Wbitfield-enum-conversion");    // https://clang.llvm.org/docs/DiagnosticsReference.html#wbitfield-enum-conversion
+
+			Arguments.Add("-Wno-enum-float-conversion");    // https://clang.llvm.org/docs/DiagnosticsReference.html#wenum-float-conversion
+			Arguments.Add("-Wno-enum-enum-conversion");     // https://clang.llvm.org/docs/DiagnosticsReference.html#wenum-enum-conversion
 
 			// https://clang.llvm.org/docs/DiagnosticsReference.html#wshadow
 			if (CompileEnvironment.ShadowVariableWarningLevel != WarningLevel.Off)
@@ -227,7 +232,6 @@ namespace UnrealBuildTool
 			}
 		}
 
-
 		/// <summary>
 		/// Common compile arguments for all files in a module.
 		/// Override and call base.GetCompileArguments_Global() in derived classes.
