Work around macOS Monterey bug that erroneously spews objc stderr messages to stdout

#rnx

diff --git a/Engine/Source/Programs/UnrealBuildTool/Platform/Mac/MacToolChain.cs b/Engine/Source/Programs/UnrealBuildTool/Platform/Mac/MacToolChain.cs
--- a/Engine/Source/Programs/UnrealBuildTool/Platform/Mac/MacToolChain.cs
+++ b/Engine/Source/Programs/UnrealBuildTool/Platform/Mac/MacToolChain.cs
@@ -731,8 +731,11 @@ namespace UnrealBuildTool
 			// Create an action that invokes the linker.
 			Action LinkAction = Graph.CreateAction(ActionType.Link);
 
+			string Linker = bIsBuildingLibrary ? MacArchiver : MacLinker;
+			string LinkerPath = Settings.ToolchainDir + Linker;
+
 			LinkAction.WorkingDirectory = GetMacDevSrcRoot();
-			LinkAction.CommandPath = new FileReference("/usr/bin/env");
+			LinkAction.CommandPath = new FileReference(LinkerPath);
 			LinkAction.CommandDescription = "Link";
 			LinkAction.CommandVersion = GetFullClangVersion();
 
@@ -740,8 +743,7 @@ namespace UnrealBuildTool
 			string EngineDisplayVersion = LoadEngineDisplayVersion(true);
 			string VersionArg = LinkEnvironment.bIsBuildingDLL ? " -current_version " + EngineAPIVersion + " -compatibility_version " + EngineDisplayVersion : "";
 
-			string Linker = bIsBuildingLibrary ? MacArchiver : MacLinker;
-			string LinkCommand = Settings.ToolchainDir + Linker + VersionArg + " " + (bIsBuildingLibrary ? GetArchiveArguments_Global(LinkEnvironment) : GetLinkArguments_Global(LinkEnvironment));
+			string LinkCommand = VersionArg + " " + (bIsBuildingLibrary ? GetArchiveArguments_Global(LinkEnvironment) : GetLinkArguments_Global(LinkEnvironment));
 
 			// Tell the action that we're building an import library here and it should conditionally be
 			// ignored as a prerequisite for other actions
@@ -922,7 +924,7 @@ namespace UnrealBuildTool
 			// Add the additional arguments specified by the environment.
 			LinkCommand += LinkEnvironment.AdditionalArguments;
 
-			LinkAction.CommandArguments = "-- " + LinkCommand;
+			LinkAction.CommandArguments = LinkCommand;
 
 			// Only execute linking on the local Mac.
 			LinkAction.bCanExecuteRemotely = false;
@@ -968,6 +970,7 @@ namespace UnrealBuildTool
 					}
 					string FixDylibLine = "pushd \"" + Directory.GetCurrentDirectory() + "\"  > /dev/null; ";
 					FixDylibLine += string.Format("TIMESTAMP=`stat -n -f \"%Sm\" -t \"%Y%m%d%H%M.%S\" \"{0}\"`; ", OutputFile.AbsolutePath);
+					FixDylibLine += LinkerPath;
 					FixDylibLine += LinkCommand.Replace("-undefined dynamic_lookup", EngineAndGameLibrariesString).Replace("$", "\\$");
 					FixDylibLine += string.Format("; touch -t $TIMESTAMP \"{0}\"; if [[ $? -ne 0 ]]; then exit 1; fi; ", OutputFile.AbsolutePath);
 					FixDylibLine += "popd > /dev/null";
