UnrealBuildTool: Fix the Rider project generator in the presence of platform extension specializations.

diff --git a/Engine/Source/Programs/UnrealBuildTool/ProjectFiles/ProjectFileGenerator.cs b/Engine/Source/Programs/UnrealBuildTool/ProjectFiles/ProjectFileGenerator.cs
--- a/Engine/Source/Programs/UnrealBuildTool/ProjectFiles/ProjectFileGenerator.cs
+++ b/Engine/Source/Programs/UnrealBuildTool/ProjectFiles/ProjectFileGenerator.cs
@@ -1628,6 +1628,35 @@ namespace UnrealBuildTool
 			return AllTargetFiles;
 		}
 
+		/// <summary>
+		/// Returns a list of all sub-targets that are specialized for a given platform. 
+		/// </summary>
+		/// <param name="AllTargetFiles">List of target files to find sub-targets for</param>
+		/// <param name="AllSubTargetFiles">List of sub-targets in a flat list</param>
+		/// <param name="AllSubTargetFilesPerTarget">List of sub-targets bucketed by their parent target</param>
+		protected void GetPlatformSpecializationsSubTargetsForAllTargets(
+			List<FileReference> AllTargetFiles,
+			out HashSet<FileReference> AllSubTargetFiles,
+			out Dictionary<string, List<FileReference>> AllSubTargetFilesPerTarget)
+		{
+			AllSubTargetFilesPerTarget = new Dictionary<string, List<FileReference>>();
+			AllSubTargetFiles = new HashSet<FileReference>();
+			foreach (FileReference TargetFilePath in AllTargetFiles)
+			{
+				string[] TargetPathSplit = TargetFilePath.GetFileNameWithoutAnyExtensions().Split(new char[] { '_' }, StringSplitOptions.RemoveEmptyEntries);
+				if (TargetPathSplit.Length > 1 && (UnrealTargetPlatform.IsValidName(TargetPathSplit.Last()) || UnrealPlatformGroup.IsValidName(TargetPathSplit.Last())))
+				{
+					string TargetName = TargetPathSplit.First();
+					if (!AllSubTargetFilesPerTarget.ContainsKey(TargetName))
+					{
+						AllSubTargetFilesPerTarget.Add(TargetName, new List<FileReference>());
+					}
+
+					AllSubTargetFilesPerTarget[TargetName].Add(TargetFilePath);
+					AllSubTargetFiles.Add(TargetFilePath);
+				}
+			}
+		}
 
 
 		/// <summary>
@@ -2298,23 +2327,9 @@ namespace UnrealBuildTool
 			Dictionary<RulesAssembly, DirectoryReference> RulesAssemblies)
 		{
 			// Separate the .target.cs files that are platform extension specializations, per target name. These will be added alongside their base target.cs
-			Dictionary<string, List<FileReference>> AllSubTargetFilesPerTarget = new Dictionary<string, List<FileReference>>();
-			HashSet<FileReference> AllSubTargetFiles = new HashSet<FileReference>();
-			foreach (FileReference TargetFilePath in AllTargetFiles)
-			{
-				string[] TargetPathSplit = TargetFilePath.GetFileNameWithoutAnyExtensions().Split(new char[] { '_' }, StringSplitOptions.RemoveEmptyEntries);
-				if (TargetPathSplit.Length > 1 && (UnrealTargetPlatform.IsValidName(TargetPathSplit.Last()) || UnrealPlatformGroup.IsValidName(TargetPathSplit.Last())))
-				{
-					string TargetName = TargetPathSplit.First();
-					if (!AllSubTargetFilesPerTarget.ContainsKey(TargetName))
-					{
-						AllSubTargetFilesPerTarget.Add(TargetName, new List<FileReference>());
-					}
-
-					AllSubTargetFilesPerTarget[TargetName].Add(TargetFilePath);
-					AllSubTargetFiles.Add(TargetFilePath);
-				}
-			}
+			HashSet<FileReference> AllSubTargetFiles;
+			Dictionary<string, List<FileReference>> AllSubTargetFilesPerTarget;
+			GetPlatformSpecializationsSubTargetsForAllTargets(AllTargetFiles, out AllSubTargetFiles, out AllSubTargetFilesPerTarget);
 
 			// Get some standard directories
 			//DirectoryReference EngineSourceProgramsDirectory = DirectoryReference.Combine(UnrealBuildTool.EngineSourceDirectory, "Programs");
diff --git a/Engine/Source/Programs/UnrealBuildTool/ProjectFiles/Rider/RiderProjectFileGenerator.cs b/Engine/Source/Programs/UnrealBuildTool/ProjectFiles/Rider/RiderProjectFileGenerator.cs
--- a/Engine/Source/Programs/UnrealBuildTool/ProjectFiles/Rider/RiderProjectFileGenerator.cs
+++ b/Engine/Source/Programs/UnrealBuildTool/ProjectFiles/Rider/RiderProjectFileGenerator.cs
@@ -182,7 +182,12 @@ namespace UnrealBuildTool
 			// Sort the targets by name. When we have multiple targets of a given type for a project, we'll use the order to determine which goes in the primary project file (so that client names with a suffix will go into their own project).
 			AllTargetFiles = AllTargetFiles.OrderBy(x => x.FullName, StringComparer.OrdinalIgnoreCase).ToList();
 
-			foreach (FileReference TargetFilePath in AllTargetFiles)
+			// Separate the .target.cs files that are platform extension specializations, per target name. These will be added alongside their base target.cs
+			HashSet<FileReference> AllSubTargetFiles;
+			Dictionary<string, List<FileReference>> AllSubTargetFilesPerTarget;
+			GetPlatformSpecializationsSubTargetsForAllTargets(AllTargetFiles, out AllSubTargetFiles, out AllSubTargetFilesPerTarget);
+
+			foreach (FileReference TargetFilePath in AllTargetFiles.Except(AllSubTargetFiles))
 			{
 				string TargetName = TargetFilePath.GetFileNameWithoutAnyExtensions();
 
@@ -388,6 +393,12 @@ namespace UnrealBuildTool
 					// Make sure the *.Target.cs file is in the project.
 					ProjectFile.AddFileToProject(TargetFilePath, BaseFolder);
 
+					// Add all matching *_<platform>.Target.cs to the same folder
+					if (AllSubTargetFilesPerTarget.ContainsKey(TargetName))
+					{
+						ProjectFile.AddFilesToProject(AllSubTargetFilesPerTarget[TargetName], BaseFolder);
+					}
+
 					Log.TraceVerbose("Generating target {0} for {1}", TargetRulesObject.Type.ToString(),
 						ProjectFilePath);
 				}
