UnrealBuildTool: Expose ParallelExecutor.ProcessPriority to xml config

#rnx

diff --git a/Engine/Source/Programs/UnrealBuildTool/Executors/ParallelExecutor.cs b/Engine/Source/Programs/UnrealBuildTool/Executors/ParallelExecutor.cs
--- a/Engine/Source/Programs/UnrealBuildTool/Executors/ParallelExecutor.cs
+++ b/Engine/Source/Programs/UnrealBuildTool/Executors/ParallelExecutor.cs
@@ -39,6 +39,14 @@ namespace UnrealBuildTool
 		[XmlConfigFile]
 		private static double MemoryPerActionBytes = 1.5 * 1024 * 1024 * 1024;
 
+		/// <summary>
+		/// The priority to set for spawned processes.
+		/// Valid Settings: Idle, BelowNormal, Normal, AboveNormal, High
+		/// Default: BelowNormal
+		/// </summary>
+		[XmlConfigFile]
+		private static ProcessPriorityClass ProcessPriority = ProcessPriorityClass.BelowNormal;
+
 		/// <summary>
 		/// When enabled, will stop compiling targets after a compile error occurs.
 		/// </summary>
@@ -283,7 +291,7 @@ namespace UnrealBuildTool
 		{
 			CancellationToken.ThrowIfCancellationRequested();
 
-			using ManagedProcess Process = new ManagedProcess(ProcessGroup, Action.CommandPath.FullName, Action.CommandArguments, Action.WorkingDirectory.FullName, null, null, ProcessPriorityClass.BelowNormal);
+			using ManagedProcess Process = new ManagedProcess(ProcessGroup, Action.CommandPath.FullName, Action.CommandArguments, Action.WorkingDirectory.FullName, null, null, ProcessPriority);
 
 			using MemoryStream StdOutStream = new MemoryStream();
 			await Process.CopyToAsync(StdOutStream, CancellationToken);
