Update minSDKVersion to be higher or equal than NDK API version to prevent runtime DSO loading errors on Android

diff --git a/Engine/Source/Programs/UnrealBuildTool/Platform/Android/AndroidToolChain.cs b/Engine/Source/Programs/UnrealBuildTool/Platform/Android/AndroidToolChain.cs
--- a/Engine/Source/Programs/UnrealBuildTool/Platform/Android/AndroidToolChain.cs
+++ b/Engine/Source/Programs/UnrealBuildTool/Platform/Android/AndroidToolChain.cs
@@ -319,7 +319,7 @@ namespace UnrealBuildTool
 
 			if (!IsNewNDKModel())
 			{
-				// We need to manually provide -D__ANDROID_API__ only for NDK versions prior to r22
+				// We need to manually provide -D__ANDROID_API__ for NDK versions prior to r22 only, for newer ones, --target=aarch64-none-linux-android + NDKApiLevel64Int does it for us
 				ToolchainParamsArm64 += " -D__ANDROID_API__=" + NDKApiLevel64Int;
 				ToolchainParamsx64 += " -D__ANDROID_API__=" + NDKApiLevel64Int;
 			}
diff --git a/Engine/Source/Programs/UnrealBuildTool/Platform/Android/UEDeployAndroid.cs b/Engine/Source/Programs/UnrealBuildTool/Platform/Android/UEDeployAndroid.cs
--- a/Engine/Source/Programs/UnrealBuildTool/Platform/Android/UEDeployAndroid.cs
+++ b/Engine/Source/Programs/UnrealBuildTool/Platform/Android/UEDeployAndroid.cs
@@ -3425,7 +3425,13 @@ namespace UnrealBuildTool
 					MinSDKVersion = 21;
 					Log.TraceInformation("Fixing minSdkVersion; NDK level above 19 requires minSdkVersion of 21 (arch={0})", Arch.Substring(1));
 				}
-			}
+
+				if (MinSDKVersion < NDKLevelInt)
+                {
+                    Log.TraceInformation("Fixing minSdkVersion; NDK level is {0} which is above minSdkVersion {1}.", NDKLevelInt, MinSDKVersion);
+					MinSDKVersion = NDKLevelInt;
+				}
+            }
 
 			if (TargetSDKVersion < MinSDKVersion)
 			{
