UnrealBuildTool: Move PGO optimize and shipping warning flags to ClangToolChain

#rnx

diff --git a/Engine/Source/Programs/UnrealBuildTool/Platform/Linux/LinuxToolChain.cs b/Engine/Source/Programs/UnrealBuildTool/Platform/Linux/LinuxToolChain.cs
--- a/Engine/Source/Programs/UnrealBuildTool/Platform/Linux/LinuxToolChain.cs
+++ b/Engine/Source/Programs/UnrealBuildTool/Platform/Linux/LinuxToolChain.cs
@@ -542,23 +542,6 @@ namespace UnrealBuildTool
 			Arguments.Add("-Wno-unknown-pragmas");          // Slate triggers this (with its optimize on/off pragmas)
 			Arguments.Add("-Wno-invalid-offsetof"); // needed to suppress warnings about using offsetof on non-POD types.
 			Arguments.Add("-Wno-gnu-string-literal-operator-template"); // we use this feature to allow static FNames.
-
-			if (CompileEnvironment.bPGOOptimize)
-			{
-				//
-				// Clang emits a warning for each compiled function that doesn't have a matching entry in the profile data.
-				// This can happen when the profile data is older than the binaries we're compiling.
-				//
-				// Disable this warning. It's far too verbose.
-				//
-				Arguments.Add("-Wno-backend-plugin");
-			}
-
-			// shipping builds will cause this warning with "ensure", so disable only in those case
-			if (CompileEnvironment.Configuration == CppConfiguration.Shipping)
-			{
-				Arguments.Add("-Wno-unused-value");
-			}
 		}
 
 		/// <inheritdoc/>
diff --git a/Engine/Source/Programs/UnrealBuildTool/ToolChain/ClangToolChain.cs b/Engine/Source/Programs/UnrealBuildTool/ToolChain/ClangToolChain.cs
--- a/Engine/Source/Programs/UnrealBuildTool/ToolChain/ClangToolChain.cs
+++ b/Engine/Source/Programs/UnrealBuildTool/ToolChain/ClangToolChain.cs
@@ -219,6 +219,36 @@ namespace UnrealBuildTool
 			Arguments.Add("-Wno-enum-float-conversion");    // https://clang.llvm.org/docs/DiagnosticsReference.html#wenum-float-conversion
 			Arguments.Add("-Wno-enum-enum-conversion");     // https://clang.llvm.org/docs/DiagnosticsReference.html#wenum-enum-conversion
 
+			// Profile Guided Optimization (PGO) and Link Time Optimization (LTO)
+			if (CompileEnvironment.bPGOOptimize)
+			{
+				//
+				// Clang emits warnings for each compiled object file that doesn't have a matching entry in the profile data.
+				// This can happen when the profile data is older than the binaries we're compiling.
+				//
+				// Disable these warnings. They are far too verbose.
+				//
+				Arguments.Add("-Wno-profile-instr-out-of-date");    // https://clang.llvm.org/docs/DiagnosticsReference.html#wprofile-instr-out-of-date
+				Arguments.Add("-Wno-profile-instr-unprofiled");     // https://clang.llvm.org/docs/DiagnosticsReference.html#wprofile-instr-unprofiled
+
+				// apparently there can be hashing conflicts with PGO which can result in:
+				// 'Function control flow change detected (hash mismatch)' warnings. 
+				Arguments.Add("-Wno-backend-plugin");               // https://clang.llvm.org/docs/DiagnosticsReference.html#wbackend-plugin
+			}
+
+			// shipping builds will cause this warning with "ensure", so disable only in those case
+			if (CompileEnvironment.Configuration == CppConfiguration.Shipping)
+			{
+				Arguments.Add("-Wno-unused-value"); // https://clang.llvm.org/docs/DiagnosticsReference.html#wunused-value
+			}
+
+			// https://clang.llvm.org/docs/DiagnosticsReference.html#wdeprecated-declarations
+			if (CompileEnvironment.DeprecationWarningLevel == WarningLevel.Error)
+			{
+				// TODO: This may be unnecessary with -Werror
+				Arguments.Add("-Werror=deprecated-declarations");
+			}
+
 			// https://clang.llvm.org/docs/DiagnosticsReference.html#wshadow
 			if (CompileEnvironment.ShadowVariableWarningLevel != WarningLevel.Off)
 			{
