Turn on using C# UHT by default.

#rnx

diff --git a/Engine/Source/Programs/UnrealBuildTool/Configuration/BuildConfiguration.cs b/Engine/Source/Programs/UnrealBuildTool/Configuration/BuildConfiguration.cs
--- a/Engine/Source/Programs/UnrealBuildTool/Configuration/BuildConfiguration.cs
+++ b/Engine/Source/Programs/UnrealBuildTool/Configuration/BuildConfiguration.cs
@@ -141,7 +141,7 @@ namespace UnrealBuildTool
 		/// If true, use UHT internal to UBT
 		/// </summary>
 		[XmlConfigFile(Category = "UEBuildConfiguration")]
-		public bool bUseBuiltInUnrealHeaderTool = false;
+		public bool bUseBuiltInUnrealHeaderTool = true;
 
 		/// <summary>
 		/// Whether to skip compiling rules assemblies and just assume they are valid
