Removed the single instance flag for running UHT

#rnx

diff --git a/Engine/Source/Programs/UnrealBuildTool/Modes/UnrealHeaderToolMode.cs b/Engine/Source/Programs/UnrealBuildTool/Modes/UnrealHeaderToolMode.cs
--- a/Engine/Source/Programs/UnrealBuildTool/Modes/UnrealHeaderToolMode.cs
+++ b/Engine/Source/Programs/UnrealBuildTool/Modes/UnrealHeaderToolMode.cs
@@ -894,7 +894,7 @@ namespace UnrealBuildTool.Modes
 	/// <summary>
 	/// Invoke UHT
 	/// </summary>
-	[ToolMode("UnrealHeaderTool", ToolModeOptions.XmlConfig | ToolModeOptions.BuildPlatforms | ToolModeOptions.SingleInstance | ToolModeOptions.ShowExecutionTime)]
+	[ToolMode("UnrealHeaderTool", ToolModeOptions.XmlConfig | ToolModeOptions.BuildPlatforms | ToolModeOptions.ShowExecutionTime)]
 	class UnrealHeaderToolMode : ToolMode
 	{
 		/// <summary>
