Build in solution by default set to true for inheriting classes of TestTargetRules

diff --git a/Engine/Source/Programs/UnrealBuildTool/Configuration/TargetRules.cs b/Engine/Source/Programs/UnrealBuildTool/Configuration/TargetRules.cs
--- a/Engine/Source/Programs/UnrealBuildTool/Configuration/TargetRules.cs
+++ b/Engine/Source/Programs/UnrealBuildTool/Configuration/TargetRules.cs
@@ -509,6 +509,7 @@ namespace UnrealBuildTool
 
 		/// <summary>
 		/// Whether the target should be included in the default solution build configuration
+		/// Setting this to false will skip building when running in the IDE
 		/// </summary>
 		public bool? bBuildInSolutionByDefault = null;
 
diff --git a/Engine/Source/Programs/UnrealBuildTool/Configuration/TestTargetRules.cs b/Engine/Source/Programs/UnrealBuildTool/Configuration/TestTargetRules.cs
--- a/Engine/Source/Programs/UnrealBuildTool/Configuration/TestTargetRules.cs
+++ b/Engine/Source/Programs/UnrealBuildTool/Configuration/TestTargetRules.cs
@@ -95,6 +95,7 @@ namespace UnrealBuildTool
 
 			if (bExplicitTestsTargetOverride)
 			{
+				bBuildInSolutionByDefault = true;
 				SolutionDirectory = "Programs/LowLevelTests";
 			}
 		}
