UBT: Allow overriding the output path for executables.

This may not work with UE code, since the engine assumes the executable is in a certain location to figure out where other known folders are in relation to it, but is useful for programs that do not link against the engine.

diff --git a/Engine/Source/Programs/UnrealBuildTool/Configuration/TargetRules.cs b/Engine/Source/Programs/UnrealBuildTool/Configuration/TargetRules.cs
--- a/Engine/Source/Programs/UnrealBuildTool/Configuration/TargetRules.cs
+++ b/Engine/Source/Programs/UnrealBuildTool/Configuration/TargetRules.cs
@@ -389,6 +389,12 @@ namespace UnrealBuildTool
 		}
 		private EngineIncludeOrderVersion? IncludeOrderVersionPrivate;
 
+		/// <summary>
+		/// Path to the output file for the main executable, relative to the Engine or project directory.
+		/// This setting is only typically useful for non-UE programs, since the engine uses paths relative to the executable to find other known folders (eg. Content).
+		/// </summary>
+		public string? OutputFile;
+
 		/// <summary>
 		/// Tracks a list of config values read while constructing this target
 		/// </summary>
@@ -2246,6 +2252,11 @@ namespace UnrealBuildTool
 			get { return Inner.ConfigValueTracker; }
 		}
 
+		public string? OutputFile
+		{
+			get { return Inner.OutputFile; }
+		}
+
 		public bool bUsesSteam
 		{
 			get { return Inner.bUsesSteam; }
diff --git a/Engine/Source/Programs/UnrealBuildTool/Configuration/UEBuildTarget.cs b/Engine/Source/Programs/UnrealBuildTool/Configuration/UEBuildTarget.cs
--- a/Engine/Source/Programs/UnrealBuildTool/Configuration/UEBuildTarget.cs
+++ b/Engine/Source/Programs/UnrealBuildTool/Configuration/UEBuildTarget.cs
@@ -3248,18 +3248,26 @@ namespace UnrealBuildTool
 		/// <param name="ProjectFile">The project file containing the target being built</param>
 		/// <param name="Rules">Rules for the target being built</param>
 		/// <returns>List of executable paths for this target</returns>
-		public static List<FileReference> MakeBinaryPaths(DirectoryReference BaseDirectory, string BinaryName, UnrealTargetPlatform Platform, UnrealTargetConfiguration Configuration, UEBuildBinaryType BinaryType, string Architecture, UnrealTargetConfiguration UndecoratedConfiguration, bool bIncludesGameModules, string ExeSubFolder, FileReference? ProjectFile, ReadOnlyTargetRules Rules)
+		public List<FileReference> MakeBinaryPaths(DirectoryReference BaseDirectory, string BinaryName, UnrealTargetPlatform Platform, UnrealTargetConfiguration Configuration, UEBuildBinaryType BinaryType, string Architecture, UnrealTargetConfiguration UndecoratedConfiguration, bool bIncludesGameModules, string ExeSubFolder, FileReference? ProjectFile, ReadOnlyTargetRules Rules)
 		{
-			// Get the configuration for the executable. If we're building DebugGame, and this executable only contains engine modules, use the same name as development.
-			UnrealTargetConfiguration ExeConfiguration = Configuration;
-
-			// Build the binary path
-			DirectoryReference BinaryDirectory = DirectoryReference.Combine(BaseDirectory, "Binaries", Platform.ToString());
-			if (!String.IsNullOrEmpty(ExeSubFolder))
+			FileReference BinaryFile;
+			if (Rules.OutputFile != null)
+			{
+				BinaryFile = FileReference.Combine(BaseDirectory, Rules.OutputFile);
+			}
+			else
 			{
-				BinaryDirectory = DirectoryReference.Combine(BinaryDirectory, ExeSubFolder);
+				// Get the configuration for the executable. If we're building DebugGame, and this executable only contains engine modules, use the same name as development.
+				UnrealTargetConfiguration ExeConfiguration = Configuration;
+
+				// Build the binary path
+				DirectoryReference BinaryDirectory = DirectoryReference.Combine(BaseDirectory, "Binaries", Platform.ToString());
+				if (!String.IsNullOrEmpty(ExeSubFolder))
+				{
+					BinaryDirectory = DirectoryReference.Combine(BinaryDirectory, ExeSubFolder);
+				}
+				BinaryFile = FileReference.Combine(BinaryDirectory, MakeBinaryFileName(BinaryName, Platform, ExeConfiguration, Architecture, UndecoratedConfiguration, BinaryType));
 			}
-			FileReference BinaryFile = FileReference.Combine(BinaryDirectory, MakeBinaryFileName(BinaryName, Platform, ExeConfiguration, Architecture, UndecoratedConfiguration, BinaryType));
 
 			// Allow the platform to customize the output path (and output several executables at once if necessary)
 			return UEBuildPlatform.GetBuildPlatform(Platform).FinalizeBinaryPaths(BinaryFile, ProjectFile, Rules);
