UnrealBuildTool: Move more common warning disables to ClangToolChain

#rnx

diff --git a/Engine/Source/Programs/UnrealBuildTool/Platform/Linux/LinuxToolChain.cs b/Engine/Source/Programs/UnrealBuildTool/Platform/Linux/LinuxToolChain.cs
--- a/Engine/Source/Programs/UnrealBuildTool/Platform/Linux/LinuxToolChain.cs
+++ b/Engine/Source/Programs/UnrealBuildTool/Platform/Linux/LinuxToolChain.cs
@@ -519,29 +519,9 @@ namespace UnrealBuildTool
 		{
 			base.GetCompileArguments_WarningsAndErrors(CompileEnvironment, Arguments);
 
-			if (CompilerVersionGreaterOrEqual(13, 0, 0))
-			{
-				Arguments.Add("-Wno-unused-but-set-variable");           // https://clang.llvm.org/docs/DiagnosticsReference.html#wunused-but-set-variable
-				Arguments.Add("-Wno-unused-but-set-parameter");          // https://clang.llvm.org/docs/DiagnosticsReference.html#wunused-but-set-parameter
-				Arguments.Add("-Wno-ordered-compare-function-pointers"); // https://clang.llvm.org/docs/DiagnosticsReference.html#wordered-compare-function-pointers
-			}
-
 			//Arguments.Add("-Wunreachable-code");            // additional warning not normally included in Wall: warns if there is code that will never be executed - not helpful due to bIsGCC and similar
 
-			Arguments.Add("-Wno-unused-private-field"); // MultichannelTcpSocket.h triggers this, possibly more
-			Arguments.Add("-Wno-tautological-compare"); // this hides the "warning : comparison of unsigned expression < 0 is always false" type warnings due to constant comparisons, which are possible with template arguments
 			Arguments.Add("-Wno-undefined-bool-conversion"); // hides checking if 'this' pointer is null
-			Arguments.Add("-Wno-unused-local-typedef"); // clang is being overly strict here? PhysX headers trigger this.
-			Arguments.Add("-Wno-inconsistent-missing-override");    // these have to be suppressed for UE 4.8, should be fixed later.
-			Arguments.Add("-Wno-undefined-var-template"); // not really a good warning to disable
-			Arguments.Add("-Wno-unused-lambda-capture");  // suppressed because capturing of compile-time constants is seemingly inconsistent. And MSVC doesn't do that.
-			Arguments.Add("-Wno-unused-variable");
-
-			Arguments.Add("-Wno-unused-function"); // this will hide the warnings about static functions in headers that aren't used in every single .cpp file
-			Arguments.Add("-Wno-switch"); // this hides the "enumeration value 'XXXXX' not handled in switch [-Wswitch]" warnings - we should maybe remove this at some point and add UE_LOG(, Fatal, ) to default cases
-			Arguments.Add("-Wno-unknown-pragmas");          // Slate triggers this (with its optimize on/off pragmas)
-			Arguments.Add("-Wno-invalid-offsetof"); // needed to suppress warnings about using offsetof on non-POD types.
-			Arguments.Add("-Wno-gnu-string-literal-operator-template"); // we use this feature to allow static FNames.
 		}
 
 		/// <inheritdoc/>
diff --git a/Engine/Source/Programs/UnrealBuildTool/ToolChain/AppleToolChain.cs b/Engine/Source/Programs/UnrealBuildTool/ToolChain/AppleToolChain.cs
--- a/Engine/Source/Programs/UnrealBuildTool/ToolChain/AppleToolChain.cs
+++ b/Engine/Source/Programs/UnrealBuildTool/ToolChain/AppleToolChain.cs
@@ -254,11 +254,12 @@ namespace UnrealBuildTool
 		{
 			base.GetCompileArguments_WarningsAndErrors(CompileEnvironment, Arguments);
 
-			// clang 12.00 has a new warning for copies in ranged loops. Instances have all been fixed up (2020/6/26) but
-			// are likely to be reintroduced due to no equivalent on other platforms at this time so disable the warning
-			if (GetClangVersion().Major >= 12)
+			// Flags added in Xcode 13.3 (Apple clang 13.1.6)
+			if (CompilerVersionGreaterOrEqual(13, 0, 0) && CompilerVersionLessThan(13, 1, 6))
 			{
-				Arguments.Add("-Wno-range-loop-analysis");
+				Arguments.Remove("-Wno-unused-but-set-variable");
+				Arguments.Remove("-Wno-unused-but-set-parameter");
+				Arguments.Remove("-Wno-ordered-compare-function-pointers");
 			}
 		}
 
diff --git a/Engine/Source/Programs/UnrealBuildTool/ToolChain/ClangToolChain.cs b/Engine/Source/Programs/UnrealBuildTool/ToolChain/ClangToolChain.cs
--- a/Engine/Source/Programs/UnrealBuildTool/ToolChain/ClangToolChain.cs
+++ b/Engine/Source/Programs/UnrealBuildTool/ToolChain/ClangToolChain.cs
@@ -209,15 +209,35 @@ namespace UnrealBuildTool
 		/// <param name="Arguments"></param>
 		protected virtual void GetCompileArguments_WarningsAndErrors(CppCompileEnvironment CompileEnvironment, List<string> Arguments)
 		{
-			Arguments.Add("-Wall");     // https://clang.llvm.org/docs/DiagnosticsReference.html#wall
-			Arguments.Add("-Werror");   // https://clang.llvm.org/docs/UsersManual.html#cmdoption-werror
+			Arguments.Add("-Wall");                                     // https://clang.llvm.org/docs/DiagnosticsReference.html#wall
+			Arguments.Add("-Werror");                                   // https://clang.llvm.org/docs/UsersManual.html#cmdoption-werror
 
-			Arguments.Add("-Wdelete-non-virtual-dtor");     // https://clang.llvm.org/docs/DiagnosticsReference.html#wdelete-non-virtual-dtor
-			Arguments.Add("-Wenum-conversion");             // https://clang.llvm.org/docs/DiagnosticsReference.html#wenum-conversion
-			Arguments.Add("-Wbitfield-enum-conversion");    // https://clang.llvm.org/docs/DiagnosticsReference.html#wbitfield-enum-conversion
+			Arguments.Add("-Wdelete-non-virtual-dtor");                 // https://clang.llvm.org/docs/DiagnosticsReference.html#wdelete-non-virtual-dtor
+			Arguments.Add("-Wenum-conversion");                         // https://clang.llvm.org/docs/DiagnosticsReference.html#wenum-conversion
+			Arguments.Add("-Wbitfield-enum-conversion");                // https://clang.llvm.org/docs/DiagnosticsReference.html#wbitfield-enum-conversion
 
-			Arguments.Add("-Wno-enum-float-conversion");    // https://clang.llvm.org/docs/DiagnosticsReference.html#wenum-float-conversion
-			Arguments.Add("-Wno-enum-enum-conversion");     // https://clang.llvm.org/docs/DiagnosticsReference.html#wenum-enum-conversion
+			Arguments.Add("-Wno-enum-enum-conversion");                 // https://clang.llvm.org/docs/DiagnosticsReference.html#wenum-enum-conversion					// ?? no reason given
+			Arguments.Add("-Wno-enum-float-conversion");                // https://clang.llvm.org/docs/DiagnosticsReference.html#wenum-float-conversion					// ?? no reason given
+
+			if (CompilerVersionGreaterOrEqual(13, 0, 0))
+			{
+				Arguments.Add("-Wno-unused-but-set-variable");           // https://clang.llvm.org/docs/DiagnosticsReference.html#wunused-but-set-variable				// new warning for clang 13
+				Arguments.Add("-Wno-unused-but-set-parameter");          // https://clang.llvm.org/docs/DiagnosticsReference.html#wunused-but-set-parameter				// new warning for clang 13
+				Arguments.Add("-Wno-ordered-compare-function-pointers"); // https://clang.llvm.org/docs/DiagnosticsReference.html#wordered-compare-function-pointers	// new warning for clang 13
+			}
+
+			Arguments.Add("-Wno-gnu-string-literal-operator-template"); // https://clang.llvm.org/docs/DiagnosticsReference.html#wgnu-string-literal-operator-template	// We use this feature to allow static FNames.
+			Arguments.Add("-Wno-inconsistent-missing-override");        // https://clang.llvm.org/docs/DiagnosticsReference.html#winconsistent-missing-override			// ?? no reason given
+			Arguments.Add("-Wno-invalid-offsetof");                     // https://clang.llvm.org/docs/DiagnosticsReference.html#winvalid-offsetof						// needed to suppress warnings about using offsetof on non-POD types.
+			Arguments.Add("-Wno-switch");                               // https://clang.llvm.org/docs/DiagnosticsReference.html#wswitch								// this hides the "enumeration value 'XXXXX' not handled in switch [-Wswitch]" warnings - we should maybe remove this at some point and add UE_LOG(, Fatal, ) to default cases
+			Arguments.Add("-Wno-tautological-compare");                 // https://clang.llvm.org/docs/DiagnosticsReference.html#wtautological-compare					// this hides the "warning : comparison of unsigned expression < 0 is always false" type warnings due to constant comparisons, which are possible with template arguments
+			Arguments.Add("-Wno-unknown-pragmas");                      // https://clang.llvm.org/docs/DiagnosticsReference.html#wunknown-pragmas						// Slate triggers this (with its optimize on/off pragmas)
+			Arguments.Add("-Wno-unused-function");                      // https://clang.llvm.org/docs/DiagnosticsReference.html#wunused-function						// this will hide the warnings about static functions in headers that aren't used in every single .cpp file
+			Arguments.Add("-Wno-unused-lambda-capture");                // https://clang.llvm.org/docs/DiagnosticsReference.html#wunused-lambda-capture					// suppressed because capturing of compile-time constants is seemingly inconsistent. And MSVC doesn't do that.
+			Arguments.Add("-Wno-unused-local-typedef");                 // https://clang.llvm.org/docs/DiagnosticsReference.html#wunused-local-typedef					// clang is being overly strict here? PhysX headers trigger this.
+			Arguments.Add("-Wno-unused-private-field");                 // https://clang.llvm.org/docs/DiagnosticsReference.html#wunused-private-field					// this will prevent the issue of warnings for unused private variables. MultichannelTcpSocket.h triggers this, possibly more
+			Arguments.Add("-Wno-unused-variable");                      // https://clang.llvm.org/docs/DiagnosticsReference.html#wunused-variable						// ?? no reason given
+			Arguments.Add("-Wno-undefined-var-template");               // https://clang.llvm.org/docs/DiagnosticsReference.html#wundefined-var-template				// not really a good warning to disable
 
 			// Profile Guided Optimization (PGO) and Link Time Optimization (LTO)
 			if (CompileEnvironment.bPGOOptimize)
@@ -228,18 +248,18 @@ namespace UnrealBuildTool
 				//
 				// Disable these warnings. They are far too verbose.
 				//
-				Arguments.Add("-Wno-profile-instr-out-of-date");    // https://clang.llvm.org/docs/DiagnosticsReference.html#wprofile-instr-out-of-date
-				Arguments.Add("-Wno-profile-instr-unprofiled");     // https://clang.llvm.org/docs/DiagnosticsReference.html#wprofile-instr-unprofiled
+				Arguments.Add("-Wno-profile-instr-out-of-date");        // https://clang.llvm.org/docs/DiagnosticsReference.html#wprofile-instr-out-of-date
+				Arguments.Add("-Wno-profile-instr-unprofiled");         // https://clang.llvm.org/docs/DiagnosticsReference.html#wprofile-instr-unprofiled
 
 				// apparently there can be hashing conflicts with PGO which can result in:
 				// 'Function control flow change detected (hash mismatch)' warnings. 
-				Arguments.Add("-Wno-backend-plugin");               // https://clang.llvm.org/docs/DiagnosticsReference.html#wbackend-plugin
+				Arguments.Add("-Wno-backend-plugin");                   // https://clang.llvm.org/docs/DiagnosticsReference.html#wbackend-plugin
 			}
 
 			// shipping builds will cause this warning with "ensure", so disable only in those case
 			if (CompileEnvironment.Configuration == CppConfiguration.Shipping)
 			{
-				Arguments.Add("-Wno-unused-value"); // https://clang.llvm.org/docs/DiagnosticsReference.html#wunused-value
+				Arguments.Add("-Wno-unused-value");                     // https://clang.llvm.org/docs/DiagnosticsReference.html#wunused-value
 			}
 
 			// https://clang.llvm.org/docs/DiagnosticsReference.html#wdeprecated-declarations
