Fixed Android builds with NDK25

diff --git a/Engine/Source/Programs/UnrealBuildTool/Platform/Android/AndroidToolChain.cs b/Engine/Source/Programs/UnrealBuildTool/Platform/Android/AndroidToolChain.cs
--- a/Engine/Source/Programs/UnrealBuildTool/Platform/Android/AndroidToolChain.cs
+++ b/Engine/Source/Programs/UnrealBuildTool/Platform/Android/AndroidToolChain.cs
@@ -567,6 +567,8 @@ namespace UnrealBuildTool
 				throw new BuildException("No NDK platforms found in {0}", PlatformsFilename);
 			}
 
+			MaxPlatform = Math.Clamp(MaxPlatform, 1, 32);	//Current NDK 25beta4 breaks existing UE NDK to SDK API mapping, so clamp it at 32
+
 			return "android-" + MaxPlatform.ToString();
 		}
 
@@ -590,12 +592,18 @@ namespace UnrealBuildTool
 			Result += " -fno-PIE";
 
 			Result += " -Wno-unused-variable";
-			if (CompilerVersionGreaterOrEqual(13, 0, 0))
+            if (CompilerVersionGreaterOrEqual(13, 0, 0))
+            {
+                Result += " -Wno-unused-but-set-variable";
+                Result += " -Wno-unused-but-set-parameter";
+                Result += " -Wno-ordered-compare-function-pointers";
+            }
+			
+			if (CompilerVersionGreaterOrEqual(14, 0, 0))
 			{
-				Result += " -Wno-unused-but-set-variable";
-				Result += " -Wno-unused-but-set-parameter";
-				Result += " -Wno-ordered-compare-function-pointers";
+				Result += " -Wno-bitwise-instead-of-logical";
 			}
+
 			// this will hide the warnings about static functions in headers that aren't used in every single .cpp file
 			Result += " -Wno-unused-function";
 			// this hides the "enumeration value 'XXXXX' not handled in switch [-Wswitch]" warnings - we should maybe remove this at some point and add UE_LOG(, Fatal, ) to default cases
diff --git a/Engine/Source/Programs/UnrealBuildTool/Platform/Android/UEDeployAndroid.cs b/Engine/Source/Programs/UnrealBuildTool/Platform/Android/UEDeployAndroid.cs
--- a/Engine/Source/Programs/UnrealBuildTool/Platform/Android/UEDeployAndroid.cs
+++ b/Engine/Source/Programs/UnrealBuildTool/Platform/Android/UEDeployAndroid.cs
@@ -1327,10 +1327,30 @@ namespace UnrealBuildTool
 			Utils.RunLocalProcessAndLogOutput(StartInfo);
 		}
 
+		private static string GetPlatformNDKHostName()
+		{
+			if (RuntimePlatform.IsLinux)
+			{
+				return "linux-x86_64";
+			}
+			else if (RuntimePlatform.IsMac)
+			{
+				return "darwin-x86_64";
+			}
+
+			return "windows-x86_64";
+		}
+
 		private static void CopySTL(AndroidToolChain ToolChain, string UnrealBuildPath, string UnrealArch, string NDKArch, bool bForDistribution)
 		{
 			// copy it in!
 			string SourceSTLSOName = Environment.ExpandEnvironmentVariables("%NDKROOT%/sources/cxx-stl/llvm-libc++/libs/") +  NDKArch + "/libc++_shared.so";
+			if (!File.Exists(SourceSTLSOName))
+			{
+				// NDK25 has changed a directory where it stores libs, check it instead
+				string NDKTargetTripletName = (NDKArch == "x86_64") ? "x86_64-linux-android" : "aarch64-linux-android";
+				SourceSTLSOName = Environment.ExpandEnvironmentVariables("%NDKROOT%/toolchains/llvm/prebuilt/") + GetPlatformNDKHostName() + "/sysroot/usr/lib/" + NDKTargetTripletName + "/libc++_shared.so";
+			}
 			string FinalSTLSOName = UnrealBuildPath + "/jni/" + NDKArch + "/libc++_shared.so";
 
 			// check to see if libc++_shared.so is newer than last time we copied
@@ -1505,22 +1525,14 @@ namespace UnrealBuildTool
 
 			string WrapSh = Path.Combine(Unreal.EngineDirectory.ToString(), "Build", "Android", "ClangSanitizers", "wrap.sh");
 
-			string PlatformName = "windows-x86_64";
-			if (RuntimePlatform.IsLinux)
-			{
-				PlatformName = "linux-x86_64";
-			}
-			else if (RuntimePlatform.IsMac)
-			{
-				PlatformName = "darwin-x86_64";
-			}
+			string PlatformHostName = GetPlatformNDKHostName();
 
-			string VersionFileName = Path.Combine(Environment.ExpandEnvironmentVariables("%NDKROOT%"), "toolchains", "llvm", "prebuilt", PlatformName, "AndroidVersion.txt");
+			string VersionFileName = Path.Combine(Environment.ExpandEnvironmentVariables("%NDKROOT%"), "toolchains", "llvm", "prebuilt", PlatformHostName, "AndroidVersion.txt");
 			System.IO.StreamReader VersionFile = new System.IO.StreamReader(VersionFileName);
 			string LibsVersion = VersionFile.ReadLine()!;
 			VersionFile.Close();
 
-			string SanitizerLib = Path.Combine(Environment.ExpandEnvironmentVariables("%NDKROOT%"), "toolchains", "llvm", "prebuilt", PlatformName, "lib64", "clang", LibsVersion, "lib", "linux", SanitizerFullLibName);
+			string SanitizerLib = Path.Combine(Environment.ExpandEnvironmentVariables("%NDKROOT%"), "toolchains", "llvm", "prebuilt", PlatformHostName, "lib64", "clang", LibsVersion, "lib", "linux", SanitizerFullLibName);
 			if (File.Exists(SanitizerLib) && File.Exists(WrapSh))
 			{
 				string LibDestDir = Path.Combine(UnrealBuildPath, "libs", NDKArch);
