UnrealBuildTool: Merge various Clang ToolChainOptions enums

#rnx

diff --git a/Engine/Source/Programs/UnrealBuildTool/Platform/Android/AndroidToolChain.cs b/Engine/Source/Programs/UnrealBuildTool/Platform/Android/AndroidToolChain.cs
--- a/Engine/Source/Programs/UnrealBuildTool/Platform/Android/AndroidToolChain.cs
+++ b/Engine/Source/Programs/UnrealBuildTool/Platform/Android/AndroidToolChain.cs
@@ -12,43 +12,6 @@
 
 namespace UnrealBuildTool
 {
-	/// <summary>
-	/// Option flags for the Android toolchain
-	/// </summary>
-	[Flags]
-	enum AndroidToolChainOptions
-	{
-		/// <summary>
-		/// No custom options
-		/// </summary>
-		None = 0,
-
-		/// <summary>
-		/// Enable address sanitizer
-		/// </summary>
-		EnableAddressSanitizer = 0x1,
-
-		/// <summary>
-		/// Enable HW address sanitizer
-		/// </summary>
-		EnableHWAddressSanitizer = 0x2,
-
-		/// <summary>
-		/// Enable thread sanitizer
-		/// </summary>
-		EnableThreadSanitizer = 0x4,
-
-		/// <summary>
-		/// Enable undefined behavior sanitizer
-		/// </summary>
-		EnableUndefinedBehaviorSanitizer = 0x8,
-
-		/// <summary>
-		/// Enable undefined behavior sanitizer
-		/// </summary>
-		EnableMinimalUndefinedBehaviorSanitizer = 0x10,
-	}
-
 	class AndroidToolChain : ClangToolChain, IAndroidToolChain
 	{
 		public static readonly string[] AllCpuSuffixes =
@@ -101,8 +64,6 @@
 		// the list of architectures we will compile for
 		protected List<string>? Arches = null;
 
-		private AndroidToolChainOptions Options;
-
 		// the "-android" suffix paths here are vcpkg triplets for the android platform
 		static private Dictionary<string, string[]> AllArchNames = new Dictionary<string, string[]> {
 			{ "-arm64", new string[] { "arm64", "arm64-v8a", "arm64-android" } },
@@ -159,16 +120,17 @@
 		protected bool bEnableGcSections = true;
 
 		public AndroidToolChain(FileReference? InProjectFile, bool bInUseLdGold, IReadOnlyList<string>? InAdditionalArches, IReadOnlyList<string>? InAdditionalGPUArches)
-			: this(InProjectFile, bInUseLdGold, InAdditionalArches, InAdditionalGPUArches, false, AndroidToolChainOptions.None)
+			: this(InProjectFile, bInUseLdGold, InAdditionalArches, InAdditionalGPUArches, false, ClangToolChainOptions.None)
 		{
 		}
 
-		public AndroidToolChain(FileReference? InProjectFile, bool bInUseLdGold, IReadOnlyList<string>? InAdditionalArches, IReadOnlyList<string>? InAdditionalGPUArches, AndroidToolChainOptions ToolchainOptions)
+		public AndroidToolChain(FileReference? InProjectFile, bool bInUseLdGold, IReadOnlyList<string>? InAdditionalArches, IReadOnlyList<string>? InAdditionalGPUArches, ClangToolChainOptions ToolchainOptions)
 			: this(InProjectFile, bInUseLdGold, InAdditionalArches, InAdditionalGPUArches, false, ToolchainOptions)
 		{
 		}
 
-		protected AndroidToolChain(FileReference? InProjectFile, bool bInUseLdGold, IReadOnlyList<string>? InAdditionalArches, IReadOnlyList<string>? InAdditionalGPUArches, bool bAllowMissingNDK, AndroidToolChainOptions ToolchainOptions)
+		protected AndroidToolChain(FileReference? InProjectFile, bool bInUseLdGold, IReadOnlyList<string>? InAdditionalArches, IReadOnlyList<string>? InAdditionalGPUArches, bool bAllowMissingNDK, ClangToolChainOptions ToolchainOptions)
+			: base(ToolchainOptions)
 		{
 			Options = ToolchainOptions;
 			ProjectFile = InProjectFile;
@@ -2213,23 +2175,23 @@
 
 		public ClangSanitizer BuildWithSanitizer()
 		{
-			if (Options.HasFlag(AndroidToolChainOptions.EnableAddressSanitizer))
+			if (Options.HasFlag(ClangToolChainOptions.EnableAddressSanitizer))
 			{
 				return ClangSanitizer.Address;
 			}
-			else if (Options.HasFlag(AndroidToolChainOptions.EnableHWAddressSanitizer))
+			else if (Options.HasFlag(ClangToolChainOptions.EnableHWAddressSanitizer))
 			{
 				return ClangSanitizer.HwAddress;
 			}
-			else if (Options.HasFlag(AndroidToolChainOptions.EnableThreadSanitizer))
+			else if (Options.HasFlag(ClangToolChainOptions.EnableThreadSanitizer))
 			{
 				return ClangSanitizer.Thread;
 			}
-			else if (Options.HasFlag(AndroidToolChainOptions.EnableUndefinedBehaviorSanitizer))
+			else if (Options.HasFlag(ClangToolChainOptions.EnableUndefinedBehaviorSanitizer))
 			{
 				return ClangSanitizer.UndefinedBehavior;
 			}
-			else if (Options.HasFlag(AndroidToolChainOptions.EnableMinimalUndefinedBehaviorSanitizer))
+			else if (Options.HasFlag(ClangToolChainOptions.EnableMinimalUndefinedBehaviorSanitizer))
 			{
 				return ClangSanitizer.UndefinedBehaviorMinimal;
 			}
diff --git a/Engine/Source/Programs/UnrealBuildTool/Platform/Android/UEBuildAndroid.cs b/Engine/Source/Programs/UnrealBuildTool/Platform/Android/UEBuildAndroid.cs
--- a/Engine/Source/Programs/UnrealBuildTool/Platform/Android/UEBuildAndroid.cs
+++ b/Engine/Source/Programs/UnrealBuildTool/Platform/Android/UEBuildAndroid.cs
@@ -485,28 +485,28 @@
 			};
 		}
 
-		static public AndroidToolChainOptions CreateToolChainOptions(AndroidTargetRules TargetRules)
+		static public ClangToolChainOptions CreateToolChainOptions(AndroidTargetRules TargetRules)
 		{
-			AndroidToolChainOptions Options = AndroidToolChainOptions.None;
+			ClangToolChainOptions Options = ClangToolChainOptions.None;
 			if (TargetRules.bEnableAddressSanitizer)
 			{
-				Options |= AndroidToolChainOptions.EnableAddressSanitizer;
+				Options |= ClangToolChainOptions.EnableAddressSanitizer;
 			}
 			else if (TargetRules.bEnableHWAddressSanitizer)
 			{
-				Options |= AndroidToolChainOptions.EnableHWAddressSanitizer;
+				Options |= ClangToolChainOptions.EnableHWAddressSanitizer;
 			}
 			if (TargetRules.bEnableThreadSanitizer)
 			{
-				Options |= AndroidToolChainOptions.EnableThreadSanitizer;
+				Options |= ClangToolChainOptions.EnableThreadSanitizer;
 			}
 			if (TargetRules.bEnableUndefinedBehaviorSanitizer)
 			{
-				Options |= AndroidToolChainOptions.EnableUndefinedBehaviorSanitizer;
+				Options |= ClangToolChainOptions.EnableUndefinedBehaviorSanitizer;
 			}
 			else if (TargetRules.bEnableMinimalUndefinedBehaviorSanitizer)
 			{
-				Options |= AndroidToolChainOptions.EnableMinimalUndefinedBehaviorSanitizer;
+				Options |= ClangToolChainOptions.EnableMinimalUndefinedBehaviorSanitizer;
 			}
 
 			return Options;
@@ -515,14 +515,14 @@
 		public override UEToolChain CreateToolChain(ReadOnlyTargetRules Target)
 		{
 			bool bUseLdGold = Target.bUseUnityBuild;
-			AndroidToolChainOptions Options = CreateToolChainOptions(Target.AndroidPlatform.TargetRules);
+			ClangToolChainOptions Options = CreateToolChainOptions(Target.AndroidPlatform.TargetRules);
 			return new AndroidToolChain(Target.ProjectFile, bUseLdGold, Target.AndroidPlatform.Architectures, Target.AndroidPlatform.GPUArchitectures, Options);
 		}
 		public virtual UEToolChain CreateTempToolChainForProject(FileReference? ProjectFile)
 		{
 			AndroidTargetRules TargetRules = new AndroidTargetRules();
 			CommandLine.ParseArguments(Environment.GetCommandLineArgs(), TargetRules);
-			AndroidToolChainOptions Options = CreateToolChainOptions(TargetRules);
+			ClangToolChainOptions Options = CreateToolChainOptions(TargetRules);
 			return new AndroidToolChain(ProjectFile, true, null, null, Options);
 		}
 
diff --git a/Engine/Source/Programs/UnrealBuildTool/Platform/Android/UEDeployAndroid.cs b/Engine/Source/Programs/UnrealBuildTool/Platform/Android/UEDeployAndroid.cs
--- a/Engine/Source/Programs/UnrealBuildTool/Platform/Android/UEDeployAndroid.cs
+++ b/Engine/Source/Programs/UnrealBuildTool/Platform/Android/UEDeployAndroid.cs
@@ -4982,7 +4982,7 @@
 			// file and instead pushes files directly via deploy
 			AndroidTargetRules TargetRules = new AndroidTargetRules();
 			CommandLine.ParseArguments(Environment.GetCommandLineArgs(), TargetRules);
-			AndroidToolChainOptions Options = AndroidPlatform.CreateToolChainOptions(TargetRules);
+			ClangToolChainOptions Options = AndroidPlatform.CreateToolChainOptions(TargetRules);
 			AndroidToolChain ToolChain = new AndroidToolChain(ProjectFile, false, null, null, Options);
 
 			SavePackageInfo(ProjectName, ProjectDirectory.FullName, Type, bSkipGradleBuild);
diff --git a/Engine/Source/Programs/UnrealBuildTool/Platform/IOS/IOSExports.cs b/Engine/Source/Programs/UnrealBuildTool/Platform/IOS/IOSExports.cs
--- a/Engine/Source/Programs/UnrealBuildTool/Platform/IOS/IOSExports.cs
+++ b/Engine/Source/Programs/UnrealBuildTool/Platform/IOS/IOSExports.cs
@@ -117,7 +117,7 @@
 		public static void StripSymbols(UnrealTargetPlatform PlatformType, FileReference SourceFile, FileReference TargetFile)
 		{
 			IOSProjectSettings ProjectSettings = ((IOSPlatform)UEBuildPlatform.GetBuildPlatform(PlatformType)).ReadProjectSettings(null);
-			IOSToolChain ToolChain = new IOSToolChain(null, ProjectSettings, IOSToolChainOptions.None);
+			IOSToolChain ToolChain = new IOSToolChain(null, ProjectSettings, ClangToolChainOptions.None);
 			ToolChain.StripSymbols(SourceFile, TargetFile);
 		}
 
diff --git a/Engine/Source/Programs/UnrealBuildTool/Platform/IOS/IOSToolChain.cs b/Engine/Source/Programs/UnrealBuildTool/Platform/IOS/IOSToolChain.cs
--- a/Engine/Source/Programs/UnrealBuildTool/Platform/IOS/IOSToolChain.cs
+++ b/Engine/Source/Programs/UnrealBuildTool/Platform/IOS/IOSToolChain.cs
@@ -19,33 +19,6 @@
 
 namespace UnrealBuildTool
 {
-	/// <summary>
-	/// Option flags for the iOS toolchain
-	/// </summary>
-	[Flags]
-	enum IOSToolChainOptions
-	{
-		/// <summary>
-		/// No custom options
-		/// </summary>
-		None = 0,
-
-		/// <summary>
-		/// Enable address sanitizer
-		/// </summary>
-		EnableAddressSanitizer = 0x1,
-
-		/// <summary>
-		/// Enable thread sanitizer
-		/// </summary>
-		EnableThreadSanitizer = 0x2,
-
-		/// <summary>
-		/// Enable undefined behavior sanitizer
-		/// </summary>
-		EnableUndefinedBehaviorSanitizer = 0x4,
-	}
-
 	class IOSToolChainSettings : AppleToolChainSettings
 	{
 		/// <summary>
@@ -115,16 +88,14 @@
 
 		public readonly ReadOnlyTargetRules? Target;
 		protected IOSProjectSettings ProjectSettings;
-		private IOSToolChainOptions Options;
 
-		public IOSToolChain(ReadOnlyTargetRules? Target, IOSProjectSettings InProjectSettings, IOSToolChainOptions ToolchainOptions)
-			: this(Target, InProjectSettings, () => new IOSToolChainSettings())
+		public IOSToolChain(ReadOnlyTargetRules? Target, IOSProjectSettings InProjectSettings, ClangToolChainOptions ToolchainOptions)
+			: this(Target, InProjectSettings, () => new IOSToolChainSettings(), ToolchainOptions)
 		{
-			Options = ToolchainOptions;
 		}
 
-		protected IOSToolChain(ReadOnlyTargetRules? Target, IOSProjectSettings InProjectSettings, Func<IOSToolChainSettings> InCreateSettings)
-			: base((Target == null) ? null : Target.ProjectFile)
+		protected IOSToolChain(ReadOnlyTargetRules? Target, IOSProjectSettings InProjectSettings, Func<IOSToolChainSettings> InCreateSettings, ClangToolChainOptions ToolchainOptions)
+			: base((Target == null) ? null : Target.ProjectFile, ToolchainOptions)
 		{
 			this.Target = Target;
 			ProjectSettings = InProjectSettings;
@@ -260,18 +231,18 @@
 			}
 
 			string? SanitizerMode = Environment.GetEnvironmentVariable("ENABLE_ADDRESS_SANITIZER");
-			if ((SanitizerMode != null && SanitizerMode == "YES") || (Options.HasFlag(IOSToolChainOptions.EnableAddressSanitizer)))
+			if ((SanitizerMode != null && SanitizerMode == "YES") || (Options.HasFlag(ClangToolChainOptions.EnableAddressSanitizer)))
 			{
 				Arguments.Add("-fsanitize=address -fno-omit-frame-pointer -DFORCE_ANSI_ALLOCATOR=1");
 			}
 
 			string? UndefSanitizerMode = Environment.GetEnvironmentVariable("ENABLE_UNDEFINED_BEHAVIOR_SANITIZER");
-			if ((UndefSanitizerMode != null && UndefSanitizerMode == "YES") || (Options.HasFlag(IOSToolChainOptions.EnableUndefinedBehaviorSanitizer)))
+			if ((UndefSanitizerMode != null && UndefSanitizerMode == "YES") || (Options.HasFlag(ClangToolChainOptions.EnableUndefinedBehaviorSanitizer)))
 			{
 				Arguments.Add("-fsanitize=undefined -fno-sanitize=bounds,enum,return,float-divide-by-zero");
 			}
 
-			if (Options.HasFlag(IOSToolChainOptions.EnableThreadSanitizer))
+			if (Options.HasFlag(ClangToolChainOptions.EnableThreadSanitizer))
 			{
 				Arguments.Add("-fsanitize=thread");
 			}
@@ -472,20 +443,20 @@
 			}
 
 			string? SanitizerMode = Environment.GetEnvironmentVariable("ENABLE_ADDRESS_SANITIZER");
-			if ((SanitizerMode != null && SanitizerMode == "YES") || (Options.HasFlag(IOSToolChainOptions.EnableAddressSanitizer)))
+			if ((SanitizerMode != null && SanitizerMode == "YES") || (Options.HasFlag(ClangToolChainOptions.EnableAddressSanitizer)))
 			{
 				Result += " -rpath \"@executable_path/Frameworks\"";
 				Result += " -fsanitize=address";
 			}
 
 			string? UndefSanitizerMode = Environment.GetEnvironmentVariable("ENABLE_UNDEFINED_BEHAVIOR_SANITIZER");
-			if ((UndefSanitizerMode != null && UndefSanitizerMode == "YES") || (Options.HasFlag(IOSToolChainOptions.EnableUndefinedBehaviorSanitizer)))
+			if ((UndefSanitizerMode != null && UndefSanitizerMode == "YES") || (Options.HasFlag(ClangToolChainOptions.EnableUndefinedBehaviorSanitizer)))
 			{
 				Result += " -rpath \"@executable_path/libclang_rt.ubsan_ios_dynamic.dylib\"";
 				Result += " -fsanitize=undefined";
 			}
 
-			if (Options.HasFlag(IOSToolChainOptions.EnableThreadSanitizer))
+			if (Options.HasFlag(ClangToolChainOptions.EnableThreadSanitizer))
 			{
 				Result += " -rpath \"@executable_path/libclang_rt.tsan_ios_dynamic.dylib\"";
 				Result += " -fsanitize=thread";
diff --git a/Engine/Source/Programs/UnrealBuildTool/Platform/IOS/UEBuildIOS.cs b/Engine/Source/Programs/UnrealBuildTool/Platform/IOS/UEBuildIOS.cs
--- a/Engine/Source/Programs/UnrealBuildTool/Platform/IOS/UEBuildIOS.cs
+++ b/Engine/Source/Programs/UnrealBuildTool/Platform/IOS/UEBuildIOS.cs
@@ -1202,18 +1202,18 @@
 		/// <returns>New toolchain instance.</returns>
 		public override UEToolChain CreateToolChain(ReadOnlyTargetRules Target)
 		{
-			IOSToolChainOptions Options = IOSToolChainOptions.None;
+			ClangToolChainOptions Options = ClangToolChainOptions.None;
 			if (Target.IOSPlatform.bEnableAddressSanitizer)
 			{
-				Options |= IOSToolChainOptions.EnableAddressSanitizer;
+				Options |= ClangToolChainOptions.EnableAddressSanitizer;
 			}
 			if (Target.IOSPlatform.bEnableThreadSanitizer)
 			{
-				Options |= IOSToolChainOptions.EnableThreadSanitizer;
+				Options |= ClangToolChainOptions.EnableThreadSanitizer;
 			}
 			if (Target.IOSPlatform.bEnableUndefinedBehaviorSanitizer)
 			{
-				Options |= IOSToolChainOptions.EnableUndefinedBehaviorSanitizer;
+				Options |= ClangToolChainOptions.EnableUndefinedBehaviorSanitizer;
 			}
 
 			IOSProjectSettings ProjectSettings = ReadProjectSettings(Target.ProjectFile);
diff --git a/Engine/Source/Programs/UnrealBuildTool/Platform/Linux/LinuxToolChain.cs b/Engine/Source/Programs/UnrealBuildTool/Platform/Linux/LinuxToolChain.cs
--- a/Engine/Source/Programs/UnrealBuildTool/Platform/Linux/LinuxToolChain.cs
+++ b/Engine/Source/Programs/UnrealBuildTool/Platform/Linux/LinuxToolChain.cs
@@ -12,64 +12,6 @@
 
 namespace UnrealBuildTool
 {
-	/// <summary>
-	/// Option flags for the Linux toolchain
-	/// </summary>
-	[Flags]
-	enum LinuxToolChainOptions
-	{
-		/// <summary>
-		/// No custom options
-		/// </summary>
-		None = 0,
-
-		/// <summary>
-		/// Enable address sanitzier
-		/// </summary>
-		EnableAddressSanitizer = 0x1,
-
-		/// <summary>
-		/// Enable thread sanitizer
-		/// </summary>
-		EnableThreadSanitizer = 0x2,
-
-		/// <summary>
-		/// Enable undefined behavior sanitizer
-		/// </summary>
-		EnableUndefinedBehaviorSanitizer = 0x4,
-
-		/// <summary>
-		/// Enable memory sanitizer
-		/// </summary>
-		EnableMemorySanitizer = 0x8,
-
-		/// <summary>
-		/// Enable thin LTO
-		/// </summary>
-		EnableThinLTO = 0x10,
-
-		/// <summary>
-		/// Enable Shared library for the Sanitizers otherwise defaults to Statically linked
-		/// </summary>
-		EnableSharedSanitizer = 0x20,
-
-		/// <summary>
-		/// If should disable using objcopy to split the debug info into its own file or now
-		/// When we support larger the 4GB files with objcopy.exe this can be removed!
-		/// </summary>
-		DisableSplitDebugInfoWithObjCopy = 0x40,
-
-		/// <summary>
-		/// Enable tuning of debug info for LLDB
-		/// </summary>
-		TuneDebugInfoForLLDB = 0x80,
-
-		/// <summary>
-		/// Whether or not to preserve the portable symbol file produced by dump_syms
-		/// </summary>
-		PreservePSYM = 0x100,
-	}
-
 	class LinuxToolChain : ClangToolChain
 	{
 		/** Flavor of the current build (target triplet)*/
@@ -93,12 +35,7 @@
 		/** Toolchain information to print during the build. */
 		protected string? ToolchainInfo;
 
-		/// <summary>
-		/// Whether to compile with ASan enabled
-		/// </summary>
-		LinuxToolChainOptions Options;
-
-		public LinuxToolChain(string InArchitecture, LinuxPlatformSDK InSDK, LinuxToolChainOptions InOptions = LinuxToolChainOptions.None)
+		public LinuxToolChain(string InArchitecture, LinuxPlatformSDK InSDK, ClangToolChainOptions InOptions = ClangToolChainOptions.None)
 			: this(UnrealTargetPlatform.Linux, InArchitecture, InSDK, InOptions)
 		{
 			MultiArchRoot = PlatformSDK.GetSDKLocation();
@@ -190,12 +127,11 @@
 			}
 		}
 
-		public LinuxToolChain(UnrealTargetPlatform InPlatform, string InArchitecture, LinuxPlatformSDK InSDK, LinuxToolChainOptions InOptions = LinuxToolChainOptions.None)
-			: base()
+		public LinuxToolChain(UnrealTargetPlatform InPlatform, string InArchitecture, LinuxPlatformSDK InSDK, ClangToolChainOptions InOptions = ClangToolChainOptions.None)
+			: base(InOptions)
 		{
 			Architecture = InArchitecture;
 			PlatformSDK = InSDK;
-			Options = InOptions;
 		}
 
 		public override void SetUpGlobalEnvironment(ReadOnlyTargetRules Target)
@@ -250,7 +186,7 @@
 			FileItem StrippedFile = FileItem.GetItemByPath(Path.Combine(LinkEnvironment.LocalShadowDirectory.FullName, OutputFile.Location.GetFileName() + "_nodebug"));
 			FileItem DebugFile = FileItem.GetItemByPath(Path.Combine(LinkEnvironment.OutputDirectory.FullName, OutputFile.Location.GetFileNameWithoutExtension() + ".debug"));
 
-			if (Options.HasFlag(LinuxToolChainOptions.PreservePSYM))
+			if (Options.HasFlag(ClangToolChainOptions.PreservePSYM))
 			{
 				SymbolsFile = FileItem.GetItemByPath(Path.Combine(LinkEnvironment.OutputDirectory.FullName, OutputFile.Location.GetFileNameWithoutExtension() + ".psym"));
 			}
@@ -272,7 +208,7 @@
 				EncodedBinarySymbolsFile.AbsolutePath
 			);
 
-			if (!Options.HasFlag(LinuxToolChainOptions.DisableSplitDebugInfoWithObjCopy) && LinkEnvironment.bCreateDebugInfo)
+			if (!Options.HasFlag(ClangToolChainOptions.DisableSplitDebugInfoWithObjCopy) && LinkEnvironment.bCreateDebugInfo)
 			{
 				if (MaxBinarySizeOverrideForObjcopy > 0 && bUseCmdExe)
 				{
@@ -539,27 +475,27 @@
 			}
 
 			// ASan
-			if (Options.HasFlag(LinuxToolChainOptions.EnableAddressSanitizer))
+			if (Options.HasFlag(ClangToolChainOptions.EnableAddressSanitizer))
 			{
 				// Force using the ANSI allocator if ASan is enabled
 				Arguments.Add("-fsanitize=address -DFORCE_ANSI_ALLOCATOR=1");
 			}
 
 			// TSan
-			if (Options.HasFlag(LinuxToolChainOptions.EnableThreadSanitizer))
+			if (Options.HasFlag(ClangToolChainOptions.EnableThreadSanitizer))
 			{
 				// Force using the ANSI allocator if TSan is enabled
 				Arguments.Add("-fsanitize=thread -DFORCE_ANSI_ALLOCATOR=1");
 			}
 
 			// UBSan
-			if (Options.HasFlag(LinuxToolChainOptions.EnableUndefinedBehaviorSanitizer))
+			if (Options.HasFlag(ClangToolChainOptions.EnableUndefinedBehaviorSanitizer))
 			{
 				Arguments.Add("-fsanitize=undefined -fno-sanitize=vptr");
 			}
 
 			// MSan
-			if (Options.HasFlag(LinuxToolChainOptions.EnableMemorySanitizer))
+			if (Options.HasFlag(ClangToolChainOptions.EnableMemorySanitizer))
 			{
 				// Force using the ANSI allocator if MSan is enabled
 				// -fsanitize-memory-track-origins adds a 1.5x-2.5x slow down ontop of MSan normal amount of overhead
@@ -626,7 +562,7 @@
 			// Whether we actually can enable that is checked in CanUseAdvancedLinkerFeatures() earlier
 			if (CompileEnvironment.bAllowLTCG)
 			{
-				if ((Options & LinuxToolChainOptions.EnableThinLTO) != 0)
+				if ((Options & ClangToolChainOptions.EnableThinLTO) != 0)
 				{
 					Arguments.Add("-flto=thin");
 				}
@@ -639,7 +575,7 @@
 			//Arguments.Add("-DOPERATOR_NEW_INLINE=FORCENOINLINE");
 
 			bool bRetainFramePointers = CompileEnvironment.bRetainFramePointers
-				|| Options.HasFlag(LinuxToolChainOptions.EnableAddressSanitizer) || Options.HasFlag(LinuxToolChainOptions.EnableMemorySanitizer)
+				|| Options.HasFlag(ClangToolChainOptions.EnableAddressSanitizer) || Options.HasFlag(ClangToolChainOptions.EnableMemorySanitizer)
 				|| CompileEnvironment.Configuration == CppConfiguration.Debug;
 
 			if (CompileEnvironment.Configuration == CppConfiguration.Shipping)
@@ -675,7 +611,7 @@
 					Arguments.Add("-ggnu-pubnames");
 				}
 
-				if (Options.HasFlag(LinuxToolChainOptions.TuneDebugInfoForLLDB))
+				if (Options.HasFlag(ClangToolChainOptions.TuneDebugInfoForLLDB))
 				{
 					Arguments.Add("-glldb");
 				}
@@ -689,14 +625,14 @@
 			else
 			{
 				// Don't over optimise if using Address/MemorySanitizer or you'll get false positive errors due to erroneous optimisation of necessary Address/MemorySanitizer instrumentation.
-				if (Options.HasFlag(LinuxToolChainOptions.EnableAddressSanitizer) || Options.HasFlag(LinuxToolChainOptions.EnableMemorySanitizer))
+				if (Options.HasFlag(ClangToolChainOptions.EnableAddressSanitizer) || Options.HasFlag(ClangToolChainOptions.EnableMemorySanitizer))
 				{
 					Arguments.Add("-O1 -g");
 
 					// This enables __asan_default_options() in UnixCommonStartup.h which disables the leak detector
 					Arguments.Add("-DDISABLE_ASAN_LEAK_DETECTOR=1");
 				}
-				else if (Options.HasFlag(LinuxToolChainOptions.EnableThreadSanitizer))
+				else if (Options.HasFlag(ClangToolChainOptions.EnableThreadSanitizer))
 				{
 					Arguments.Add("-O1 -g");
 				}
@@ -804,31 +740,31 @@
 				Result += string.Format(" -Wl,--unresolved-symbols=ignore-in-shared-libs");
 			}
 
-			if (Options.HasFlag(LinuxToolChainOptions.EnableAddressSanitizer) ||
-				Options.HasFlag(LinuxToolChainOptions.EnableThreadSanitizer) ||
-				Options.HasFlag(LinuxToolChainOptions.EnableUndefinedBehaviorSanitizer) ||
-				Options.HasFlag(LinuxToolChainOptions.EnableMemorySanitizer))
+			if (Options.HasFlag(ClangToolChainOptions.EnableAddressSanitizer) ||
+				Options.HasFlag(ClangToolChainOptions.EnableThreadSanitizer) ||
+				Options.HasFlag(ClangToolChainOptions.EnableUndefinedBehaviorSanitizer) ||
+				Options.HasFlag(ClangToolChainOptions.EnableMemorySanitizer))
 			{
 				Result += " -g";
 
-				if (Options.HasFlag(LinuxToolChainOptions.EnableSharedSanitizer))
+				if (Options.HasFlag(ClangToolChainOptions.EnableSharedSanitizer))
 				{
 					Result += " -shared-libsan";
 				}
 
-				if (Options.HasFlag(LinuxToolChainOptions.EnableAddressSanitizer))
+				if (Options.HasFlag(ClangToolChainOptions.EnableAddressSanitizer))
 				{
 					Result += " -fsanitize=address";
 				}
-				else if (Options.HasFlag(LinuxToolChainOptions.EnableThreadSanitizer))
+				else if (Options.HasFlag(ClangToolChainOptions.EnableThreadSanitizer))
 				{
 					Result += " -fsanitize=thread";
 				}
-				else if (Options.HasFlag(LinuxToolChainOptions.EnableUndefinedBehaviorSanitizer))
+				else if (Options.HasFlag(ClangToolChainOptions.EnableUndefinedBehaviorSanitizer))
 				{
 					Result += " -fsanitize=undefined";
 				}
-				else if (Options.HasFlag(LinuxToolChainOptions.EnableMemorySanitizer))
+				else if (Options.HasFlag(ClangToolChainOptions.EnableMemorySanitizer))
 				{
 					// -fsanitize-memory-track-origins adds a 1.5x-2.5x slow ontop of MSan normal amount of overhead
 					// -fsanitize-memory-track-origins=1 is faster but collects only allocation points but not intermediate stores
@@ -915,7 +851,7 @@
 			// whether we actually can do that is checked in CanUseAdvancedLinkerFeatures() earlier
 			if (LinkEnvironment.bAllowLTCG)
 			{
-				if ((Options & LinuxToolChainOptions.EnableThinLTO) != 0)
+				if ((Options & ClangToolChainOptions.EnableThinLTO) != 0)
 				{
 					Result += String.Format(" -flto=thin -Wl,--thinlto-jobs={0}", Utils.GetPhysicalProcessorCount());
 				}
@@ -1026,18 +962,18 @@
 			Log.TraceInformation("Using {0}", UsingLld(CompileEnvironment.Architecture) ? "lld linker" : "default linker (ld)");
 			Log.TraceInformation("Using {0}", !String.IsNullOrEmpty(LlvmArPath) ? String.Format("llvm-ar : {0}", LlvmArPath) : String.Format("ar and ranlib: {0}, {1}", GetArPath(CompileEnvironment.Architecture), GetRanlibPath(CompileEnvironment.Architecture)));
 
-			if (Options.HasFlag(LinuxToolChainOptions.EnableAddressSanitizer) ||
-				Options.HasFlag(LinuxToolChainOptions.EnableThreadSanitizer) ||
-				Options.HasFlag(LinuxToolChainOptions.EnableUndefinedBehaviorSanitizer) ||
-				Options.HasFlag(LinuxToolChainOptions.EnableMemorySanitizer))
+			if (Options.HasFlag(ClangToolChainOptions.EnableAddressSanitizer) ||
+				Options.HasFlag(ClangToolChainOptions.EnableThreadSanitizer) ||
+				Options.HasFlag(ClangToolChainOptions.EnableUndefinedBehaviorSanitizer) ||
+				Options.HasFlag(ClangToolChainOptions.EnableMemorySanitizer))
 			{
 				string SanitizerInfo = "Building with:";
-				string StaticOrShared = Options.HasFlag(LinuxToolChainOptions.EnableSharedSanitizer) ? " dynamically" : " statically";
+				string StaticOrShared = Options.HasFlag(ClangToolChainOptions.EnableSharedSanitizer) ? " dynamically" : " statically";
 
-				SanitizerInfo += Options.HasFlag(LinuxToolChainOptions.EnableAddressSanitizer) ? StaticOrShared + " linked AddressSanitizer" : "";
-				SanitizerInfo += Options.HasFlag(LinuxToolChainOptions.EnableThreadSanitizer) ? StaticOrShared + " linked ThreadSanitizer" : "";
-				SanitizerInfo += Options.HasFlag(LinuxToolChainOptions.EnableUndefinedBehaviorSanitizer) ? StaticOrShared + " linked UndefinedBehaviorSanitizer" : "";
-				SanitizerInfo += Options.HasFlag(LinuxToolChainOptions.EnableMemorySanitizer) ? StaticOrShared + " linked MemorySanitizer" : "";
+				SanitizerInfo += Options.HasFlag(ClangToolChainOptions.EnableAddressSanitizer) ? StaticOrShared + " linked AddressSanitizer" : "";
+				SanitizerInfo += Options.HasFlag(ClangToolChainOptions.EnableThreadSanitizer) ? StaticOrShared + " linked ThreadSanitizer" : "";
+				SanitizerInfo += Options.HasFlag(ClangToolChainOptions.EnableUndefinedBehaviorSanitizer) ? StaticOrShared + " linked UndefinedBehaviorSanitizer" : "";
+				SanitizerInfo += Options.HasFlag(ClangToolChainOptions.EnableMemorySanitizer) ? StaticOrShared + " linked MemorySanitizer" : "";
 
 				Log.TraceInformation(SanitizerInfo);
 			}
diff --git a/Engine/Source/Programs/UnrealBuildTool/Platform/Linux/UEBuildLinux.cs b/Engine/Source/Programs/UnrealBuildTool/Platform/Linux/UEBuildLinux.cs
--- a/Engine/Source/Programs/UnrealBuildTool/Platform/Linux/UEBuildLinux.cs
+++ b/Engine/Source/Programs/UnrealBuildTool/Platform/Linux/UEBuildLinux.cs
@@ -547,20 +547,20 @@
 		/// <returns>New toolchain instance.</returns>
 		public override UEToolChain CreateToolChain(ReadOnlyTargetRules Target)
 		{
-			LinuxToolChainOptions Options = LinuxToolChainOptions.None;
+			ClangToolChainOptions Options = ClangToolChainOptions.None;
 
 			if (Target.LinuxPlatform.bEnableAddressSanitizer)
 			{
-				Options |= LinuxToolChainOptions.EnableAddressSanitizer;
+				Options |= ClangToolChainOptions.EnableAddressSanitizer;
 
 				if (Target.LinkType != TargetLinkType.Monolithic)
 				{
-					Options |= LinuxToolChainOptions.EnableSharedSanitizer;
+					Options |= ClangToolChainOptions.EnableSharedSanitizer;
 				}
 			}
 			if (Target.LinuxPlatform.bEnableThreadSanitizer)
 			{
-				Options |= LinuxToolChainOptions.EnableThreadSanitizer;
+				Options |= ClangToolChainOptions.EnableThreadSanitizer;
 
 				if (Target.LinkType != TargetLinkType.Monolithic)
 				{
@@ -569,16 +569,16 @@
 			}
 			if (Target.LinuxPlatform.bEnableUndefinedBehaviorSanitizer)
 			{
-				Options |= LinuxToolChainOptions.EnableUndefinedBehaviorSanitizer;
+				Options |= ClangToolChainOptions.EnableUndefinedBehaviorSanitizer;
 
 				if (Target.LinkType != TargetLinkType.Monolithic)
 				{
-					Options |= LinuxToolChainOptions.EnableSharedSanitizer;
+					Options |= ClangToolChainOptions.EnableSharedSanitizer;
 				}
 			}
 			if (Target.LinuxPlatform.bEnableMemorySanitizer)
 			{
-				Options |= LinuxToolChainOptions.EnableMemorySanitizer;
+				Options |= ClangToolChainOptions.EnableMemorySanitizer;
 
 				if (Target.LinkType != TargetLinkType.Monolithic)
 				{
@@ -587,7 +587,7 @@
 			}
 			if (Target.bAllowLTCG && Target.bPreferThinLTO)
 			{
-				Options |= LinuxToolChainOptions.EnableThinLTO;
+				Options |= ClangToolChainOptions.EnableThinLTO;
 			}
 
 			// When building a monolithic editor we have to avoid using objcopy.exe as it cannot handle files
@@ -597,17 +597,17 @@
 				(Target.LinkType == TargetLinkType.Monolithic) &&
 				(Target.Type == TargetType.Editor))
 			{
-				Options |= LinuxToolChainOptions.DisableSplitDebugInfoWithObjCopy;
+				Options |= ClangToolChainOptions.DisableSplitDebugInfoWithObjCopy;
 			}
 
 			if (Target.LinuxPlatform.bTuneDebugInfoForLLDB)
 			{
-				Options |= LinuxToolChainOptions.TuneDebugInfoForLLDB;
+				Options |= ClangToolChainOptions.TuneDebugInfoForLLDB;
 			}
 
 			if (Target.LinuxPlatform.bPreservePSYM)
 			{
-				Options |= LinuxToolChainOptions.PreservePSYM;
+				Options |= ClangToolChainOptions.PreservePSYM;
 			}
 
 			// Disable color logging if we are on a build machine
diff --git a/Engine/Source/Programs/UnrealBuildTool/Platform/Mac/MacExports.cs b/Engine/Source/Programs/UnrealBuildTool/Platform/Mac/MacExports.cs
--- a/Engine/Source/Programs/UnrealBuildTool/Platform/Mac/MacExports.cs
+++ b/Engine/Source/Programs/UnrealBuildTool/Platform/Mac/MacExports.cs
@@ -138,7 +138,7 @@
 		/// <param name="TargetFile">The output file</param>
 		public static void StripSymbols(FileReference SourceFile, FileReference TargetFile)
 		{
-			MacToolChain ToolChain = new MacToolChain(null, MacToolChainOptions.None);
+			MacToolChain ToolChain = new MacToolChain(null, ClangToolChainOptions.None);
 			ToolChain.StripSymbols(SourceFile, TargetFile);
 		}		
 	}
diff --git a/Engine/Source/Programs/UnrealBuildTool/Platform/Mac/MacToolChain.cs b/Engine/Source/Programs/UnrealBuildTool/Platform/Mac/MacToolChain.cs
--- a/Engine/Source/Programs/UnrealBuildTool/Platform/Mac/MacToolChain.cs
+++ b/Engine/Source/Programs/UnrealBuildTool/Platform/Mac/MacToolChain.cs
@@ -68,51 +68,13 @@
 	}
 
 	/// <summary>
-	/// Option flags for the Mac toolchain
-	/// </summary>
-	[Flags]
-	enum MacToolChainOptions
-	{
-		/// <summary>
-		/// No custom options
-		/// </summary>
-		None = 0,
-
-		/// <summary>
-		/// Enable address sanitzier
-		/// </summary>
-		EnableAddressSanitizer = 0x1,
-
-		/// <summary>
-		/// Enable thread sanitizer
-		/// </summary>
-		EnableThreadSanitizer = 0x2,
-
-		/// <summary>
-		/// Enable undefined behavior sanitizer
-		/// </summary>
-		EnableUndefinedBehaviorSanitizer = 0x4,
-
-		/// <summary>
-		/// Whether we're outputting a dylib instead of an executable
-		/// </summary>
-		OutputDylib = 0x08,
-	}
-
-	/// <summary>
 	/// Mac toolchain wrapper
 	/// </summary>
 	class MacToolChain : AppleToolChain
 	{
-		/// <summary>
-		/// Whether to compile with ASan enabled
-		/// </summary>
-		MacToolChainOptions Options;
-
-		public MacToolChain(FileReference? InProjectFile, MacToolChainOptions InOptions)
-			: base(InProjectFile)
+		public MacToolChain(FileReference? InProjectFile, ClangToolChainOptions InOptions)
+			: base(InProjectFile, InOptions)
 		{
-			this.Options = InOptions;			
 		}
 
 		public static Lazy<MacToolChainSettings> SettingsPrivate = new Lazy<MacToolChainSettings>(() => new MacToolChainSettings(false));
@@ -214,15 +176,15 @@
 				Arguments.Add("-fvisibility-ms-compat");
 				Arguments.Add("-fvisibility-inlines-hidden");
 			}
-			if (Options.HasFlag(MacToolChainOptions.EnableAddressSanitizer))
+			if (Options.HasFlag(ClangToolChainOptions.EnableAddressSanitizer))
 			{
 				Arguments.Add("-fsanitize=address");
 			}
-			if (Options.HasFlag(MacToolChainOptions.EnableThreadSanitizer))
+			if (Options.HasFlag(ClangToolChainOptions.EnableThreadSanitizer))
 			{
 				Arguments.Add("-fsanitize=thread");
 			}
-			if (Options.HasFlag(MacToolChainOptions.EnableUndefinedBehaviorSanitizer))
+			if (Options.HasFlag(ClangToolChainOptions.EnableUndefinedBehaviorSanitizer))
 			{
 				Arguments.Add("-fsanitize=undefined");
 			}					
@@ -248,11 +210,11 @@
 			if (CompileEnvironment.bOptimizeCode && !bStaticAnalysis)
 			{
 				// Don't over optimise if using AddressSanitizer or you'll get false positive errors due to erroneous optimisation of necessary AddressSanitizer instrumentation.
-				if (Options.HasFlag(MacToolChainOptions.EnableAddressSanitizer))
+				if (Options.HasFlag(ClangToolChainOptions.EnableAddressSanitizer))
 				{
 					Arguments.Add("-O1 -g -fno-optimize-sibling-calls -fno-omit-frame-pointer");
 				}
-				else if (Options.HasFlag(MacToolChainOptions.EnableThreadSanitizer))
+				else if (Options.HasFlag(ClangToolChainOptions.EnableThreadSanitizer))
 				{
 					Arguments.Add("-O1 -g");
 				}
@@ -305,18 +267,18 @@
 			Result += " -dead_strip";
 			Result += " -Wl,-fatal_warnings";
 
-			if (Options.HasFlag(MacToolChainOptions.EnableAddressSanitizer) || Options.HasFlag(MacToolChainOptions.EnableThreadSanitizer) || Options.HasFlag(MacToolChainOptions.EnableUndefinedBehaviorSanitizer))
+			if (Options.HasFlag(ClangToolChainOptions.EnableAddressSanitizer) || Options.HasFlag(ClangToolChainOptions.EnableThreadSanitizer) || Options.HasFlag(ClangToolChainOptions.EnableUndefinedBehaviorSanitizer))
 			{
 				Result += " -g";
-				if (Options.HasFlag(MacToolChainOptions.EnableAddressSanitizer))
+				if (Options.HasFlag(ClangToolChainOptions.EnableAddressSanitizer))
 				{
 					Result += " -fsanitize=address";
 				}
-				else if (Options.HasFlag(MacToolChainOptions.EnableThreadSanitizer))
+				else if (Options.HasFlag(ClangToolChainOptions.EnableThreadSanitizer))
 				{
 					Result += " -fsanitize=thread";
 				}
-				else if (Options.HasFlag(MacToolChainOptions.EnableUndefinedBehaviorSanitizer))
+				else if (Options.HasFlag(ClangToolChainOptions.EnableUndefinedBehaviorSanitizer))
 				{
 					Result += " -fsanitize=undefined";
 				}
@@ -1435,7 +1397,7 @@
 				}
 			}
 
-			if ((BinaryLinkEnvironment.bIsBuildingDLL && (Options & MacToolChainOptions.OutputDylib) == 0) || (BinaryLinkEnvironment.bIsBuildingConsoleApplication && Executable.Name.EndsWith("-Cmd")))
+			if ((BinaryLinkEnvironment.bIsBuildingDLL && (Options & ClangToolChainOptions.OutputDylib) == 0) || (BinaryLinkEnvironment.bIsBuildingConsoleApplication && Executable.Name.EndsWith("-Cmd")))
 			{
 				return OutputFiles;
 			}
diff --git a/Engine/Source/Programs/UnrealBuildTool/Platform/Mac/UEBuildMac.cs b/Engine/Source/Programs/UnrealBuildTool/Platform/Mac/UEBuildMac.cs
--- a/Engine/Source/Programs/UnrealBuildTool/Platform/Mac/UEBuildMac.cs
+++ b/Engine/Source/Programs/UnrealBuildTool/Platform/Mac/UEBuildMac.cs
@@ -376,7 +376,7 @@
 		/// <returns>New toolchain instance.</returns>
 		public override UEToolChain CreateToolChain(ReadOnlyTargetRules Target)
 		{
-			MacToolChainOptions Options = MacToolChainOptions.None;
+			ClangToolChainOptions Options = ClangToolChainOptions.None;
 
 			string? AddressSanitizer = Environment.GetEnvironmentVariable("ENABLE_ADDRESS_SANITIZER");
 			string? ThreadSanitizer = Environment.GetEnvironmentVariable("ENABLE_THREAD_SANITIZER");
@@ -384,19 +384,19 @@
 
 			if(Target.MacPlatform.bEnableAddressSanitizer || (AddressSanitizer != null && AddressSanitizer == "YES"))
 			{
-				Options |= MacToolChainOptions.EnableAddressSanitizer;
+				Options |= ClangToolChainOptions.EnableAddressSanitizer;
 			}
 			if(Target.MacPlatform.bEnableThreadSanitizer || (ThreadSanitizer != null && ThreadSanitizer == "YES"))
 			{
-				Options |= MacToolChainOptions.EnableThreadSanitizer;
+				Options |= ClangToolChainOptions.EnableThreadSanitizer;
 			}
 			if(Target.MacPlatform.bEnableUndefinedBehaviorSanitizer || (UndefSanitizerMode != null && UndefSanitizerMode == "YES"))
 			{
-				Options |= MacToolChainOptions.EnableUndefinedBehaviorSanitizer;
+				Options |= ClangToolChainOptions.EnableUndefinedBehaviorSanitizer;
 			}
 			if(Target.bShouldCompileAsDLL)
 			{
-				Options |= MacToolChainOptions.OutputDylib;
+				Options |= ClangToolChainOptions.OutputDylib;
 			}
 
 			return new MacToolChain(Target.ProjectFile, Options);
diff --git a/Engine/Source/Programs/UnrealBuildTool/Platform/TVOS/TVOSToolChain.cs b/Engine/Source/Programs/UnrealBuildTool/Platform/TVOS/TVOSToolChain.cs
--- a/Engine/Source/Programs/UnrealBuildTool/Platform/TVOS/TVOSToolChain.cs
+++ b/Engine/Source/Programs/UnrealBuildTool/Platform/TVOS/TVOSToolChain.cs
@@ -26,7 +26,7 @@
 	class TVOSToolChain : IOSToolChain
 	{
 		public TVOSToolChain(ReadOnlyTargetRules InTarget, TVOSProjectSettings InProjectSettings)
-			: base(InTarget, InProjectSettings, () => new TVOSToolChainSettings())
+			: base(InTarget, InProjectSettings, () => new TVOSToolChainSettings(), ClangToolChainOptions.None)
 		{
 		}
 
diff --git a/Engine/Source/Programs/UnrealBuildTool/ToolChain/AppleToolChain.cs b/Engine/Source/Programs/UnrealBuildTool/ToolChain/AppleToolChain.cs
--- a/Engine/Source/Programs/UnrealBuildTool/ToolChain/AppleToolChain.cs
+++ b/Engine/Source/Programs/UnrealBuildTool/ToolChain/AppleToolChain.cs
@@ -149,7 +149,7 @@
 	{
 		protected FileReference? ProjectFile;
 
-		public AppleToolChain(FileReference? InProjectFile)
+		public AppleToolChain(FileReference? InProjectFile, ClangToolChainOptions InOptions) : base(InOptions)
 		{
 			ProjectFile = InProjectFile;
 		}
diff --git a/Engine/Source/Programs/UnrealBuildTool/ToolChain/ClangToolChain.cs b/Engine/Source/Programs/UnrealBuildTool/ToolChain/ClangToolChain.cs
--- a/Engine/Source/Programs/UnrealBuildTool/ToolChain/ClangToolChain.cs
+++ b/Engine/Source/Programs/UnrealBuildTool/ToolChain/ClangToolChain.cs
@@ -8,6 +8,100 @@
 
 namespace UnrealBuildTool
 {
+	/// <summary>
+	/// Common option flags for the Clang toolchains.
+	/// Usage of these flags is currently inconsistent between various toolchains.
+	/// </summary>
+	[Flags]
+	enum ClangToolChainOptions
+	{
+		/// <summary>
+		/// No custom options
+		/// </summary>
+		None = 0,
+
+		/// <summary>
+		/// Enable address sanitzier
+		/// </summary>
+		EnableAddressSanitizer = 1 << 0,
+
+		/// <summary>
+		/// Enable hardware address sanitzier
+		/// </summary>
+		EnableHWAddressSanitizer = 1 << 1,
+
+		/// <summary>
+		/// Enable thread sanitizer
+		/// </summary>
+		EnableThreadSanitizer = 1 << 2,
+
+		/// <summary>
+		/// Enable undefined behavior sanitizer
+		/// </summary>
+		EnableUndefinedBehaviorSanitizer = 1 << 3,
+
+		/// <summary>
+		/// Enable minimal undefined behavior sanitizer
+		/// </summary>
+		EnableMinimalUndefinedBehaviorSanitizer = 1 << 4,
+
+		/// <summary>
+		/// Enable memory sanitizer
+		/// </summary>
+		EnableMemorySanitizer = 1 << 5,
+
+		/// <summary>
+		/// Enable Shared library for the Sanitizers otherwise defaults to Statically linked
+		/// </summary>
+		EnableSharedSanitizer = 1 << 6,
+
+		/// <summary>
+		/// Enables link time optimization (LTO). Link times will significantly increase.
+		/// </summary>
+		EnableLinkTimeOptimization = 1 << 7,
+
+		/// <summary>
+		/// Enable thin LTO
+		/// </summary>
+		EnableThinLTO = 1 << 8,
+
+		/// <summary>
+		/// If should disable using objcopy to split the debug info into its own file or now
+		/// When we support larger the 4GB files with objcopy.exe this can be removed!
+		/// </summary>
+		DisableSplitDebugInfoWithObjCopy = 1 << 9,
+
+		/// <summary>
+		/// Enable tuning of debug info for LLDB
+		/// </summary>
+		TuneDebugInfoForLLDB = 1 << 10,
+
+		/// <summary>
+		/// Whether or not to preserve the portable symbol file produced by dump_syms
+		/// </summary>
+		PreservePSYM = 1 << 11,
+
+		/// <summary>
+		/// (Apple toolchains) Whether we're outputting a dylib instead of an executable
+		/// </summary>
+		OutputDylib = 1 << 12,
+
+		/// <summary>
+		/// Enables the creation of custom symbol files used for runtime symbol resolution.
+		/// </summary>
+		GenerateSymbols = 1 << 13,
+
+		/// <summary>
+		/// Enables live code editing
+		/// </summary>
+		EnableLiveCodeEditing = 1 << 14,
+
+		/// <summary>
+		/// Enables dead code/data stripping and common code folding.
+		/// </summary>
+		EnableDeadStripping = 1 << 15
+	}
+
 	abstract class ClangToolChain : ISPCToolChain
 	{
 		// The Clang version being used to compile
@@ -16,8 +110,11 @@
 		protected int ClangVersionMinor = -1;
 		protected int ClangVersionPatch = -1;
 
-		public ClangToolChain() : base()
+		protected ClangToolChainOptions Options;
+
+		public ClangToolChain(ClangToolChainOptions InOptions) : base()
 		{
+			Options = InOptions;
 		}
 
 		/// <summary>
