Add LTO type to Linux build information

#rnx

diff --git a/Engine/Source/Programs/UnrealBuildTool/Platform/Linux/LinuxToolChain.cs b/Engine/Source/Programs/UnrealBuildTool/Platform/Linux/LinuxToolChain.cs
--- a/Engine/Source/Programs/UnrealBuildTool/Platform/Linux/LinuxToolChain.cs
+++ b/Engine/Source/Programs/UnrealBuildTool/Platform/Linux/LinuxToolChain.cs
@@ -1002,7 +1002,8 @@ namespace UnrealBuildTool
 
 			if (CompileEnvironment.bAllowLTCG)
 			{
-				Log.TraceInformation("Using LTO (link-time optimization).");
+				Log.TraceInformation("Using {0}LTO (link-time optimization).",
+					((Options & LinuxToolChainOptions.EnableThinLTO) != 0) ? "Thin" : "");
 			}
 
 			if (bSuppressPIE)
