Fix up toolchain in light of clang tool chain options refactor
#rnx

diff --git a/Engine/Source/Programs/UnrealBuildTool/Platform/Linux/LinuxToolChain.cs b/Engine/Source/Programs/UnrealBuildTool/Platform/Linux/LinuxToolChain.cs
--- a/Engine/Source/Programs/UnrealBuildTool/Platform/Linux/LinuxToolChain.cs
+++ b/Engine/Source/Programs/UnrealBuildTool/Platform/Linux/LinuxToolChain.cs
@@ -1003,7 +1003,7 @@ namespace UnrealBuildTool
 			if (CompileEnvironment.bAllowLTCG)
 			{
 				Log.TraceInformation("Using {0}LTO (link-time optimization).",
-					((Options & LinuxToolChainOptions.EnableThinLTO) != 0) ? "Thin" : "");
+					((Options & ClangToolChainOptions.EnableThinLTO) != 0) ? "Thin" : "");
 			}
 
 			if (bSuppressPIE)
