Fix build break issue 174391

diff --git a/Engine/Source/Programs/UnrealBuildTool/Platform/Linux/LinuxToolChain.cs b/Engine/Source/Programs/UnrealBuildTool/Platform/Linux/LinuxToolChain.cs
--- a/Engine/Source/Programs/UnrealBuildTool/Platform/Linux/LinuxToolChain.cs
+++ b/Engine/Source/Programs/UnrealBuildTool/Platform/Linux/LinuxToolChain.cs
@@ -1002,8 +1002,7 @@ namespace UnrealBuildTool
 
 			if (CompileEnvironment.bAllowLTCG)
 			{
-				Log.TraceInformation("Using {0}LTO (link-time optimization).",
-					((Options & ClangToolChainOptions.EnableThinLTO) != 0) ? "Thin" : "");
+				Log.TraceInformation("Using LTO (link-time optimization).");
 			}
 
 			if (bSuppressPIE)
