UnrealBuildTool: Relative pathing for ISPCToolChain, clean up argument strings, only process global arguments once

#rnx

diff --git a/Engine/Source/Programs/UnrealBuildTool/Platform/Windows/VCToolChain.cs b/Engine/Source/Programs/UnrealBuildTool/Platform/Windows/VCToolChain.cs
--- a/Engine/Source/Programs/UnrealBuildTool/Platform/Windows/VCToolChain.cs
+++ b/Engine/Source/Programs/UnrealBuildTool/Platform/Windows/VCToolChain.cs
@@ -136,9 +136,9 @@
 			}
 		}
 
-		public static string NormalizeCommandLinePath(FileSystemReference Reference)
+		public static new string NormalizeCommandLinePath(FileSystemReference Reference)
 		{
-			// Try to use a relative path to shorten command line length and to enabled remote distribution where absolute paths are not desired
+			// Try to use a relative path to shorten command line length and to enable remote distribution where absolute paths are not desired
 			if (Reference.IsUnderDirectory(Unreal.EngineDirectory))
 			{
 				string RelativePath = Reference.MakeRelativeTo(UnrealBuildTool.EngineSourceDirectory);
@@ -150,7 +150,7 @@
 			return Reference.FullName;
 		}
 
-		public static string NormalizeCommandLinePath(FileItem Item)
+		public static new string NormalizeCommandLinePath(FileItem Item)
 		{
 			return NormalizeCommandLinePath(Item.Location);
 		}
diff --git a/Engine/Source/Programs/UnrealBuildTool/ToolChain/ClangToolChain.cs b/Engine/Source/Programs/UnrealBuildTool/ToolChain/ClangToolChain.cs
--- a/Engine/Source/Programs/UnrealBuildTool/ToolChain/ClangToolChain.cs
+++ b/Engine/Source/Programs/UnrealBuildTool/ToolChain/ClangToolChain.cs
@@ -117,32 +117,6 @@ namespace UnrealBuildTool
 			Options = InOptions;
 		}
 
-		/// <summary>
-		/// Normalize a path for use in a command line, making it relative to Engine/Source if under the root directory
-		/// </summary>
-		/// <param name="Reference">The FileSystemReference to normalize</param>
-		/// <returns>Normalized path as a string</returns>
-		protected static string NormalizeCommandLinePath(FileSystemReference Reference)
-		{
-			// Try to use a relative path to shorten command line length.
-			if (Reference.IsUnderDirectory(Unreal.RootDirectory))
-			{
-				return Reference.MakeRelativeTo(UnrealBuildTool.EngineSourceDirectory).Replace("\\", "/");
-			}
-
-			return Reference.FullName.Replace("\\", "/");
-		}
-
-		/// <summary>
-		/// Normalize a path for use in a command line, making it relative if under the Root Directory
-		/// </summary>
-		/// <param name="Item">The FileItem to normalize</param>
-		/// <returns>Normalized path as a string</returns>
-		protected static string NormalizeCommandLinePath(FileItem Item)
-		{
-			return NormalizeCommandLinePath(Item.Location);
-		}
-
 		/// <summary>
 		/// Sanitizes a preprocessor definition argument if needed.
 		/// </summary>
diff --git a/Engine/Source/Programs/UnrealBuildTool/ToolChain/ISPCToolChain.cs b/Engine/Source/Programs/UnrealBuildTool/ToolChain/ISPCToolChain.cs
--- a/Engine/Source/Programs/UnrealBuildTool/ToolChain/ISPCToolChain.cs
+++ b/Engine/Source/Programs/UnrealBuildTool/ToolChain/ISPCToolChain.cs
@@ -319,6 +319,32 @@
 			return Escaped;
 		}
 
+		/// <summary>
+		/// Normalize a path for use in a command line, making it relative to Engine/Source if under the root directory
+		/// </summary>
+		/// <param name="Reference">The FileSystemReference to normalize</param>
+		/// <returns>Normalized path as a string</returns>
+		protected static string NormalizeCommandLinePath(FileSystemReference Reference)
+		{
+			// Try to use a relative path to shorten command line length.
+			if (Reference.IsUnderDirectory(Unreal.RootDirectory))
+			{
+				return Reference.MakeRelativeTo(UnrealBuildTool.EngineSourceDirectory).Replace("\\", "/");
+			}
+
+			return Reference.FullName.Replace("\\", "/");
+		}
+
+		/// <summary>
+		/// Normalize a path for use in a command line, making it relative if under the Root Directory
+		/// </summary>
+		/// <param name="Item">The FileItem to normalize</param>
+		/// <returns>Normalized path as a string</returns>
+		protected static string NormalizeCommandLinePath(FileItem Item)
+		{
+			return NormalizeCommandLinePath(Item.Location);
+		}
+
 		public override CPPOutput GenerateISPCHeaders(CppCompileEnvironment CompileEnvironment, List<FileItem> InputFiles, DirectoryReference OutputDir, IActionGraphBuilder Graph)
 		{
 			CPPOutput Result = new CPPOutput();
@@ -330,6 +356,65 @@
 
 			List<string> CompileTargets = GetISPCCompileTargets(CompileEnvironment.Platform, null);
 
+			List<string> GlobalArguments = new List<string>();
+
+			// Build target string. No comma on last
+			string TargetString = "";
+			foreach (string Target in CompileTargets)
+			{
+				if (Target == CompileTargets[CompileTargets.Count - 1]) // .Last()
+				{
+					TargetString += Target;
+				}
+				else
+				{
+					TargetString += Target + ",";
+				}
+			}
+
+			// Build target triplet
+			GlobalArguments.Add($"--target-os={GetISPCOSTarget(CompileEnvironment.Platform)}");
+			GlobalArguments.Add($"--arch={GetISPCArchTarget(CompileEnvironment.Platform, null)}");
+			GlobalArguments.Add($"--target={TargetString}");
+			GlobalArguments.Add($"--emit-{GetISPCObjectFileFormat(CompileEnvironment.Platform)}");
+
+			string? CpuTarget = GetISPCCpuTarget(CompileEnvironment.Platform);
+			if (!String.IsNullOrEmpty(CpuTarget))
+			{
+				GlobalArguments.Add($"--cpu={CpuTarget}");
+			}
+
+			// PIC is needed for modular builds except on Microsoft platforms
+			if ((CompileEnvironment.bIsBuildingDLL ||
+				CompileEnvironment.bIsBuildingLibrary) &&
+				!UEBuildPlatform.IsPlatformInGroup(CompileEnvironment.Platform, UnrealPlatformGroup.Microsoft))
+			{
+				GlobalArguments.Add("--pic");
+			}
+
+			// Include paths. Don't use AddIncludePath() here, since it uses the full path and exceeds the max command line length.
+			// Because ISPC response files don't support white space in arguments, paths with white space need to be passed to the command line directly.
+			foreach (DirectoryReference IncludePath in CompileEnvironment.UserIncludePaths)
+			{
+				GlobalArguments.Add($"-I\"{NormalizeCommandLinePath(IncludePath)}\"");
+			}
+
+			// System include paths.
+			foreach (DirectoryReference SystemIncludePath in CompileEnvironment.SystemIncludePaths)
+			{
+				GlobalArguments.Add($"-I\"{NormalizeCommandLinePath(SystemIncludePath)}\"");
+			}
+
+			// Preprocessor definitions.
+			foreach (string Definition in CompileEnvironment.Definitions)
+			{
+				// TODO: Causes ISPC compiler to generate a spurious warning about the universal character set
+				if (!Definition.Contains("\\\\U") && !Definition.Contains("\\\\u"))
+				{
+					GlobalArguments.Add($"-D{EscapeDefinitionForISPC(Definition)}");
+				}
+			}
+
 			foreach (FileItem ISPCFile in InputFiles)
 			{
 				Action CompileAction = Graph.CreateAction(ActionType.Compile);
@@ -344,7 +429,7 @@
 				List<string> Arguments = new List<string>();
 
 				// Add the ISPC obj file as a prerequisite of the action.
-				CompileAction.CommandArguments = String.Format("\"{0}\" ", ISPCFile.AbsolutePath);
+				Arguments.Add($"\"{NormalizeCommandLinePath(ISPCFile)}\"");
 
 				// Add the ISPC h file to the produced item list.
 				FileItem ISPCIncludeHeaderFile = FileItem.GetItemByFileReference(
@@ -355,79 +440,24 @@
 					);
 
 				// Add the ISPC file to be compiled.
-				Arguments.Add(String.Format("-h \"{0}\"", ISPCIncludeHeaderFile));
+				Arguments.Add($"-h \"{NormalizeCommandLinePath(ISPCIncludeHeaderFile)}\"");
 
-				// Build target string. No comma on last
-				string TargetString = "";
-				foreach (string Target in CompileTargets)
-				{
-					if (Target == CompileTargets[CompileTargets.Count - 1]) // .Last()
-					{
-						TargetString += Target;
-					}
-					else
-					{
-						TargetString += Target + ",";
-					}
-				}
-
-				// Build target triplet
-				Arguments.Add(String.Format("--target-os={0}", GetISPCOSTarget(CompileEnvironment.Platform)));
-				Arguments.Add(String.Format("--arch={0}", GetISPCArchTarget(CompileEnvironment.Platform, null)));
-				Arguments.Add(String.Format("--target={0}", TargetString));
-				Arguments.Add(String.Format("--emit-{0}", GetISPCObjectFileFormat(CompileEnvironment.Platform)));
-
-				string? CpuTarget = GetISPCCpuTarget(CompileEnvironment.Platform);
-				if (!String.IsNullOrEmpty(CpuTarget))
-				{
-					Arguments.Add(String.Format("--cpu={0}", CpuTarget));
-				}
-
-				// PIC is needed for modular builds except on Microsoft platforms
-				if ((CompileEnvironment.bIsBuildingDLL ||
-					CompileEnvironment.bIsBuildingLibrary) &&
-					!UEBuildPlatform.IsPlatformInGroup(CompileEnvironment.Platform, UnrealPlatformGroup.Microsoft))
-				{
-					Arguments.Add("--pic");
-				}
-
-				// Include paths. Don't use AddIncludePath() here, since it uses the full path and exceeds the max command line length.
-				// Because ISPC response files don't support white space in arguments, paths with white space need to be passed to the command line directly.
-				foreach (DirectoryReference IncludePath in CompileEnvironment.UserIncludePaths)
-				{
-					Arguments.Add(String.Format("-I\"{0}\"", IncludePath));
-				}
-
-				// System include paths.
-				foreach (DirectoryReference SystemIncludePath in CompileEnvironment.SystemIncludePaths)
-				{
-					Arguments.Add(String.Format("-I\"{0}\"", SystemIncludePath));
-				}
-
-				// Preprocessor definitions.
-				foreach (string Definition in CompileEnvironment.Definitions)
-				{
-					// TODO: Causes ISPC compiler to generate a spurious warning about the universal character set
-					if (!Definition.Contains("\\\\U") && !Definition.Contains("\\\\u"))
-					{
-						Arguments.Add($"-D{EscapeDefinitionForISPC(Definition)}");
-					}
-				}
-
 				// Generate the included header dependency list
 				if (CompileEnvironment.bGenerateDependenciesFile)
 				{
 					FileItem DependencyListFile = FileItem.GetItemByFileReference(FileReference.Combine(OutputDir, Path.GetFileName(ISPCFile.AbsolutePath) + ".txt"));
-					Arguments.Add(String.Format("-MMM \"{0}\"", DependencyListFile.AbsolutePath.Replace('\\', '/')));
+					Arguments.Add($"-MMM \"{NormalizeCommandLinePath(DependencyListFile)}\"");
 					CompileAction.DependencyListFile = DependencyListFile;
 					CompileAction.ProducedItems.Add(DependencyListFile);
 				}
 
+				Arguments.AddRange(GlobalArguments);
+
 				CompileAction.ProducedItems.Add(ISPCIncludeHeaderFile);
 
 				FileReference ResponseFileName = new FileReference(ISPCIncludeHeaderFile.AbsolutePath + ".response");
 				FileItem ResponseFileItem = Graph.CreateIntermediateTextFile(ResponseFileName, Arguments.Select(x => Utils.ExpandVariables(x)));
-				CompileAction.CommandArguments += String.Format("@\"{0}\"", ResponseFileName);
+				CompileAction.CommandArguments = $"@\"{NormalizeCommandLinePath(ResponseFileName)}\"";
 				CompileAction.PrerequisiteItems.Add(ResponseFileItem);
 
 				// Add the source file and its included files to the prerequisite item list.
@@ -452,11 +482,11 @@
 				CopyAction.CommandPath = BuildHostPlatform.Current.Shell;
 				if (BuildHostPlatform.Current.ShellType == ShellType.Cmd)
 				{
-					CopyAction.CommandArguments = String.Format("/C \"copy /Y \"{0}\" \"{1}\" 1>nul\"", SourceFile, TargetFile);
+					CopyAction.CommandArguments = $"/C \"copy /Y \"{SourceFile}\" \"{TargetFile}\" 1>nul\"";
 				}
 				else
 				{
-					CopyAction.CommandArguments = String.Format("-c 'cp -f \"\"{0}\"\" \"\"{1}\"'", SourceFile.FullName, TargetFile.FullName);
+					CopyAction.CommandArguments = $"-c 'cp -f \"\"{SourceFile}\"\" \"\"{TargetFile}\"'";
 				}
 				CopyAction.WorkingDirectory = UnrealBuildTool.EngineSourceDirectory;
 				CopyAction.PrerequisiteItems.Add(SourceFileItem);
@@ -484,6 +514,79 @@
 
 			List<string> CompileTargets = GetISPCCompileTargets(CompileEnvironment.Platform, null);
 
+			List<string> GlobalArguments = new List<string>();
+
+			// Build target string. No comma on last
+			string TargetString = "";
+			foreach (string Target in CompileTargets)
+			{
+				if (Target == CompileTargets[CompileTargets.Count - 1]) // .Last()
+				{
+					TargetString += Target;
+				}
+				else
+				{
+					TargetString += Target + ",";
+				}
+			}
+
+			// Build target triplet
+			GlobalArguments.Add($"--target-os={GetISPCOSTarget(CompileEnvironment.Platform)}");
+			GlobalArguments.Add($"--arch={GetISPCArchTarget(CompileEnvironment.Platform, null)}");
+			GlobalArguments.Add($"--target={TargetString}");
+			GlobalArguments.Add($"--emit-{GetISPCObjectFileFormat(CompileEnvironment.Platform)}");
+
+			if (CompileEnvironment.Configuration == CppConfiguration.Debug)
+			{
+				if (CompileEnvironment.Platform == UnrealTargetPlatform.Mac)
+				{
+					// Turn off debug symbols on Mac due to dsym generation issue
+					GlobalArguments.Add("-O0");
+					// Ideally we would be able to turn on symbols and specify the dwarf version, but that does
+					// does not seem to be working currently, ie:
+					//    GlobalArguments.Add("-g -O0 --dwarf-version=2");
+
+				}
+				else
+				{
+					GlobalArguments.Add("-g -O0");
+				}
+			}
+			else
+			{
+				GlobalArguments.Add("-O2");
+			}
+
+			// PIC is needed for modular builds except on Microsoft platforms
+			if ((CompileEnvironment.bIsBuildingDLL ||
+				CompileEnvironment.bIsBuildingLibrary) &&
+				!UEBuildPlatform.IsPlatformInGroup(CompileEnvironment.Platform, UnrealPlatformGroup.Microsoft))
+			{
+				GlobalArguments.Add("--pic");
+			}
+
+			// Include paths. Don't use AddIncludePath() here, since it uses the full path and exceeds the max command line length.
+			foreach (DirectoryReference IncludePath in CompileEnvironment.UserIncludePaths)
+			{
+				GlobalArguments.Add($"-I\"{NormalizeCommandLinePath(IncludePath)}\"");
+			}
+
+			// System include paths.
+			foreach (DirectoryReference SystemIncludePath in CompileEnvironment.SystemIncludePaths)
+			{
+				GlobalArguments.Add($"-I\"{NormalizeCommandLinePath(SystemIncludePath)}\"");
+			}
+
+			// Preprocessor definitions.
+			foreach (string Definition in CompileEnvironment.Definitions)
+			{
+				// TODO: Causes ISPC compiler to generate a spurious warning about the universal character set
+				if (!Definition.Contains("\\\\U") && !Definition.Contains("\\\\u"))
+				{
+					GlobalArguments.Add($"-D{EscapeDefinitionForISPC(Definition)}");
+				}
+			}
+
 			foreach (FileItem ISPCFile in InputFiles)
 			{
 				Action CompileAction = Graph.CreateAction(ActionType.Compile);
@@ -498,10 +601,9 @@
 				List<string> Arguments = new List<string>();
 
 				// Add the ISPC file to be compiled.
-				Arguments.Add(String.Format(" \"{0}\"", ISPCFile.AbsolutePath));
+				Arguments.Add($"\"{NormalizeCommandLinePath(ISPCFile)}\"");
 
 				List<FileItem> CompiledISPCObjFiles = new List<FileItem>();
-				string TargetString = "";
 
 				foreach (string Target in CompileTargets)
 				{
@@ -536,16 +638,6 @@
 
 					// Add the ISA specific ISPC obj files to the produced item list.
 					CompiledISPCObjFiles.Add(CompiledISPCObjFile);
-
-					// Build target string. No comma on last
-					if (Target == CompileTargets[CompileTargets.Count - 1]) // .Last()
-					{
-						TargetString += Target;
-					}
-					else
-					{
-						TargetString += Target + ",";
-					}
 				}
 
 				// Add the common ISPC obj file to the produced item list.
@@ -559,65 +651,18 @@
 				CompiledISPCObjFiles.Add(CompiledISPCObjFileNoISA);
 
 				// Add the output ISPC obj file
-				Arguments.Add(String.Format("-o \"{0}\"", CompiledISPCObjFileNoISA));
+				Arguments.Add($"-o \"{NormalizeCommandLinePath(CompiledISPCObjFileNoISA)}\"");
 
-				// Build target triplet
-				Arguments.Add(String.Format("--target-os=\"{0}\"", GetISPCOSTarget(CompileEnvironment.Platform)));
-				Arguments.Add(String.Format("--arch=\"{0}\"", GetISPCArchTarget(CompileEnvironment.Platform, null)));
-				Arguments.Add(String.Format("--target=\"{0}\"", TargetString));
-				Arguments.Add(String.Format("--emit-{0}", GetISPCObjectFileFormat(CompileEnvironment.Platform)));
-
-				if (CompileEnvironment.Configuration == CppConfiguration.Debug)
+				// Generate the timing info
+				if (CompileEnvironment.bPrintTimingInfo)
 				{
-					if (CompileEnvironment.Platform == UnrealTargetPlatform.Mac)
-					{
-						// Turn off debug symbols on Mac due to dsym generation issue
-						Arguments.Add("-O0");
-						// Ideally we would be able to turn on symbols and specify the dwarf version, but that does
-						// does not seem to be working currently, ie:
-						//    Arguments.Add("-g -O0 --dwarf-version=2");
-
-					}
-					else
-					{
-						Arguments.Add("-g -O0");
-					}
-				}
-				else
-				{
-					Arguments.Add("-O2");
+					FileItem TraceFile = FileItem.GetItemByFileReference(FileReference.FromString($"{CompiledISPCObjFileNoISA}.json"));
+					Arguments.Add("--time-trace");
+					CompileAction.ProducedItems.Add(TraceFile);
 				}
 
-				// PIC is needed for modular builds except on Microsoft platforms
-				if ((CompileEnvironment.bIsBuildingDLL ||
-					CompileEnvironment.bIsBuildingLibrary) &&
-					!UEBuildPlatform.IsPlatformInGroup(CompileEnvironment.Platform, UnrealPlatformGroup.Microsoft))
-				{
-					Arguments.Add("--pic");
-				}
+				Arguments.AddRange(GlobalArguments);
 
-				// Include paths. Don't use AddIncludePath() here, since it uses the full path and exceeds the max command line length.
-				foreach (DirectoryReference IncludePath in CompileEnvironment.UserIncludePaths)
-				{
-					Arguments.Add(String.Format("-I\"{0}\"", IncludePath));
-				}
-
-				// System include paths.
-				foreach (DirectoryReference SystemIncludePath in CompileEnvironment.SystemIncludePaths)
-				{
-					Arguments.Add(String.Format("-I\"{0}\"", SystemIncludePath));
-				}
-
-				// Preprocessor definitions.
-				foreach (string Definition in CompileEnvironment.Definitions)
-				{
-					// TODO: Causes ISPC compiler to generate a spurious warning about the universal character set
-					if (!Definition.Contains("\\\\U") && !Definition.Contains("\\\\u"))
-					{
-						Arguments.Add($"-D{EscapeDefinitionForISPC(Definition)}");
-					}
-				}
-
 				// Consume the included header dependency list
 				if (CompileEnvironment.bGenerateDependenciesFile)
 				{
@@ -631,7 +676,7 @@
 
 				FileReference ResponseFileName = new FileReference(CompiledISPCObjFileNoISA.AbsolutePath + ".response");
 				FileItem ResponseFileItem = Graph.CreateIntermediateTextFile(ResponseFileName, Arguments.Select(x => Utils.ExpandVariables(x)));
-				CompileAction.CommandArguments = " @\"" + ResponseFileName + "\"";
+				CompileAction.CommandArguments = $"@\"{ResponseFileName}\"";
 				CompileAction.PrerequisiteItems.Add(ResponseFileItem);
 
 				// Add the source file and its included files to the prerequisite item list.
