- Changed ClangToolchain to use absolute paths for errors, so IDEs are better able to go to the file/line of the error

diff --git a/Engine/Source/Programs/UnrealBuildTool/Platform/Linux/LinuxToolChain.cs b/Engine/Source/Programs/UnrealBuildTool/Platform/Linux/LinuxToolChain.cs
--- a/Engine/Source/Programs/UnrealBuildTool/Platform/Linux/LinuxToolChain.cs
+++ b/Engine/Source/Programs/UnrealBuildTool/Platform/Linux/LinuxToolChain.cs
@@ -1,4 +1,4 @@
-// Copyright Epic Games, Inc. All Rights Reserved.
+﻿// Copyright Epic Games, Inc. All Rights Reserved.
 
 using System;
 using System.Collections.Generic;
@@ -543,8 +543,6 @@ namespace UnrealBuildTool
 				}
 			}
 
-			Arguments.Add("-fdiagnostics-absolute-paths"); // output full paths to the files when the build fails
-
 			// Profile Guided Optimization (PGO) and Link Time Optimization (LTO)
 			// Whether we actually can enable that is checked in CanUseAdvancedLinkerFeatures() earlier
 			if (CompileEnvironment.bPGOOptimize)
diff --git a/Engine/Source/Programs/UnrealBuildTool/ToolChain/ClangToolChain.cs b/Engine/Source/Programs/UnrealBuildTool/ToolChain/ClangToolChain.cs
--- a/Engine/Source/Programs/UnrealBuildTool/ToolChain/ClangToolChain.cs
+++ b/Engine/Source/Programs/UnrealBuildTool/ToolChain/ClangToolChain.cs
@@ -1,4 +1,4 @@
-// Copyright Epic Games, Inc. All Rights Reserved.
+﻿// Copyright Epic Games, Inc. All Rights Reserved.
 
 using EpicGames.Core;
 using System;
@@ -351,6 +351,9 @@ namespace UnrealBuildTool
 			{
 				Arguments.Add("-Wundef" + (CompileEnvironment.bUndefinedIdentifierWarningsAsErrors ? "" : " -Wno-error=undef"));
 			}
+
+			// always use absolute paths for errors, this can help IDEs go to the error properly
+			Arguments.Add("-fdiagnostics-absolute-paths");
 		}
 
 		/// <summary>
