UnrealBuildTool: Print preferred version if the only available toolchain can not be used due to being banned

#rnx

diff --git a/Engine/Source/Programs/UnrealBuildTool/Platform/Windows/MicrosoftPlatformSDK.cs b/Engine/Source/Programs/UnrealBuildTool/Platform/Windows/MicrosoftPlatformSDK.cs
--- a/Engine/Source/Programs/UnrealBuildTool/Platform/Windows/MicrosoftPlatformSDK.cs
+++ b/Engine/Source/Programs/UnrealBuildTool/Platform/Windows/MicrosoftPlatformSDK.cs
@@ -872,13 +872,13 @@
 			string? Error = null;
 			if (Version < MinimumVisualCppVersion)
 			{
-				Error = $"UnrealBuildTool requires at minimum the MSVC {MinimumVisualCppVersion} toolchain. Please install a later toolchain from the Visual Studio installer.";
+				Error = $"UnrealBuildTool requires at minimum the MSVC {MinimumVisualCppVersion} toolchain. Please install a later toolchain such as {PreferredVisualCppVersions.Select(x => x.Min).Max()} from the Visual Studio installer.";
 			}
 
 			VersionNumberRange? Banned = BannedVisualCppVersions.FirstOrDefault(x => x.Contains(Version));
 			if (Banned != null)
 			{
-				Error = $"UnrealBuildTool has banned the MSVC {Banned} toolchains due to compiler issues. Please install a different toolchain from the Visual Studio installer.";
+				Error = $"UnrealBuildTool has banned the MSVC {Banned} toolchains due to compiler issues. Please install a different toolchain such as {PreferredVisualCppVersions.Select(x => x.Min).Max()} from the Visual Studio installer.";
 			}
 
 			Log.TraceLog("Found Visual Studio toolchain: {0} (Family={1}, FamilyRank={2}, Version={3}, Is64Bit={4}, Preview={5}, Architecture={6}, Error={7}, Redist={8})", ToolChainDir, Family, FamilyRank, Version, Is64Bit, bPreview, WindowsArchitecture.x64.ToString(), Error != null, RedistDir);
@@ -911,7 +911,7 @@
 				string? Error = null;
 				if (Version < MinimumClangVersion)
 				{
-					Error = $"UnrealBuildTool requires at minimum the Clang {MinimumClangVersion} toolchain. Please install a later toolchain from LLVM.";
+					Error = $"UnrealBuildTool requires at minimum the Clang {MinimumClangVersion} toolchain. Please install a later toolchain such as {PreferredClangVersions.Select(x => x.Min).Max()} from LLVM.";
 				}
 
 				Log.TraceLog("Found Clang toolchain: {0} (Version={1}, Is64Bit={2}, Rank={3})", ToolChainDir, Version, Is64Bit, Rank);
@@ -945,7 +945,7 @@
 				string? Error = null;
 				if (Version < MinimumIntelOneApiVersion)
 				{
-					Error = $"UnrealBuildTool requires at minimum the Intel OneAPI {MinimumIntelOneApiVersion} toolchain. Please install a later toolchain from Intel.";
+					Error = $"UnrealBuildTool requires at minimum the Intel OneAPI {MinimumIntelOneApiVersion} toolchain. Please install a later toolchain such as {PreferredIntelOneApiVersions.Select(x => x.Min).Max()} from Intel.";
 				}
 
 				Log.TraceLog("Found Intel OneAPI toolchain: {0} (Version={1}, Is64Bit={2}, Rank={3})", ToolChainDir, Version, Is64Bit, Rank);
