Rename UnrealBuildTool.EngineSourceDirectory -> Unreal.EngineSourceDirectory

#rnx

diff --git a/Engine/Source/Programs/Shared/EpicGames.Build/Unreal.cs b/Engine/Source/Programs/Shared/EpicGames.Build/Unreal.cs
--- a/Engine/Source/Programs/Shared/EpicGames.Build/Unreal.cs
+++ b/Engine/Source/Programs/Shared/EpicGames.Build/Unreal.cs
@@ -122,6 +122,11 @@
 		public static readonly DirectoryReference EngineDirectory = DirectoryReference.Combine(RootDirectory, "Engine");
 
 		/// <summary>
+		/// The full name of the Engine/Source directory
+		/// </summary>
+		public static readonly DirectoryReference EngineSourceDirectory = DirectoryReference.Combine(EngineDirectory, "Source");
+
+		/// <summary>
 		/// The path to UBT
 		/// </summary>
 		[Obsolete("Deprecated in UE5.1; to launch UnrealBuildTool, use this dll as the first argument with DonetPath")]
diff --git a/Engine/Source/Programs/UnrealBuildTool.Tests/SourceFileTests.cs b/Engine/Source/Programs/UnrealBuildTool.Tests/SourceFileTests.cs
--- a/Engine/Source/Programs/UnrealBuildTool.Tests/SourceFileTests.cs
+++ b/Engine/Source/Programs/UnrealBuildTool.Tests/SourceFileTests.cs
@@ -29,9 +29,9 @@
 			Unreal.LocationOverride.RootDirectory = DirectoryReference.Combine(DirectoryReference.GetCurrentDirectory(), @"..\..\..\..\..\..\..");
 
 			List<DirectoryReference> BaseDirectories = new List<DirectoryReference>();
-			BaseDirectories.Add(DirectoryReference.Combine(UnrealBuildTool.UnrealBuildTool.EngineSourceDirectory, "Runtime"));
-			BaseDirectories.Add(DirectoryReference.Combine(UnrealBuildTool.UnrealBuildTool.EngineSourceDirectory, "Developer"));
-			BaseDirectories.Add(DirectoryReference.Combine(UnrealBuildTool.UnrealBuildTool.EngineSourceDirectory, "Editor"));
+			BaseDirectories.Add(DirectoryReference.Combine(Unreal.EngineSourceDirectory, "Runtime"));
+			BaseDirectories.Add(DirectoryReference.Combine(Unreal.EngineSourceDirectory, "Developer"));
+			BaseDirectories.Add(DirectoryReference.Combine(Unreal.EngineSourceDirectory, "Editor"));
 
 			foreach(FileReference PluginFile in PluginsBase.EnumeratePlugins((FileReference)null))
 			{
diff --git a/Engine/Source/Programs/UnrealBuildTool/Configuration/UEBuildBinary.cs b/Engine/Source/Programs/UnrealBuildTool/Configuration/UEBuildBinary.cs
--- a/Engine/Source/Programs/UnrealBuildTool/Configuration/UEBuildBinary.cs
+++ b/Engine/Source/Programs/UnrealBuildTool/Configuration/UEBuildBinary.cs
@@ -277,8 +277,8 @@
 					else
 					{
 						// Parse the source and target patterns
-						FilePattern SourcePattern = new FilePattern(UnrealBuildTool.EngineSourceDirectory, Module.ExpandPathVariables(Dependency.SourcePath, OutputDir, ExeDir));
-						FilePattern TargetPattern = new FilePattern(UnrealBuildTool.EngineSourceDirectory, Module.ExpandPathVariables(Dependency.Path, OutputDir, ExeDir));
+						FilePattern SourcePattern = new FilePattern(Unreal.EngineSourceDirectory, Module.ExpandPathVariables(Dependency.SourcePath, OutputDir, ExeDir));
+						FilePattern TargetPattern = new FilePattern(Unreal.EngineSourceDirectory, Module.ExpandPathVariables(Dependency.Path, OutputDir, ExeDir));
 
 						// Skip non-essential single files if they do not exist
 						if (Dependency.Type == StagedFileType.DebugNonUFS && !SourcePattern.ContainsWildcards() && !FileReference.Exists(SourcePattern.GetSingleFile()))
diff --git a/Engine/Source/Programs/UnrealBuildTool/Configuration/UEBuildModuleCPP.cs b/Engine/Source/Programs/UnrealBuildTool/Configuration/UEBuildModuleCPP.cs
--- a/Engine/Source/Programs/UnrealBuildTool/Configuration/UEBuildModuleCPP.cs
+++ b/Engine/Source/Programs/UnrealBuildTool/Configuration/UEBuildModuleCPP.cs
@@ -1089,7 +1089,7 @@
 					string FileString = File.AbsolutePath;
 					if (File.Location.IsUnderDirectory(Unreal.RootDirectory))
 					{
-						FileString = File.Location.MakeRelativeTo(UnrealBuildTool.EngineSourceDirectory);
+						FileString = File.Location.MakeRelativeTo(Unreal.EngineSourceDirectory);
 					}
 					FileString = FileString.Replace('\\', '/');
 					Writer.WriteLine("// Dedicated PCH for {0}", FileString);
@@ -1277,7 +1277,7 @@
 				string IncludeFileString = IncludedFile.AbsolutePath;
 				if (IncludedFile.Location.IsUnderDirectory(Unreal.RootDirectory))
 				{
-					IncludeFileString = IncludedFile.Location.MakeRelativeTo(UnrealBuildTool.EngineSourceDirectory);
+					IncludeFileString = IncludedFile.Location.MakeRelativeTo(Unreal.EngineSourceDirectory);
 				}
 				IncludeFileString = IncludeFileString.Replace('\\', '/');
 				Writer.WriteLine("// PCH for {0}", IncludeFileString);
diff --git a/Engine/Source/Programs/UnrealBuildTool/Configuration/UEBuildTarget.cs b/Engine/Source/Programs/UnrealBuildTool/Configuration/UEBuildTarget.cs
--- a/Engine/Source/Programs/UnrealBuildTool/Configuration/UEBuildTarget.cs
+++ b/Engine/Source/Programs/UnrealBuildTool/Configuration/UEBuildTarget.cs
@@ -2087,7 +2087,7 @@
 					{
 						PostBuildStepAction.CommandArguments = String.Format("\"{0}\" && touch \"{1}\"", PostBuildScript, OutputFile);
 					}
-					PostBuildStepAction.WorkingDirectory = UnrealBuildTool.EngineSourceDirectory;
+					PostBuildStepAction.WorkingDirectory = Unreal.EngineSourceDirectory;
 					PostBuildStepAction.StatusDescription = String.Format("Executing post build script ({0})", PostBuildScript.GetFileName());
 					PostBuildStepAction.bCanExecuteRemotely = false;
 					PostBuildStepAction.PrerequisiteItems.Add(FileItem.GetItemByFileReference(ReceiptFileName));
@@ -2233,7 +2233,7 @@
 
 				// Create the action
 				Action WriteMetadataAction = Makefile.CreateRecursiveAction<WriteMetadataMode>(ActionType.WriteMetadata, WriteMetadataArguments.ToString());
-				WriteMetadataAction.WorkingDirectory = UnrealBuildTool.EngineSourceDirectory;
+				WriteMetadataAction.WorkingDirectory = Unreal.EngineSourceDirectory;
 				WriteMetadataAction.StatusDescription = StatusDescription;
 				WriteMetadataAction.bCanExecuteRemotely = false;
 				WriteMetadataAction.bUseActionHistory = false; // Different files for each target; do not want to invalidate based on this.
@@ -3820,7 +3820,7 @@
 			GlobalCompileEnvironment.Definitions.Add("USE_XMA2_FOR_STREAMING=1");
 
 			// Add the 'Engine/Source' path as a global include path for all modules
-			GlobalCompileEnvironment.UserIncludePaths.Add(UnrealBuildTool.EngineSourceDirectory);
+			GlobalCompileEnvironment.UserIncludePaths.Add(Unreal.EngineSourceDirectory);
 
 			[REDACTED_EMAIL]: Do any platform specific tool chain initialization here if required
 
@@ -4447,7 +4447,7 @@
 		/// </summary>
 		public static string NormalizeIncludePath(DirectoryReference Directory)
 		{
-			return Utils.CleanDirectorySeparators(Directory.MakeRelativeTo(UnrealBuildTool.EngineSourceDirectory), '/');
+			return Utils.CleanDirectorySeparators(Directory.MakeRelativeTo(Unreal.EngineSourceDirectory), '/');
 		}
 
 		/// <summary>
diff --git a/Engine/Source/Programs/UnrealBuildTool/Executors/Experimental/HordeExecutor.cs b/Engine/Source/Programs/UnrealBuildTool/Executors/Experimental/HordeExecutor.cs
--- a/Engine/Source/Programs/UnrealBuildTool/Executors/Experimental/HordeExecutor.cs
+++ b/Engine/Source/Programs/UnrealBuildTool/Executors/Experimental/HordeExecutor.cs
@@ -716,10 +716,10 @@
 
 			ZenWorker? Worker = new ZenWorker
 			{
-				Path = Action.CommandPath.IsUnderDirectory(Unreal.RootDirectory) ? Action.CommandPath.MakeRelativeTo(UnrealBuildTool.EngineSourceDirectory) : Action.CommandPath.FullName,
-				WorkingDirectory = UnrealBuildTool.EngineSourceDirectory.MakeRelativeTo(Unreal.RootDirectory),
+				Path = Action.CommandPath.IsUnderDirectory(Unreal.RootDirectory) ? Action.CommandPath.MakeRelativeTo(Unreal.EngineSourceDirectory) : Action.CommandPath.FullName,
+				WorkingDirectory = Unreal.EngineSourceDirectory.MakeRelativeTo(Unreal.RootDirectory),
 				Arguments = new List<string>() { Action.CommandArguments.Trim() },
-				Directories = new List<string>() { UnrealBuildTool.EngineSourceDirectory.MakeRelativeTo(Unreal.RootDirectory) },
+				Directories = new List<string>() { Unreal.EngineSourceDirectory.MakeRelativeTo(Unreal.RootDirectory) },
 				HostPlatform = BuildHostPlatform.Current.Platform.ToString(),
 				Cores = 1,
 				Memory = MemoryPerActionBytesOverride > 0 ? (long)MemoryPerActionBytesOverride : 1 * 1024 * 1024 * 1024, // 1 GB default
diff --git a/Engine/Source/Programs/UnrealBuildTool/Modes/GenerateClangDatabase.cs b/Engine/Source/Programs/UnrealBuildTool/Modes/GenerateClangDatabase.cs
--- a/Engine/Source/Programs/UnrealBuildTool/Modes/GenerateClangDatabase.cs
+++ b/Engine/Source/Programs/UnrealBuildTool/Modes/GenerateClangDatabase.cs
@@ -156,7 +156,7 @@
 						Writer.WriteObjectStart();
 						Writer.WriteValue("file", FileCommandPair.Key.FullName);
 						Writer.WriteValue("command", FileCommandPair.Value);
-						Writer.WriteValue("directory", UnrealBuildTool.EngineSourceDirectory.ToString());
+						Writer.WriteValue("directory", Unreal.EngineSourceDirectory.ToString());
 						Writer.WriteObjectEnd();
 					}
 					Writer.WriteArrayEnd();
diff --git a/Engine/Source/Programs/UnrealBuildTool/Platform/Android/AndroidToolChain.cs b/Engine/Source/Programs/UnrealBuildTool/Platform/Android/AndroidToolChain.cs
--- a/Engine/Source/Programs/UnrealBuildTool/Platform/Android/AndroidToolChain.cs
+++ b/Engine/Source/Programs/UnrealBuildTool/Platform/Android/AndroidToolChain.cs
@@ -1377,7 +1377,7 @@ namespace UnrealBuildTool
 					FileItem ResponseFileItem = Graph.CreateIntermediateTextFile(ResponseFileName, new List<string> { AllArguments });
 					string ResponseArgument = string.Format("@\"{0}\"", ResponseFileName);
 
-					CompileAction.WorkingDirectory = UnrealBuildTool.EngineSourceDirectory;
+					CompileAction.WorkingDirectory = Unreal.EngineSourceDirectory;
 					if(bExecuteCompilerThroughShell)
 					{
 						SetupActionToExecuteCompilerThroughShell(ref CompileAction, ClangPath!, ResponseArgument, "Compile");
@@ -1459,7 +1459,7 @@ namespace UnrealBuildTool
 				{
 					Action CompileAction = Graph.CreateAction(ActionType.Compile);
 					CompileAction.CommandDescription = "Compile";
-					CompileAction.WorkingDirectory = UnrealBuildTool.EngineSourceDirectory;
+					CompileAction.WorkingDirectory = Unreal.EngineSourceDirectory;
 					CompileAction.CommandPath = new FileReference(GetISPCHostCompilerPath(BuildHostPlatform.Current.Platform));
 					CompileAction.StatusDescription = Path.GetFileName(ISPCFile.AbsolutePath);
 					CompileAction.CommandVersion = GetISPCHostCompilerVersion(BuildHostPlatform.Current.Platform).ToString();
@@ -1558,7 +1558,7 @@ namespace UnrealBuildTool
 					{
 						CopyAction.CommandArguments = String.Format("-c 'cp -f \"{0}\" \"{1}\"'", SourceFile.FullName, TargetFile.FullName);
 					}
-					CopyAction.WorkingDirectory = UnrealBuildTool.EngineSourceDirectory;
+					CopyAction.WorkingDirectory = Unreal.EngineSourceDirectory;
 					CopyAction.PrerequisiteItems.Add(SourceFileItem);
 					CopyAction.ProducedItems.Add(TargetFileItem);
 					CopyAction.StatusDescription = TargetFileItem.Location.GetFileName();
@@ -1596,7 +1596,7 @@ namespace UnrealBuildTool
 				{
 					Action CompileAction = Graph.CreateAction(ActionType.Compile);
 					CompileAction.CommandDescription = "Compile";
-					CompileAction.WorkingDirectory = UnrealBuildTool.EngineSourceDirectory;
+					CompileAction.WorkingDirectory = Unreal.EngineSourceDirectory;
 					CompileAction.CommandPath = new FileReference(GetISPCHostCompilerPath(BuildHostPlatform.Current.Platform));
 					CompileAction.StatusDescription = Path.GetFileName(ISPCFile.AbsolutePath);
 
@@ -1762,7 +1762,7 @@ namespace UnrealBuildTool
 						{
 							CopyAction.CommandArguments = String.Format("-c 'cp -f \"{0}\" \"{1}\"'", SourceFile.FullName, TargetFile.FullName);
 						}
-						CopyAction.WorkingDirectory = UnrealBuildTool.EngineSourceDirectory;
+						CopyAction.WorkingDirectory = Unreal.EngineSourceDirectory;
 						CopyAction.PrerequisiteItems.Add(SourceFileItem);
 						CopyAction.ProducedItems.Add(TargetFileItem);
 						CopyAction.StatusDescription = TargetFileItem.Location.GetFileName();
@@ -1805,7 +1805,7 @@ namespace UnrealBuildTool
 
 				// Create an action that invokes the linker.
 				Action LinkAction = Graph.CreateAction(ActionType.Link);
-				LinkAction.WorkingDirectory = UnrealBuildTool.EngineSourceDirectory;
+				LinkAction.WorkingDirectory = Unreal.EngineSourceDirectory;
 
 				if (LinkEnvironment.bIsBuildingLibrary)
 				{
diff --git a/Engine/Source/Programs/UnrealBuildTool/Platform/IOS/IOSToolChain.cs b/Engine/Source/Programs/UnrealBuildTool/Platform/IOS/IOSToolChain.cs
--- a/Engine/Source/Programs/UnrealBuildTool/Platform/IOS/IOSToolChain.cs
+++ b/Engine/Source/Programs/UnrealBuildTool/Platform/IOS/IOSToolChain.cs
@@ -1491,7 +1491,7 @@ namespace UnrealBuildTool
 				}
 
 				Action PostBuildSyncAction = Graph.CreateRecursiveAction<IOSPostBuildSyncMode>(ActionType.CreateAppBundle, PostBuildSyncArguments);
-				PostBuildSyncAction.WorkingDirectory = UnrealBuildTool.EngineSourceDirectory;
+				PostBuildSyncAction.WorkingDirectory = Unreal.EngineSourceDirectory;
 				PostBuildSyncAction.PrerequisiteItems.Add(Executable);
 				PostBuildSyncAction.PrerequisiteItems.AddRange(OutputFiles);
 				PostBuildSyncAction.ProducedItems.Add(FileItem.GetItemByFileReference(GetStagedExecutablePath(Executable.Location, Target.Name)));
diff --git a/Engine/Source/Programs/UnrealBuildTool/Platform/Linux/LinuxToolChain.cs b/Engine/Source/Programs/UnrealBuildTool/Platform/Linux/LinuxToolChain.cs
--- a/Engine/Source/Programs/UnrealBuildTool/Platform/Linux/LinuxToolChain.cs
+++ b/Engine/Source/Programs/UnrealBuildTool/Platform/Linux/LinuxToolChain.cs
@@ -1131,7 +1131,7 @@ namespace UnrealBuildTool
 					CompileAction.ProducedItems.Add(DependencyListFile);
 				}
 
-				CompileAction.WorkingDirectory = UnrealBuildTool.EngineSourceDirectory;
+				CompileAction.WorkingDirectory = Unreal.EngineSourceDirectory;
 				CompileAction.CommandPath = new FileReference(ClangPath!);
 
 				List<string> ResponseFileContents = new();
@@ -1221,7 +1221,7 @@ namespace UnrealBuildTool
 		{
 			// Create an archive action
 			Action ArchiveAction = Graph.CreateAction(ActionType.Link);
-			ArchiveAction.WorkingDirectory = UnrealBuildTool.EngineSourceDirectory;
+			ArchiveAction.WorkingDirectory = Unreal.EngineSourceDirectory;
 			ArchiveAction.CommandPath = BuildHostPlatform.Current.Shell;
 
 			if (BuildHostPlatform.Current.ShellType == ShellType.Sh)
@@ -1306,7 +1306,7 @@ namespace UnrealBuildTool
 				}
 
 				Action PostLinkAction = Graph.CreateAction(ActionType.Link);
-				PostLinkAction.WorkingDirectory = UnrealBuildTool.EngineSourceDirectory;
+				PostLinkAction.WorkingDirectory = Unreal.EngineSourceDirectory;
 				PostLinkAction.CommandPath = ShellBinary;
 				PostLinkAction.StatusDescription = string.Format("{0}", Path.GetFileName(Executable.AbsolutePath));
 				PostLinkAction.CommandDescription = "FixDeps";
@@ -1369,7 +1369,7 @@ namespace UnrealBuildTool
 
 			// Create an action that invokes the linker.
 			Action LinkAction = Graph.CreateAction(ActionType.Link);
-			LinkAction.WorkingDirectory = UnrealBuildTool.EngineSourceDirectory;
+			LinkAction.WorkingDirectory = Unreal.EngineSourceDirectory;
 
 			string LinkCommandString;
 			LinkCommandString = "\"" + ClangPath + "\"";
diff --git a/Engine/Source/Programs/UnrealBuildTool/Platform/Mac/MacToolChain.cs b/Engine/Source/Programs/UnrealBuildTool/Platform/Mac/MacToolChain.cs
--- a/Engine/Source/Programs/UnrealBuildTool/Platform/Mac/MacToolChain.cs
+++ b/Engine/Source/Programs/UnrealBuildTool/Platform/Mac/MacToolChain.cs
@@ -1080,7 +1080,7 @@ namespace UnrealBuildTool
 		FileItem FixDylibDependencies(LinkEnvironment LinkEnvironment, FileItem Executable, IActionGraphBuilder Graph)
 		{
 			Action FixDylibAction = Graph.CreateAction(ActionType.PostBuildStep);
-			FixDylibAction.WorkingDirectory = UnrealBuildTool.EngineSourceDirectory;
+			FixDylibAction.WorkingDirectory = Unreal.EngineSourceDirectory;
 			FixDylibAction.CommandPath = BuildHostPlatform.Current.Shell;
 			FixDylibAction.CommandDescription = "";
 
diff --git a/Engine/Source/Programs/UnrealBuildTool/Platform/Windows/PVSToolChain.cs b/Engine/Source/Programs/UnrealBuildTool/Platform/Windows/PVSToolChain.cs
--- a/Engine/Source/Programs/UnrealBuildTool/Platform/Windows/PVSToolChain.cs
+++ b/Engine/Source/Programs/UnrealBuildTool/Platform/Windows/PVSToolChain.cs
@@ -584,7 +584,7 @@
 				Action AnalyzeAction = Graph.CreateAction(ActionType.Compile);
 				AnalyzeAction.CommandDescription = "Analyzing";
 				AnalyzeAction.StatusDescription = BaseFileName;
-				AnalyzeAction.WorkingDirectory = UnrealBuildTool.EngineSourceDirectory;
+				AnalyzeAction.WorkingDirectory = Unreal.EngineSourceDirectory;
 				AnalyzeAction.CommandPath = AnalyzerFile;
 
 				StringBuilder Arguments = new StringBuilder();
@@ -637,7 +637,7 @@
 			AnalyzeAction.CommandDescription = "Process PVS-Studio Results";
 			AnalyzeAction.CommandPath = Unreal.DotnetPath;
 			AnalyzeAction.CommandArguments = $"\"{Unreal.UnrealBuildToolDllPath}\" -Mode=PVSGather -Input=\"{InputFileListItem.Location}\" -Output=\"{OutputFile}\" ";
-			AnalyzeAction.WorkingDirectory = UnrealBuildTool.EngineSourceDirectory;
+			AnalyzeAction.WorkingDirectory = Unreal.EngineSourceDirectory;
 			AnalyzeAction.PrerequisiteItems.Add(InputFileListItem);
 			AnalyzeAction.PrerequisiteItems.AddRange(Makefile.OutputItems);
 			AnalyzeAction.PrerequisiteItems.AddRange(CompileSourceFiles);
diff --git a/Engine/Source/Programs/UnrealBuildTool/Platform/Windows/VCCompileAction.cs b/Engine/Source/Programs/UnrealBuildTool/Platform/Windows/VCCompileAction.cs
--- a/Engine/Source/Programs/UnrealBuildTool/Platform/Windows/VCCompileAction.cs
+++ b/Engine/Source/Programs/UnrealBuildTool/Platform/Windows/VCCompileAction.cs
@@ -159,7 +159,7 @@
 		#region Implementation of IAction
 
 		IEnumerable<FileItem> IExternalAction.DeleteItems => DeleteItems;
-		public DirectoryReference WorkingDirectory => UnrealBuildTool.EngineSourceDirectory;
+		public DirectoryReference WorkingDirectory => Unreal.EngineSourceDirectory;
 		string IExternalAction.CommandDescription => "Compile";
 		bool IExternalAction.bIsGCCCompiler => false;
 		bool IExternalAction.bProducesImportLibrary => false;
diff --git a/Engine/Source/Programs/UnrealBuildTool/Platform/Windows/VCToolChain.cs b/Engine/Source/Programs/UnrealBuildTool/Platform/Windows/VCToolChain.cs
--- a/Engine/Source/Programs/UnrealBuildTool/Platform/Windows/VCToolChain.cs
+++ b/Engine/Source/Programs/UnrealBuildTool/Platform/Windows/VCToolChain.cs
@@ -124,7 +124,7 @@ namespace UnrealBuildTool
 		{
 			if (!Path.IsPathRooted(PathString))
 			{
-				string ResolvedPath = Path.Combine(UnrealBuildTool.EngineSourceDirectory.FullName, PathString);
+				string ResolvedPath = Path.Combine(Unreal.EngineSourceDirectory.FullName, PathString);
 				if (ResolvedPath.Length > MaxPathWarningLength)
 				{
 					Log.TraceWarningOnce($"Relative path '{PathString}' when resolved will have length '{ResolvedPath.Length}' which is greater than MAX_PATH (260) and may cause unexpected errors with the MSVC toolchain.");
@@ -141,7 +141,7 @@ namespace UnrealBuildTool
 			// Try to use a relative path to shorten command line length and to enable remote distribution where absolute paths are not desired
 			if (Reference.IsUnderDirectory(Unreal.EngineDirectory))
 			{
-				string RelativePath = Reference.MakeRelativeTo(UnrealBuildTool.EngineSourceDirectory);
+				string RelativePath = Reference.MakeRelativeTo(Unreal.EngineSourceDirectory);
 				CheckCommandLinePathLength(RelativePath);
 				return RelativePath;
 			}
@@ -1455,7 +1455,7 @@ namespace UnrealBuildTool
 			}
 
 			Action ParseTimingInfoAction = Graph.CreateRecursiveAction<ParseMsvcTimingInfoMode>(ActionType.ParseTimingInfo, ParseTimingArguments);
-			ParseTimingInfoAction.WorkingDirectory = UnrealBuildTool.EngineSourceDirectory;
+			ParseTimingInfoAction.WorkingDirectory = Unreal.EngineSourceDirectory;
 			ParseTimingInfoAction.StatusDescription = Path.GetFileName(TimingFile.AbsolutePath);
 			ParseTimingInfoAction.bCanExecuteRemotely = true;
 			ParseTimingInfoAction.bCanExecuteRemotelyWithSNDBS = true;
@@ -1493,7 +1493,7 @@ namespace UnrealBuildTool
 					};
 
 					Action AggregateTimingInfoAction = Makefile.CreateRecursiveAction<AggregateParsedTimingInfo>(ActionType.ParseTimingInfo, string.Join(" ", ActionArgs));
-					AggregateTimingInfoAction.WorkingDirectory = UnrealBuildTool.EngineSourceDirectory;
+					AggregateTimingInfoAction.WorkingDirectory = Unreal.EngineSourceDirectory;
 					AggregateTimingInfoAction.StatusDescription = $"Aggregating {TimingJsonFiles.Count} Timing File(s)";
 					AggregateTimingInfoAction.bCanExecuteRemotely = false;
 					AggregateTimingInfoAction.bCanExecuteRemotelyWithSNDBS = false;
@@ -1551,7 +1551,7 @@ namespace UnrealBuildTool
 			{
 				Action CompileAction = Graph.CreateAction(ActionType.Compile);
 				CompileAction.CommandDescription = "Resource";
-				CompileAction.WorkingDirectory = UnrealBuildTool.EngineSourceDirectory;
+				CompileAction.WorkingDirectory = Unreal.EngineSourceDirectory;
 				CompileAction.CommandPath = EnvVars.ResourceCompilerPath;
 				CompileAction.StatusDescription = Path.GetFileName(RCFile.AbsolutePath);
 				CompileAction.PrerequisiteItems.AddRange(CompileEnvironment.ForceIncludeFiles);
@@ -1728,7 +1728,7 @@ namespace UnrealBuildTool
 			CompileAction.ProducedItems.Add(ObjectFile);
 			CompileAction.DeleteItems.Add(FileItem.GetItemByFileReference(OutputFile));
 			CompileAction.StatusDescription = TypeLibrary.Header;
-			CompileAction.WorkingDirectory = UnrealBuildTool.EngineSourceDirectory;
+			CompileAction.WorkingDirectory = Unreal.EngineSourceDirectory;
 			CompileAction.CommandPath = EnvVars.CompilerPath;
 			CompileAction.CommandArguments = String.Join(" ", Arguments);
 			CompileAction.bShouldOutputStatusDescription = Target.WindowsPlatform.Compiler.IsClang();
@@ -1739,7 +1739,7 @@ namespace UnrealBuildTool
 			TouchAction.CommandDescription = "Touch";
 			TouchAction.CommandPath = BuildHostPlatform.Current.Shell;
 			TouchAction.CommandArguments = $"/C \"copy /b \"{OutputFile.FullName}\"+,, \"{OutputFile.FullName}\" 1>nul:\"";
-			TouchAction.WorkingDirectory = UnrealBuildTool.EngineSourceDirectory;
+			TouchAction.WorkingDirectory = Unreal.EngineSourceDirectory;
 			TouchAction.PrerequisiteItems.Add(ObjectFile);
 			TouchAction.ProducedItems.Add(FileItem.GetItemByFileReference(OutputFile));
 			TouchAction.StatusDescription = OutputFile.GetFileName();
@@ -2028,7 +2028,7 @@ namespace UnrealBuildTool
 			// Create an action that invokes the linker.
 			Action LinkAction = Graph.CreateAction(ActionType.Link);
 			LinkAction.CommandDescription = "Link";
-			LinkAction.WorkingDirectory = UnrealBuildTool.EngineSourceDirectory;
+			LinkAction.WorkingDirectory = Unreal.EngineSourceDirectory;
 			if(bIsBuildingLibraryOrImportLibrary)
 			{
 				LinkAction.CommandPath = EnvVars.LibraryManagerPath;
diff --git a/Engine/Source/Programs/UnrealBuildTool/ProjectFiles/ProjectFileGenerator.cs b/Engine/Source/Programs/UnrealBuildTool/ProjectFiles/ProjectFileGenerator.cs
--- a/Engine/Source/Programs/UnrealBuildTool/ProjectFiles/ProjectFileGenerator.cs
+++ b/Engine/Source/Programs/UnrealBuildTool/ProjectFiles/ProjectFileGenerator.cs
@@ -1382,7 +1382,7 @@
 					throw new BuildException("A valid game project was not found in the specified location (" + OnlyGameProject.Directory.FullName + ")");
 				}
 
-				bool bInstalledEngineWithSource = Unreal.IsEngineInstalled() && DirectoryReference.Exists(UnrealBuildTool.EngineSourceDirectory);
+				bool bInstalledEngineWithSource = Unreal.IsEngineInstalled() && DirectoryReference.Exists(Unreal.EngineSourceDirectory);
 
 				bIncludeEngineSource = !Unreal.IsEngineInstalled() || bInstalledEngineWithSource;
 				bIncludeDocumentation = false;
@@ -1812,7 +1812,7 @@
 		/// </summary>
 		private void AddUnrealBuildToolProject(PrimaryProjectFolder ProgramsFolder)
 		{
-			DirectoryReference ProgramsDirectory = DirectoryReference.Combine(UnrealBuildTool.EngineSourceDirectory, "Programs");
+			DirectoryReference ProgramsDirectory = DirectoryReference.Combine(Unreal.EngineSourceDirectory, "Programs");
 
 			FileReference ProjectFileName = FileReference.Combine(ProgramsDirectory, "UnrealBuildTool", "UnrealBuildTool.csproj");
 			if (FileReference.Exists(ProjectFileName))
@@ -1851,7 +1851,7 @@
 		{
 			VCSharpProjectFile Project;
 
-			FileReference ProjectFileName = FileReference.Combine(UnrealBuildTool.EngineSourceDirectory, "Programs", ProjectName, Path.GetFileName(ProjectName) + ".csproj");
+			FileReference ProjectFileName = FileReference.Combine(Unreal.EngineSourceDirectory, "Programs", ProjectName, Path.GetFileName(ProjectName) + ".csproj");
 
 			FileInfo Info = new FileInfo(ProjectFileName.FullName);
 			if (Info.Exists)
@@ -2332,7 +2332,7 @@
 			GetPlatformSpecializationsSubTargetsForAllTargets(AllTargetFiles, out AllSubTargetFiles, out AllSubTargetFilesPerTarget);
 
 			// Get some standard directories
-			//DirectoryReference EngineSourceProgramsDirectory = DirectoryReference.Combine(UnrealBuildTool.EngineSourceDirectory, "Programs");
+			//DirectoryReference EngineSourceProgramsDirectory = DirectoryReference.Combine(Unreal.EngineSourceDirectory, "Programs");
 
 			// Keep a cache of the default editor target for each project so that we don't have to interrogate the configs for each target
 			Dictionary<string, string?> DefaultProjectEditorTargetCache = new Dictionary<string, string?>();
diff --git a/Engine/Source/Programs/UnrealBuildTool/ProjectFiles/Rider/RiderProjectFileGenerator.cs b/Engine/Source/Programs/UnrealBuildTool/ProjectFiles/Rider/RiderProjectFileGenerator.cs
--- a/Engine/Source/Programs/UnrealBuildTool/ProjectFiles/Rider/RiderProjectFileGenerator.cs
+++ b/Engine/Source/Programs/UnrealBuildTool/ProjectFiles/Rider/RiderProjectFileGenerator.cs
@@ -14,7 +14,7 @@
 		public override string ProjectFileExtension => ".json";
 
 		private readonly DirectoryReference EngineSourceProgramsDirectory =
-			DirectoryReference.Combine(UnrealBuildTool.EngineSourceDirectory, "Programs");
+			DirectoryReference.Combine(Unreal.EngineSourceDirectory, "Programs");
 
 		private readonly CommandLineArguments Arguments;
 		
@@ -203,7 +203,7 @@
 					{
 						WantProjectFileForTarget = bIncludeEnginePrograms;
 					}
-					else if (TargetFilePath.IsUnderDirectory(UnrealBuildTool.EngineSourceDirectory))
+					else if (TargetFilePath.IsUnderDirectory(Unreal.EngineSourceDirectory))
 					{
 						WantProjectFileForTarget = bIncludeEngineSource;
 					}
diff --git a/Engine/Source/Programs/UnrealBuildTool/ProjectFiles/VisualStudio/VCProject.cs b/Engine/Source/Programs/UnrealBuildTool/ProjectFiles/VisualStudio/VCProject.cs
--- a/Engine/Source/Programs/UnrealBuildTool/ProjectFiles/VisualStudio/VCProject.cs
+++ b/Engine/Source/Programs/UnrealBuildTool/ProjectFiles/VisualStudio/VCProject.cs
@@ -1302,7 +1302,7 @@
 			{
 				VCProjectFileContent.AppendLine("  <PropertyGroup>");
 				VCProjectFileContent.Append("    <SourcePath>");
-				foreach (string DirectoryName in Directory.EnumerateDirectories(UnrealBuildTool.EngineSourceDirectory.FullName, "*", SearchOption.AllDirectories))
+				foreach (string DirectoryName in Directory.EnumerateDirectories(Unreal.EngineSourceDirectory.FullName, "*", SearchOption.AllDirectories))
 				{
 					if (Directory.EnumerateFiles(DirectoryName, "*.cpp").Any())
 					{
diff --git a/Engine/Source/Programs/UnrealBuildTool/ProjectFiles/VisualStudio/VCProjectFileGenerator.cs b/Engine/Source/Programs/UnrealBuildTool/ProjectFiles/VisualStudio/VCProjectFileGenerator.cs
--- a/Engine/Source/Programs/UnrealBuildTool/ProjectFiles/VisualStudio/VCProjectFileGenerator.cs
+++ b/Engine/Source/Programs/UnrealBuildTool/ProjectFiles/VisualStudio/VCProjectFileGenerator.cs
@@ -450,7 +450,7 @@
 			{
 				XNamespace NS = XNamespace.Get("http://schemas.microsoft.com/developer/msbuild/2003");
 
-				DirectoryReference AutomationToolDir = DirectoryReference.Combine(UnrealBuildTool.EngineSourceDirectory, "Programs", "AutomationTool");
+				DirectoryReference AutomationToolDir = DirectoryReference.Combine(Unreal.EngineSourceDirectory, "Programs", "AutomationTool");
 				new XDocument(
 					new XElement(NS + "Project",
 						new XAttribute("ToolsVersion", VCProjectFileGenerator.GetProjectFileToolVersionString(Settings.ProjectFileFormat)),
diff --git a/Engine/Source/Programs/UnrealBuildTool/ProjectFiles/VisualStudioCode/VSCodeProjectFileGenerator.cs b/Engine/Source/Programs/UnrealBuildTool/ProjectFiles/VisualStudioCode/VSCodeProjectFileGenerator.cs
--- a/Engine/Source/Programs/UnrealBuildTool/ProjectFiles/VisualStudioCode/VSCodeProjectFileGenerator.cs
+++ b/Engine/Source/Programs/UnrealBuildTool/ProjectFiles/VisualStudioCode/VSCodeProjectFileGenerator.cs
@@ -880,7 +880,7 @@
 					Writer.WriteValue(MakePathString(CompilerPath, bInAbsolute: true, bForceSkipQuotes: true));
 					Writer.WriteValue($"@{MakePathString(ResponseFile, bInAbsolute: true, bForceSkipQuotes: true)}");
 					Writer.WriteArrayEnd();
-					Writer.WriteValue("directory", UnrealBuildTool.EngineSourceDirectory.ToString());
+					Writer.WriteValue("directory", Unreal.EngineSourceDirectory.ToString());
 					Writer.WriteObjectEnd();
 				}
 				Writer.WriteArrayEnd();
diff --git a/Engine/Source/Programs/UnrealBuildTool/ProjectFiles/VisualStudioMac/VisualStudioMacProjectFileGenerator.cs b/Engine/Source/Programs/UnrealBuildTool/ProjectFiles/VisualStudioMac/VisualStudioMacProjectFileGenerator.cs
--- a/Engine/Source/Programs/UnrealBuildTool/ProjectFiles/VisualStudioMac/VisualStudioMacProjectFileGenerator.cs
+++ b/Engine/Source/Programs/UnrealBuildTool/ProjectFiles/VisualStudioMac/VisualStudioMacProjectFileGenerator.cs
@@ -8,8 +8,8 @@
 using System.Xml;
 using System.Xml.Linq;
 using EpicGames.Core;
+using UnrealBuildBase;
 
-
 namespace UnrealBuildTool
 {
     
@@ -91,7 +91,7 @@
             {
                 XNamespace NS = XNamespace.Get("http://schemas.microsoft.com/developer/msbuild/2003");
 
-                DirectoryReference AutomationToolDir = DirectoryReference.Combine(UnrealBuildTool.EngineSourceDirectory, "Programs", "AutomationTool");
+                DirectoryReference AutomationToolDir = DirectoryReference.Combine(Unreal.EngineSourceDirectory, "Programs", "AutomationTool");
                 new XDocument(
                     new XElement(NS + "Project",
                         new XAttribute("ToolsVersion", VCProjectFileGenerator.GetProjectFileToolVersionString(Settings.ProjectFileFormat)),
diff --git a/Engine/Source/Programs/UnrealBuildTool/System/ActionGraphBuilder.cs b/Engine/Source/Programs/UnrealBuildTool/System/ActionGraphBuilder.cs
--- a/Engine/Source/Programs/UnrealBuildTool/System/ActionGraphBuilder.cs
+++ b/Engine/Source/Programs/UnrealBuildTool/System/ActionGraphBuilder.cs
@@ -242,7 +242,7 @@
 			{
 				CopyAction.CommandArguments = String.Format("-c 'cp -f \"\"{0}\"\" \"\"{1}\"'", SourceFile.AbsolutePath, TargetFile.AbsolutePath);
 			}
-			CopyAction.WorkingDirectory = UnrealBuildTool.EngineSourceDirectory;
+			CopyAction.WorkingDirectory = Unreal.EngineSourceDirectory;
 			CopyAction.PrerequisiteItems.Add(SourceFile);
 			CopyAction.ProducedItems.Add(TargetFile);
 			CopyAction.DeleteItems.Add(TargetFile);
diff --git a/Engine/Source/Programs/UnrealBuildTool/System/Unity.cs b/Engine/Source/Programs/UnrealBuildTool/System/Unity.cs
--- a/Engine/Source/Programs/UnrealBuildTool/System/Unity.cs
+++ b/Engine/Source/Programs/UnrealBuildTool/System/Unity.cs
@@ -256,7 +256,7 @@
 					string CPPFileString = CPPFile.AbsolutePath;
 					if (CPPFile.Location.IsUnderDirectory(Unreal.RootDirectory))
 					{
-						CPPFileString = CPPFile.Location.MakeRelativeTo(UnrealBuildTool.EngineSourceDirectory);
+						CPPFileString = CPPFile.Location.MakeRelativeTo(Unreal.EngineSourceDirectory);
 					}
 					OutputUnityCPPWriter.WriteLine("#include \"{0}\"", CPPFileString.Replace('\\', '/'));
 				}
diff --git a/Engine/Source/Programs/UnrealBuildTool/ToolChain/AppleToolChain.cs b/Engine/Source/Programs/UnrealBuildTool/ToolChain/AppleToolChain.cs
--- a/Engine/Source/Programs/UnrealBuildTool/ToolChain/AppleToolChain.cs
+++ b/Engine/Source/Programs/UnrealBuildTool/ToolChain/AppleToolChain.cs
@@ -156,7 +156,7 @@ namespace UnrealBuildTool
 
 		protected DirectoryReference GetMacDevSrcRoot()
 		{
-			return UnrealBuildTool.EngineSourceDirectory;
+			return Unreal.EngineSourceDirectory;
 		}
 
 		protected void StripSymbolsWithXcode(FileReference SourceFile, FileReference TargetFile, string ToolchainDir)
diff --git a/Engine/Source/Programs/UnrealBuildTool/ToolChain/ISPCToolChain.cs b/Engine/Source/Programs/UnrealBuildTool/ToolChain/ISPCToolChain.cs
--- a/Engine/Source/Programs/UnrealBuildTool/ToolChain/ISPCToolChain.cs
+++ b/Engine/Source/Programs/UnrealBuildTool/ToolChain/ISPCToolChain.cs
@@ -140,7 +140,7 @@ namespace UnrealBuildTool
 		/// <returns>Path to ISPC compiler</returns>
 		public virtual string GetISPCHostCompilerPath(UnrealTargetPlatform Platform)
 		{
-			string ISPCCompilerPathCommon = Path.Combine(UnrealBuildTool.EngineSourceDirectory.FullName, "ThirdParty", "Intel", "ISPC", "bin");
+			string ISPCCompilerPathCommon = Path.Combine(Unreal.EngineSourceDirectory.FullName, "ThirdParty", "Intel", "ISPC", "bin");
 			string ISPCArchitecturePath = "";
 			string ExeExtension = ".exe";
 
@@ -329,7 +329,7 @@ namespace UnrealBuildTool
 			// Try to use a relative path to shorten command line length.
 			if (Reference.IsUnderDirectory(Unreal.RootDirectory))
 			{
-				return Reference.MakeRelativeTo(UnrealBuildTool.EngineSourceDirectory).Replace("\\", "/");
+				return Reference.MakeRelativeTo(Unreal.EngineSourceDirectory).Replace("\\", "/");
 			}
 
 			return Reference.FullName.Replace("\\", "/");
@@ -419,7 +419,7 @@ namespace UnrealBuildTool
 			{
 				Action CompileAction = Graph.CreateAction(ActionType.Compile);
 				CompileAction.CommandDescription = "Compile";
-				CompileAction.WorkingDirectory = UnrealBuildTool.EngineSourceDirectory;
+				CompileAction.WorkingDirectory = Unreal.EngineSourceDirectory;
 				CompileAction.CommandPath = new FileReference(GetISPCHostCompilerPath(BuildHostPlatform.Current.Platform));
 				CompileAction.StatusDescription = Path.GetFileName(ISPCFile.AbsolutePath);
 
@@ -488,7 +488,7 @@ namespace UnrealBuildTool
 				{
 					CopyAction.CommandArguments = $"-c 'cp -f \"\"{SourceFile}\"\" \"\"{TargetFile}\"'";
 				}
-				CopyAction.WorkingDirectory = UnrealBuildTool.EngineSourceDirectory;
+				CopyAction.WorkingDirectory = Unreal.EngineSourceDirectory;
 				CopyAction.PrerequisiteItems.Add(SourceFileItem);
 				CopyAction.ProducedItems.Add(TargetFileItem);
 				CopyAction.StatusDescription = TargetFileItem.Location.GetFileName();
@@ -591,7 +591,7 @@ namespace UnrealBuildTool
 			{
 				Action CompileAction = Graph.CreateAction(ActionType.Compile);
 				CompileAction.CommandDescription = "Compile";
-				CompileAction.WorkingDirectory = UnrealBuildTool.EngineSourceDirectory;
+				CompileAction.WorkingDirectory = Unreal.EngineSourceDirectory;
 				CompileAction.CommandPath = new FileReference(GetISPCHostCompilerPath(BuildHostPlatform.Current.Platform));
 				CompileAction.StatusDescription = Path.GetFileName(ISPCFile.AbsolutePath);
 
diff --git a/Engine/Source/Programs/UnrealBuildTool/ToolChain/RemoteMac.cs b/Engine/Source/Programs/UnrealBuildTool/ToolChain/RemoteMac.cs
--- a/Engine/Source/Programs/UnrealBuildTool/ToolChain/RemoteMac.cs
+++ b/Engine/Source/Programs/UnrealBuildTool/ToolChain/RemoteMac.cs
@@ -508,7 +508,7 @@
 					StringBuilder Arguments = new StringBuilder("ExportCertificate");
 					if(TargetDesc.ProjectFile == null)
 					{
-						Arguments.AppendFormat(" \"{0}\"", UnrealBuildTool.EngineSourceDirectory);
+						Arguments.AppendFormat(" \"{0}\"", Unreal.EngineSourceDirectory);
 					}
 					else
 					{
diff --git a/Engine/Source/Programs/UnrealBuildTool/UnrealBuildTool.cs b/Engine/Source/Programs/UnrealBuildTool/UnrealBuildTool.cs
--- a/Engine/Source/Programs/UnrealBuildTool/UnrealBuildTool.cs
+++ b/Engine/Source/Programs/UnrealBuildTool/UnrealBuildTool.cs
@@ -46,6 +46,7 @@
 		/// <summary>
 		/// The full name of the Engine/Source directory
 		/// </summary>
+		[Obsolete("Replace with Unreal.EngineSourceDirectory")]
 		public static readonly DirectoryReference EngineSourceDirectory = DirectoryReference.Combine(Unreal.EngineDirectory, "Source");
 
 		/// <summary>
@@ -498,7 +499,7 @@
 
 				// Change the working directory to be the Engine/Source folder. We are likely running from Engine/Binaries/DotNET
 				// This is critical to be done early so any code that relies on the current directory being Engine/Source will work.
-				DirectoryReference.SetCurrentDirectory(UnrealBuildTool.EngineSourceDirectory);
+				DirectoryReference.SetCurrentDirectory(Unreal.EngineSourceDirectory);
 
 				// Register encodings from Net FW as this is required when using Ionic as we do in multiple toolchains
 				Encoding.RegisterProvider(CodePagesEncodingProvider.Instance);
