UnrealBuildTool: Readd -Wno-range-loop-analysis warnings for Apple platforms, there appears to be outstanding issues that presubmit did not notify of

#rnx

diff --git a/Engine/Source/Programs/UnrealBuildTool/ToolChain/AppleToolChain.cs b/Engine/Source/Programs/UnrealBuildTool/ToolChain/AppleToolChain.cs
--- a/Engine/Source/Programs/UnrealBuildTool/ToolChain/AppleToolChain.cs
+++ b/Engine/Source/Programs/UnrealBuildTool/ToolChain/AppleToolChain.cs
@@ -261,6 +261,13 @@ namespace UnrealBuildTool
 				Arguments.Remove("-Wno-unused-but-set-parameter");
 				Arguments.Remove("-Wno-ordered-compare-function-pointers");
 			}
+
+			// clang 12.00 has a new warning for copies in ranged loops. Instances have all been fixed up (2020/6/26) but
+			// are likely to be reintroduced due to no equivalent on other platforms at this time so disable the warning
+			if (GetClangVersion().Major >= 12)
+			{
+				Arguments.Add("-Wno-range-loop-analysis");
+			}
 		}
 
 		/// <inheritdoc/>
