UBT: Use ILogger throughout UBT, rather than legacy EpicGames.Core.Log methods. All output going forwards should use structured logging rather than String.Format style logging.

diff --git a/Engine/Source/Programs/UnrealBuildTool/Configuration/UEBuildTarget.cs b/Engine/Source/Programs/UnrealBuildTool/Configuration/UEBuildTarget.cs
--- a/Engine/Source/Programs/UnrealBuildTool/Configuration/UEBuildTarget.cs
+++ b/Engine/Source/Programs/UnrealBuildTool/Configuration/UEBuildTarget.cs
@@ -17,6 +17,7 @@ using System.Reflection;
 using OpenTracing.Util;
 using UnrealBuildBase;
 using EpicGames.UHT.Utils;
+using Microsoft.Extensions.Logging;
 
 namespace UnrealBuildTool
 {
@@ -693,8 +694,9 @@ namespace UnrealBuildTool
 		/// <param name="bSkipRulesCompile">Whether to skip compiling any rules assemblies</param>
 		/// <param name="bForceRulesCompile">Whether to always compile all rules assemblies</param>
 		/// <param name="bUsePrecompiled">Whether to use a precompiled engine build</param>
+		/// <param name="Logger">Logger for output</param>
 		/// <returns>The build target object for the specified build rules source file</returns>
-		public static UEBuildTarget Create(TargetDescriptor Descriptor, bool bSkipRulesCompile, bool bForceRulesCompile, bool bUsePrecompiled)
+		public static UEBuildTarget Create(TargetDescriptor Descriptor, bool bSkipRulesCompile, bool bForceRulesCompile, bool bUsePrecompiled, ILogger Logger)
 		{
 			// make sure we are allowed to build this platform
 			if (!UEBuildPlatform.IsPlatformAvailable(Descriptor.Platform))
@@ -705,13 +707,13 @@ namespace UnrealBuildTool
 			RulesAssembly RulesAssembly;
 			using (GlobalTracer.Instance.BuildSpan("RulesCompiler.CreateTargetRulesAssembly()").StartActive())
 			{
-				RulesAssembly = RulesCompiler.CreateTargetRulesAssembly(Descriptor.ProjectFile, Descriptor.Name, bSkipRulesCompile, bForceRulesCompile, bUsePrecompiled, Descriptor.ForeignPlugin);
+				RulesAssembly = RulesCompiler.CreateTargetRulesAssembly(Descriptor.ProjectFile, Descriptor.Name, bSkipRulesCompile, bForceRulesCompile, bUsePrecompiled, Descriptor.ForeignPlugin, Logger);
 			}
 
 			TargetRules RulesObject;
 			using (GlobalTracer.Instance.BuildSpan("RulesAssembly.CreateTargetRules()").StartActive())
 			{
-				RulesObject = RulesAssembly.CreateTargetRules(Descriptor.Name, Descriptor.Platform, Descriptor.Configuration, Descriptor.Architecture, Descriptor.ProjectFile, Descriptor.AdditionalArguments, Descriptor.IsTestsTarget);
+				RulesObject = RulesAssembly.CreateTargetRules(Descriptor.Name, Descriptor.Platform, Descriptor.Configuration, Descriptor.Architecture, Descriptor.ProjectFile, Descriptor.AdditionalArguments, Logger, Descriptor.IsTestsTarget);
 			}
 			if ((ProjectFileGenerator.bGenerateProjectFiles == false) && !RulesObject.GetSupportedPlatforms().Contains(Descriptor.Platform))
 			{
@@ -748,7 +750,7 @@ namespace UnrealBuildTool
 			// If we're using the shared build environment, make sure all the settings are valid
 			if (RulesObject.BuildEnvironment == TargetBuildEnvironment.Shared)
 			{
-				ValidateSharedEnvironment(RulesAssembly, Descriptor.Name, Descriptor.AdditionalArguments, RulesObject);
+				ValidateSharedEnvironment(RulesAssembly, Descriptor.Name, Descriptor.AdditionalArguments, RulesObject, Logger);
 			}
 
 			// If we're precompiling, generate a list of all the files that we depend on
@@ -808,10 +810,10 @@ namespace UnrealBuildTool
 
 				FileReference IncludeCache = FileReference.Combine(CacheDir, "Intermediate", "HeaderLookup.dat");
 
-				Log.TraceInformation("Building dependency cache for specified headers");
+				Logger.LogInformation("Building dependency cache for specified headers");
 				foreach (FileReference HeaderFile in SpecificHeaderFiles)
 				{
-					Log.TraceLog("  {0}", HeaderFile);
+					Logger.LogDebug("  {HeaderFile}", HeaderFile);
 				}
 
 				CppIncludeLookup IncludeLookup = new CppIncludeLookup(IncludeCache);
@@ -823,10 +825,10 @@ namespace UnrealBuildTool
 				NewFiles.RemoveAll(x => IsHeader(x));
 				NewFiles = NewFiles.Except(Descriptor.SpecificFilesToCompile).ToList();
 
-				Log.TraceInformation("Found {0} source files", NewFiles.Count);
+				Logger.LogInformation("Found {NewFilesCount} source files", NewFiles.Count);
 				foreach (FileReference NewFile in NewFiles)
 				{
-					Log.TraceLog("  {0}", NewFile);
+					Logger.LogDebug("  {NewFile}", NewFile);
 				}
 
 				Descriptor.OptionalFilesToCompile.AddRange(NewFiles.Except(Descriptor.SpecificFilesToCompile));
@@ -852,11 +854,11 @@ namespace UnrealBuildTool
 			UEBuildTarget Target;
 			using (GlobalTracer.Instance.BuildSpan("UEBuildTarget constructor").StartActive())
 			{
-				Target = new UEBuildTarget(Descriptor, new ReadOnlyTargetRules(RulesObject), RulesAssembly);
+				Target = new UEBuildTarget(Descriptor, new ReadOnlyTargetRules(RulesObject), RulesAssembly, Logger);
 			}
 			using (GlobalTracer.Instance.BuildSpan("UEBuildTarget.PreBuildSetup()").StartActive())
 			{
-				Target.PreBuildSetup();
+				Target.PreBuildSetup(Logger);
 			}
 
 			return Target;
@@ -865,7 +867,7 @@ namespace UnrealBuildTool
 		/// <summary>
 		/// Validates that the build environment matches the shared build environment, by comparing the TargetRules instance to the vanilla target rules for the current target type.
 		/// </summary>
-		static void ValidateSharedEnvironment(RulesAssembly RulesAssembly, string ThisTargetName, CommandLineArguments Arguments, TargetRules ThisRules)
+		static void ValidateSharedEnvironment(RulesAssembly RulesAssembly, string ThisTargetName, CommandLineArguments Arguments, TargetRules ThisRules, ILogger Logger)
 		{
 			// Allow disabling these checks
 			if (ThisRules.bOverrideBuildEnvironment)
@@ -894,7 +896,7 @@ namespace UnrealBuildTool
 			}
 
 			// Create the target rules for it
-			TargetRules BaseRules = RulesAssembly.CreateTargetRules(BaseTargetName, ThisRules.Platform, ThisRules.Configuration, ThisRules.Architecture, null, Arguments);
+			TargetRules BaseRules = RulesAssembly.CreateTargetRules(BaseTargetName, ThisRules.Platform, ThisRules.Configuration, ThisRules.Architecture, null, Arguments, Logger);
 
 			// Get all the configurable objects
 			object[] BaseObjects = BaseRules.GetConfigurableObjects().ToArray();
@@ -1140,9 +1142,10 @@ namespace UnrealBuildTool
 		/// <param name="InDescriptor">Target descriptor</param>
 		/// <param name="InRules">The target rules, as created by RulesCompiler.</param>
 		/// <param name="InRulesAssembly">The chain of rules assemblies that this target was created with</param>
-		private UEBuildTarget(TargetDescriptor InDescriptor, ReadOnlyTargetRules InRules, RulesAssembly InRulesAssembly)
+		/// <param name="Logger">Logger for output</param>
+		private UEBuildTarget(TargetDescriptor InDescriptor, ReadOnlyTargetRules InRules, RulesAssembly InRulesAssembly, ILogger Logger)
 		{
-			MetadataCache = SourceFileMetadataCache.CreateHierarchy(InDescriptor.ProjectFile);
+			MetadataCache = SourceFileMetadataCache.CreateHierarchy(InDescriptor.ProjectFile, Logger);
 			ProjectFile = InDescriptor.ProjectFile;
 			TargetName = InRules.Name;
 			AppName = InRules.Name;
@@ -1254,7 +1257,8 @@ namespace UnrealBuildTool
 		/// <param name="Location">Path to the dependency list</param>
 		/// <param name="RuntimeDependencies">List of all the runtime dependencies for this target</param>
 		/// <param name="RuntimeDependencyTargetFileToSourceFile">Map of runtime dependencies to their location in the source tree, before they are staged</param>
-		void WriteDependencyList(FileReference Location, List<RuntimeDependency> RuntimeDependencies, Dictionary<FileReference, FileReference> RuntimeDependencyTargetFileToSourceFile)
+		/// <param name="Logger">Logger for output</param>
+		void WriteDependencyList(FileReference Location, List<RuntimeDependency> RuntimeDependencies, Dictionary<FileReference, FileReference> RuntimeDependencyTargetFileToSourceFile, ILogger Logger)
 		{
 			HashSet<FileReference> Files = new HashSet<FileReference>();
 
@@ -1294,7 +1298,7 @@ namespace UnrealBuildTool
 				}
 
 				// Create the module rules
-				ModuleRules Rules = CreateModuleRulesAndSetDefaults(Module.Name, "external file list option");
+				ModuleRules Rules = CreateModuleRulesAndSetDefaults(Module.Name, "external file list option", Logger);
 
 				// Add Additional Bundle Resources for all modules
 				foreach (ModuleRules.BundleResource Resource in Rules.AdditionalBundleResources)
@@ -1386,7 +1390,7 @@ namespace UnrealBuildTool
 			}
 
 			// Write the file
-			Log.TraceInformation("Writing dependency list to {0}", Location);
+			Logger.LogInformation("Writing dependency list to {Location}", Location);
 			DirectoryReference.CreateDirectory(Location.Directory);
 			FileReference.WriteAllLines(Location, Files.Where(x => x.IsUnderDirectory(Unreal.RootDirectory)).Select(x => x.MakeRelativeTo(Unreal.RootDirectory).Replace(Path.DirectorySeparatorChar, '/')).OrderBy(x => x));
 		}
@@ -1411,7 +1415,7 @@ namespace UnrealBuildTool
 		/// <summary>
 		/// Generates a public manifest file for writing out
 		/// </summary>
-		public void GenerateManifest(FileReference ManifestPath, List<KeyValuePair<FileReference, BuildProductType>> BuildProducts)
+		public void GenerateManifest(FileReference ManifestPath, List<KeyValuePair<FileReference, BuildProductType>> BuildProducts, ILogger Logger)
 		{
 			BuildManifest Manifest = new BuildManifest();
 
@@ -1447,8 +1451,8 @@ namespace UnrealBuildTool
 			Manifest.BuildProducts.Sort();
 			Manifest.DeployTargetFiles.Sort();
 
-			Log.TraceInformation("Writing manifest to {0}", ManifestPath);
-			Utils.WriteClass<BuildManifest>(Manifest, ManifestPath.FullName, "");
+			Logger.LogInformation("Writing manifest to {ManifestPath}", ManifestPath);
+			Utils.WriteClass<BuildManifest>(Manifest, ManifestPath.FullName, "", Logger);
 		}
 
 		/// <summary>
@@ -1673,11 +1677,11 @@ namespace UnrealBuildTool
 		/// Creates a global compile environment suitable for generating project files.
 		/// </summary>
 		/// <returns>New compile environment</returns>
-		public CppCompileEnvironment CreateCompileEnvironmentForProjectFiles()
+		public CppCompileEnvironment CreateCompileEnvironmentForProjectFiles(ILogger Logger)
 		{
 			CppConfiguration CppConfiguration = GetCppConfiguration(Configuration);
 
-			SourceFileMetadataCache MetadataCache = SourceFileMetadataCache.CreateHierarchy(ProjectFile);
+			SourceFileMetadataCache MetadataCache = SourceFileMetadataCache.CreateHierarchy(ProjectFile, Logger);
 
 			CppCompileEnvironment GlobalCompileEnvironment = new CppCompileEnvironment(Platform, CppConfiguration, Architecture, MetadataCache);
 			LinkEnvironment GlobalLinkEnvironment = new LinkEnvironment(GlobalCompileEnvironment.Platform, GlobalCompileEnvironment.Configuration, GlobalCompileEnvironment.Architecture);
@@ -1685,7 +1689,7 @@ namespace UnrealBuildTool
 			UEToolChain TargetToolChain = CreateToolchain(Platform);
 			SetupGlobalEnvironment(TargetToolChain, GlobalCompileEnvironment, GlobalLinkEnvironment);
 
-			FindSharedPCHs(Binaries, GlobalCompileEnvironment);
+			FindSharedPCHs(Binaries, GlobalCompileEnvironment, Logger);
 
 			return GlobalCompileEnvironment;
 		}
@@ -1693,11 +1697,11 @@ namespace UnrealBuildTool
 		/// <summary>
 		/// Builds the target, appending list of output files and returns building result.
 		/// </summary>
-		public TargetMakefile Build(BuildConfiguration BuildConfiguration, ISourceFileWorkingSet WorkingSet, TargetDescriptor TargetDescriptor, bool bInitOnly = false)
+		public TargetMakefile Build(BuildConfiguration BuildConfiguration, ISourceFileWorkingSet WorkingSet, TargetDescriptor TargetDescriptor, ILogger Logger, bool bInitOnly = false)
 		{
 			CppConfiguration CppConfiguration = GetCppConfiguration(Configuration);
 
-			SourceFileMetadataCache MetadataCache = SourceFileMetadataCache.CreateHierarchy(ProjectFile);
+			SourceFileMetadataCache MetadataCache = SourceFileMetadataCache.CreateHierarchy(ProjectFile, Logger);
 
 			CppCompileEnvironment GlobalCompileEnvironment = new CppCompileEnvironment(Platform, CppConfiguration, Architecture, MetadataCache);
 			LinkEnvironment GlobalLinkEnvironment = new LinkEnvironment(GlobalCompileEnvironment.Platform, GlobalCompileEnvironment.Configuration, GlobalCompileEnvironment.Architecture);
@@ -1772,7 +1776,7 @@ namespace UnrealBuildTool
 							&& !Rules.InternalPluginDependencies.Contains(DependencyPlugin.Name)
 							&& !Plugin.Dependencies!.Contains(DependencyPlugin))
 						{
-							Log.TraceWarning("Warning: Plugin '{0}' does not list plugin '{1}' as a dependency, but module '{2}' depends on '{3}'.", Plugin.Name, DependencyPlugin.Name, Module.Name, DependencyModule.Name);
+							Logger.LogWarning("Warning: Plugin '{PluginName}' does not list plugin '{DependencyPluginName}' as a dependency, but module '{ModuleName}' depends on '{DependencyModuleName}'.", Plugin.Name, DependencyPlugin.Name, Module.Name, DependencyModule.Name);
 						}
 					}
 				}
@@ -1784,6 +1788,7 @@ namespace UnrealBuildTool
 				Rules.ConfigValueTracker, bDeployAfterCompile, bHasProjectScriptPlugin, bHasRequiredProjectScriptPlugin, UbtPlugins?.ToArray(), EnabledUbtPlugins?.ToArray(), 
 				EnabledUhtPlugins?.ToArray());
 			Makefile.IsTestTarget = Rules.IsTestTarget;
+			TargetMakefileBuilder MakefileBuilder = new TargetMakefileBuilder(Makefile, Logger);
 
 			// Get diagnostic info to be printed before each build
 			TargetToolChain.GetVersionInfo(Makefile.Diagnostics);
@@ -1820,7 +1825,7 @@ namespace UnrealBuildTool
 			// Prepare cached data for UHT header generation
 			using (GlobalTracer.Instance.BuildSpan("ExternalExecution.SetupUObjectModules()").StartActive())
 			{
-				ExternalExecution.SetupUObjectModules(ModulesToGenerateHeadersFor, Rules.Platform, ProjectDescriptor, Makefile.UObjectModules, Makefile.UObjectModuleHeaders, Rules.GeneratedCodeVersion, MetadataCache);
+				ExternalExecution.SetupUObjectModules(ModulesToGenerateHeadersFor, Rules.Platform, ProjectDescriptor, Makefile.UObjectModules, Makefile.UObjectModuleHeaders, Rules.GeneratedCodeVersion, MetadataCache, Logger);
 			}
 
 			List<string> UHTAdditionalArguments = new List<string>();
@@ -1845,7 +1850,7 @@ namespace UnrealBuildTool
 			// files that are injected as top level prerequisites.  If UHT only emitted included header files, we wouldn't need to run it during the Gather phase at all.
 			if (Makefile.UObjectModules.Count > 0)
 			{
-				ExternalExecution.ExecuteHeaderToolIfNecessary(BuildConfiguration, ProjectFile, Makefile, TargetName, WorkingSet);
+				ExternalExecution.ExecuteHeaderToolIfNecessary(BuildConfiguration, ProjectFile, Makefile, TargetName, WorkingSet, Logger);
 			}
 
 #if __VPROJECT_AVAILABLE__
@@ -1862,14 +1867,14 @@ namespace UnrealBuildTool
 			// files that are injected as top level prerequisites.  If VNI only emitted included header files, we wouldn't need to run it during the Gather phase at all.
 			if (Makefile.VNIModules.Count > 0)
 			{
-				VNIExecution.ExecuteVNITool(Makefile, TargetDescriptor);
+				VNIExecution.ExecuteVNITool(Makefile, TargetDescriptor, Logger);
 			}
 #endif
 
 			// Find all the shared PCHs.
 			if (Rules.bUseSharedPCHs)
 			{
-				FindSharedPCHs(OriginalBinaries, GlobalCompileEnvironment);
+				FindSharedPCHs(OriginalBinaries, GlobalCompileEnvironment, Logger);
 			}
 
 			// Compile the resource files common to all DLLs on Windows
@@ -1886,7 +1891,7 @@ namespace UnrealBuildTool
 
 							FileItem DefaultResourceFile = FileItem.GetItemByFileReference(DefaultResourceLocation);
 
-							CPPOutput DefaultResourceOutput = TargetToolChain.CompileRCFiles(DefaultResourceCompileEnvironment, new List<FileItem> { DefaultResourceFile }, EngineIntermediateDirectory, Makefile);
+							CPPOutput DefaultResourceOutput = TargetToolChain.CompileRCFiles(DefaultResourceCompileEnvironment, new List<FileItem> { DefaultResourceFile }, EngineIntermediateDirectory, MakefileBuilder);
 							GlobalLinkEnvironment.DefaultResourceFiles.AddRange(DefaultResourceOutput.ObjectFiles);
 						}
 					}
@@ -1912,7 +1917,7 @@ namespace UnrealBuildTool
 			{
 				foreach (UEBuildBinary Binary in Binaries)
 				{
-					List<FileItem> BinaryOutputItems = Binary.Build(Rules, TargetToolChain, GlobalCompileEnvironment, GlobalLinkEnvironment, TargetDescriptor.SpecificFilesToCompile, WorkingSet, ExeDir, Makefile);
+					List<FileItem> BinaryOutputItems = Binary.Build(Rules, TargetToolChain, GlobalCompileEnvironment, GlobalLinkEnvironment, TargetDescriptor.SpecificFilesToCompile, WorkingSet, ExeDir, MakefileBuilder, Logger);
 					Makefile.OutputItems.AddRange(BinaryOutputItems);
 				}
 			}
@@ -1930,7 +1935,7 @@ namespace UnrealBuildTool
 			{
 				if (!UnrealBuildTool.IsFileInstalled(Pair.Key))
 				{
-					Makefile.OutputItems.Add(Makefile.CreateCopyAction(Pair.Value, Pair.Key));
+					Makefile.OutputItems.Add(MakefileBuilder.CreateCopyAction(Pair.Value, Pair.Key));
 				}
 			}
 
@@ -1956,7 +1961,7 @@ namespace UnrealBuildTool
 			}
 
 			// Allow the toolchain to modify the final output items
-			TargetToolChain.FinalizeOutput(Rules, Makefile);
+			TargetToolChain.FinalizeOutput(Rules, MakefileBuilder);
 
 			// Get all the regular build products
 			List<KeyValuePair<FileReference, BuildProductType>> BuildProducts = new List<KeyValuePair<FileReference, BuildProductType>>();
@@ -2028,7 +2033,7 @@ namespace UnrealBuildTool
 				{
 					WriteMetadataTargetInfo TargetInfo = new WriteMetadataTargetInfo(ProjectFile, null, null, ReceiptFileName, Receipt, FileNameToModuleManifest);
 					FileReference TargetInfoFile = FileReference.Combine(ProjectIntermediateDirectory, "TargetMetadata.dat");
-					CreateWriteMetadataAction(Makefile, ReceiptFileName.GetFileName(), TargetInfoFile, TargetInfo, Makefile.OutputItems);
+					CreateWriteMetadataAction(MakefileBuilder, ReceiptFileName.GetFileName(), TargetInfoFile, TargetInfo, Makefile.OutputItems);
 				}
 				else
 				{
@@ -2064,11 +2069,11 @@ namespace UnrealBuildTool
 
 					WriteMetadataTargetInfo EngineInfo = new WriteMetadataTargetInfo(null, VersionFile, Receipt.Version, null, null, EngineFileToManifest);
 					FileReference EngineInfoFile = FileReference.Combine(ProjectIntermediateDirectory, "EngineMetadata.dat");
-					CreateWriteMetadataAction(Makefile, VersionFile.GetFileName(), EngineInfoFile, EngineInfo, EnginePrereqItems);
+					CreateWriteMetadataAction(MakefileBuilder, VersionFile.GetFileName(), EngineInfoFile, EngineInfo, EnginePrereqItems);
 
 					WriteMetadataTargetInfo TargetInfo = new WriteMetadataTargetInfo(ProjectFile, VersionFile, null, ReceiptFileName, Receipt, TargetFileToManifest);
 					FileReference TargetInfoFile = FileReference.Combine(ProjectIntermediateDirectory, "TargetMetadata.dat");
-					CreateWriteMetadataAction(Makefile, ReceiptFileName.GetFileName(), TargetInfoFile, TargetInfo, TargetPrereqItems);
+					CreateWriteMetadataAction(MakefileBuilder, ReceiptFileName.GetFileName(), TargetInfoFile, TargetInfo, TargetPrereqItems);
 				}
 
 				// Create actions to run the post build steps
@@ -2077,7 +2082,7 @@ namespace UnrealBuildTool
 				{
 					FileReference OutputFile = new FileReference(PostBuildScript.FullName + ".ran");
 
-					Action PostBuildStepAction = Makefile.CreateAction(ActionType.PostBuildStep);
+					Action PostBuildStepAction = MakefileBuilder.CreateAction(ActionType.PostBuildStep);
 					PostBuildStepAction.CommandPath = BuildHostPlatform.Current.Shell;
 					if (BuildHostPlatform.Current.ShellType == ShellType.Cmd)
 					{
@@ -2100,13 +2105,13 @@ namespace UnrealBuildTool
 			// Build a list of all the files required to build
 			foreach (FileReference DependencyListFileName in Rules.DependencyListFileNames)
 			{
-				WriteDependencyList(DependencyListFileName, RuntimeDependencies, RuntimeDependencyTargetFileToSourceFile);
+				WriteDependencyList(DependencyListFileName, RuntimeDependencies, RuntimeDependencyTargetFileToSourceFile, Logger);
 			}
 
 			// If we're only generating the manifest, return now
 			foreach (FileReference ManifestFileName in Rules.ManifestFileNames)
 			{
-				GenerateManifest(ManifestFileName, BuildProducts);
+				GenerateManifest(ManifestFileName, BuildProducts, Logger);
 			}
 
 			// Check there are no EULA or restricted folder violations
@@ -2132,11 +2137,11 @@ namespace UnrealBuildTool
 					bool bResult = true;
 					foreach (UEBuildBinary Binary in Binaries)
 					{
-						bResult &= Binary.CheckRestrictedFolders(RootDirectories, ModuleRestrictedFolderCache);
+						bResult &= Binary.CheckRestrictedFolders(RootDirectories, ModuleRestrictedFolderCache, Logger);
 					}
 					foreach (KeyValuePair<FileReference, FileReference> Pair in RuntimeDependencyTargetFileToSourceFile)
 					{
-						bResult &= CheckRestrictedFolders(Pair.Key, Pair.Value);
+						bResult &= CheckRestrictedFolders(Pair.Key, Pair.Value, Logger);
 					}
 
 					if (!bResult)
@@ -2146,7 +2151,7 @@ namespace UnrealBuildTool
 				}
 
 				// Check for linking against modules prohibited by the EULA
-				CheckForEULAViolation();
+				CheckForEULAViolation(Logger);
 			}
 
 			// Add all the plugins to be tracked
@@ -2189,16 +2194,16 @@ namespace UnrealBuildTool
 			{
 				UEBuildBinary Binary = Binaries[0];
 				FileReference Header = FileReference.Combine(Binary.OutputDir, Rules.ExportPublicHeader);
-				WritePublicHeader(Binary, Header, GlobalCompileEnvironment);
+				WritePublicHeader(Binary, Header, GlobalCompileEnvironment, Logger);
 			}
 
 			// Clean any stale modules which exist in multiple output directories. This can lead to the wrong DLL being loaded on Windows.
-			CleanStaleModules();
+			CleanStaleModules(Logger);
 
 			return Makefile;
 		}
 
-		void CreateWriteMetadataAction(TargetMakefile Makefile, string StatusDescription, FileReference InfoFile, WriteMetadataTargetInfo Info, IEnumerable<FileItem> PrerequisiteItems)
+		void CreateWriteMetadataAction(TargetMakefileBuilder MakefileBuilder, string StatusDescription, FileReference InfoFile, WriteMetadataTargetInfo Info, IEnumerable<FileItem> PrerequisiteItems)
 		{
 			List<FileReference> ProducedItems = new List<FileReference>(Info.FileToManifest.Keys);
 			if (Info.Version != null && Info.VersionFile != null)
@@ -2212,6 +2217,8 @@ namespace UnrealBuildTool
 
 			if (!ProducedItems.Any(x => UnrealBuildTool.IsFileInstalled(x)))
 			{
+				TargetMakefile Makefile = MakefileBuilder.Makefile;
+
 				// Save the info file, but deliberately do not add a prerequisite from the action onto it unless we're updating the
 				// target file. Since the outputs are not target specific, but the input file is always beneath the project intermediate
 				// directory (even for the engine), it will cause -NoEngineChanges to block the build when updated. The behavior we want
@@ -2232,7 +2239,7 @@ namespace UnrealBuildTool
 				}
 
 				// Create the action
-				Action WriteMetadataAction = Makefile.CreateRecursiveAction<WriteMetadataMode>(ActionType.WriteMetadata, WriteMetadataArguments.ToString());
+				Action WriteMetadataAction = MakefileBuilder.CreateRecursiveAction<WriteMetadataMode>(ActionType.WriteMetadata, WriteMetadataArguments.ToString());
 				WriteMetadataAction.WorkingDirectory = Unreal.EngineSourceDirectory;
 				WriteMetadataAction.StatusDescription = StatusDescription;
 				WriteMetadataAction.bCanExecuteRemotely = false;
@@ -2274,8 +2281,9 @@ namespace UnrealBuildTool
 		/// </summary>
 		/// <param name="TargetFile">The destination location for the file</param>
 		/// <param name="SourceFile">The source location of the file</param>
+		/// <param name="Logger">Logger for output</param>
 		/// <returns>True if the copy is permitted, false otherwise</returns>
-		bool CheckRestrictedFolders(FileReference TargetFile, FileReference SourceFile)
+		bool CheckRestrictedFolders(FileReference TargetFile, FileReference SourceFile, ILogger Logger)
 		{
 			List<RestrictedFolder> TargetRestrictedFolders = GetRestrictedFolders(TargetFile);
 			List<RestrictedFolder> SourceRestrictedFolders = GetRestrictedFolders(SourceFile);
@@ -2283,7 +2291,7 @@ namespace UnrealBuildTool
 			{
 				if (!TargetRestrictedFolders.Contains(SourceRestrictedFolder))
 				{
-					Log.TraceError("Runtime dependency '{0}' is copied to '{1}', which does not contain a '{2}' folder.", SourceFile, TargetFile, SourceRestrictedFolder);
+					Logger.LogError("Runtime dependency '{SourceFile}' is copied to '{TargetFile}', which does not contain a '{SourceRestrictedFolder}' folder.", SourceFile, TargetFile, SourceRestrictedFolder);
 					return false;
 				}
 			}
@@ -2343,7 +2351,7 @@ namespace UnrealBuildTool
 		/// On Windows, the loader reads imported DLLs from the first location it finds them. If modules are moved from one place to another, we have to be sure to clean up the old versions
 		/// so that they're not loaded accidentally causing unintuitive import errors.
 		/// </summary>
-		void CleanStaleModules()
+		void CleanStaleModules(ILogger Logger)
 		{
 			// Find all the output files
 			HashSet<FileReference> OutputFiles = new HashSet<FileReference>();
@@ -2369,14 +2377,14 @@ namespace UnrealBuildTool
 						FileReference? OutputFile;
 						if (OutputNameToLocation.TryGetValue(ExistingFile.GetFileName(), out OutputFile) && !OutputFiles.Contains(ExistingFile))
 						{
-							Log.TraceInformation("Deleting '{0}' to avoid ambiguity with '{1}'", ExistingFile, OutputFile);
+							Logger.LogInformation("Deleting '{ExistingFile}' to avoid ambiguity with '{OutputFile}'", ExistingFile, OutputFile);
 							try
 							{
 								FileReference.Delete(ExistingFile);
 							}
 							catch (Exception Ex)
 							{
-								Log.TraceError("Unable to delete {0} ({1})", ExistingFile, Ex.Message.TrimEnd());
+								Logger.LogError("Unable to delete {File} ({Ex})", ExistingFile, Ex.Message.TrimEnd());
 							}
 						}
 					}
@@ -2450,7 +2458,8 @@ namespace UnrealBuildTool
 		/// <param name="Binary">Binary to write a header for</param>
 		/// <param name="HeaderFile">Path to the header to output</param>
 		/// <param name="GlobalCompileEnvironment">The global compile environment for this target</param>
-		void WritePublicHeader(UEBuildBinary Binary, FileReference HeaderFile, CppCompileEnvironment GlobalCompileEnvironment)
+		/// <param name="Logger"></param>
+		void WritePublicHeader(UEBuildBinary Binary, FileReference HeaderFile, CppCompileEnvironment GlobalCompileEnvironment, ILogger Logger)
 		{
 			DirectoryReference.CreateDirectory(HeaderFile.Directory);
 
@@ -2479,13 +2488,13 @@ namespace UnrealBuildTool
 					}
 				}
 			}
-			Log.TraceInformation("Written public header to {0}", HeaderFile);
+			Logger.LogInformation("Written public header to {HeaderFile}", HeaderFile);
 		}
 
 		/// <summary>
 		/// Check for EULA violation dependency issues.
 		/// </summary>
-		private void CheckForEULAViolation()
+		private void CheckForEULAViolation(ILogger Logger)
 		{
 			if (TargetType != TargetType.Editor && TargetType != TargetType.Program && Configuration == UnrealTargetConfiguration.Shipping &&
 				Rules.bCheckLicenseViolations)
@@ -2510,11 +2519,11 @@ namespace UnrealBuildTool
 						}
 						if (Rules.bBreakBuildOnLicenseViolation)
 						{
-							Log.TraceError("ERROR: {0}", Message);
+							Logger.LogError("ERROR: {Message}", Message);
 						}
 						else
 						{
-							Log.TraceWarning("WARNING: {0}", Message);
+							Logger.LogWarning("WARNING: {Message}", Message);
 						}
 						bLicenseViolation = true;
 					}
@@ -2560,16 +2569,16 @@ namespace UnrealBuildTool
 		/// <summary>
 		/// Setup target before build. This method finds dependencies, sets up global environment etc.
 		/// </summary>
-		public void PreBuildSetup()
+		public void PreBuildSetup(ILogger Logger)
 		{
 			// Describe what's being built.
-			Log.TraceVerbose("Building {0} - {1} - {2} - {3}", AppName, TargetName, Platform, Configuration);
+			Logger.LogDebug("Building {AppName} - {TargetName} - {Platform} - {Configuration}", AppName, TargetName, Platform, Configuration);
 
 			// Setup the target's binaries.
-			SetupBinaries();
+			SetupBinaries(Logger);
 
 			// Setup the target's plugins
-			SetupPlugins();
+			SetupPlugins(Logger);
 
 			// Add the plugin binaries to the build
 			foreach (UEBuildPlugin Plugin in BuildPlugins!)
@@ -2583,12 +2592,12 @@ namespace UnrealBuildTool
 			// Add all of the extra modules, including game modules, that need to be compiled along
 			// with this app.  These modules are always statically linked in monolithic targets, but not necessarily linked to anything in modular targets,
 			// and may still be required at runtime in order for the application to load and function properly!
-			AddExtraModules();
+			AddExtraModules(Logger);
 
 			// Create all the modules referenced by the existing binaries
 			foreach (UEBuildBinary Binary in Binaries)
 			{
-				Binary.CreateAllDependentModules(FindOrCreateModuleByName);
+				Binary.CreateAllDependentModules((Name, RefChain) => FindOrCreateModuleByName(Name, RefChain, Logger), Logger);
 			}
 
 			// Bind every referenced C++ module to a binary
@@ -2607,7 +2616,7 @@ namespace UnrealBuildTool
 			// Add all the modules to the target if necessary.
 			if (Rules.bBuildAllModules)
 			{
-				AddAllValidModulesToTarget();
+				AddAllValidModulesToTarget(Logger);
 			}
 
 			// Add the external and non-C++ referenced modules to the binaries that reference them.
@@ -2821,7 +2830,8 @@ namespace UnrealBuildTool
 		/// </summary>
 		/// <param name="OriginalBinaries">The list of binaries</param>
 		/// <param name="GlobalCompileEnvironment">The compile environment. The shared PCHs will be added to the SharedPCHs list in this.</param>
-		void FindSharedPCHs(List<UEBuildBinary> OriginalBinaries, CppCompileEnvironment GlobalCompileEnvironment)
+		/// <param name="Logger">Logger for output</param>
+		void FindSharedPCHs(List<UEBuildBinary> OriginalBinaries, CppCompileEnvironment GlobalCompileEnvironment, ILogger Logger)
 		{
 			// Find how many other shared PCH modules each module depends on, and use that to sort the shared PCHs by reverse order of size.
 			HashSet<UEBuildModuleCPP> SharedPCHModules = new HashSet<UEBuildModuleCPP>();
@@ -2855,10 +2865,10 @@ namespace UnrealBuildTool
 			// Print the ordered list of shared PCHs
 			if (OrderedSharedPCHModules.Count > 0)
 			{
-				Log.TraceVerbose("Found {0} shared PCH headers (listed in order of preference):", SharedPCHModules.Count);
+				Logger.LogDebug("Found {SharedPCHModulesCount} shared PCH headers (listed in order of preference):", SharedPCHModules.Count);
 				foreach (PrecompiledHeaderTemplate SharedPCHModule in OrderedSharedPCHModules)
 				{
-					Log.TraceVerbose("	" + SharedPCHModule.Module.Name);
+					Logger.LogDebug("	{ModuleName}", SharedPCHModule.Module.Name);
 				}
 			}
 		}
@@ -2867,7 +2877,7 @@ namespace UnrealBuildTool
 		/// When building a target, this is called to add any additional modules that should be compiled along
 		/// with the main target.  If you override this in a derived class, remember to call the base implementation!
 		/// </summary>
-		protected void AddExtraModules()
+		protected void AddExtraModules(ILogger Logger)
 		{
 			// Find all the extra module names
 			List<string> ExtraModuleNames = new List<string>();
@@ -2877,7 +2887,7 @@ namespace UnrealBuildTool
 			// Add extra modules that will either link into the main binary (monolithic), or be linked into separate DLL files (modular)
 			foreach (string ModuleName in ExtraModuleNames)
 			{
-				UEBuildModuleCPP Module = FindOrCreateCppModuleByName(ModuleName, TargetRulesFile.GetFileName());
+				UEBuildModuleCPP Module = FindOrCreateCppModuleByName(ModuleName, TargetRulesFile.GetFileName(), Logger);
 				if (Module.Binary == null)
 				{
 					AddModuleToBinary(Module);
@@ -2888,7 +2898,7 @@ namespace UnrealBuildTool
 		/// <summary>
 		/// Adds all the precompiled modules into the target. Precompiled modules are compiled alongside the target, but not linked into it unless directly referenced.
 		/// </summary>
-		protected void AddAllValidModulesToTarget()
+		protected void AddAllValidModulesToTarget(ILogger Logger)
 		{
 			// Find all the modules that are part of the target
 			HashSet<string> ValidModuleNames = new HashSet<string>();
@@ -3019,7 +3029,7 @@ namespace UnrealBuildTool
 				ModuleRules? ModuleRules;
 				try
 				{
-					ModuleRules = RulesAssembly.CreateModuleRules(FilteredModuleName, this.Rules, "all modules option");
+					ModuleRules = RulesAssembly.CreateModuleRules(FilteredModuleName, this.Rules, "all modules option", Logger);
 				}
 				catch (BuildException)
 				{
@@ -3043,16 +3053,16 @@ namespace UnrealBuildTool
 			const string PrecompileReferenceChain = "allmodules option";
 			foreach (string ModuleName in ValidModuleNames)
 			{
-				UEBuildModule Module = FindOrCreateModuleByName(ModuleName, PrecompileReferenceChain);
+				UEBuildModule Module = FindOrCreateModuleByName(ModuleName, PrecompileReferenceChain, Logger);
 				AllModules.Add(Module);
 				Module.RecursivelyCreateModules(
 					(string ModuleName, string ReferenceChain) => 
 					{ 
-						UEBuildModule FoundModule = FindOrCreateModuleByName(ModuleName, ReferenceChain);
+						UEBuildModule FoundModule = FindOrCreateModuleByName(ModuleName, ReferenceChain, Logger);
 						AllModules.Add(FoundModule);
 						return FoundModule; 
 					}, 
-					PrecompileReferenceChain);
+					PrecompileReferenceChain, Logger);
 			}
 
 			// Exclude additional modules that were added only for include-path-only purposes, and those that will have their Verse dependency satisfied
@@ -3066,7 +3076,7 @@ namespace UnrealBuildTool
 				{
 					if (!ValidModules.Contains(ReferencedModule))
 					{
-						Log.TraceError("Module '{0}' is not usable without module '{1}', which is not valid for this target.", ValidModule.Name, ReferencedModule.Name);
+						Logger.LogError("Module '{ValidModuleName}' is not usable without module '{ReferencedModuleName}', which is not valid for this target.", ValidModule.Name, ReferencedModule.Name);
 					}
 				}
 
@@ -3276,7 +3286,7 @@ namespace UnrealBuildTool
 		/// <summary>
 		/// Sets up the plugins for this target
 		/// </summary>
-		public void SetupPlugins()
+		public void SetupPlugins(ILogger Logger)
 		{
 			// Find all the valid plugins
 			Dictionary<string, PluginInfo> NameToInfo = RulesAssembly.EnumeratePlugins().ToDictionary(x => x.Name, x => x, StringComparer.InvariantCultureIgnoreCase);
@@ -3304,7 +3314,7 @@ namespace UnrealBuildTool
 			if (ForeignPlugin != null)
 			{
 				PluginReferenceDescriptor PluginReference = new PluginReferenceDescriptor(ForeignPlugin.GetFileNameWithoutExtension(), null, true);
-				AddPlugin(PluginReference, "command line", ExcludeFolders, NameToInstance, NameToInfo);
+				AddPlugin(PluginReference, "command line", ExcludeFolders, NameToInstance, NameToInfo, Logger);
 			}
 
 			// Configure plugins explicitly enabled via target settings
@@ -3313,7 +3323,7 @@ namespace UnrealBuildTool
 				if (ReferencedNames.Add(PluginName))
 				{
 					PluginReferenceDescriptor PluginReference = new PluginReferenceDescriptor(PluginName, null, true);
-					AddPlugin(PluginReference, "target settings", ExcludeFolders, NameToInstance, NameToInfo);
+					AddPlugin(PluginReference, "target settings", ExcludeFolders, NameToInstance, NameToInfo, Logger);
 				}
 			}
 
@@ -3323,7 +3333,7 @@ namespace UnrealBuildTool
 				if (ReferencedNames.Add(PluginName))
 				{
 					PluginReferenceDescriptor PluginReference = new PluginReferenceDescriptor(PluginName, null, false);
-					AddPlugin(PluginReference, "target settings", ExcludeFolders, NameToInstance, NameToInfo);
+					AddPlugin(PluginReference, "target settings", ExcludeFolders, NameToInstance, NameToInfo, Logger);
 				}
 			}
 
@@ -3343,11 +3353,11 @@ namespace UnrealBuildTool
 							// Make sure we don't have multiple references to the same plugin
 							if (!ReferencedNames.Add(PluginReference.Name))
 							{
-								Log.TraceWarning("Plugin '{0}' is listed multiple times in project file '{1}'.", PluginReference.Name, ProjectFile);
+								Logger.LogWarning("Plugin '{PluginReferenceName}' is listed multiple times in project file '{ProjectFile}'.", PluginReference.Name, ProjectFile);
 							}
 							else
 							{
-								AddPlugin(PluginReference, ProjectReferenceChain, ExcludeFolders, NameToInstance, NameToInfo);
+								AddPlugin(PluginReference, ProjectReferenceChain, ExcludeFolders, NameToInstance, NameToInfo, Logger);
 							}
 						}
 					}
@@ -3373,7 +3383,7 @@ namespace UnrealBuildTool
 							PluginReference.PlatformAllowList = PluginReference.SupportedTargetPlatforms; //synthesize allow list if it must be explicit
 						}
 
-						AddPlugin(PluginReference, "default plugins", ExcludeFolders, NameToInstance, NameToInfo);
+						AddPlugin(PluginReference, "default plugins", ExcludeFolders, NameToInstance, NameToInfo, Logger);
 					}
 				}
 			}
@@ -3391,7 +3401,7 @@ namespace UnrealBuildTool
 						if (ReferencedNames.Add(PluginName))
 						{
 							PluginReferenceDescriptor PluginReference = new PluginReferenceDescriptor(PluginName, null, true);
-							AddPlugin(PluginReference, "DefaultEngine.ini", ExcludeFolders, NameToInstance, NameToInfo);
+							AddPlugin(PluginReference, "DefaultEngine.ini", ExcludeFolders, NameToInstance, NameToInfo, Logger);
 						}
 					}
 				}
@@ -3406,7 +3416,7 @@ namespace UnrealBuildTool
 				if (ReferencedNames.Add(PluginName))
 				{
 					PluginReferenceDescriptor PluginReference = new PluginReferenceDescriptor(PluginName, null, true);
-					AddPlugin(PluginReference, "target settings", ExcludeFolders, NameToInstance, NameToInfo);
+					AddPlugin(PluginReference, "target settings", ExcludeFolders, NameToInstance, NameToInfo, Logger);
 				}
 			}
 
@@ -3414,7 +3424,7 @@ namespace UnrealBuildTool
 			BuildPlugins = new List<UEBuildPlugin>(NameToInstance.Values);
 
 			// Setup the UHT plugins
-			SetupUbtPlugins();
+			SetupUbtPlugins(Logger);
 		}
 
 		/// <summary>
@@ -3425,8 +3435,9 @@ namespace UnrealBuildTool
 		/// <param name="ExcludeFolders">Array of folder names to be excluded</param>
 		/// <param name="NameToInstance">Map from plugin name to instance of it</param>
 		/// <param name="NameToInfo">Map from plugin name to information</param>
+		/// <param name="Logger">Logger for diagnostic output</param>
 		/// <returns>Instance of the plugin, or null if it should not be used</returns>
-		private UEBuildPlugin? AddPlugin(PluginReferenceDescriptor Reference, string ReferenceChain, string[] ExcludeFolders, Dictionary<string, UEBuildPlugin> NameToInstance, Dictionary<string, PluginInfo> NameToInfo)
+		private UEBuildPlugin? AddPlugin(PluginReferenceDescriptor Reference, string ReferenceChain, string[] ExcludeFolders, Dictionary<string, UEBuildPlugin> NameToInstance, Dictionary<string, PluginInfo> NameToInfo, ILogger Logger)
 		{
 			// Ignore disabled references
 			if (!Reference.bEnabled)
@@ -3447,7 +3458,7 @@ namespace UnrealBuildTool
 						foreach (PluginReferenceDescriptor NextReference in Instance.Descriptor.Plugins)
 						{
 							string NextReferenceChain = String.Format("{0} -> {1}", ReferenceChain, Instance.File.GetFileName());
-							AddPlugin(NextReference, NextReferenceChain, ExcludeFolders, NameToInstance, NameToInfo);
+							AddPlugin(NextReference, NextReferenceChain, ExcludeFolders, NameToInstance, NameToInfo, Logger);
 						}
 					}
 				}
@@ -3457,26 +3468,26 @@ namespace UnrealBuildTool
 				// Check if the plugin is required for this platform
 				if (!Reference.IsEnabledForPlatform(Platform))
 				{
-					Log.TraceLog("Ignoring plugin '{0}' (referenced via {1}), not enabled for platform {2}", Reference.Name, ReferenceChain, Platform);
+					Logger.LogDebug("Ignoring plugin '{ReferenceName}' (referenced via {ReferenceChain}), not enabled for platform {Platform}", Reference.Name, ReferenceChain, Platform);
 					return null;
 				}
 
 				if (!Reference.IsEnabledForTargetConfiguration(Configuration))
 				{
-					Log.TraceLog("Ignoring plugin '{0}' (referenced via {1}), not enabled for target configuration {2}", Reference.Name, ReferenceChain, Configuration);
+					Logger.LogDebug("Ignoring plugin '{ReferenceName}' (referenced via {ReferenceChain}), not enabled for target configuration {Configuration}", Reference.Name, ReferenceChain, Configuration);
 					return null;
 				}
 
 				if (!Reference.IsEnabledForTarget(TargetType))
 				{
-					Log.TraceLog("Ignoring plugin '{0}' (referenced via {1}), not enabled for target {2}", Reference.Name, ReferenceChain, TargetType);
+					Logger.LogDebug("Ignoring plugin '{ReferenceName}' (referenced via {ReferenceChain}), not enabled for target {TargetType}", Reference.Name, ReferenceChain, TargetType);
 					return null;
 				}
 
 				// Disable any plugin reference which does not support the target platform
 				if (!Rules.bIncludePluginsForTargetPlatforms && !Reference.IsSupportedTargetPlatform(Platform))
 				{
-					Log.TraceLog("Ignoring plugin '{0}' (referenced via {1}) due to unsupported target platform.", Reference.Name, ReferenceChain);
+					Logger.LogDebug("Ignoring plugin '{ReferenceName}' (referenced via {ReferenceChain}) due to unsupported target platform.", Reference.Name, ReferenceChain);
 					return null;
 				}
 
@@ -3503,19 +3514,19 @@ namespace UnrealBuildTool
 				// Disable any plugin that requires the build platform
 				if (Info.Descriptor.bRequiresBuildPlatform && ShouldExcludePlugin(Info, ExcludeFolders))
 				{
-					Log.TraceLog("Ignoring plugin '{0}' (referenced via {1}) due to missing build platform", Reference.Name, ReferenceChain);
+					Logger.LogDebug("Ignoring plugin '{ReferenceName}' (referenced via {ReferenceChain}) due to missing build platform", Reference.Name, ReferenceChain);
 					return null;
 				}
 
 				// Disable any plugins that aren't compatible with this program
 				if (Rules.Type == TargetType.Program && (Info.Descriptor.SupportedPrograms == null || !Info.Descriptor.SupportedPrograms.Contains(AppName)))
 				{
-					Log.TraceLog("Ignoring plugin '{0}' (referenced via {1}) due to absence from supported programs list.", Reference.Name, ReferenceChain);
+					Logger.LogDebug("Ignoring plugin '{ReferenceName}' (referenced via {ReferenceChain}) due to absence from supported programs list.", Reference.Name, ReferenceChain);
 					return null;
 				}
 
 				// Create the new instance and add it to the cache
-				Log.TraceLog("Enabling plugin '{0}' (referenced via {1})", Reference.Name, ReferenceChain);
+				Logger.LogDebug("Enabling plugin '{ReferenceName}' (referenced via {ReferenceChain})", Reference.Name, ReferenceChain);
 				Instance = new UEBuildPlugin(Info, ReferenceChain);
 				Instance.bDescriptorReferencedExplicitly = !Reference.bOptional;
 				NameToInstance.Add(Info.Name, Instance);
@@ -3531,7 +3542,7 @@ namespace UnrealBuildTool
 					{
 						if (ModuleInfo.IsCompiledInConfiguration(Platform, Configuration, TargetName, TargetType, Rules.bBuildDeveloperTools, Rules.bBuildRequiresCookedData))
 						{
-							UEBuildModuleCPP Module = FindOrCreateCppModuleByName(ModuleInfo.Name, PluginReferenceChain);
+							UEBuildModuleCPP Module = FindOrCreateCppModuleByName(ModuleInfo.Name, PluginReferenceChain, Logger);
 							if (!Instance.Modules.Contains(Module))
 							{
 								// This could be in a child plugin so scan thorugh those as well
@@ -3552,7 +3563,7 @@ namespace UnrealBuildTool
 				{
 					foreach (PluginReferenceDescriptor NextReference in Info.Descriptor.Plugins)
 					{
-						UEBuildPlugin? NextInstance = AddPlugin(NextReference, PluginReferenceChain, ExcludeFolders, NameToInstance, NameToInfo);
+						UEBuildPlugin? NextInstance = AddPlugin(NextReference, PluginReferenceChain, ExcludeFolders, NameToInstance, NameToInfo, Logger);
 						if (NextInstance != null)
 						{
 							Dependencies.Add(NextInstance);
@@ -3598,7 +3609,7 @@ namespace UnrealBuildTool
 		/// Scan for UBT plugins.  This also detects if the current target has a C++ style UHT that don't have
 		/// an associated UBT/UHT plugin.  When this happens, the old UHT must be used.
 		/// </summary>
-		private void SetupUbtPlugins()
+		private void SetupUbtPlugins(ILogger Logger)
 		{
 			if (BuildPlugins == null)
 			{
@@ -3624,7 +3635,7 @@ namespace UnrealBuildTool
 				EnabledUbtPlugins.SortBy(P => P.FullName);
 
 				// Build the plugins
-				bool bCompiled = PluginsBase.BuildUbtPlugins(ProjectFile, EnabledUbtPlugins, out BuiltPlugins);
+				bool bCompiled = PluginsBase.BuildUbtPlugins(ProjectFile, EnabledUbtPlugins, Logger, out BuiltPlugins);
 				if (!bCompiled)
 				{
 					throw new BuildException("Not all plugins compiled");
@@ -3674,7 +3685,7 @@ namespace UnrealBuildTool
 		/// <summary>
 		/// Sets up the binaries for the target.
 		/// </summary>
-		protected void SetupBinaries()
+		protected void SetupBinaries(ILogger Logger)
 		{
 			// If we're using the new method for specifying binaries, fill in the binary configurations now
 			if (Rules.LaunchModuleName == null)
@@ -3683,7 +3694,7 @@ namespace UnrealBuildTool
 			}
 
 			// Create the launch module
-			UEBuildModuleCPP LaunchModule = FindOrCreateCppModuleByName(Rules.LaunchModuleName, TargetRulesFile.GetFileName());
+			UEBuildModuleCPP LaunchModule = FindOrCreateCppModuleByName(Rules.LaunchModuleName, TargetRulesFile.GetFileName(), Logger);
 
 			// Get the intermediate directory for the launch module directory. This can differ from the standard engine intermediate directory because it is always configuration-specific.
 			DirectoryReference IntermediateDirectory;
@@ -4165,10 +4176,10 @@ namespace UnrealBuildTool
 		/// <summary>
 		/// Create a rules object for the given module, and set any default values for this target
 		/// </summary>
-		private ModuleRules CreateModuleRulesAndSetDefaults(string ModuleName, string ReferenceChain)
+		private ModuleRules CreateModuleRulesAndSetDefaults(string ModuleName, string ReferenceChain, ILogger Logger)
 		{
 			// Create the rules from the assembly
-			ModuleRules RulesObject = RulesAssembly.CreateModuleRules(ModuleName, Rules, ReferenceChain);
+			ModuleRules RulesObject = RulesAssembly.CreateModuleRules(ModuleName, Rules, ReferenceChain, Logger);
 
 			// Set whether the module requires an IMPLEMENT_MODULE macro
 			if (!RulesObject.bRequiresImplementModule.HasValue)
@@ -4207,7 +4218,7 @@ namespace UnrealBuildTool
 				{
 					if (ProjectFile == null || !RulesObject.File.IsUnderDirectory(ProjectFile.Directory))
 					{
-						Log.TraceWarning("{0} module has shared PCHs disabled, but does not have a private PCH set", ModuleName);
+						Logger.LogWarning("{ModuleName} module has shared PCHs disabled, but does not have a private PCH set", ModuleName);
 					}
 				}
 
@@ -4284,14 +4295,15 @@ namespace UnrealBuildTool
 		/// </summary>
 		/// <param name="ModuleName">Name of the module</param>
 		/// <param name="ReferenceChain">Chain of references causing this module to be instantiated, for display in error messages</param>
-		public UEBuildModule FindOrCreateModuleByName(string ModuleName, string ReferenceChain)
+		/// <param name="Logger">Logger for output</param>
+		public UEBuildModule FindOrCreateModuleByName(string ModuleName, string ReferenceChain, ILogger Logger)
 		{
 			UEBuildModule? Module;
 			if (!Modules.TryGetValue(ModuleName, out Module))
 			{
 				// @todo projectfiles: Cross-platform modules can appear here during project generation, but they may have already
 				//   been filtered out by the project generator.  This causes the projects to not be added to directories properly.
-				ModuleRules RulesObject = CreateModuleRulesAndSetDefaults(ModuleName, ReferenceChain);
+				ModuleRules RulesObject = CreateModuleRulesAndSetDefaults(ModuleName, ReferenceChain, Logger);
 				DirectoryReference ModuleDirectory = RulesObject.File.Directory;
 
 				// Clear the bUsePrecompiled flag if we're compiling a foreign plugin; since it's treated like an engine module, it will default to true in an installed build.
@@ -4387,20 +4399,20 @@ namespace UnrealBuildTool
 				UEBuildPlatform.PlatformModifyHostModuleRules(ModuleName, RulesObject, Rules);
 
 				// Now, go ahead and create the module builder instance
-				Module = InstantiateModule(RulesObject, GeneratedCodeDirectory);
+				Module = InstantiateModule(RulesObject, GeneratedCodeDirectory, Logger);
 				Modules.Add(Module.Name, Module);
 
 				// Module must not have Verse if Verse is not enabled
 				if (Module.bHasVerse && !Rules.bUseVerse && !Rules.bBuildAllModules && !ProjectFileGenerator.bGenerateProjectFiles)
 				{
-					Log.TraceInformation("Module '{0}' has associated Verse code but target '{1}' does not have Verse enabled. C++ include errors are likely to follow.", Module.Name, TargetName);
+					Logger.LogInformation("Module '{ModuleName}' has associated Verse code but target '{TargetName}' does not have Verse enabled. C++ include errors are likely to follow.", Module.Name, TargetName);
 				}
 			}
 
 			// Warn if the module reference has incorrect text case
 			if (!Module.Name.Equals(ModuleName) && Module.Name.Equals(ModuleName, StringComparison.InvariantCultureIgnoreCase))
 			{
-				Log.TraceWarning("Module '{0}' (referenced via {1}) has incorrect text case. Did you mean '{2}'?", ModuleName, ReferenceChain, Module.Name);
+				Logger.LogWarning("Module '{ModuleName}' (referenced via {ReferenceChain}) has incorrect text case. Did you mean '{PossibleModuleName}'?", ModuleName, ReferenceChain, Module.Name);
 			}
 			return Module;
 		}
@@ -4410,10 +4422,11 @@ namespace UnrealBuildTool
 		/// </summary>
 		/// <param name="ModuleName">Name of the module</param>
 		/// <param name="ReferenceChain">Chain of references causing this module to be instantiated, for display in error messages</param>
+		/// <param name="Logger">Logger for output</param>
 		/// <returns>New C++ module</returns>
-		public UEBuildModuleCPP FindOrCreateCppModuleByName(string ModuleName, string ReferenceChain)
+		public UEBuildModuleCPP FindOrCreateCppModuleByName(string ModuleName, string ReferenceChain, ILogger Logger)
 		{
-			UEBuildModuleCPP? CppModule = FindOrCreateModuleByName(ModuleName, ReferenceChain) as UEBuildModuleCPP;
+			UEBuildModuleCPP? CppModule = FindOrCreateModuleByName(ModuleName, ReferenceChain, Logger) as UEBuildModuleCPP;
 			if (CppModule == null)
 			{
 				throw new BuildException("'{0}' is not a C++ module (referenced via {1})", ModuleName, ReferenceChain);
@@ -4423,7 +4436,8 @@ namespace UnrealBuildTool
 
 		protected UEBuildModule InstantiateModule(
 			ModuleRules RulesObject,
-			DirectoryReference? GeneratedCodeDirectory)
+			DirectoryReference? GeneratedCodeDirectory,
+			ILogger Logger)
 		{
 			switch (RulesObject.Type)
 			{
@@ -4431,11 +4445,12 @@ namespace UnrealBuildTool
 					return new UEBuildModuleCPP(
 							Rules: RulesObject,
 							IntermediateDirectory: GetModuleIntermediateDirectory(RulesObject),
-							GeneratedCodeDirectory: GeneratedCodeDirectory
+							GeneratedCodeDirectory: GeneratedCodeDirectory,
+							Logger
 						);
 
 				case ModuleRules.ModuleType.External:
-					return new UEBuildModuleExternal(RulesObject, GetModuleIntermediateDirectory(RulesObject));
+					return new UEBuildModuleExternal(RulesObject, GetModuleIntermediateDirectory(RulesObject), Logger);
 
 				default:
 					throw new BuildException("Unrecognized module type specified by 'Rules' object {0}", RulesObject.ToString());
diff --git a/Engine/Source/Programs/UnrealBuildTool/Executors/Experimental/FASTBuild.cs b/Engine/Source/Programs/UnrealBuildTool/Executors/Experimental/FASTBuild.cs
--- a/Engine/Source/Programs/UnrealBuildTool/Executors/Experimental/FASTBuild.cs
+++ b/Engine/Source/Programs/UnrealBuildTool/Executors/Experimental/FASTBuild.cs
@@ -23,6 +23,7 @@
 using System.Text.RegularExpressions;
 using UnrealBuildBase;
 using System.Runtime.Versioning;
+using Microsoft.Extensions.Logging;
 
 namespace UnrealBuildTool
 {
@@ -156,11 +157,11 @@
 		/// <summary>
 		/// Constructor
 		/// </summary>
-		public FASTBuild(int MaxLocalActions, bool bCompactOutput)
+		public FASTBuild(int MaxLocalActions, bool bCompactOutput, ILogger Logger)
 		{
 			XmlConfig.ApplyTo(this);
 
-			this.LocalExecutor = new ParallelExecutor(MaxLocalActions, bCompactOutput);
+			this.LocalExecutor = new ParallelExecutor(MaxLocalActions, bCompactOutput, Logger);
 		}
 
 		public override string Name
@@ -229,7 +230,7 @@
 			return FBuildCoordinator;
 		}
 
-		public static bool IsAvailable()
+		public static bool IsAvailable(ILogger Logger)
 		{
 			string? ExecutablePath = GetExecutablePath();
 			if (BuildHostPlatform.Current.Platform == UnrealTargetPlatform.Mac)
@@ -270,7 +271,7 @@
 				if (File.Exists(FBuildExecutablePath))
 					return true;
 
-				Log.TraceWarning($"FBuildExecutablePath '{FBuildExecutablePath}' doesn't exist! Attempting to find executable in PATH.");
+				Logger.LogWarning("FBuildExecutablePath '{FBuildExecutablePath}' doesn't exist! Attempting to find executable in PATH.", FBuildExecutablePath);
 			}
 
 			// Get the name of the FASTBuild executable.
@@ -298,7 +299,7 @@
 				}
 			}
 
-			Log.TraceError("FASTBuild disabled. Unable to find any executable to use.");
+			Logger.LogError("FASTBuild disabled. Unable to find any executable to use.");
 			return false;
 		}
 
@@ -389,7 +390,7 @@
 			),
 		};
 
-		private bool DetectBuildType(IEnumerable<LinkedAction> Actions)
+		private bool DetectBuildType(IEnumerable<LinkedAction> Actions, ILogger Logger)
 		{
 			foreach (LinkedAction Action in Actions)
 			{
@@ -398,16 +399,16 @@
 					if (BuildTypeSearchParam.Item2(Action).Contains(BuildTypeSearchParam.Item1))
 					{
 						BuildType = BuildTypeSearchParam.Item3;
-						Log.TraceInformation($"Detected build type as {BuildTypeSearchParam.Item1.ToString()} from '{BuildTypeSearchParam.Item2(Action)}' using search term '{BuildTypeSearchParam.Item1}'");
+						Logger.LogInformation("Detected build type as {Type} from '{From}' using search term '{Term}'", BuildTypeSearchParam.Item1.ToString(), BuildTypeSearchParam.Item2(Action), BuildTypeSearchParam.Item1);
 						return true;
 					}
 				}
 			}
 
-			Log.TraceError("Couldn't detect build type from actions! Unsupported platform?");
+			Logger.LogError("Couldn't detect build type from actions! Unsupported platform?");
 			foreach (LinkedAction Action in Actions)
 			{
-				PrintActionDetails(Action);
+				PrintActionDetails(Action, Logger);
 			}
 			return false;
 		}
@@ -426,7 +427,7 @@
 		}
 
 		[SupportedOSPlatform("windows")]
-		public override bool ExecuteActions(List<LinkedAction> Actions)
+		public override bool ExecuteActions(List<LinkedAction> Actions, ILogger Logger)
 		{
 			if (Actions.Count <= 0)
 				return true;
@@ -454,7 +455,7 @@
 
 			if (PreCompileActions.Any())
 			{
-				bool bResult = LocalExecutor.ExecuteActions(PreCompileActions);
+				bool bResult = LocalExecutor.ExecuteActions(PreCompileActions, Logger);
 
 				if (!bResult)
 					return false;
@@ -465,14 +466,14 @@
 
 			if (CompileActions.Any())
 			{
-				if (!DetectBuildType(CompileActions))
+				if (!DetectBuildType(CompileActions, Logger))
 					return false;
 
 				string FASTBuildFilePath = Path.Combine(Unreal.EngineDirectory.FullName, "Intermediate", "Build", "fbuild.bff");
-				if (!CreateBffFile(CompileActions, FASTBuildFilePath))
+				if (!CreateBffFile(CompileActions, FASTBuildFilePath, Logger))
 					return false;
 
-				if (!ExecuteBffFile(FASTBuildFilePath))
+				if (!ExecuteBffFile(FASTBuildFilePath, Logger))
 					return false;
 			}
 
@@ -483,7 +484,7 @@
 
 			if (PostCompileActions.Any())
 			{
-				bool bResult = LocalExecutor.ExecuteActions(PostCompileActions);
+				bool bResult = LocalExecutor.ExecuteActions(PostCompileActions, Logger);
 
 				if (!bResult)
 					return false;
@@ -515,7 +516,7 @@
 				.Replace("$(CommonProgramFiles)", "$CommonProgramFiles$");
 		}
 
-		private Dictionary<string, string> ParseCommandLineOptions(string LocalToolName, string CompilerCommandLine, string[] SpecialOptions, bool SaveResponseFile = false)
+		private Dictionary<string, string> ParseCommandLineOptions(string LocalToolName, string CompilerCommandLine, string[] SpecialOptions, ILogger Logger, bool SaveResponseFile = false)
 		{
 			Dictionary<string, string> ParsedCompilerOptions = new Dictionary<string, string>();
 
@@ -575,8 +576,8 @@
 				catch (Exception e)
 				{
 					if (!string.IsNullOrEmpty(e.Message))
-						Log.TraceInformation(e.Message);
-					Log.TraceError("Looks like a response file in: " + CompilerCommandLine + ", but we could not load it! " + e.Message);
+						Logger.LogInformation("{Message}", e.Message);
+					Logger.LogError("Looks like a response file in: {CompilerCommandLine}, but we could not load it! {Ex}", CompilerCommandLine, e.Message);
 					ResponseFilePath = "";
 				}
 			}
@@ -645,7 +646,7 @@
 						}
 						if (!AddedToken)
 						{
-							Log.TraceWarning("Warning! Looks like an unterminated string in tokens. Adding PartialToken and hoping for the best. Command line: " + CompilerCommandLine);
+							Logger.LogWarning("Warning! Looks like an unterminated string in tokens. Adding PartialToken and hoping for the best. Command line: {CompilerCommandLine}", CompilerCommandLine);
 							ProcessedTokens.Add(PartialToken);
 						}
 					}
@@ -766,7 +767,7 @@
 			return ParsedCompilerOptions;
 		}
 
-		private string GetOptionValue(Dictionary<string, string> OptionsDictionary, string Key, LinkedAction Action, bool ProblemIfNotFound = false)
+		private string GetOptionValue(Dictionary<string, string> OptionsDictionary, string Key, LinkedAction Action, ILogger Logger, bool ProblemIfNotFound = false)
 		{
 			string? Value = string.Empty;
 			if (OptionsDictionary.TryGetValue(Key, out Value))
@@ -776,8 +777,8 @@
 
 			if (ProblemIfNotFound)
 			{
-				Log.TraceWarning("We failed to find " + Key + ", which may be a problem.");
-				Log.TraceWarning("Action.CommandArguments: " + Action.CommandArguments);
+				Logger.LogWarning("We failed to find {Key}, which may be a problem.", Key);
+				Logger.LogWarning("Action.CommandArguments: {CommandArguments}", Action.CommandArguments);
 			}
 
 			return String.Empty;
@@ -806,7 +807,7 @@
 		}
 
 		[SupportedOSPlatform("windows")]
-		private void WriteEnvironmentSetup()
+		private void WriteEnvironmentSetup(ILogger Logger)
 		{
 			VCEnvironment? VCEnv = null;
 
@@ -816,12 +817,12 @@
 				// it probably means we are building for another platform.
                 if(BuildType == FBBuildType.Windows)
                 {
-					VCEnv = VCEnvironment.Create(WindowsPlatform.GetDefaultCompiler(null, WindowsArchitecture.x64), UnrealTargetPlatform.Win64, WindowsArchitecture.x64, null, null, null);
+					VCEnv = VCEnvironment.Create(WindowsPlatform.GetDefaultCompiler(null, WindowsArchitecture.x64, Logger), UnrealTargetPlatform.Win64, WindowsArchitecture.x64, null, null, null, Logger);
 				}
             }
 			catch (Exception)
 			{
-				Log.TraceWarning("Failed to get Visual Studio environment.");
+				Logger.LogWarning("Failed to get Visual Studio environment.");
 			}
 
 			// Copy environment into a case-insensitive dictionary for easier key lookups
@@ -870,7 +871,7 @@
 
 					default:
 						string exceptionString = "Error: Unsupported Visual Studio Version.";
-						Log.TraceError(exceptionString);
+						Logger.LogError("{Ex}", exceptionString);
 						throw new BuildException(exceptionString);
 				}
 
@@ -969,13 +970,13 @@
 
 					if (MsvcCRTRedistVersion.Length > 0)
 					{
-						Log.TraceInformation("Couldn't find redist path for given MsvcCRTRedistVersion {" + MsvcCRTRedistVersion.ToString()
-						+ "} (in BuildConfiguration.xml). \n\t...Using this path instead: {" + PrefferedMSVCRedistPath.ToString() + "}");
+						Logger.LogInformation("Couldn't find redist path for given MsvcCRTRedistVersion {MsvcCRTRedistVersion}"
+						+ " (in BuildConfiguration.xml). \n\t...Using this path instead: {PrefferedMSVCRedistPath}", MsvcCRTRedistVersion, PrefferedMSVCRedistPath);
 					}
 					else
 					{
-						Log.TraceInformation("Using path : {" + PrefferedMSVCRedistPath.ToString() + "} for vccorlib_.dll (MSVC redist)..." +
-							"\n\t...Add an entry for MsvcCRTRedistVersion in BuildConfiguration.xml to specify a version number");
+						Logger.LogInformation("Using path : {PrefferedMSVCRedistPath} for vccorlib_.dll (MSVC redist)..." +
+							"\n\t...Add an entry for MsvcCRTRedistVersion in BuildConfiguration.xml to specify a version number", PrefferedMSVCRedistPath.ToString());
 					}
 
 				}
@@ -1057,7 +1058,7 @@
 			AddText("}\n\n"); //End Settings
 		}
 
-		private void AddCompileAction(LinkedAction Action, IEnumerable<LinkedAction> DependencyActions)
+		private void AddCompileAction(LinkedAction Action, IEnumerable<LinkedAction> DependencyActions, ILogger Logger)
 		{
 			string CompilerName = GetCompilerName();
 			if (Action.CommandPath.FullName.Contains("rc.exe"))
@@ -1066,13 +1067,13 @@
 			}
 
 			string[] SpecialCompilerOptions = { "/Fo", "/fo", "/Yc", "/Yu", "/Fp", "-o", "-dependencies=", "-compiler=" };
-			var ParsedCompilerOptions = ParseCommandLineOptions(Action.CommandPath.GetFileName(), Action.CommandArguments, SpecialCompilerOptions);
+			var ParsedCompilerOptions = ParseCommandLineOptions(Action.CommandPath.GetFileName(), Action.CommandArguments, SpecialCompilerOptions, Logger);
 
-			string OutputObjectFileName = GetOptionValue(ParsedCompilerOptions, IsMSVC() ? "/Fo" : "-o", Action, ProblemIfNotFound: !IsMSVC());
+			string OutputObjectFileName = GetOptionValue(ParsedCompilerOptions, IsMSVC() ? "/Fo" : "-o", Action, Logger, ProblemIfNotFound: !IsMSVC());
 
 			if (IsMSVC() && string.IsNullOrEmpty(OutputObjectFileName)) // Didn't find /Fo, try /fo
 			{
-				OutputObjectFileName = GetOptionValue(ParsedCompilerOptions, "/fo", Action, ProblemIfNotFound: true);
+				OutputObjectFileName = GetOptionValue(ParsedCompilerOptions, "/fo", Action, Logger, ProblemIfNotFound: true);
 			}
 
 			if (string.IsNullOrEmpty(OutputObjectFileName)) //No /Fo or /fo, we're probably in trouble.
@@ -1088,7 +1089,7 @@
 
 			IntermediatePath = IsApple() ? IntermediatePath.Replace("\\", "/") : IntermediatePath;
 
-			string InputFile = GetOptionValue(ParsedCompilerOptions, "InputFile", Action, ProblemIfNotFound: true);
+			string InputFile = GetOptionValue(ParsedCompilerOptions, "InputFile", Action, Logger, ProblemIfNotFound: true);
 			if (string.IsNullOrEmpty(InputFile))
 			{
 				throw new Exception("We have no InputFile. Bailing. Our Action.CommandArguments were: " + Action.CommandArguments);
@@ -1104,7 +1105,7 @@
 				AddText("\t.AllowDistribution = false\n");
 			}
 
-			string OtherCompilerOptions = GetOptionValue(ParsedCompilerOptions, "OtherOptions", Action);
+			string OtherCompilerOptions = GetOptionValue(ParsedCompilerOptions, "OtherOptions", Action, Logger);
 			string CompilerOutputExtension = ".unset";
 			string CLFilterParams = "";
 			string ShowIncludesParam = "";
@@ -1116,8 +1117,8 @@
 
 			if (ParsedCompilerOptions.ContainsKey("/Yc")) //Create PCH
 			{
-				string PCHIncludeHeader = GetOptionValue(ParsedCompilerOptions, "/Yc", Action, ProblemIfNotFound: true);
-				string PCHOutputFile = GetOptionValue(ParsedCompilerOptions, "/Fp", Action, ProblemIfNotFound: true);
+				string PCHIncludeHeader = GetOptionValue(ParsedCompilerOptions, "/Yc", Action, Logger, ProblemIfNotFound: true);
+				string PCHOutputFile = GetOptionValue(ParsedCompilerOptions, "/Fp", Action, Logger, ProblemIfNotFound: true);
 
 				AddText($"\t.CompilerOptions = '{CLFilterParams}\"%1\" /Fo\"%2\" /Fp\"{PCHOutputFile}\" /Yu\"{PCHIncludeHeader}\" {OtherCompilerOptions} '\n");
 
@@ -1128,8 +1129,8 @@
 			}
 			else if (ParsedCompilerOptions.ContainsKey("/Yu")) //Use PCH
 			{
-				string PCHIncludeHeader = GetOptionValue(ParsedCompilerOptions, "/Yu", Action, ProblemIfNotFound: true);
-				string PCHOutputFile = GetOptionValue(ParsedCompilerOptions, "/Fp", Action, ProblemIfNotFound: true);
+				string PCHIncludeHeader = GetOptionValue(ParsedCompilerOptions, "/Yu", Action, Logger, ProblemIfNotFound: true);
+				string PCHOutputFile = GetOptionValue(ParsedCompilerOptions, "/Fp", Action, Logger, ProblemIfNotFound: true);
 				string PCHToForceInclude = PCHOutputFile.Replace(".pch", "");
 				AddText($"\t.CompilerOptions = '{CLFilterParams}\"%1\" /Fo\"%2\" /Fp\"{PCHOutputFile}\" /Yu\"{PCHIncludeHeader}\" /FI\"{PCHToForceInclude}\" {OtherCompilerOptions} {ShowIncludesParam} '\n");
 				string InputFileExt = Path.GetExtension(InputFile);
@@ -1172,7 +1173,7 @@
 			AddText("}\n\n");
 		}
 
-		private void AddLinkAction(LinkedAction Action, IEnumerable<LinkedAction> DependencyActions)
+		private void AddLinkAction(LinkedAction Action, IEnumerable<LinkedAction> DependencyActions, ILogger Logger)
 		{
 			string[] SpecialLinkerOptions = IsApple() ? new string[] { "-o" } : new string[] { "/OUT:", "@", "-o" };
 
@@ -1186,31 +1187,31 @@
 				CommandArgsToParse = SplitCommandArgs[0];
 			}
 
-			Dictionary<string, string> ParsedLinkerOptions = ParseCommandLineOptions(Action.CommandPath.GetFileName(), CommandArgsToParse, SpecialLinkerOptions, SaveResponseFile: true);
+			Dictionary<string, string> ParsedLinkerOptions = ParseCommandLineOptions(Action.CommandPath.GetFileName(), CommandArgsToParse, SpecialLinkerOptions, Logger, SaveResponseFile: true);
 
 			string OutputFile;
 
 			if (IsMSVC())
 			{
-				OutputFile = GetOptionValue(ParsedLinkerOptions, "/OUT:", Action, ProblemIfNotFound: true);
+				OutputFile = GetOptionValue(ParsedLinkerOptions, "/OUT:", Action, Logger, ProblemIfNotFound: true);
 			}
 			else // Apple
 			{
-				OutputFile = GetOptionValue(ParsedLinkerOptions, "-o", Action, ProblemIfNotFound: false);
+				OutputFile = GetOptionValue(ParsedLinkerOptions, "-o", Action, Logger, ProblemIfNotFound: false);
 				if (string.IsNullOrEmpty(OutputFile))
 				{
-					OutputFile = GetOptionValue(ParsedLinkerOptions, "InputFile", Action, ProblemIfNotFound: true);
+					OutputFile = GetOptionValue(ParsedLinkerOptions, "InputFile", Action, Logger, ProblemIfNotFound: true);
 				}
 			}
 
 			if (string.IsNullOrEmpty(OutputFile))
 			{
-				Log.TraceError("Failed to find output file. Bailing.");
+				Logger.LogError("Failed to find output file. Bailing.");
 				return;
 			}
 
-			string ResponseFilePath = GetOptionValue(ParsedLinkerOptions, "@", Action);
-			string OtherCompilerOptions = GetOptionValue(ParsedLinkerOptions, "OtherOptions", Action);
+			string ResponseFilePath = GetOptionValue(ParsedLinkerOptions, "@", Action, Logger);
+			string OtherCompilerOptions = GetOptionValue(ParsedLinkerOptions, "OtherOptions", Action, Logger);
 
 			IEnumerable<LinkedAction>? PrebuildDependencies = null;
 
@@ -1270,7 +1271,7 @@
 				}
 				else
 				{
-					AddText(string.Format("\t.LibrarianAdditionalInputs = {{ '{0}' }} \n", GetOptionValue(ParsedLinkerOptions, "InputFile", Action, ProblemIfNotFound: true)));
+					AddText(string.Format("\t.LibrarianAdditionalInputs = {{ '{0}' }} \n", GetOptionValue(ParsedLinkerOptions, "InputFile", Action, Logger, ProblemIfNotFound: true)));
 				}
 
 				AddText($"\t.LibrarianOutput = '{OutputFile}' \n");
@@ -1316,7 +1317,7 @@
 				AddText($"Executable('{ActionToActionString(Action)}')\n{{ \n");
 				AddText("\t.Linker = '$MacToolchainDir$/clang++' \n");
 
-				string InputFile = GetOptionValue(ParsedLinkerOptions, "InputFile", Action, ProblemIfNotFound: true);
+				string InputFile = GetOptionValue(ParsedLinkerOptions, "InputFile", Action, Logger, ProblemIfNotFound: true);
 				if (!string.IsNullOrEmpty(InputFile))
 				{
 					LinkedAction? InputFileAction = DependencyActions
@@ -1362,23 +1363,23 @@
 			}
 			else
 			{
-				Log.TraceError("Failed to add link action!");
-				PrintActionDetails(Action);
+				Logger.LogError("Failed to add link action!");
+				PrintActionDetails(Action, Logger);
 			}
 		}
 
-		private void PrintActionDetails(LinkedAction ActionToPrint)
+		private void PrintActionDetails(LinkedAction ActionToPrint, ILogger Logger)
 		{
-			Log.TraceInformation(ActionToActionString(ActionToPrint));
-			Log.TraceInformation($"Action Type: {ActionToPrint.ActionType.ToString()}");
-			Log.TraceInformation($"Action CommandPath: {ActionToPrint.CommandPath.FullName}");
-			Log.TraceInformation($"Action CommandArgs: {ActionToPrint.CommandArguments}");
+			Logger.LogInformation("{Action}", ActionToActionString(ActionToPrint));
+			Logger.LogInformation("Action Type: {Type}", ActionToPrint.ActionType.ToString());
+			Logger.LogInformation("Action CommandPath: {Path}", ActionToPrint.CommandPath.FullName);
+			Logger.LogInformation("Action CommandArgs: {Args}", ActionToPrint.CommandArguments);
 		}
 
 		private MemoryStream? bffOutputMemoryStream = null;
 
 		[SupportedOSPlatform("windows")]
-		private bool CreateBffFile(IEnumerable<LinkedAction> Actions, string BffFilePath)
+		private bool CreateBffFile(IEnumerable<LinkedAction> Actions, string BffFilePath, ILogger Logger)
 		{
 			try
 			{
@@ -1388,7 +1389,7 @@
 				AddText(";* Autogenerated bff - see FASTBuild.cs for how this file was generated. *\n");
 				AddText(";*************************************************************************\n\n");
 
-				WriteEnvironmentSetup(); //Compiler, environment variables and base paths
+				WriteEnvironmentSetup(Logger); //Compiler, environment variables and base paths
 
 				foreach (LinkedAction Action in Actions)
 				{
@@ -1404,14 +1405,14 @@
 					switch (Action.ActionType)
 					{
 						case ActionType.Compile:
-							AddCompileAction(Action, DependencyActions);
+							AddCompileAction(Action, DependencyActions, Logger);
 							break;
 						case ActionType.Link:
-							AddLinkAction(Action, DependencyActions);
+							AddLinkAction(Action, DependencyActions, Logger);
 							break;
 						default:
-							Log.TraceWarning("FASTBuild is ignoring an unsupported action!");
-							PrintActionDetails(Action);
+							Logger.LogWarning("FASTBuild is ignoring an unsupported action!");
+							PrintActionDetails(Action, Logger);
 							break;
 					}
 				}
@@ -1437,14 +1438,14 @@
 			}
 			catch (Exception e)
 			{
-				Log.TraceError("Exception while creating bff file: " + e.ToString());
+				Logger.LogError("Exception while creating bff file: {Ex}", e.ToString());
 				return false;
 			}
 
 			return true;
 		}
 
-		private bool ExecuteBffFile(string BffFilePath)
+		private bool ExecuteBffFile(string BffFilePath, ILogger Logger)
 		{
 			string CacheArgument = "";
 
@@ -1476,7 +1477,7 @@
 			// Basically we want FB to stupidly compile what UBT tells it to.
 			string FBCommandLine	= $"-monitor -summary {DistArgument} {CacheArgument} {IDEArgument} -clean -config \"{BffFilePath}\" {NoStopOnErrorArgument} {ForceRemoteArgument}";
 
-			Log.TraceInformation($"FBuild Command Line Arguments: '{FBCommandLine}");
+			Logger.LogInformation("FBuild Command Line Arguments: '{FBCommandLine}", FBCommandLine);
 
 			string FBExecutable		= GetExecutablePath()!;
 			string WorkingDirectory	= Path.GetFullPath(Path.Combine(Unreal.EngineDirectory.MakeRelativeTo(DirectoryReference.GetCurrentDirectory()), "Source"));
@@ -1513,7 +1514,7 @@
 				DataReceivedEventHandler OutputEventHandler = (Sender, Args) =>
 				{
 					if (Args.Data != null)
-						Log.TraceInformation(Args.Data);
+						Logger.LogInformation("{Output}", Args.Data);
 				};
 
 				FBProcess.OutputDataReceived += OutputEventHandler;
@@ -1529,7 +1530,7 @@
 			}
 			catch (Exception e)
 			{
-				Log.TraceError("Exception launching fbuild process. Is it in your path?" + e.ToString());
+				Logger.LogError("Exception launching fbuild process. Is it in your path? {Ex}", e.ToString());
 				return false;
 			}
 		}
diff --git a/Engine/Source/Programs/UnrealBuildTool/Executors/Experimental/HordeExecutor.cs b/Engine/Source/Programs/UnrealBuildTool/Executors/Experimental/HordeExecutor.cs
--- a/Engine/Source/Programs/UnrealBuildTool/Executors/Experimental/HordeExecutor.cs
+++ b/Engine/Source/Programs/UnrealBuildTool/Executors/Experimental/HordeExecutor.cs
@@ -2,6 +2,7 @@
 
 using EpicGames.Core;
 using EpicGames.Serialization;
+using Microsoft.Extensions.Logging;
 using System;
 using System.Collections.Generic;
 using System.Diagnostics;
@@ -279,7 +280,8 @@ namespace UnrealBuildTool
 		/// </summary>
 		/// <param name="MaxLocalActions">How many local actions to execute in parallel</param>
 		/// <param name="bCompactOutput"></param>
-		public HordeExecutor(int MaxLocalActions, bool bCompactOutput) : base(MaxLocalActions, bCompactOutput)
+		/// <param name="Logger"></param>
+		public HordeExecutor(int MaxLocalActions, bool bCompactOutput, ILogger Logger) : base(MaxLocalActions, bCompactOutput, Logger)
 		{
 			XmlConfig.ApplyTo(this);
 		}
@@ -332,7 +334,7 @@ namespace UnrealBuildTool
 			}
 		}
 
-		private static bool LaunchZen(ManagedProcessGroup ProcessGroup, [System.Diagnostics.CodeAnalysis.NotNullWhen(true)] out ManagedProcess? ZenProcess)
+		private static bool LaunchZen(ManagedProcessGroup ProcessGroup, ILogger Logger, [System.Diagnostics.CodeAnalysis.NotNullWhen(true)] out ManagedProcess? ZenProcess)
 		{
 			ZenProcess = null;
 			try
@@ -341,19 +343,19 @@ namespace UnrealBuildTool
 
 				if (!EngineConfig.GetBool("Zen", "AutoLaunch", out bool AutoLaunch) || !AutoLaunch)
 				{
-					Log.TraceWarning("Unable to launch zenserver: '[Zen] AutoLaunch' not enabled in Engine config");
+					Logger.LogWarning("Unable to launch zenserver: '[Zen] AutoLaunch' not enabled in Engine config");
 					return false;
 				}
 
 				if (!EngineConfig.GetString("Zen.AutoLaunch", "DataPath", out string DataPath) || string.IsNullOrEmpty(DataPath))
 				{
-					Log.TraceWarning("Unable to launch zenserver: '[Zen.AutoLaunch] DataPath' not set in Engine config");
+					Logger.LogWarning("Unable to launch zenserver: '[Zen.AutoLaunch] DataPath' not set in Engine config");
 					return false;
 				}
 
 				if (!EngineConfig.GetString("Zen.AutoLaunch", "ExtraArgs", out string ExtraArgs) || string.IsNullOrEmpty(ExtraArgs))
 				{
-					Log.TraceWarning("Unable to launch zenserver: '[Zen.AutoLaunch] ExtraArgs' not set in Engine config");
+					Logger.LogWarning("Unable to launch zenserver: '[Zen.AutoLaunch] ExtraArgs' not set in Engine config");
 					return false;
 				}
 
@@ -375,7 +377,7 @@ namespace UnrealBuildTool
 
 					if (CommonPath == null)
 					{
-						Log.TraceWarning("Unable to launch zenserver: data-dir root directory not found");
+						Logger.LogWarning("Unable to launch zenserver: data-dir root directory not found");
 						return false;
 					}
 
@@ -385,7 +387,7 @@ namespace UnrealBuildTool
 				DirectoryReference? DataPathReference = DirectoryReference.FromString(DataPath);
 				if (DataPathReference == null || DataPathReference.ParentDirectory == null)
 				{
-					Log.TraceWarning("Unable to launch zenserver: data-dir not found");
+					Logger.LogWarning("Unable to launch zenserver: data-dir not found");
 					return false;
 				}
 
@@ -393,13 +395,13 @@ namespace UnrealBuildTool
 				FileReference ZenSourcePath = FileReference.Combine(Unreal.EngineDirectory, "Binaries", BuildHostPlatform.Current.Platform.ToString(), $"zenserver{(RuntimeInformation.IsOSPlatform(OSPlatform.Windows) ? ".exe" : "")}");
 				FileReference ZenDestPath = FileReference.Combine(DataPathReference.ParentDirectory, "Install", ZenSourcePath.GetFileName());
 
-				Log.TraceInformation($"Copying zenserver to '{ZenDestPath}'");
+				Logger.LogInformation("Copying zenserver to '{ZenDestPath}'", ZenDestPath);
 				DirectoryReference.CreateDirectory(DataPathReference);
 				DirectoryReference.CreateDirectory(ZenDestPath.Directory);
 				FileReference.MakeWriteable(ZenDestPath);
 				FileReference.Copy(ZenSourcePath, ZenDestPath, true);
 
-				Log.TraceInformation($"Launching zenserver with data-dir '{DataPathReference.FullName}'...");
+				Logger.LogInformation("Launching zenserver with data-dir '{DataPathReference}'...", DataPathReference);
 				ZenProcess = new ManagedProcess(ProcessGroup, ZenDestPath.FullName, $"--port 1337 --data-dir \"{DataPathReference.FullName}\" {ExtraArgs}", ZenDestPath.Directory.FullName, null, ProcessPriorityClass.Normal);
 
 				return true;
@@ -415,7 +417,7 @@ namespace UnrealBuildTool
 		/// Executes the specified actions locally.
 		/// </summary>
 		/// <returns>True if all the tasks successfully executed, or false if any of them failed.</returns>
-		public override bool ExecuteActions(List<LinkedAction> InputActions)
+		public override bool ExecuteActions(List<LinkedAction> InputActions, ILogger Logger)
 		{
 			int NumCompletedActions = 0;
 			int TotalActions = InputActions.Count;
@@ -427,7 +429,7 @@ namespace UnrealBuildTool
 			{
 				const double TotalWaitTime = 10.0;
 				DateTime WaitStartTime = DateTime.Now;
-				if (LaunchZenServer && LaunchZen(ProcessGroup, out ZenProcess))
+				if (LaunchZenServer && LaunchZen(ProcessGroup, Logger, out ZenProcess))
 				{
 					while ((DateTime.Now - WaitStartTime).TotalSeconds < TotalWaitTime)
 					{
@@ -442,26 +444,27 @@ namespace UnrealBuildTool
 				{
 					if (LaunchZenServer)
 					{
-						Log.TraceError($"Unable to establish connection to zenserver after {TotalWaitTime}s");
+						Logger.LogError("Unable to establish connection to zenserver after {TotalWaitTime}s", TotalWaitTime);
 						return false;
 					}
 					Log.TraceInformationOnce("Unable to establish connection to zenserver, disabling HordeExecutor");
-					return base.ExecuteActions(InputActions);
+					return base.ExecuteActions(InputActions, Logger);
 				}
 				else
 				{
-					Log.TraceInformation($"Established connection to zenserver in {(DateTime.Now - WaitStartTime).TotalSeconds}s");
+					Logger.LogInformation("Established connection to zenserver in {Time}s", (DateTime.Now - WaitStartTime).TotalSeconds);
 				}
 			}
 
 			using SemaphoreSlim MaxProcessSemaphore = new SemaphoreSlim(ActualNumParallelProcesses, ActualNumParallelProcesses);
 			using SemaphoreSlim MaxRemoteProcessSemaphore = new SemaphoreSlim(NumRemoteParallelProcesses, NumRemoteParallelProcesses);
-			using ProgressWriter ProgressWriter = new ProgressWriter("Compiling C++ source code...", false);
+			using ProgressWriter ProgressWriter = new ProgressWriter("Compiling C++ source code...", false, Logger);
+
+			Logger.LogInformation("Building {NumActions} {Actions} with {NumLocalProcesses} local {LocalProcesses} and {NumRemoteProcesses} remote {RemoteProcesses}...", TotalActions, (TotalActions == 1) ? "action" : "actions", ActualNumParallelProcesses, (ActualNumParallelProcesses == 1) ? "process" : "processes", NumRemoteParallelProcesses, (NumRemoteParallelProcesses == 1) ? "process" : "processes");
 
-			Log.TraceInformation($"Building {TotalActions} {(TotalActions == 1 ? "action" : "actions")} with {ActualNumParallelProcesses} local {(ActualNumParallelProcesses == 1 ? "process" : "processes")} and {NumRemoteParallelProcesses} remote {(NumRemoteParallelProcesses == 1 ? "process" : "processes")} ...");
 			if (RemoteProcessOnly)
 			{
-				Log.TraceWarning("Processing all distributable actions remotely. This is not recommended unless testing remote compute.");
+				Logger.LogWarning("Processing all distributable actions remotely. This is not recommended unless testing remote compute.");
 			}
 
 			Dictionary<LinkedAction, Task<ExecuteResults>> ExecuteTasks = new Dictionary<LinkedAction, Task<ExecuteResults>>();
@@ -480,22 +483,22 @@ namespace UnrealBuildTool
 					continue;
 				}
 
-				Task<ExecuteResults> ExecuteTask = CreateRemoteExecuteTask(Action, ExecuteTasks, ProcessGroup, MaxProcessSemaphore, MaxRemoteProcessSemaphore, CancellationToken);
-				Task LogTask = ExecuteTask.ContinueWith(antecedent => LogCompletedAction(Action, antecedent, CancellationTokenSource, ProgressWriter, TotalActions, ref NumCompletedActions), CancellationToken);
+				Task<ExecuteResults> ExecuteTask = CreateRemoteExecuteTask(Action, ExecuteTasks, ProcessGroup, MaxProcessSemaphore, MaxRemoteProcessSemaphore, CancellationToken, Logger);
+				Task LogTask = ExecuteTask.ContinueWith(antecedent => LogCompletedAction(Action, antecedent, CancellationTokenSource, ProgressWriter, TotalActions, ref NumCompletedActions, Logger), CancellationToken);
 
 				ExecuteTasks.Add(Action, ExecuteTask);
 				LogTasks.Add(LogTask);
 			}
 
-			Task.Factory.ContinueWhenAll(LogTasks.ToArray(), (AntecedentTasks) => TraceSummary(ExecuteTasks, ProcessGroup), CancellationToken)
-				.ContinueWith(antecedent => TraceRemoteSummary(ExecuteTasks))
+			Task.Factory.ContinueWhenAll(LogTasks.ToArray(), (AntecedentTasks) => TraceSummary(ExecuteTasks, ProcessGroup, Logger), CancellationToken)
+				.ContinueWith(antecedent => TraceRemoteSummary(ExecuteTasks, Logger))
 				.Wait();
 
 			// Return if all tasks succeeded
 			return ExecuteTasks.Values.All(x => x.Result.ExitCode == 0);
 		}
 
-		private static void TraceStats(List<RemoteExecuteResults> RemoteExecuteResults, string Description, string StartTime, string EndTime)
+		private static void TraceStats(List<RemoteExecuteResults> RemoteExecuteResults, string Description, string StartTime, string EndTime, ILogger Logger)
 		{
 			var Results = RemoteExecuteResults.Where(x => x.Timepoints.ContainsKey(StartTime) && x.Timepoints.ContainsKey(EndTime)).ToArray();
 			if (Results.Length == 0)
@@ -503,35 +506,35 @@ namespace UnrealBuildTool
 				return;
 			}
 
-			Log.TraceInformation($"  {Description}:");
-			Log.TraceInformation($"    Total {Results.Sum(x => (x.Timepoints[EndTime] - x.Timepoints[StartTime]).TotalSeconds):0.00}");
-			Log.TraceInformation($"    Avg   {Results.Average(x => (x.Timepoints[EndTime] - x.Timepoints[StartTime]).TotalSeconds):0.00}");
-			Log.TraceInformation($"    Min   {Results.Min(x => (x.Timepoints[EndTime] - x.Timepoints[StartTime]).TotalSeconds):0.00}");
-			Log.TraceInformation($"    Max   {Results.Max(x => (x.Timepoints[EndTime] - x.Timepoints[StartTime]).TotalSeconds):0.00}");
+			Logger.LogInformation("  {Description}:", Description);
+			Logger.LogInformation("    Total {Total:0.00}", Results.Sum(x => (x.Timepoints[EndTime] - x.Timepoints[StartTime]).TotalSeconds));
+			Logger.LogInformation("    Avg   {Average:0.00}", Results.Average(x => (x.Timepoints[EndTime] - x.Timepoints[StartTime]).TotalSeconds));
+			Logger.LogInformation("    Min   {Min:0.00}", Results.Min(x => (x.Timepoints[EndTime] - x.Timepoints[StartTime]).TotalSeconds));
+			Logger.LogInformation("    Max   {Max:0.00}", Results.Max(x => (x.Timepoints[EndTime] - x.Timepoints[StartTime]).TotalSeconds));
 		}
 
-		private static void TraceRemoteSummary(Dictionary<LinkedAction, Task<ExecuteResults>> Tasks)
+		private static void TraceRemoteSummary(Dictionary<LinkedAction, Task<ExecuteResults>> Tasks, ILogger Logger)
 		{
 			// TODO: Beautify summary
 			List<ExecuteResults> LocalExecuteResults = Tasks.Values.Where(x => x.IsCompletedSuccessfully && x.Result.ExitCode == 0 && !(x.Result is RemoteExecuteResults)).Select(x => x.Result).ToList();
 			List<RemoteExecuteResults> RemoteExecuteResults = Tasks.Values.Where(x => x.IsCompletedSuccessfully && x.Result.ExitCode == 0 && x.Result is RemoteExecuteResults).Select(x => (RemoteExecuteResults)x.Result).ToList();
-			Log.TraceInformation("");
-			Log.TraceInformation($"Successful Tasks Run: Local {LocalExecuteResults.Count} Remote {RemoteExecuteResults.Count}");
-			TraceStats(RemoteExecuteResults, "ubt iohash", "ubt-queue-dispatched", "ubt-hash-files");
-			TraceStats(RemoteExecuteResults, "ubt put files", "ubt-hash-files", "ubt-put-files");
-			TraceStats(RemoteExecuteResults, "zen queue wait", "zen-queue-added", "zen-queue-dispatched");
-			TraceStats(RemoteExecuteResults, "zen put horde", "zen-storage-build-ref", "zen-storage-put-ref");
-			TraceStats(RemoteExecuteResults, "horde queue wait", "horde-queue-added", "horde-queue-dispatched");
-			TraceStats(RemoteExecuteResults, "horde prepare sandbox", "horde-execution-start", "horde-execution-download-input");
-			TraceStats(RemoteExecuteResults, "horde execute", "horde-execution-download-input", "horde-execution-execute");
-			TraceStats(RemoteExecuteResults, "horde upload results", "horde-execution-execute", "horde-execution-upload-ref");
-			TraceStats(RemoteExecuteResults, "zen queue results wait", "zen-complete-queue-added", "zen-complete-queue-dispatched");
-			TraceStats(RemoteExecuteResults, "zen get horde", "zen-complete-queue-dispatched", "zen-storage-get-blobs");
-			TraceStats(RemoteExecuteResults, "zen finalize", "zen-storage-get-blobs", "zen-queue-complete");
-			Log.TraceInformation("");
+			Logger.LogInformation("");
+			Logger.LogInformation("Successful Tasks Run: Local {LocalCount} Remote {RemoteCount}", LocalExecuteResults.Count, RemoteExecuteResults.Count);
+			TraceStats(RemoteExecuteResults, "ubt iohash", "ubt-queue-dispatched", "ubt-hash-files", Logger);
+			TraceStats(RemoteExecuteResults, "ubt put files", "ubt-hash-files", "ubt-put-files", Logger);
+			TraceStats(RemoteExecuteResults, "zen queue wait", "zen-queue-added", "zen-queue-dispatched", Logger);
+			TraceStats(RemoteExecuteResults, "zen put horde", "zen-storage-build-ref", "zen-storage-put-ref", Logger);
+			TraceStats(RemoteExecuteResults, "horde queue wait", "horde-queue-added", "horde-queue-dispatched", Logger);
+			TraceStats(RemoteExecuteResults, "horde prepare sandbox", "horde-execution-start", "horde-execution-download-input", Logger);
+			TraceStats(RemoteExecuteResults, "horde execute", "horde-execution-download-input", "horde-execution-execute", Logger);
+			TraceStats(RemoteExecuteResults, "horde upload results", "horde-execution-execute", "horde-execution-upload-ref", Logger);
+			TraceStats(RemoteExecuteResults, "zen queue results wait", "zen-complete-queue-added", "zen-complete-queue-dispatched", Logger);
+			TraceStats(RemoteExecuteResults, "zen get horde", "zen-complete-queue-dispatched", "zen-storage-get-blobs", Logger);
+			TraceStats(RemoteExecuteResults, "zen finalize", "zen-storage-get-blobs", "zen-queue-complete", Logger);
+			Logger.LogInformation("");
 		}
 
-		private static Task<ExecuteResults> CreateRemoteExecuteTask(LinkedAction Action, Dictionary<LinkedAction, Task<ExecuteResults>> ExecuteTasks, ManagedProcessGroup ProcessGroup, SemaphoreSlim MaxProcessSemaphore, SemaphoreSlim MaxRemoteProcessSemaphore, CancellationToken CancellationToken)
+		private static Task<ExecuteResults> CreateRemoteExecuteTask(LinkedAction Action, Dictionary<LinkedAction, Task<ExecuteResults>> ExecuteTasks, ManagedProcessGroup ProcessGroup, SemaphoreSlim MaxProcessSemaphore, SemaphoreSlim MaxRemoteProcessSemaphore, CancellationToken CancellationToken, ILogger Logger)
 		{
 			Task<ExecuteResults> ActionTask;
 			if (Action.PrerequisiteActions.Count == 0)
@@ -552,7 +555,7 @@ namespace UnrealBuildTool
 				{
 					if (!ExecuteTasks.ContainsKey(PrerequisiteAction))
 					{
-						ExecuteTasks.Add(PrerequisiteAction, CreateRemoteExecuteTask(PrerequisiteAction, ExecuteTasks, ProcessGroup, MaxProcessSemaphore, MaxRemoteProcessSemaphore, CancellationToken));
+						ExecuteTasks.Add(PrerequisiteAction, CreateRemoteExecuteTask(PrerequisiteAction, ExecuteTasks, ProcessGroup, MaxProcessSemaphore, MaxRemoteProcessSemaphore, CancellationToken, Logger));
 					}
 					PrerequisiteTasks.Add(ExecuteTasks[PrerequisiteAction]);
 				}
@@ -574,7 +577,7 @@ namespace UnrealBuildTool
 					{
 						return Task.FromResult(ancendent.Result);
 					}
-					Log.TraceInformation($"Remote execute for {(Action.CommandDescription ?? Action.CommandPath.GetFileNameWithoutExtension())} {Action.StatusDescription} failed, rescheduling for local execution.");
+					Logger.LogInformation("Remote execute for {CommandDescription} {StatusDescription} failed, rescheduling for local execution.", (Action.CommandDescription ?? Action.CommandPath.GetFileNameWithoutExtension()), Action.StatusDescription);
 					return ExecuteAction(Array.Empty<Task<ExecuteResults>>(), Action, ProcessGroup, MaxProcessSemaphore, CancellationToken);
 				},
 				CancellationToken,
diff --git a/Engine/Source/Programs/UnrealBuildTool/Executors/HybridExecutor.cs b/Engine/Source/Programs/UnrealBuildTool/Executors/HybridExecutor.cs
--- a/Engine/Source/Programs/UnrealBuildTool/Executors/HybridExecutor.cs
+++ b/Engine/Source/Programs/UnrealBuildTool/Executors/HybridExecutor.cs
@@ -1,6 +1,7 @@
 // Copyright Epic Games, Inc. All Rights Reserved.
 
 using EpicGames.Core;
+using Microsoft.Extensions.Logging;
 using System.Collections.Generic;
 using System.Linq;
 
@@ -31,11 +32,11 @@
 		/// <summary>
 		/// Constructor
 		/// </summary>
-		public HybridExecutor(List<TargetDescriptor> TargetDescriptors, int InMaxLocalActions, bool bCompactOutput, ActionExecutor? InLocalExecutor = null, ActionExecutor? InRemoteExecutor = null)
+		public HybridExecutor(List<TargetDescriptor> TargetDescriptors, int InMaxLocalActions, bool bCompactOutput, ILogger Logger, ActionExecutor? InLocalExecutor = null, ActionExecutor? InRemoteExecutor = null)
 		{
 			MaxLocalActions = InMaxLocalActions;
-			LocalExecutor = InLocalExecutor ?? new ParallelExecutor(MaxLocalActions, bCompactOutput);
-			RemoteExecutor = InRemoteExecutor ?? (XGE.IsAvailable() ? (ActionExecutor)new XGE() : new SNDBS(TargetDescriptors));
+			LocalExecutor = InLocalExecutor ?? new ParallelExecutor(MaxLocalActions, bCompactOutput, Logger);
+			RemoteExecutor = InRemoteExecutor ?? (XGE.IsAvailable(Logger) ? (ActionExecutor)new XGE() : new SNDBS(TargetDescriptors));
 
 			XmlConfig.ApplyTo(this);
 
@@ -49,9 +50,9 @@
 		/// Tests whether this executor can be used
 		/// </summary>
 		/// <returns>True if the executor may be used</returns>
-		public static bool IsAvailable()
+		public static bool IsAvailable(ILogger Logger)
 		{
-			return (XGE.IsAvailable() || SNDBS.IsAvailable()) && ParallelExecutor.IsAvailable();
+			return (XGE.IsAvailable(Logger) || SNDBS.IsAvailable(Logger)) && ParallelExecutor.IsAvailable();
 		}
 
 		/// <summary>
@@ -63,8 +64,9 @@
 		/// Execute the given actions
 		/// </summary>
 		/// <param name="ActionsToExecute">Actions to be executed</param>
+		/// <param name="Logger">Logger for output</param>
 		/// <returns>True if the build succeeded, false otherwise</returns>
-		public override bool ExecuteActions(List<LinkedAction> ActionsToExecute)
+		public override bool ExecuteActions(List<LinkedAction> ActionsToExecute, ILogger Logger)
 		{
 			// Find the number of dependants for each action
 			Dictionary<LinkedAction, int> ActionToNumDependents = ActionsToExecute.ToDictionary(x => x, x => 0);
@@ -125,13 +127,13 @@
 			}
 
 			// Execute the remote actions
-			if (RemoteActionsToExecute.Count > 0 && !RemoteExecutor.ExecuteActions(RemoteActionsToExecute))
+			if (RemoteActionsToExecute.Count > 0 && !RemoteExecutor.ExecuteActions(RemoteActionsToExecute, Logger))
 			{
 				return false;
 			}
 
 			// Pass all the local actions through to the parallel executor
-			if (LocalActionsToExecute.Count > 0 && !LocalExecutor.ExecuteActions(LocalActionsToExecute))
+			if (LocalActionsToExecute.Count > 0 && !LocalExecutor.ExecuteActions(LocalActionsToExecute, Logger))
 			{
 				return false;
 			}
