UBT: Use ILogger throughout UBT, rather than legacy EpicGames.Core.Log methods. All output going forwards should use structured logging rather than String.Format style logging.

diff --git a/Engine/Source/Programs/UnrealBuildTool/Executors/ParallelExecutor.cs b/Engine/Source/Programs/UnrealBuildTool/Executors/ParallelExecutor.cs
--- a/Engine/Source/Programs/UnrealBuildTool/Executors/ParallelExecutor.cs
+++ b/Engine/Source/Programs/UnrealBuildTool/Executors/ParallelExecutor.cs
@@ -1,6 +1,7 @@
 // Copyright Epic Games, Inc. All Rights Reserved.
 
 using EpicGames.Core;
+using Microsoft.Extensions.Logging;
 using System;
 using System.Collections.Generic;
 using System.Diagnostics;
@@ -89,12 +90,12 @@
 
 		private static readonly char[] LineEndingSplit = new char[] { '\n', '\r' };
 
-		public static int GetDefaultNumParallelProcesses()
+		public static int GetDefaultNumParallelProcesses(ILogger Logger)
 		{
 			double MemoryPerActionBytesComputed = Math.Max(MemoryPerActionBytes, MemoryPerActionBytesOverride);
 			if (MemoryPerActionBytesComputed > MemoryPerActionBytes)
 			{
-				Log.TraceInformation($"Overriding MemoryPerAction with target-defined value of {MemoryPerActionBytesComputed / 1024 / 1024 / 1024} bytes");
+				Logger.LogInformation("Overriding MemoryPerAction with target-defined value of {Memory} bytes", MemoryPerActionBytesComputed / 1024 / 1024 / 1024);
 			}
 
 			return Utils.GetMaxActionsToExecuteInParallel(MaxProcessorCount, ProcessorCountMultiplier, Convert.ToInt64(MemoryPerActionBytesComputed));
@@ -105,7 +106,8 @@
 		/// </summary>
 		/// <param name="MaxLocalActions">How many actions to execute in parallel</param>
 		/// <param name="bCompactOutput">Should output be written in a compact fashion</param>
-		public ParallelExecutor(int MaxLocalActions, bool bCompactOutput)
+		/// <param name="Logger">Logger for output</param>
+		public ParallelExecutor(int MaxLocalActions, bool bCompactOutput, ILogger Logger)
 		{
 			XmlConfig.ApplyTo(this);
 
@@ -117,7 +119,7 @@
 			else
 			{
 				// Figure out how many processors to use
-				NumParallelProcesses = GetDefaultNumParallelProcesses();
+				NumParallelProcesses = GetDefaultNumParallelProcesses(Logger);
 			}
 
 			this.bCompactOutput = bCompactOutput;
@@ -166,7 +168,7 @@
 		/// Executes the specified actions locally.
 		/// </summary>
 		/// <returns>True if all the tasks successfully executed, or false if any of them failed.</returns>
-		public override bool ExecuteActions(List<LinkedAction> InputActions)
+		public override bool ExecuteActions(List<LinkedAction> InputActions, ILogger Logger)
 		{
 			int NumCompletedActions = 0;
 			int TotalActions = InputActions.Count;
@@ -174,9 +176,9 @@
 
 			using ManagedProcessGroup ProcessGroup = new ManagedProcessGroup();
 			using SemaphoreSlim MaxProcessSemaphore = new SemaphoreSlim(ActualNumParallelProcesses, ActualNumParallelProcesses);
-			using ProgressWriter ProgressWriter = new ProgressWriter("Compiling C++ source code...", false);
+			using ProgressWriter ProgressWriter = new ProgressWriter("Compiling C++ source code...", false, Logger);
 
-			Log.TraceInformation($"Building {TotalActions} {(TotalActions == 1 ? "action" : "actions")} with {ActualNumParallelProcesses} {(ActualNumParallelProcesses == 1 ? "process" : "processes")}...");
+			Logger.LogInformation("Building {NumActions} {Actions} with {NumProcesses} {Processes}...", TotalActions, (TotalActions == 1 ? "action" : "actions"), ActualNumParallelProcesses, (ActualNumParallelProcesses == 1 ? "process" : "processes"));
 
 			Dictionary<LinkedAction, Task<ExecuteResults>> ExecuteTasks = new Dictionary<LinkedAction, Task<ExecuteResults>>();
 			List<Task> LogTasks = new List<Task>();
@@ -195,13 +197,13 @@
 				}
 
 				Task<ExecuteResults> ExecuteTask = CreateExecuteTask(Action, InputActions, ExecuteTasks, ProcessGroup, MaxProcessSemaphore, CancellationToken);
-				Task LogTask = ExecuteTask.ContinueWith(antecedent => LogCompletedAction(Action, antecedent, CancellationTokenSource, ProgressWriter, TotalActions, ref NumCompletedActions), CancellationToken);
+				Task LogTask = ExecuteTask.ContinueWith(antecedent => LogCompletedAction(Action, antecedent, CancellationTokenSource, ProgressWriter, TotalActions, ref NumCompletedActions, Logger), CancellationToken);
 
 				ExecuteTasks.Add(Action, ExecuteTask);
 				LogTasks.Add(LogTask);
 			}
 
-			Task SummaryTask = Task.Factory.ContinueWhenAll(LogTasks.ToArray(), (AntecedentTasks) => TraceSummary(ExecuteTasks, ProcessGroup), CancellationToken);
+			Task SummaryTask = Task.Factory.ContinueWhenAll(LogTasks.ToArray(), (AntecedentTasks) => TraceSummary(ExecuteTasks, ProcessGroup, Logger), CancellationToken);
 			SummaryTask.Wait();
 
 			// Return if all tasks succeeded
@@ -309,7 +311,7 @@
 
 		private static int s_previousLineLength = -1;
 
-		protected void LogCompletedAction(LinkedAction Action, Task<ExecuteResults> ExecuteTask, CancellationTokenSource CancellationTokenSource, ProgressWriter ProgressWriter, int TotalActions, ref int NumCompletedActions)
+		protected void LogCompletedAction(LinkedAction Action, Task<ExecuteResults> ExecuteTask, CancellationTokenSource CancellationTokenSource, ProgressWriter ProgressWriter, int TotalActions, ref int NumCompletedActions, ILogger Logger)
 		{
 			List<string> LogLines = new List<string>();
 			int ExitCode = int.MaxValue;
@@ -350,7 +352,7 @@
 				// Cancelled
 				if (ExitCode == int.MaxValue)
 				{
-					Log.TraceInformation($"[{CompletedActions}/{TotalActions}] {Description} cancelled");
+					Logger.LogInformation("[{CompletedActions}/{TotalActions}] {Description} cancelled", CompletedActions, TotalActions, Description);
 					return;
 				}
 
@@ -363,7 +365,7 @@
 
 				if (bLogActionCommandLines)
 				{
-					Log.TraceLog($"[{CompletedActions}/{TotalActions}]{TargetDetails} Command: {Action.CommandPath} {Action.CommandArguments}");
+					Logger.LogDebug("[{CompletedActions}/{TotalActions}]{TargetDetails} Command: {CommandPath} {CommandArguments}", CompletedActions, TotalActions, TargetDetails, Action.CommandPath, Action.CommandArguments);
 				}
 
 				string CompilationTimes = "";
@@ -398,7 +400,7 @@
 
 				s_previousLineLength = message.Length;
 
-				Log.TraceInformation(message);
+				Logger.LogInformation("{Message}", message);
 				foreach (string Line in LogLines.Skip(Action.bShouldOutputStatusDescription ? 0 : 1))
 				{
 					// suppress library creation messages when writing compact output
@@ -407,7 +409,7 @@
 						continue;
 					}
 
-					Log.TraceInformation(Line);
+					Logger.LogInformation("{Message}", Line);
 
 					// Prevent overwriting of logged lines
 					s_previousLineLength = -1;
@@ -418,9 +420,9 @@
 					// BEGIN TEMPORARY TO CATCH PVS-STUDIO ISSUES
 					if (LogLines.Count == 0)
 					{
-						Log.TraceError($"{TargetDetails} {Description}: Exited with error code {ExitCode}");
-						Log.TraceInformation($"{TargetDetails} {Description}: WorkingDirectory {Action.WorkingDirectory}");
-						Log.TraceInformation($"{TargetDetails} {Description}: {Action.CommandPath} {Action.CommandArguments}");
+						Logger.LogError("{TargetDetails} {Description}: Exited with error code {ExitCode}", TargetDetails, Description, ExitCode);
+						Logger.LogInformation("{TargetDetails} {Description}: WorkingDirectory {WorkingDirectory}", TargetDetails, Description, Action.WorkingDirectory);
+						Logger.LogInformation("{TargetDetails} {Description}: {CommandPath} {CommandArguments}", TargetDetails, Description, Action.CommandPath, Action.CommandArguments);
 					}
 					// END TEMPORARY
 
@@ -436,40 +438,40 @@
 			}
 		}
 
-		protected static void TraceSummary(Dictionary<LinkedAction, Task<ExecuteResults>> Tasks, ManagedProcessGroup ProcessGroup)
+		protected static void TraceSummary(Dictionary<LinkedAction, Task<ExecuteResults>> Tasks, ManagedProcessGroup ProcessGroup, ILogger Logger)
 		{
 			if (!bShowCompilationTimes)
 			{
 				return;
 			}
 
-			Log.TraceInformation("");
+			Logger.LogInformation("");
 			if (ProcessGroup.TotalProcessorTime.Ticks > 0)
 			{
-				Log.TraceInformation($"Total CPU Time: {ProcessGroup.TotalProcessorTime.TotalSeconds} s");
-				Log.TraceInformation("");
+				Logger.LogInformation("Total CPU Time: {TotalSeconds} s", ProcessGroup.TotalProcessorTime.TotalSeconds);
+				Logger.LogInformation("");
 			}
 
 			var CompletedTasks = Tasks.Where(x => x.Value.Status == TaskStatus.RanToCompletion).OrderByDescending(x => x.Value.Result.ExecutionTime).Take(20);
 
 			if (CompletedTasks.Any())
 			{
-				Log.TraceInformation($"Compilation Time Top {CompletedTasks.Count()}");
-				Log.TraceInformation("");
+				Logger.LogInformation("Compilation Time Top {CompletedTaskCount}", CompletedTasks.Count());
+				Logger.LogInformation("");
 				foreach (var Pair in CompletedTasks)
 				{
 					string Description = $"{(Pair.Key.Inner.CommandDescription ?? Pair.Key.Inner.CommandPath.GetFileNameWithoutExtension())} {Pair.Key.Inner.StatusDescription}".Trim();
 					if (Pair.Value.Result.ProcessorTime.Ticks > 0)
 					{
-						Log.TraceInformation($"{Description} [ Wall Time {Pair.Value.Result.ExecutionTime.TotalSeconds:0.00} s / CPU Time {Pair.Value.Result.ProcessorTime.TotalSeconds:0.00} s ]");
+						Logger.LogInformation("{Description} [ Wall Time {ExecutionTime:0.00} s / CPU Time {ProcessorTime:0.00} s ]", Description, Pair.Value.Result.ExecutionTime.TotalSeconds, Pair.Value.Result.ProcessorTime.TotalSeconds);
 					}
 					else
 					{
-						Log.TraceInformation($"{Description} [ Time {Pair.Value.Result.ExecutionTime.TotalSeconds:0.00} s ]");
+						Logger.LogInformation("{Description} [ Time {ExecutionTime:0.00} s ]", Description, Pair.Value.Result.ExecutionTime.TotalSeconds);
 					}
 
 				}
-				Log.TraceInformation("");
+				Logger.LogInformation("");
 			}
 		}
 	}
@@ -482,6 +484,6 @@
 		/// <summary>
 		/// Maximum number of processes that should be used for execution
 		/// </summary>
-		public static int MaxParallelProcesses { get { return ParallelExecutor.GetDefaultNumParallelProcesses(); } }
+		public static int GetMaxParallelProcesses(ILogger Logger) => ParallelExecutor.GetDefaultNumParallelProcesses(Logger);
 	}
 }
diff --git a/Engine/Source/Programs/UnrealBuildTool/Executors/SNDBS.cs b/Engine/Source/Programs/UnrealBuildTool/Executors/SNDBS.cs
--- a/Engine/Source/Programs/UnrealBuildTool/Executors/SNDBS.cs
+++ b/Engine/Source/Programs/UnrealBuildTool/Executors/SNDBS.cs
@@ -15,6 +15,7 @@
 using System.Text.Json;
 using System.Text.RegularExpressions;
 using EpicGames.Core;
+using Microsoft.Extensions.Logging;
 using UnrealBuildBase;
 
 namespace UnrealBuildTool
@@ -123,7 +124,7 @@
 			return null;
 		}
 
-		public static bool IsHostOnVpn(string HostName)
+		public static bool IsHostOnVpn(string HostName, ILogger Logger)
 		{
 			// If there aren't any defined subnets, just early out
 			if (VpnSubnets == null || VpnSubnets.Length == 0)
@@ -164,12 +165,12 @@
 			}
 			catch (Exception Ex)
 			{
-				Log.TraceWarning("Unable to check whether host {0} is connected to VPN:\n{1}", HostName, ExceptionUtils.FormatExceptionDetails(Ex));
+				Logger.LogWarning("Unable to check whether host {HostName} is connected to VPN:\n{Ex}", HostName, ExceptionUtils.FormatExceptionDetails(Ex));
 			}
 			return false;
 		}
 
-		public static bool IsAvailable()
+		public static bool IsAvailable(ILogger Logger)
 		{
 			// Check the executable exists on disk
 			if (SCERoot == null || !File.Exists(SNDBSBuildExe))
@@ -187,7 +188,7 @@
 			if (!bAllowOverVpn && VpnSubnets != null && VpnSubnets.Length > 0)
 			{
 				string? BrokerHost;
-				if (TryGetBrokerHost(out BrokerHost) && IsHostOnVpn(BrokerHost))
+				if (TryGetBrokerHost(out BrokerHost) && IsHostOnVpn(BrokerHost, Logger))
 				{
 					return false;
 				}
@@ -196,7 +197,7 @@
 			return true;
 		}
 
-		public override bool ExecuteActions(List<LinkedAction> Actions)
+		public override bool ExecuteActions(List<LinkedAction> Actions, ILogger Logger)
 		{
 			if (Actions.Count == 0)
 				return true;
@@ -264,15 +265,15 @@
 			{
 				StartInfo.Arguments += " -k";
 			}
-			return ExecuteProcessWithProgressMarkup(StartInfo, Actions.Count);
+			return ExecuteProcessWithProgressMarkup(StartInfo, Actions.Count, Logger);
 		}
 
 		/// <summary>
 		/// Executes the process, parsing progress markup as part of the output.
 		/// </summary>
-		private bool ExecuteProcessWithProgressMarkup(ProcessStartInfo SnDbsStartInfo, int NumActions)
+		private bool ExecuteProcessWithProgressMarkup(ProcessStartInfo SnDbsStartInfo, int NumActions, ILogger Logger)
 		{
-			using (ProgressWriter Writer = new ProgressWriter("Compiling C++ source files...", false))
+			using (ProgressWriter Writer = new ProgressWriter("Compiling C++ source files...", false, Logger))
 			{
 				int NumCompletedActions = 0;
 				string CurrentStatus = "";
@@ -289,14 +290,14 @@
 
 							Text = Args.Data.Substring(ProgressMarkupPrefix.Length).Trim();
 							var ActionInfo = Text.Split(':');
-							Log.TraceInformation($"[{NumCompletedActions}/{NumActions}] {ActionInfo[0]} {ActionInfo[1]}");
+							Logger.LogInformation("[{NumCompletedActions}/{NumActions}] {ActionInfo0} {ActionInfo1}", NumCompletedActions, NumActions, ActionInfo[0], ActionInfo[1]);
 							CurrentStatus = ActionInfo[1];
 							return;
 						}
 						// Suppress redundant tool output of status we already printed (e.g., msvc cl prints compile unit name always)
 						if (!Text.Equals(CurrentStatus))
 						{
-							Log.TraceInformation(Text);
+							Logger.LogInformation("{Output}", Text);
 						}
 					}
 				};
@@ -322,7 +323,7 @@
 						LocalProcess.BeginErrorReadLine();
 					}
 
-					Log.TraceInformation("Distributing {0} action{1} to SN-DBS",
+					Logger.LogInformation("Distributing {NumAction} action{ActionS} to SN-DBS",
 						NumActions,
 						NumActions == 1 ? "" : "s");
 
diff --git a/Engine/Source/Programs/UnrealBuildTool/Executors/XGE.cs b/Engine/Source/Programs/UnrealBuildTool/Executors/XGE.cs
--- a/Engine/Source/Programs/UnrealBuildTool/Executors/XGE.cs
+++ b/Engine/Source/Programs/UnrealBuildTool/Executors/XGE.cs
@@ -18,6 +18,7 @@
 using UnrealBuildBase;
 using System.Diagnostics.CodeAnalysis;
 using System.Runtime.Versioning;
+using Microsoft.Extensions.Logging;
 
 namespace UnrealBuildTool
 {
@@ -250,7 +251,7 @@
 			return null;
 		}
 
-		public static bool IsHostOnVpn(string HostName)
+		public static bool IsHostOnVpn(string HostName, ILogger Logger)
 		{
 			if (!OperatingSystem.IsWindows())
 			{
@@ -296,12 +297,12 @@
 			}
 			catch (Exception Ex)
 			{
-				Log.TraceWarning("Unable to check whether host {0} is connected to VPN:\n{1}", HostName, ExceptionUtils.FormatExceptionDetails(Ex));
+				Logger.LogWarning("Unable to check whether host {Host} is connected to VPN:\n{Ex}", HostName, ExceptionUtils.FormatExceptionDetails(Ex));
 			}
 			return false;
 		}
 
-		public static bool IsAvailable()
+		public static bool IsAvailable(ILogger Logger)
 		{
 			if (!OperatingSystem.IsWindows())
 			{
@@ -328,7 +329,7 @@
 				}
 				catch(Exception Ex)
 				{
-					Log.TraceLog("Unable to query for status of Incredibuild service: {0}", ExceptionUtils.FormatExceptionDetails(Ex));
+					Logger.LogDebug("Unable to query for status of Incredibuild service: {Ex}", ExceptionUtils.FormatExceptionDetails(Ex));
 					return false;
 				}
 			}
@@ -337,7 +338,7 @@
 			if (!bAllowOverVpn && VpnSubnets != null && VpnSubnets.Length > 0)
 			{
 				string? CoordinatorHost;
-				if (TryGetCoordinatorHost(out CoordinatorHost) && IsHostOnVpn(CoordinatorHost))
+				if (TryGetCoordinatorHost(out CoordinatorHost) && IsHostOnVpn(CoordinatorHost, Logger))
 				{
 					return false;
 				}
@@ -357,7 +358,7 @@
 				};
 				if (Utils.RunLocalProcess(XGEProcess) == 4)
 				{
-					Log.TraceWarning("Unable to use Incredibuild executor because a build is already in progress");
+					Logger.LogWarning("Unable to use Incredibuild executor because a build is already in progress");
 					return false;
 
 				}
@@ -369,27 +370,27 @@
 		// precompile the Regex needed to parse the XGE output (the ones we want are of the form "File (Duration at +time)"
 		//private static Regex XGEDurationRegex = new Regex(@"(?<Filename>.*) *\((?<Duration>[0-9:\.]+) at [0-9\+:\.]+\)", RegexOptions.ExplicitCapture);
 
-		public static void ExportActions(List<LinkedAction> ActionsToExecute)
+		public static void ExportActions(List<LinkedAction> ActionsToExecute, ILogger Logger)
 		{
 			for(int FileNum = 0;;FileNum++)
 			{
 				string OutFile = Path.Combine(Unreal.EngineDirectory.FullName, "Intermediate", "Build", String.Format("UBTExport.{0}.xge.xml", FileNum.ToString("D3")));
 				if(!File.Exists(OutFile))
 				{
-					ExportActions(ActionsToExecute, OutFile);
+					ExportActions(ActionsToExecute, OutFile, Logger);
 					break;
 				}
 			}
 		}
 
-		public static void ExportActions(List<LinkedAction> ActionsToExecute, string OutFile)
+		public static void ExportActions(List<LinkedAction> ActionsToExecute, string OutFile, ILogger Logger)
 		{
-			WriteTaskFile(ActionsToExecute, OutFile, ProgressWriter.bWriteMarkup, bXGEExport: true);
-			Log.TraceInformation("XGEEXPORT: Exported '{0}'", OutFile);
+			WriteTaskFile(ActionsToExecute, OutFile, ProgressWriter.bWriteMarkup, bXGEExport: true, Logger);
+			Logger.LogInformation("XGEEXPORT: Exported '{OutFile}'", OutFile);
 		}
 
 		[SupportedOSPlatform("windows")]
-		public override bool ExecuteActions(List<LinkedAction> ActionsToExecute)
+		public override bool ExecuteActions(List<LinkedAction> ActionsToExecute, ILogger Logger)
 		{
 			bool XGEResult = true;
 
@@ -403,7 +404,7 @@
 			}
 			if (ActionBatch.Count > 0 && XGEResult)
 			{
-				XGEResult = ExecuteActionBatch(ActionBatch);
+				XGEResult = ExecuteActionBatch(ActionBatch, Logger);
 				ActionBatch.Clear();
 			}
 
@@ -411,16 +412,16 @@
 		}
 
 		[SupportedOSPlatform("windows")]
-		bool ExecuteActionBatch(List<LinkedAction> Actions)
+		bool ExecuteActionBatch(List<LinkedAction> Actions, ILogger Logger)
 		{
 			bool XGEResult = true;
 			if (Actions.Count > 0)
 			{
 				// Write the actions to execute to a XGE task file.
 				string XGETaskFilePath = FileReference.Combine(Unreal.EngineDirectory, "Intermediate", "Build", "XGETasks.xml").FullName;
-				WriteTaskFile(Actions, XGETaskFilePath, true, false);
+				WriteTaskFile(Actions, XGETaskFilePath, true, false, Logger);
 
-				XGEResult = ExecuteTaskFileWithProgressMarkup(XGETaskFilePath, Actions.Count);
+				XGEResult = ExecuteTaskFileWithProgressMarkup(XGETaskFilePath, Actions.Count, Logger);
 			}
 			return XGEResult;
 		}
@@ -428,9 +429,9 @@
 		/// <summary>
 		/// Writes a XGE task file containing the specified actions to the specified file path.
 		/// </summary>
-		static void WriteTaskFile(List<LinkedAction> InActions, string TaskFilePath, bool bProgressMarkup, bool bXGEExport)
+		static void WriteTaskFile(List<LinkedAction> InActions, string TaskFilePath, bool bProgressMarkup, bool bXGEExport, ILogger Logger)
 		{
-			bool HostOnVpn = TryGetCoordinatorHost(out string? CoordinatorHost) && IsHostOnVpn(CoordinatorHost);
+			bool HostOnVpn = TryGetCoordinatorHost(out string? CoordinatorHost) && IsHostOnVpn(CoordinatorHost, Logger);
 
 			Dictionary<string, string> ExportEnv = new Dictionary<string, string>();
 
@@ -629,9 +630,10 @@
 		/// <param name="TaskFilePath">- The path to the file containing the tasks to execute in XGE XML format.</param>
 		/// <param name="OutputEventHandler"></param>
 		/// <param name="ActionCount"></param>
+		/// <param name="Logger"></param>
 		/// <returns>Indicates whether the tasks were successfully executed.</returns>
 		[SupportedOSPlatform("windows")]
-		bool ExecuteTaskFile(string TaskFilePath, DataReceivedEventHandler OutputEventHandler, int ActionCount)
+		bool ExecuteTaskFile(string TaskFilePath, DataReceivedEventHandler OutputEventHandler, int ActionCount, ILogger Logger)
 		{
 			// A bug in the UCRT can cause XGE to hang on VS2015 builds. Figure out if this hang is likely to effect this build and workaround it if able.
 			// @todo: There is a KB coming that will fix this. Once that KB is available, test if it is present. Stalls will not be a problem if it is.
@@ -702,7 +704,7 @@
 					XGEProcess.BeginErrorReadLine();
 				}
 
-				Log.TraceInformation("Distributing {0} action{1} to XGE",
+				Logger.LogInformation("Distributing {NumAction} action{ActionS} to XGE",
 					ActionCount,
 					ActionCount == 1 ? "" : "s");
 
@@ -721,9 +723,9 @@
 		/// Executes the tasks in the specified file, parsing progress markup as part of the output.
 		/// </summary>
 		[SupportedOSPlatform("windows")]
-		bool ExecuteTaskFileWithProgressMarkup(string TaskFilePath, int NumActions)
+		bool ExecuteTaskFileWithProgressMarkup(string TaskFilePath, int NumActions, ILogger Logger)
 		{
-			using (ProgressWriter Writer = new ProgressWriter("Compiling C++ source files...", false))
+			using (ProgressWriter Writer = new ProgressWriter("Compiling C++ source files...", false, Logger))
 			{
 				int NumCompletedActions = 0;
 				string ProgressText = string.Empty;
@@ -739,7 +741,7 @@
 							// Flush old progress text
 							if (!string.IsNullOrEmpty(ProgressText))
 							{
-								Log.TraceInformation($"[{NumCompletedActions}/{NumActions}] Complete {ProgressText}");
+								Logger.LogInformation("[{NumCompletedActions}/{NumActions}] Complete {ProgressText}", NumCompletedActions, NumActions, ProgressText);
 								ProgressText = string.Empty;
 							}
 							Writer.Write(++NumCompletedActions, NumActions);
@@ -752,21 +754,21 @@
 								ProgressText = Text.Trim();
 								return;
 							}
-							Log.TraceInformation($"[{NumCompletedActions}/{NumActions}] {Text}");
+							Logger.LogInformation("[{NumCompletedActions}/{NumActions}] {Text}", NumCompletedActions, NumActions, Text);
 							return;
 						}
 						if (!string.IsNullOrEmpty(ProgressText))
 						{
-							Log.TraceInformation($"[{NumCompletedActions}/{NumActions}] {Text} {ProgressText}");
+							Logger.LogInformation("[{NumCompletedActions}/{NumActions}] {Text} {ProgressText}", NumCompletedActions, NumActions, Text, ProgressText);
 							ProgressText = string.Empty;
 							return;
 						}
-						Log.TraceInformation(Text);
+						Logger.LogInformation("{Status}", Text);
 					}
 				};
 
 				// Run through the standard XGE executor
-				return ExecuteTaskFile(TaskFilePath, EventHandlerWrapper, NumActions);
+				return ExecuteTaskFile(TaskFilePath, EventHandlerWrapper, NumActions, Logger);
 			}
 		}
 	}
diff --git a/Engine/Source/Programs/UnrealBuildTool/Modes/AggregateParsedTimingInfo.cs b/Engine/Source/Programs/UnrealBuildTool/Modes/AggregateParsedTimingInfo.cs
--- a/Engine/Source/Programs/UnrealBuildTool/Modes/AggregateParsedTimingInfo.cs
+++ b/Engine/Source/Programs/UnrealBuildTool/Modes/AggregateParsedTimingInfo.cs
@@ -11,6 +11,7 @@
 using System.Threading;
 using System.Threading.Tasks;
 using EpicGames.Core;
+using Microsoft.Extensions.Logging;
 
 namespace UnrealBuildTool
 {
@@ -28,14 +29,14 @@
 		private ConcurrentBag<TimingData> AggregateFunctions = new ConcurrentBag<TimingData>();
 		private object AddFileLock = new object();
 
-		public override int Execute(CommandLineArguments Arguments)
+		public override int Execute(CommandLineArguments Arguments, ILogger Logger)
 		{
 			FileReference ManifestFile = Arguments.GetFileReference("-ManifestFile=");
 			string[] ParsedFileNames = FileReference.ReadAllLines(ManifestFile);
 
 			// Load all the file timing data and summarize them for the aggregate.
 			FileTimingData = new TimingData("Files", TimingDataType.Summary);
-			Parallel.ForEach(ParsedFileNames, new ParallelOptions() { MaxDegreeOfParallelism = 4 }, ParseTimingDataFile);
+			Parallel.ForEach(ParsedFileNames, new ParallelOptions() { MaxDegreeOfParallelism = 4 }, x => ParseTimingDataFile(x, Logger));
 
 			// Create aggregate summary. Duration is the duration of the files in the aggregate.
 			string AggregateName = Arguments.GetString("-Name=");
@@ -49,7 +50,7 @@
 
 			// Write out aggregate summary.
 			string OutputFile = Path.Combine(ManifestFile.Directory.FullName, String.Format("{0}.cta", AggregateName));
-			Log.TraceLog("Writing {0}", OutputFile);
+			Logger.LogDebug("Writing {OutputFile}", OutputFile);
 			using (BinaryWriter Writer = new BinaryWriter(File.Open(OutputFile, FileMode.Create)))
 			{
 				// Write out the aggregate data.
@@ -104,9 +105,9 @@
 			AggregateData.ExclusiveDuration -= GroupedTimingData.InclusiveDuration;
 		}
 
-		private void ParseTimingDataFile(string ParsedFileName)
+		private void ParseTimingDataFile(string ParsedFileName, ILogger Logger)
 		{
-			Log.TraceLog("Parsing {0}", ParsedFileName);
+			Logger.LogDebug("Parsing {ParsedFileName}", ParsedFileName);
 			// Convert input file back into summary objects.
 			using (BinaryReader Reader = new BinaryReader(File.Open(ParsedFileName, FileMode.Open, FileAccess.Read)))
 			{
diff --git a/Engine/Source/Programs/UnrealBuildTool/Modes/AnalyzeMode.cs b/Engine/Source/Programs/UnrealBuildTool/Modes/AnalyzeMode.cs
--- a/Engine/Source/Programs/UnrealBuildTool/Modes/AnalyzeMode.cs
+++ b/Engine/Source/Programs/UnrealBuildTool/Modes/AnalyzeMode.cs
@@ -7,6 +7,7 @@
 using System.Linq;
 using System.Text;
 using System.Xml;
+using Microsoft.Extensions.Logging;
 
 namespace UnrealBuildTool.Modes
 {
@@ -21,7 +22,8 @@
 		/// </summary>
 		/// <param name="Arguments">Command line arguments</param>
 		/// <returns>Exit code</returns>
-		public override int Execute(CommandLineArguments Arguments)
+		/// <param name="Logger"></param>
+		public override int Execute(CommandLineArguments Arguments, ILogger Logger)
 		{
 			Arguments.ApplyTo(this);
 
@@ -31,7 +33,7 @@
 			Arguments.ApplyTo(BuildConfiguration);
 
 			// Parse all the target descriptors
-			List<TargetDescriptor> TargetDescriptors = TargetDescriptor.ParseCommandLine(Arguments, BuildConfiguration.bUsePrecompiled, BuildConfiguration.bSkipRulesCompile, BuildConfiguration.bForceRulesCompile);
+			List<TargetDescriptor> TargetDescriptors = TargetDescriptor.ParseCommandLine(Arguments, BuildConfiguration.bUsePrecompiled, BuildConfiguration.bSkipRulesCompile, BuildConfiguration.bForceRulesCompile, Logger);
 
 			// Generate the compile DB for each target
 			using (ISourceFileWorkingSet WorkingSet = new EmptySourceFileWorkingSet())
@@ -40,7 +42,7 @@
 				Dictionary<FileReference, string> FileToCommand = new Dictionary<FileReference, string>();
 				foreach (TargetDescriptor TargetDescriptor in TargetDescriptors)
 				{
-					AnalyzeTarget(TargetDescriptor, BuildConfiguration);
+					AnalyzeTarget(TargetDescriptor, BuildConfiguration, Logger);
 				}
 			}
 
@@ -68,10 +70,10 @@
 			}
 		}
 
-		private void AnalyzeTarget(TargetDescriptor TargetDescriptor, BuildConfiguration BuildConfiguration)
+		private void AnalyzeTarget(TargetDescriptor TargetDescriptor, BuildConfiguration BuildConfiguration, ILogger Logger)
 		{
 			// Create a makefile for the target
-			UEBuildTarget Target = UEBuildTarget.Create(TargetDescriptor, BuildConfiguration.bSkipRulesCompile, BuildConfiguration.bForceRulesCompile, BuildConfiguration.bUsePrecompiled);
+			UEBuildTarget Target = UEBuildTarget.Create(TargetDescriptor, BuildConfiguration.bSkipRulesCompile, BuildConfiguration.bForceRulesCompile, BuildConfiguration.bUsePrecompiled, Logger);
 			DirectoryReference.CreateDirectory(Target.ReceiptFileName.Directory);
 
 			// Find the shortest path from the target to each module
@@ -212,22 +214,22 @@
 			// Warn about any missing modules
 			foreach (UEBuildModule MissingModule in MissingModules.OrderBy(x => x.Name))
 			{
-				Log.TraceWarning("Missing module '{0}'", MissingModule.Name);
+				Logger.LogWarning("Missing module '{MissingModuleName}'", MissingModule.Name);
 			}
 
 			// Generate the dependency graph between modules
 			FileReference DependencyGraphFile = Target.ReceiptFileName.ChangeExtension(".Dependencies.gexf");
-			Log.TraceInformation("Writing dependency graph to {0}...", DependencyGraphFile);
+			Logger.LogInformation("Writing dependency graph to {DependencyGraphFile}...", DependencyGraphFile);
 			WriteDependencyGraph(Target, ModuleToInfo, DependencyGraphFile);
 
 			// Generate the dependency graph between modules
 			FileReference ShortestPathGraphFile = Target.ReceiptFileName.ChangeExtension(".ShortestPath.gexf");
-			Log.TraceInformation("Writing shortest-path graph to {0}...", ShortestPathGraphFile);
+			Logger.LogInformation("Writing shortest-path graph to {ShortestPathGraphFile}...", ShortestPathGraphFile);
 			WriteShortestPathGraph(Target, ModuleToInfo, ShortestPathGraphFile);
 
 			// Write all the target stats as a text file
 			FileReference TextFile = Target.ReceiptFileName.ChangeExtension(".txt");
-			Log.TraceInformation("Writing module information to {0}", TextFile);
+			Logger.LogInformation("Writing module information to {TextFile}", TextFile);
 			using (StreamWriter Writer = new StreamWriter(TextFile.FullName))
 			{
 				Writer.WriteLine("All modules in {0}, ordered by number of indirect references", Target.TargetName);
@@ -250,7 +252,7 @@
 
 			// Write all the target stats as a CSV file
 			FileReference CsvFile = Target.ReceiptFileName.ChangeExtension(".csv");
-			Log.TraceInformation("Writing module information to {0}", CsvFile);
+			Logger.LogInformation("Writing module information to {CsvFile}", CsvFile);
 			using (StreamWriter Writer = new StreamWriter(CsvFile.FullName))
 			{
 				List<string> Columns = new List<string>();
diff --git a/Engine/Source/Programs/UnrealBuildTool/Modes/BuildMode.cs b/Engine/Source/Programs/UnrealBuildTool/Modes/BuildMode.cs
--- a/Engine/Source/Programs/UnrealBuildTool/Modes/BuildMode.cs
+++ b/Engine/Source/Programs/UnrealBuildTool/Modes/BuildMode.cs
@@ -8,6 +8,7 @@
 using EpicGames.Core;
 using OpenTracing.Util;
 using UnrealBuildBase;
+using Microsoft.Extensions.Logging;
 
 namespace UnrealBuildTool
 {
@@ -98,12 +99,13 @@
 		/// </summary>
 		/// <param name="Arguments">Command-line arguments</param>
 		/// <returns>One of the values of ECompilationResult</returns>
-		public override int Execute(CommandLineArguments Arguments)
+		/// <param name="Logger"></param>
+		public override int Execute(CommandLineArguments Arguments, ILogger Logger)
 		{
 			Arguments.ApplyTo(this);
 
 			// Write the command line
-			Log.TraceLog("Command line: {0}", Environment.CommandLine);
+			Logger.LogDebug("Command line: {EnvironmentCommandLine}", Environment.CommandLine);
 
 			// Grab the environment.
 			UnrealBuildTool.InitialEnvironment = Environment.GetEnvironmentVariables();
@@ -143,7 +145,7 @@
 			// Check the root path length isn't too long
 			if (BuildHostPlatform.Current.Platform == UnrealTargetPlatform.Win64 && Unreal.RootDirectory.FullName.Length > BuildConfiguration.MaxRootPathLength)
 			{
-				Log.TraceWarning("Running from a path with a long directory name (\"{0}\" = {1} characters). Root paths shorter than {2} characters are recommended to avoid exceeding maximum path lengths on Windows.", Unreal.RootDirectory, Unreal.RootDirectory.FullName.Length, BuildConfiguration.MaxRootPathLength);
+				Logger.LogWarning("Running from a path with a long directory name (\"{Path}\" = {NumChars} characters). Root paths shorter than {MaxChars} characters are recommended to avoid exceeding maximum path lengths on Windows.", Unreal.RootDirectory, Unreal.RootDirectory.FullName.Length, BuildConfiguration.MaxRootPathLength);
 			}
 
 			// now that we know the available platforms, we can delete other platforms' junk. if we're only building specific modules from the editor, don't touch anything else (it may be in use).
@@ -151,7 +153,7 @@
 			{
 				using (GlobalTracer.Instance.BuildSpan("DeleteJunk()").StartActive())
 				{
-					JunkDeleter.DeleteJunk();
+					JunkDeleter.DeleteJunk(Logger);
 				}
 			}
 
@@ -163,7 +165,7 @@
 				// Parse all the target descriptors
 				using (GlobalTracer.Instance.BuildSpan("TargetDescriptor.ParseCommandLine()").StartActive())
 				{
-					TargetDescriptors = TargetDescriptor.ParseCommandLine(Arguments, BuildConfiguration.bUsePrecompiled, BuildConfiguration.bSkipRulesCompile, BuildConfiguration.bForceRulesCompile);
+					TargetDescriptors = TargetDescriptor.ParseCommandLine(Arguments, BuildConfiguration.bUsePrecompiled, BuildConfiguration.bSkipRulesCompile, BuildConfiguration.bForceRulesCompile, Logger);
 				}
 
 				// Hack for specific files compile; don't build the ShaderCompileWorker target that's added to the command line for generated project files
@@ -177,7 +179,7 @@
 				{
 					CleanMode CleanMode = new CleanMode();
 					CleanMode.bSkipPreBuildTargets = bSkipPreBuildTargets;
-					CleanMode.Clean(TargetDescriptors.Where(D => D.bRebuild).ToList(), BuildConfiguration);
+					CleanMode.Clean(TargetDescriptors.Where(D => D.bRebuild).ToList(), BuildConfiguration, Logger);
 				}
 
 
@@ -190,8 +192,8 @@
 						FileReference BaseLogFile = Log.OutputFile ?? new FileReference(BaseLogFileName);
 						FileReference RemoteLogFile = FileReference.Combine(BaseLogFile.Directory, BaseLogFile.GetFileNameWithoutExtension() + "_Remote.txt");
 
-						RemoteMac RemoteMac = new RemoteMac(TargetDesc.ProjectFile);
-						if (!RemoteMac.Build(TargetDesc, RemoteLogFile, bSkipPreBuildTargets))
+						RemoteMac RemoteMac = new RemoteMac(TargetDesc.ProjectFile, Logger);
+						if (!RemoteMac.Build(TargetDesc, RemoteLogFile, bSkipPreBuildTargets, Logger))
 						{
 							return (int)CompilationResult.Unknown;
 						}
@@ -234,16 +236,16 @@
 					}
 
 					// Create the working set provider per group.
-					using (ISourceFileWorkingSet WorkingSet = SourceFileWorkingSet.Create(Unreal.RootDirectory, ProjectDirs))
+					using (ISourceFileWorkingSet WorkingSet = SourceFileWorkingSet.Create(Unreal.RootDirectory, ProjectDirs, Logger))
 					{
-						Build(TargetDescriptors, BuildConfiguration, WorkingSet, Options, WriteOutdatedActionsFile, bSkipPreBuildTargets);
+						Build(TargetDescriptors, BuildConfiguration, WorkingSet, Options, WriteOutdatedActionsFile, Logger, bSkipPreBuildTargets);
 					}
 				}
 			}
 			finally
 			{
 				// Check if anything failed during our run, and act accordingly.
-				ProcessCoreDumps(SaveCrashDumpDirectory);
+				ProcessCoreDumps(SaveCrashDumpDirectory, Logger);
 
 				// Save all the caches
 				SourceFileMetadataCache.SaveAll();
@@ -260,15 +262,16 @@
 		/// <param name="WorkingSet">The source file working set</param>
 		/// <param name="Options">Additional options for the build</param>
 		/// <param name="WriteOutdatedActionsFile">Files to write the list of outdated actions to (rather than building them)</param>
+		/// <param name="Logger">Logger for output</param>
 		/// <param name="bSkipPreBuildTargets">If true then only the current target descriptors will be built.</param>
 		/// <returns>Result from the compilation</returns>
-		public static void Build(List<TargetDescriptor> TargetDescriptors, BuildConfiguration BuildConfiguration, ISourceFileWorkingSet WorkingSet, BuildOptions Options, FileReference? WriteOutdatedActionsFile, bool bSkipPreBuildTargets = false)
+		public static void Build(List<TargetDescriptor> TargetDescriptors, BuildConfiguration BuildConfiguration, ISourceFileWorkingSet WorkingSet, BuildOptions Options, FileReference? WriteOutdatedActionsFile, ILogger Logger, bool bSkipPreBuildTargets = false)
 		{
 			List<TargetMakefile> TargetMakefiles = new List<TargetMakefile>();
 
 			for (int Idx = 0; Idx < TargetDescriptors.Count; ++Idx)
 			{
-				TargetMakefile NewMakefile = CreateMakefile(BuildConfiguration, TargetDescriptors[Idx], WorkingSet);
+				TargetMakefile NewMakefile = CreateMakefile(BuildConfiguration, TargetDescriptors[Idx], WorkingSet, Logger);
 				TargetMakefiles.Add(NewMakefile);
 				if (!bSkipPreBuildTargets)
 				{
@@ -283,7 +286,7 @@
 				}
 			}
 
-			Build(TargetMakefiles.ToArray(), TargetDescriptors, BuildConfiguration, Options, WriteOutdatedActionsFile);
+			Build(TargetMakefiles.ToArray(), TargetDescriptors, BuildConfiguration, Options, WriteOutdatedActionsFile, Logger);
 		}
 
 		/// <summary>
@@ -294,19 +297,20 @@
 		/// <param name="BuildConfiguration">Current build configuration</param>
 		/// <param name="Options">Additional options for the build</param>
 		/// <param name="WriteOutdatedActionsFile">Files to write the list of outdated actions to (rather than building them)</param>
+		/// <param name="Logger">Logger for output</param>
 		/// <returns>Result from the compilation</returns>
-		static void Build(TargetMakefile[] Makefiles, List<TargetDescriptor> TargetDescriptors, BuildConfiguration BuildConfiguration, BuildOptions Options, FileReference? WriteOutdatedActionsFile)
+		static void Build(TargetMakefile[] Makefiles, List<TargetDescriptor> TargetDescriptors, BuildConfiguration BuildConfiguration, BuildOptions Options, FileReference? WriteOutdatedActionsFile, ILogger Logger)
 		{
 			// Execute the build
 			if ((Options & BuildOptions.SkipBuild) == 0)
 			{
 				// Make sure that none of the actions conflict with any other (producing output files differently, etc...)
-				ActionGraph.CheckForConflicts(Makefiles.SelectMany(x => x.Actions));
+				ActionGraph.CheckForConflicts(Makefiles.SelectMany(x => x.Actions), Logger);
 
 				// Check we don't exceed the nominal max path length
 				using (GlobalTracer.Instance.BuildSpan("ActionGraph.CheckPathLengths").StartActive())
 				{
-					ActionGraph.CheckPathLengths(BuildConfiguration, Makefiles.SelectMany(x => x.Actions));
+					ActionGraph.CheckPathLengths(BuildConfiguration, Makefiles.SelectMany(x => x.Actions), Logger);
 				}
 
 				// Create a QueuedAction instance for each action in the makefiles
@@ -320,7 +324,7 @@
 				Dictionary<FileReference, FileReference>? InitialPatchedOldLocationToNewLocation = null;
 				for (int TargetIdx = 0; TargetIdx < TargetDescriptors.Count; TargetIdx++)
 				{
-					InitialPatchedOldLocationToNewLocation = HotReload.Setup(TargetDescriptors[TargetIdx], Makefiles[TargetIdx], QueuedActions[TargetIdx], BuildConfiguration);
+					InitialPatchedOldLocationToNewLocation = HotReload.Setup(TargetDescriptors[TargetIdx], Makefiles[TargetIdx], QueuedActions[TargetIdx], BuildConfiguration, Logger);
 				}
 
 				// Merge the action graphs together
@@ -368,7 +372,7 @@
 					using (GlobalTracer.Instance.BuildSpan("Reading dependency cache").StartActive())
 					{
 						TargetDescriptor TargetDescriptor = TargetDescriptors[TargetIdx];
-						CppDependencies.Mount(TargetDescriptor.ProjectFile, TargetDescriptor.Name, TargetDescriptor.Platform, TargetDescriptor.Configuration, Makefiles[TargetIdx].TargetType, TargetDescriptor.Architecture);
+						CppDependencies.Mount(TargetDescriptor.ProjectFile, TargetDescriptor.Name, TargetDescriptor.Platform, TargetDescriptor.Configuration, Makefiles[TargetIdx].TargetType, TargetDescriptor.Architecture, Logger);
 					}
 				}
 
@@ -377,13 +381,13 @@
 				if (ModuleDependencyActions.Count > 0)
 				{
 					Dictionary<LinkedAction, bool> PreprocessActionToOutdatedFlag = new Dictionary<LinkedAction, bool>();
-					ActionGraph.GatherAllOutdatedActions(ModuleDependencyActions, History, PreprocessActionToOutdatedFlag, CppDependencies, BuildConfiguration.bIgnoreOutdatedImportLibraries);
+					ActionGraph.GatherAllOutdatedActions(ModuleDependencyActions, History, PreprocessActionToOutdatedFlag, CppDependencies, BuildConfiguration.bIgnoreOutdatedImportLibraries, Logger);
 
 					List<LinkedAction> PreprocessActions = PreprocessActionToOutdatedFlag.Where(x => x.Value).Select(x => x.Key).ToList();
 					if (PreprocessActions.Count > 0)
 					{
-						Log.TraceInformation("Updating module dependencies...");
-						ActionGraph.ExecuteActions(BuildConfiguration, PreprocessActions, TargetDescriptors);
+						Logger.LogInformation("Updating module dependencies...");
+						ActionGraph.ExecuteActions(BuildConfiguration, PreprocessActions, TargetDescriptors, Logger);
 
 						foreach (FileItem ProducedItem in PreprocessActions.SelectMany(x => x.ProducedItems))
 						{
@@ -411,13 +415,13 @@
 							if (CppModulesAction != null && CppModulesAction.CompiledModuleInterfaceFile != null)
 							{
 								string? ProducedModule;
-								if (CppDependencies.TryGetProducedModule(ModuleDependencyAction.DependencyListFile!, out ProducedModule))
+								if (CppDependencies.TryGetProducedModule(ModuleDependencyAction.DependencyListFile!, Logger, out ProducedModule))
 								{
 									ModuleOutputs[ProducedModule] = CppModulesAction.CompiledModuleInterfaceFile;
 								}
 
 								List<(string Name, string BMI)>? ImportedModules;
-								if (CppDependencies.TryGetImportedModules(ModuleDependencyAction.DependencyListFile!, out ImportedModules))
+								if (CppDependencies.TryGetImportedModules(ModuleDependencyAction.DependencyListFile!, Logger, out ImportedModules))
 								{
 									ModuleImports[CppModulesAction.CompiledModuleInterfaceFile] = ImportedModules.Select(x => x.Name).ToList();
 								}
@@ -481,7 +485,7 @@
 					// Plan the actions to execute for the build. For single file compiles, always rebuild the source file regardless of whether it's out of date.
 					if (TargetDescriptor.SpecificFilesToCompile.Count == 0)
 					{
-						ActionGraph.GatherAllOutdatedActions(PrerequisiteActions, History, ActionToOutdatedFlag, CppDependencies, BuildConfiguration.bIgnoreOutdatedImportLibraries);
+						ActionGraph.GatherAllOutdatedActions(PrerequisiteActions, History, ActionToOutdatedFlag, CppDependencies, BuildConfiguration.bIgnoreOutdatedImportLibraries, Logger);
 					}
 					else
 					{
@@ -512,22 +516,22 @@
 						}
 						else
 						{
-							MergedActionsToExecute = HotReload.PatchActionsForTarget(BuildConfiguration, TargetDescriptors[Idx], Makefiles[Idx], PrerequisiteActions, MergedActionsToExecute, InitialPatchedOldLocationToNewLocation);
+							MergedActionsToExecute = HotReload.PatchActionsForTarget(BuildConfiguration, TargetDescriptors[Idx], Makefiles[Idx], PrerequisiteActions, MergedActionsToExecute, InitialPatchedOldLocationToNewLocation, Logger);
 						}
 						HotReloadTargetIdx = Idx;
 					}
 					else if (MergedActionsToExecute.Count > 0)
 					{
-						HotReload.CheckForLiveCodingSessionActive(TargetDescriptors[Idx], Makefiles[Idx], BuildConfiguration);
+						HotReload.CheckForLiveCodingSessionActive(TargetDescriptors[Idx], Makefiles[Idx], BuildConfiguration, Logger);
 					}
 				}	
 
 				if (HotReloadTargetIdx != -1)
 				{
-					Log.TraceLog("Re-evaluating action graph");
+					Logger.LogDebug("Re-evaluating action graph");
 					// Re-check the graph to remove any LiveCoding actions added by PatchActionsForTarget() that are already up to date.
 					Dictionary<LinkedAction, bool> LiveActionToOutdatedFlag = new Dictionary<LinkedAction, bool>(MergedActionsToExecute.Count);
-					ActionGraph.GatherAllOutdatedActions(MergedActionsToExecute, History, LiveActionToOutdatedFlag, CppDependencies, BuildConfiguration.bIgnoreOutdatedImportLibraries);
+					ActionGraph.GatherAllOutdatedActions(MergedActionsToExecute, History, LiveActionToOutdatedFlag, CppDependencies, BuildConfiguration.bIgnoreOutdatedImportLibraries, Logger);
 					List<LinkedAction> LiveCodingActionsToExecute = LiveActionToOutdatedFlag.Where(x => x.Value).Select(x => x.Key).ToList();
 					ActionGraph.Link(LiveCodingActionsToExecute);
 					MergedActionsToExecute = LiveCodingActionsToExecute;
@@ -545,7 +549,7 @@
 							Result.AppendFormat("\n{0}", EngineChange.FullName);
 						}
 						Result.Append("\n\nPlease rebuild from an IDE instead.");
-						Log.TraceError("{0}", Result.ToString());
+						Logger.LogError("{Result}", Result.ToString());
 						throw new CompilationResultException(CompilationResult.FailedDueToEngineChange);
 					}
 				}
@@ -560,7 +564,7 @@
 				}
 
 				// Delete produced items that are outdated.
-				ActionGraph.DeleteOutdatedProducedItems(MergedActionsToExecute);
+				ActionGraph.DeleteOutdatedProducedItems(MergedActionsToExecute, Logger);
 
 				// Save all the action histories now that files have been removed. We have to do this after deleting produced items to ensure that any
 				// items created during the build don't have the wrong command line.
@@ -572,17 +576,17 @@
 				// Execute the actions
 				if ((Options & BuildOptions.XGEExport) != 0)
 				{
-					OutputToolchainInfo(TargetDescriptors, Makefiles);
+					OutputToolchainInfo(TargetDescriptors, Makefiles, Logger);
 
 					// Just export to an XML file
 					using (GlobalTracer.Instance.BuildSpan("XGE.ExportActions()").StartActive())
 					{
-						XGE.ExportActions(MergedActionsToExecute);
+						XGE.ExportActions(MergedActionsToExecute, Logger);
 					}
 				}
 				else if (WriteOutdatedActionsFile != null)
 				{
-					OutputToolchainInfo(TargetDescriptors, Makefiles);
+					OutputToolchainInfo(TargetDescriptors, Makefiles, Logger);
 
 					// Write actions to an output file
 					using (GlobalTracer.Instance.BuildSpan("ActionGraph.WriteActions").StartActive())
@@ -597,23 +601,30 @@
 					{
 						if (TargetDescriptors.Any(x => !x.bQuiet))
 						{
-							Log.TraceInformation((TargetDescriptors.Count == 1)? "Target is up to date" : "Targets are up to date");
+							if (TargetDescriptors.Count == 1)
+							{
+								Logger.LogInformation("Target is up to date");
+							}
+							else
+							{
+								Logger.LogInformation("Targets are up to date");
+							}
 						}
 					}
 					else
 					{
 						if (TargetDescriptors.Any(x => !x.bQuiet))
 						{
-							Log.TraceInformation("Building {0}...", StringUtils.FormatList(TargetDescriptors.Select(x => x.Name).Distinct()));
+							Logger.LogInformation("Building {Targets}...", StringUtils.FormatList(TargetDescriptors.Select(x => x.Name).Distinct()));
 						}
 
-						OutputToolchainInfo(TargetDescriptors, Makefiles);
+						OutputToolchainInfo(TargetDescriptors, Makefiles, Logger);
 
 						ActionExecutor.SetMemoryPerActionOverride(Makefiles.Select(x => x.MemoryPerActionGB).Max());
 
 						using (GlobalTracer.Instance.BuildSpan("ActionGraph.ExecuteActions()").StartActive())
 						{
-							ActionGraph.ExecuteActions(BuildConfiguration, MergedActionsToExecute, TargetDescriptors);
+							ActionGraph.ExecuteActions(BuildConfiguration, MergedActionsToExecute, TargetDescriptors, Logger);
 						}
 					}
 
@@ -623,7 +634,7 @@
 						if (Makefile.bDeployAfterCompile)
 						{
 							TargetReceipt Receipt = TargetReceipt.Read(Makefile.ReceiptFile);
-							Log.TraceInformation("Deploying {0} {1} {2}...", Receipt.TargetName, Receipt.Platform, Receipt.Configuration);
+							Logger.LogInformation("Deploying {ReceiptTargetName} {ReceiptPlatform} {ReceiptConfiguration}...", Receipt.TargetName, Receipt.Platform, Receipt.Configuration);
 
 							UEBuildPlatform.GetBuildPlatform(Receipt.Platform).Deploy(Receipt);
 						}
@@ -637,7 +648,8 @@
 		/// </summary>
 		/// <param name="TargetDescriptors">List of targets being built</param>
 		/// <param name="Makefiles">Matching array of makefiles for each target</param>
-		static void OutputToolchainInfo(List<TargetDescriptor> TargetDescriptors, TargetMakefile[] Makefiles)
+		/// <param name="Logger">Logger for output</param>
+		static void OutputToolchainInfo(List<TargetDescriptor> TargetDescriptors, TargetMakefile[] Makefiles, ILogger Logger)
 		{
 			List<int> OutputIndices = new List<int>();
 			for (int Idx = 0; Idx < TargetDescriptors.Count; Idx++)
@@ -652,7 +664,7 @@
 			{
 				foreach(string Diagnostic in Makefiles[OutputIndices[0]].Diagnostics)
 				{
-					Log.TraceInformation("{0}", Diagnostic);
+					Logger.LogInformation("{Diagnostic}", Diagnostic);
 				}
 			}
 			else
@@ -661,7 +673,7 @@
 				{
 					foreach(string Diagnostic in Makefiles[OutputIndex].Diagnostics)
 					{
-						Log.TraceInformation("{0}: {1}", TargetDescriptors[OutputIndex].Name, Diagnostic);
+						Logger.LogInformation("{Name}: {Diagnostic}", TargetDescriptors[OutputIndex].Name, Diagnostic);
 					}
 				}
 			}
@@ -673,8 +685,9 @@
 		/// <param name="BuildConfiguration">The build configuration</param>
 		/// <param name="TargetDescriptor">Target being built</param>
 		/// <param name="WorkingSet">Set of source files which are part of the working set</param>
+		/// <param name="Logger">Logger for output</param>
 		/// <returns>Makefile for the given target</returns>
-		static TargetMakefile CreateMakefile(BuildConfiguration BuildConfiguration, TargetDescriptor TargetDescriptor, ISourceFileWorkingSet WorkingSet)
+		static TargetMakefile CreateMakefile(BuildConfiguration BuildConfiguration, TargetDescriptor TargetDescriptor, ISourceFileWorkingSet WorkingSet, ILogger Logger)
 		{
 			// Get the path to the makefile for this target
 			FileReference? MakefileLocation = null;
@@ -690,10 +703,10 @@
 				using (GlobalTracer.Instance.BuildSpan("TargetMakefile.Load()").StartActive())
 				{
 					string? ReasonNotLoaded;
-					Makefile = TargetMakefile.Load(MakefileLocation, TargetDescriptor.ProjectFile, TargetDescriptor.Platform, TargetDescriptor.AdditionalArguments.GetRawArray(), out ReasonNotLoaded);
+					Makefile = TargetMakefile.Load(MakefileLocation, TargetDescriptor.ProjectFile, TargetDescriptor.Platform, TargetDescriptor.AdditionalArguments.GetRawArray(), Logger, out ReasonNotLoaded);
 					if (Makefile == null)
 					{
-						Log.TraceInformation("Creating makefile for {0} ({1})", TargetDescriptor.Name, ReasonNotLoaded);
+						Logger.LogInformation("Creating makefile for {TargetDescriptorName} ({ReasonNotLoaded})", TargetDescriptor.Name, ReasonNotLoaded);
 					}
 				}
 			}
@@ -705,7 +718,7 @@
 				// Execute the scripts. We have to invalidate all cached file info after doing so, because we don't know what may have changed.
 				if(Makefile.PreBuildScripts.Length > 0)
 				{
-					Utils.ExecuteCustomBuildSteps(Makefile.PreBuildScripts);
+					Utils.ExecuteCustomBuildSteps(Makefile.PreBuildScripts, Logger);
 					DirectoryItem.ResetAllCachedInfo_SLOW();
 				}
 
@@ -714,9 +727,9 @@
 
 				// Check that the makefile is still valid
 				string? Reason;
-				if(!TargetMakefile.IsValidForSourceFiles(Makefile, TargetDescriptor.ProjectFile, TargetDescriptor.Platform, WorkingSet, out Reason))
+				if(!TargetMakefile.IsValidForSourceFiles(Makefile, TargetDescriptor.ProjectFile, TargetDescriptor.Platform, WorkingSet, Logger, out Reason))
 				{
-					Log.TraceInformation("Invalidating makefile for {0} ({1})", TargetDescriptor.Name, Reason);
+					Logger.LogInformation("Invalidating makefile for {TargetDescriptorName} ({Reason})", TargetDescriptor.Name, Reason);
 					Makefile = null;
 				}
 			}
@@ -728,7 +741,7 @@
 				UEBuildTarget Target;
 				using (GlobalTracer.Instance.BuildSpan("UEBuildTarget.Create()").StartActive())
 				{
-					Target = UEBuildTarget.Create(TargetDescriptor, BuildConfiguration.bSkipRulesCompile, BuildConfiguration.bForceRulesCompile, BuildConfiguration.bUsePrecompiled);
+					Target = UEBuildTarget.Create(TargetDescriptor, BuildConfiguration.bSkipRulesCompile, BuildConfiguration.bForceRulesCompile, BuildConfiguration.bUsePrecompiled, Logger);
 				}
 
 				// Create the pre-build scripts
@@ -737,14 +750,14 @@
 				// Execute the pre-build scripts
 				if(!bHasRunPreBuildScripts)
 				{
-					Utils.ExecuteCustomBuildSteps(PreBuildScripts);
+					Utils.ExecuteCustomBuildSteps(PreBuildScripts, Logger);
 					bHasRunPreBuildScripts = true;
 				}
 
 				// Build the target
 				using (GlobalTracer.Instance.BuildSpan("UEBuildTarget.Build()").StartActive())
 				{
-					Makefile = Target.Build(BuildConfiguration, WorkingSet, TargetDescriptor);
+					Makefile = Target.Build(BuildConfiguration, WorkingSet, TargetDescriptor, Logger);
 				}
 
 				Makefile.MemoryPerActionGB = Target.Rules.MemoryPerActionGB;
@@ -784,7 +797,7 @@
 				// If the target needs UHT to be run, we'll go ahead and do that now
 				if (Makefile.UObjectModules.Count > 0)
 				{
-					ExternalExecution.ExecuteHeaderToolIfNecessary(BuildConfiguration, TargetDescriptor.ProjectFile, Makefile, TargetDescriptor.Name, WorkingSet);
+					ExternalExecution.ExecuteHeaderToolIfNecessary(BuildConfiguration, TargetDescriptor.ProjectFile, Makefile, TargetDescriptor.Name, WorkingSet, Logger);
 				}
 
 #if __VPROJECT_AVAILABLE__
@@ -792,7 +805,7 @@
 				if (Makefile.VNIModules.Count > 0)
 				{
 
-					VNIExecution.ExecuteVNITool(Makefile, TargetDescriptor);
+					VNIExecution.ExecuteVNITool(Makefile, TargetDescriptor, Logger);
 				}
 #endif
 			}
@@ -865,7 +878,7 @@
 			return new List<LinkedAction>(OutputItemToProducingAction.Values);
 		}
 
-		void ProcessCoreDumps(DirectoryReference? SaveCrashDumpDirectory)
+		void ProcessCoreDumps(DirectoryReference? SaveCrashDumpDirectory, ILogger Logger)
 		{
 			if (SaveCrashDumpDirectory == null)
 			{
@@ -902,7 +915,7 @@
 					{
 						if (FileReference.GetLastWriteTime(CrashDump) > UBTStartTime)
 						{
-							Log.TraceWarning($"Crash dump {CrashDump} was created duing UnrealBuildTool execution");
+							Logger.LogWarning("Crash dump {CrashDump} was created duing UnrealBuildTool execution", CrashDump);
 							FoundCrashDumps.Add(CrashDump);
 						}
 					}
@@ -935,7 +948,7 @@
 						{
 							if (TmpFile.HasExtension(".dmp") && FileReference.GetLastWriteTime(TmpFile) > UBTStartTime)
 							{
-								Log.TraceWarning($"Crash dump {TmpFile} was created duing UnrealBuildTool execution");
+								Logger.LogWarning("Crash dump {TmpFile} was created duing UnrealBuildTool execution", TmpFile);
 								FoundCrashDumps.Add(TmpFile);
 							}
 						}
@@ -952,7 +965,7 @@
 				DirectoryReference.CreateDirectory(SaveCrashDumpDirectory);
 				foreach (FileReference CrashDump in FoundCrashDumps)
 				{
-					Log.TraceInformation($"Copying {CrashDump} to {SaveCrashDumpDirectory}");
+					Logger.LogInformation("Copying {CrashDump} to {SaveCrashDumpDirectory}", CrashDump, SaveCrashDumpDirectory);
 					try
 					{
 						FileReference.Copy(CrashDump, FileReference.Combine(SaveCrashDumpDirectory, CrashDump.GetFileName()));
@@ -960,7 +973,7 @@
 					catch(Exception Ex)
 					{
 						// don't stop if there was a problem copying one of the files
-						Log.TraceWarning($"Failed to copy crash dump {CrashDump}: {ExceptionUtils.FormatException(Ex)}");
+						Logger.LogWarning("Failed to copy crash dump {CrashDump}: {Ex}", CrashDump, ExceptionUtils.FormatException(Ex));
 					}
 				}
 			}
diff --git a/Engine/Source/Programs/UnrealBuildTool/Modes/CleanMode.cs b/Engine/Source/Programs/UnrealBuildTool/Modes/CleanMode.cs
--- a/Engine/Source/Programs/UnrealBuildTool/Modes/CleanMode.cs
+++ b/Engine/Source/Programs/UnrealBuildTool/Modes/CleanMode.cs
@@ -11,6 +11,7 @@
 using OpenTracing;
 using OpenTracing.Util;
 using UnrealBuildBase;
+using Microsoft.Extensions.Logging;
 
 namespace UnrealBuildTool
 {
@@ -38,7 +39,8 @@
 		/// </summary>
 		/// <param name="Arguments">Command-line arguments</param>
 		/// <returns>One of the values of ECompilationResult</returns>
-		public override int Execute(CommandLineArguments Arguments)
+		/// <param name="Logger"></param>
+		public override int Execute(CommandLineArguments Arguments, ILogger Logger)
 		{
 			Arguments.ApplyTo(this);
 
@@ -48,16 +50,16 @@
 			Arguments.ApplyTo(BuildConfiguration);
 
 			// Parse all the targets being built
-			List<TargetDescriptor> TargetDescriptors = TargetDescriptor.ParseCommandLine(Arguments, BuildConfiguration.bUsePrecompiled, bSkipRulesCompile, BuildConfiguration.bForceRulesCompile);
+			List<TargetDescriptor> TargetDescriptors = TargetDescriptor.ParseCommandLine(Arguments, BuildConfiguration.bUsePrecompiled, bSkipRulesCompile, BuildConfiguration.bForceRulesCompile, Logger);
 
-			Clean(TargetDescriptors, BuildConfiguration);
+			Clean(TargetDescriptors, BuildConfiguration, Logger);
 
 			return 0;
 		}
 
-		public void Clean(List<TargetDescriptor> TargetDescriptors, BuildConfiguration BuildConfiguration)
+		public void Clean(List<TargetDescriptor> TargetDescriptors, BuildConfiguration BuildConfiguration, ILogger Logger)
 		{
-			using ScopedTimer CleanTimer = new ScopedTimer("CleanMode.Clean()");
+			using ScopedTimer CleanTimer = new ScopedTimer("CleanMode.Clean()", Logger);
 			using IScope Scope = GlobalTracer.Instance.BuildSpan("CleanMode.Clean()").StartActive();
 
 			if (TargetDescriptors.Count == 0)
@@ -100,23 +102,23 @@
 			}
 
 			// Output the list of targets that we're cleaning
-			Log.TraceInformation("Cleaning {0} binaries...", StringUtils.FormatList(TargetDescriptors.Select(x => x.Name).Distinct()));
+			Logger.LogInformation("Cleaning {TargetNames} binaries...", StringUtils.FormatList(TargetDescriptors.Select(x => x.Name).Distinct()));
 
 			// Loop through all the targets, and clean them all
 			HashSet<FileReference> FilesToDelete = new HashSet<FileReference>();
 			HashSet<DirectoryReference> DirectoriesToDelete = new HashSet<DirectoryReference>();
 
-			using (ScopedTimer GatherTimer = new ScopedTimer("Find paths to clean"))
+			using (ScopedTimer GatherTimer = new ScopedTimer("Find paths to clean", Logger))
 			{
 				for (int Idx = 0; Idx < TargetDescriptors.Count; ++Idx)
 				{
 					TargetDescriptor TargetDescriptor = TargetDescriptors[Idx];
 
 					// Create the rules assembly
-					RulesAssembly RulesAssembly = RulesCompiler.CreateTargetRulesAssembly(TargetDescriptor.ProjectFile, TargetDescriptor.Name, bSkipRulesCompile, BuildConfiguration.bForceRulesCompile, BuildConfiguration.bUsePrecompiled, TargetDescriptor.ForeignPlugin);
+					RulesAssembly RulesAssembly = RulesCompiler.CreateTargetRulesAssembly(TargetDescriptor.ProjectFile, TargetDescriptor.Name, bSkipRulesCompile, BuildConfiguration.bForceRulesCompile, BuildConfiguration.bUsePrecompiled, TargetDescriptor.ForeignPlugin, Logger);
 
 					// Create the rules object
-					ReadOnlyTargetRules Target = new ReadOnlyTargetRules(RulesAssembly.CreateTargetRules(TargetDescriptor.Name, TargetDescriptor.Platform, TargetDescriptor.Configuration, TargetDescriptor.Architecture, TargetDescriptor.ProjectFile, TargetDescriptor.AdditionalArguments));
+					ReadOnlyTargetRules Target = new ReadOnlyTargetRules(RulesAssembly.CreateTargetRules(TargetDescriptor.Name, TargetDescriptor.Platform, TargetDescriptor.Configuration, TargetDescriptor.Architecture, TargetDescriptor.ProjectFile, TargetDescriptor.AdditionalArguments, Logger));
 
 					if (!bSkipPreBuildTargets && Target.PreBuildTargets.Count > 0)
 					{
@@ -216,7 +218,7 @@
 			}
 
 			// Ensure no overlap between directories
-			using (ScopedTimer Timer = new ScopedTimer("Ensure no directory overlap"))
+			using (ScopedTimer Timer = new ScopedTimer("Ensure no directory overlap", Logger))
 			{
 				HashSet<DirectoryReference> SubdirectoriesToDelete = new(DirectoriesToDelete.Count);
 				foreach (DirectoryReference Directory in DirectoriesToDelete)
@@ -235,7 +237,7 @@
 			}
 
 			// Remove any files that are contained within one of the directories
-			using (ScopedTimer Time = new ScopedTimer("Ensure no file overlap"))
+			using (ScopedTimer Time = new ScopedTimer("Ensure no file overlap", Logger))
 			{
 				FilesToDelete.RemoveWhere(File =>
 				{
@@ -253,14 +255,14 @@
 
 			var DeleteDirectories = Task.Run(() =>
 			{
-				using ScopedTimer Timer = new ScopedTimer($"Delete {DirectoriesToDelete.Count} directories");
+				using ScopedTimer Timer = new ScopedTimer($"Delete {DirectoriesToDelete.Count} directories", Logger);
 				Parallel.ForEach(DirectoriesToDelete, DirectoryToDelete =>
 				{
-					using ScopedTimer Timer = new ScopedTimer($"Delete directory '{DirectoryToDelete}'", bIncreaseIndent: false);
+					using ScopedTimer Timer = new ScopedTimer($"Delete directory '{DirectoryToDelete}'", Logger, bIncreaseIndent: false);
 
 					if (DirectoryReference.Exists(DirectoryToDelete))
 					{
-						Log.TraceVerbose($"    Deleting {DirectoryToDelete}{Path.DirectorySeparatorChar}...");
+						Logger.LogDebug("    Deleting {DirectoryToDelete}...", $"{DirectoryToDelete}{Path.DirectorySeparatorChar}");
 
 						try
 						{
@@ -276,12 +278,12 @@
 
 			var DeleteFiles = Task.Run(() =>
 			{
-				using ScopedTimer Timer = new ScopedTimer($"Delete {FilesToDelete.Count} files");
+				using ScopedTimer Timer = new ScopedTimer($"Delete {FilesToDelete.Count} files", Logger);
 				Parallel.ForEach(FilesToDelete, FileToDelete =>
 				{
 					if (FileReference.Exists(FileToDelete))
 					{
-						Log.TraceVerbose($"    Deleting {FileToDelete}...");
+						Logger.LogDebug("    Deleting {FileToDelete}...", FileToDelete);
 
 						try
 						{
@@ -304,8 +306,8 @@
 				TargetDescriptor TargetDescriptor = TargetDescriptors[Idx];
 				if (RemoteMac.HandlesTargetPlatform(TargetDescriptor.Platform))
 				{
-					RemoteMac RemoteMac = new RemoteMac(TargetDescriptor.ProjectFile);
-					RemoteMac.Clean(TargetDescriptor);
+					RemoteMac RemoteMac = new RemoteMac(TargetDescriptor.ProjectFile, Logger);
+					RemoteMac.Clean(TargetDescriptor, Logger);
 				}
 			}
 		}
diff --git a/Engine/Source/Programs/UnrealBuildTool/Modes/DeployMode.cs b/Engine/Source/Programs/UnrealBuildTool/Modes/DeployMode.cs
--- a/Engine/Source/Programs/UnrealBuildTool/Modes/DeployMode.cs
+++ b/Engine/Source/Programs/UnrealBuildTool/Modes/DeployMode.cs
@@ -6,6 +6,7 @@
 using System.Text;
 using System.Threading.Tasks;
 using EpicGames.Core;
+using Microsoft.Extensions.Logging;
 
 namespace UnrealBuildTool
 {
@@ -26,7 +27,8 @@
 		/// </summary>
 		/// <param name="Arguments">Command line arguments</param>
 		/// <returns>Exit code</returns>
-		public override int Execute(CommandLineArguments Arguments)
+		/// <param name="Logger"></param>
+		public override int Execute(CommandLineArguments Arguments, ILogger Logger)
 		{
 			// Apply the arguments
 			Arguments.ApplyTo(this);
@@ -34,7 +36,7 @@
 
 			// Execute the deploy
 			TargetReceipt Receipt = TargetReceipt.Read(ReceiptFile!);
-			Log.WriteLine(LogEventType.Console, "Deploying {0} {1} {2}...", Receipt.TargetName, Receipt.Platform, Receipt.Configuration);
+			Logger.LogInformation("Deploying {ReceiptTargetName} {ReceiptPlatform} {ReceiptConfiguration}...", Receipt.TargetName, Receipt.Platform, Receipt.Configuration);
 			UEBuildPlatform.GetBuildPlatform(Receipt.Platform).Deploy(Receipt);
 
 			return (int)CompilationResult.Succeeded;
diff --git a/Engine/Source/Programs/UnrealBuildTool/Modes/ExecuteMode.cs b/Engine/Source/Programs/UnrealBuildTool/Modes/ExecuteMode.cs
--- a/Engine/Source/Programs/UnrealBuildTool/Modes/ExecuteMode.cs
+++ b/Engine/Source/Programs/UnrealBuildTool/Modes/ExecuteMode.cs
@@ -10,6 +10,7 @@
 using System.Threading.Tasks;
 using EpicGames.Core;
 using OpenTracing.Util;
+using Microsoft.Extensions.Logging;
 
 namespace UnrealBuildTool
 {
@@ -30,7 +31,8 @@
 		/// </summary>
 		/// <param name="Arguments">Command-line arguments</param>
 		/// <returns>One of the values of ECompilationResult</returns>
-		public override int Execute(CommandLineArguments Arguments)
+		/// <param name="Logger"></param>
+		public override int Execute(CommandLineArguments Arguments, ILogger Logger)
 		{
 			Arguments.ApplyTo(this);
 
@@ -58,8 +60,8 @@
 			// Execute the actions
 			using (GlobalTracer.Instance.BuildSpan("ActionGraph.ExecuteActions()").StartActive())
 			{
-				List<TargetDescriptor> TargetDescriptors = TargetDescriptor.ParseCommandLine(Arguments, false, false, false);
-				ActionGraph.ExecuteActions(BuildConfiguration, Actions, TargetDescriptors);
+				List<TargetDescriptor> TargetDescriptors = TargetDescriptor.ParseCommandLine(Arguments, false, false, false, Logger);
+				ActionGraph.ExecuteActions(BuildConfiguration, Actions, TargetDescriptors, Logger);
 			}
 
 			return 0;
diff --git a/Engine/Source/Programs/UnrealBuildTool/Modes/GenerateClangDatabase.cs b/Engine/Source/Programs/UnrealBuildTool/Modes/GenerateClangDatabase.cs
--- a/Engine/Source/Programs/UnrealBuildTool/Modes/GenerateClangDatabase.cs
+++ b/Engine/Source/Programs/UnrealBuildTool/Modes/GenerateClangDatabase.cs
@@ -7,6 +7,7 @@ using System.Text;
 using System.Threading.Tasks;
 using EpicGames.Core;
 using UnrealBuildBase;
+using Microsoft.Extensions.Logging;
 
 namespace UnrealBuildTool
 {
@@ -27,7 +28,8 @@ namespace UnrealBuildTool
 		/// </summary>
 		/// <param name="Arguments">Command line arguments</param>
 		/// <returns>Exit code</returns>
-		public override int Execute(CommandLineArguments Arguments)
+		/// <param name="Logger"></param>
+		public override int Execute(CommandLineArguments Arguments, ILogger Logger)
 		{
 			Arguments.ApplyTo(this);
 
@@ -51,7 +53,7 @@ namespace UnrealBuildTool
 			UEBuildModuleCPP.bForceAddGeneratedCodeIncludePath = true;
 
 			// Parse all the target descriptors
-			List<TargetDescriptor> TargetDescriptors = TargetDescriptor.ParseCommandLine(Arguments, BuildConfiguration.bUsePrecompiled, BuildConfiguration.bSkipRulesCompile, BuildConfiguration.bForceRulesCompile);
+			List<TargetDescriptor> TargetDescriptors = TargetDescriptor.ParseCommandLine(Arguments, BuildConfiguration.bUsePrecompiled, BuildConfiguration.bSkipRulesCompile, BuildConfiguration.bForceRulesCompile, Logger);
 
 			// Generate the compile DB for each target
 			using (ISourceFileWorkingSet WorkingSet = new EmptySourceFileWorkingSet())
@@ -64,15 +66,15 @@ namespace UnrealBuildTool
 					TargetDescriptor.AdditionalArguments = TargetDescriptor.AdditionalArguments.Append(new string[] { "-NoPCH", "-DisableUnity" });
 
 					// Create a makefile for the target
-					UEBuildTarget Target = UEBuildTarget.Create(TargetDescriptor, BuildConfiguration.bSkipRulesCompile, BuildConfiguration.bForceRulesCompile, BuildConfiguration.bUsePrecompiled);
+					UEBuildTarget Target = UEBuildTarget.Create(TargetDescriptor, BuildConfiguration.bSkipRulesCompile, BuildConfiguration.bForceRulesCompile, BuildConfiguration.bUsePrecompiled, Logger);
 
 					// Find the location of the compiler
-					FileReference ClangPath = FindClangCompiler(Target);
+					FileReference ClangPath = FindClangCompiler(Target, Logger);
 
 					bool IsWindowsClang = ClangPath.GetFileName().Equals("clang-cl.exe", StringComparison.OrdinalIgnoreCase);
 
 					// Create all the binaries and modules
-					CppCompileEnvironment GlobalCompileEnvironment = Target.CreateCompileEnvironmentForProjectFiles();
+					CppCompileEnvironment GlobalCompileEnvironment = Target.CreateCompileEnvironmentForProjectFiles(Logger);
 					foreach (UEBuildBinary Binary in Target.Binaries)
 					{
 						CppCompileEnvironment BinaryCompileEnvironment = Binary.CreateBinaryCompileEnvironment(GlobalCompileEnvironment);
@@ -170,8 +172,9 @@ namespace UnrealBuildTool
 		/// Searches for the Clang compiler for the given platform.
 		/// </summary>
 		/// <param name="Target">The build platform to use to search for the Clang compiler.</param>
+		/// <param name="Logger">Logger for output</param>
 		/// <returns>The path to the Clang compiler.</returns>
-		private static FileReference FindClangCompiler(UEBuildTarget Target)
+		private static FileReference FindClangCompiler(UEBuildTarget Target, ILogger Logger)
 		{
 			UnrealTargetPlatform HostPlatform = BuildHostPlatform.Current.Platform;
 
@@ -179,13 +182,14 @@ namespace UnrealBuildTool
 			{
 				VCEnvironment Environment = VCEnvironment.Create(WindowsCompiler.Clang, Target.Platform,
 					Target.Rules.WindowsPlatform.Architecture, null,
-					Target.Rules.WindowsPlatform.WindowsSdkVersion, null);
+					Target.Rules.WindowsPlatform.WindowsSdkVersion, null,
+					Logger);
 
 				return Environment.CompilerPath;
 			}
 			else if (OperatingSystem.IsLinux())
 			{
-				string? Clang = LinuxCommon.WhichClang();
+				string? Clang = LinuxCommon.WhichClang(Logger);
 
 				if (Clang != null)
 				{
@@ -194,7 +198,7 @@ namespace UnrealBuildTool
 			}
 			else if (OperatingSystem.IsMacOS())
 			{
-				MacToolChainSettings Settings = new MacToolChainSettings(false);
+				MacToolChainSettings Settings = new MacToolChainSettings(false, Logger);
 				DirectoryReference? ToolchainDir = DirectoryReference.FromString(Settings.ToolchainDir);
 
 				if (ToolchainDir != null)
diff --git a/Engine/Source/Programs/UnrealBuildTool/Modes/GenerateProjectFilesMode.cs b/Engine/Source/Programs/UnrealBuildTool/Modes/GenerateProjectFilesMode.cs
--- a/Engine/Source/Programs/UnrealBuildTool/Modes/GenerateProjectFilesMode.cs
+++ b/Engine/Source/Programs/UnrealBuildTool/Modes/GenerateProjectFilesMode.cs
@@ -11,6 +11,7 @@
 using System.Diagnostics;
 using UnrealBuildBase;
 using System.Diagnostics.CodeAnalysis;
+using Microsoft.Extensions.Logging;
 
 namespace UnrealBuildTool
 {
@@ -59,7 +60,8 @@
 		/// </summary>
 		/// <param name="Arguments">Command line arguments</param>
 		/// <returns>Exit code</returns>
-		public override int Execute(CommandLineArguments Arguments)
+		/// <param name="Logger"></param>
+		public override int Execute(CommandLineArguments Arguments, ILogger Logger)
 		{
 			// Apply any command line arguments to this class
 			Arguments.ApplyTo(this);
@@ -73,7 +75,7 @@
 
 			// Parse rocket-specific arguments.
 			FileReference? ProjectFile;
-			TryParseProjectFileArgument(Arguments, out ProjectFile);
+			TryParseProjectFileArgument(Arguments, Logger, out ProjectFile);
 
 			// Warn if there are explicit project file formats specified
 			if (ProjectFileFormats.Count > 0 && !bAutomated)
@@ -94,7 +96,7 @@
 				}
 				Configuration.Append("  </ProjectFileGenerator>\n");
 				Configuration.Append("</Configuration>\n");
-				Log.TraceWarning("{0}", Configuration.ToString());
+				Logger.LogWarning("{Configuration}", Configuration.ToString());
 			}
 
 			// If there aren't any formats set, read the default project file format from the config file
@@ -103,7 +105,7 @@
 				// Read from the XML config
 				if (!String.IsNullOrEmpty(ProjectFileGeneratorSettings.Format))
 				{
-					ProjectFileFormats.UnionWith(ProjectFileGeneratorSettings.ParseFormatList(ProjectFileGeneratorSettings.Format));
+					ProjectFileFormats.UnionWith(ProjectFileGeneratorSettings.ParseFormatList(ProjectFileGeneratorSettings.Format, Logger));
 				}
 
 				// Read from the editor config
@@ -131,8 +133,8 @@
 					{
 						if(DisablePlatformProjectGenerators == null || !DisablePlatformProjectGenerators.Any(x => x.Equals(Platform.ToString(), StringComparison.OrdinalIgnoreCase)))
 						{
-							Log.TraceVerbose("Registering project generator {0} for {1}", CheckType, Platform);
-							PlatformProjectGenerators.RegisterPlatformProjectGenerator(Platform, Generator);
+							Logger.LogDebug("Registering project generator {CheckType} for {Platform}", CheckType, Platform);
+							PlatformProjectGenerators.RegisterPlatformProjectGenerator(Platform, Generator, Logger);
 						}
 					}
 				}
@@ -140,7 +142,7 @@
 
 			// print out any errors to the log
 			List<string> BadPlatformNames = new List<string>();
-			Log.TraceLog("\n---   SDK INFO START   ---");
+			Logger.LogDebug("\n---   SDK INFO START   ---");
 			foreach (string PlatformName in UnrealTargetPlatform.GetValidPlatformNames())
 			{
 				UEBuildPlatformSDK? SDK = UEBuildPlatformSDK.GetSDKForPlatform(PlatformName);
@@ -157,10 +159,10 @@
 			if (BadPlatformNames.Count > 0)
 			{
 				Log.TraceInformationOnce("\nSome Platforms were skipped due to invalid SDK setup: {0}.\nSee the log file for detailed information\n\n", string.Join(", ", BadPlatformNames));
-				Log.TraceInformation("");
+				Logger.LogInformation("");
 			}
-			Log.TraceLog("---   SDK INFO END   ---");
-			Log.TraceLog("");
+			Logger.LogDebug("---   SDK INFO END   ---");
+			Logger.LogDebug("");
 
 			// Create each project generator and run it
 			List<ProjectFileGenerator> Generators = new List<ProjectFileGenerator>();
@@ -231,7 +233,7 @@
 			foreach(ProjectFileGenerator Generator in Generators)
 			{
 				ProjectFileGenerator.Current = Generator;
-				bool bGenerateSuccess = Generator.GenerateProjectFiles(PlatformProjectGenerators, Arguments.GetRawArray());
+				bool bGenerateSuccess = Generator.GenerateProjectFiles(PlatformProjectGenerators, Arguments.GetRawArray(), Logger);
 				ProjectFileGenerator.Current = null;
 
 				if (!bGenerateSuccess)
@@ -246,9 +248,10 @@
 		/// Try to parse the project file from the command line
 		/// </summary>
 		/// <param name="Arguments">The command line arguments</param>
+		/// <param name="Logger">Logger for output</param>
 		/// <param name="ProjectFile">The project file that was parsed</param>
 		/// <returns>True if the project file was parsed, false otherwise</returns>
-		private static bool TryParseProjectFileArgument(CommandLineArguments Arguments, [NotNullWhen(true)] out FileReference? ProjectFile)
+		private static bool TryParseProjectFileArgument(CommandLineArguments Arguments, ILogger Logger, [NotNullWhen(true)] out FileReference? ProjectFile)
 		{
 			string? CandidateProjectPath = null;
 
@@ -283,7 +286,7 @@
 					string ProjectName = CandidateProjectFile.ChangeExtension("uproject").GetFileName();
 
 					// check native project paths (uprojectdirs)
-					IEnumerable<FileReference> NativeProjectFiles = NativeProjects.EnumerateProjectFiles();
+					IEnumerable<FileReference> NativeProjectFiles = NativeProjects.EnumerateProjectFiles(Logger);
 
 					CandidateProjectFile = NativeProjectFiles.Where(F => F.GetFileName().Equals(ProjectName, StringComparison.InvariantCultureIgnoreCase)).FirstOrDefault();
 				}
@@ -294,7 +297,7 @@
 					throw new Exception(string.Format("Unable to find project file based on argument {0}", CandidateProjectPath));
 				}
 
-				Log.TraceVerbose("Resolved project argument {0} to {1}", CandidateProjectPath, CandidateProjectFile);
+				Logger.LogDebug("Resolved project argument {CandidateProjectPath} to {CandidateProjectFile}", CandidateProjectPath, CandidateProjectFile);
 				ProjectFile = CandidateProjectFile;
 				return true;
 			}
diff --git a/Engine/Source/Programs/UnrealBuildTool/Modes/JsonExportMode.cs b/Engine/Source/Programs/UnrealBuildTool/Modes/JsonExportMode.cs
--- a/Engine/Source/Programs/UnrealBuildTool/Modes/JsonExportMode.cs
+++ b/Engine/Source/Programs/UnrealBuildTool/Modes/JsonExportMode.cs
@@ -7,6 +7,7 @@
 using System.Threading.Tasks;
 using EpicGames.Core;
 using UnrealBuildBase;
+using Microsoft.Extensions.Logging;
 
 namespace UnrealBuildTool
 {
@@ -27,15 +28,16 @@
 		/// </summary>
 		/// <param name="Arguments">Command line arguments</param>
 		/// <returns>Exit code (always zero)</returns>
-		public override int Execute(CommandLineArguments Arguments)
+		/// <param name="Logger"></param>
+		public override int Execute(CommandLineArguments Arguments, ILogger Logger)
 		{
 			Arguments.ApplyTo(this);
 
-			List<TargetDescriptor> TargetDescriptors = TargetDescriptor.ParseCommandLine(Arguments, false, false, false);
+			List<TargetDescriptor> TargetDescriptors = TargetDescriptor.ParseCommandLine(Arguments, false, false, false, Logger);
 			foreach(TargetDescriptor TargetDescriptor in TargetDescriptors)
 			{
 				// Create the target
-				UEBuildTarget Target = UEBuildTarget.Create(TargetDescriptor, false, false, false);
+				UEBuildTarget Target = UEBuildTarget.Create(TargetDescriptor, false, false, false, Logger);
 
 				// Get the output file
 				FileReference? OutputFile = TargetDescriptor.AdditionalArguments.GetFileReferenceOrDefault("-OutputFile=", null);
@@ -55,7 +57,7 @@
 						Arguments.ApplyTo(BuildConfiguration);
 
 						// Create the makefile
-						TargetMakefile Makefile = Target.Build(BuildConfiguration, WorkingSet, TargetDescriptor);
+						TargetMakefile Makefile = Target.Build(BuildConfiguration, WorkingSet, TargetDescriptor, Logger);
 						List<LinkedAction> Actions = Makefile.Actions.ConvertAll(x => new LinkedAction(x, TargetDescriptor));
 						ActionGraph.Link(Actions);
 
@@ -66,14 +68,14 @@
 						// Execute these actions
 						if (PrerequisiteActions.Count > 0)
 						{
-							Log.TraceInformation("Exeucting actions that produce source files...");
-							ActionGraph.ExecuteActions(BuildConfiguration, PrerequisiteActions, new List<TargetDescriptor> { TargetDescriptor });
+							Logger.LogInformation("Exeucting actions that produce source files...");
+							ActionGraph.ExecuteActions(BuildConfiguration, PrerequisiteActions, new List<TargetDescriptor> { TargetDescriptor }, Logger);
 						}
 					}
 				}
 
 				// Write the output file
-				Log.TraceInformation("Writing {0}...", OutputFile);
+				Logger.LogInformation("Writing {OutputFile}...", OutputFile);
 				Target.ExportJson(OutputFile);
 			}
 			return 0;
diff --git a/Engine/Source/Programs/UnrealBuildTool/Modes/ParseMsvcTimingInfo.cs b/Engine/Source/Programs/UnrealBuildTool/Modes/ParseMsvcTimingInfo.cs
--- a/Engine/Source/Programs/UnrealBuildTool/Modes/ParseMsvcTimingInfo.cs
+++ b/Engine/Source/Programs/UnrealBuildTool/Modes/ParseMsvcTimingInfo.cs
@@ -6,6 +6,7 @@
 using System.Linq;
 using System.Text.RegularExpressions;
 using EpicGames.Core;
+using Microsoft.Extensions.Logging;
 
 namespace UnrealBuildTool
 {
@@ -18,7 +19,7 @@
 	{
 		const string TimingDataRegex = @"^\t\t(?<Indent>\t*)(?<Name>[^\t]+):\s*(?<Duration>[0-9\.]+)s$";
 
-		public override int Execute(CommandLineArguments Arguments)
+		public override int Execute(CommandLineArguments Arguments, ILogger Logger)
 		{
 			FileReference InputFile = Arguments.GetFileReference("-TimingFile=");
 
diff --git a/Engine/Source/Programs/UnrealBuildTool/Modes/QueryTargetsMode.cs b/Engine/Source/Programs/UnrealBuildTool/Modes/QueryTargetsMode.cs
--- a/Engine/Source/Programs/UnrealBuildTool/Modes/QueryTargetsMode.cs
+++ b/Engine/Source/Programs/UnrealBuildTool/Modes/QueryTargetsMode.cs
@@ -7,6 +7,7 @@
 using System.Threading.Tasks;
 using EpicGames.Core;
 using UnrealBuildBase;
+using Microsoft.Extensions.Logging;
 
 namespace UnrealBuildTool
 {
@@ -33,7 +34,8 @@
 		/// </summary>
 		/// <param name="Arguments">Command line arguments</param>
 		/// <returns></returns>
-		public override int Execute(CommandLineArguments Arguments)
+		/// <param name="Logger"></param>
+		public override int Execute(CommandLineArguments Arguments, ILogger Logger)
 		{
 			Arguments.ApplyTo(this);
 
@@ -52,16 +54,16 @@
 			RulesAssembly Assembly;
 			if (ProjectFile != null)
 			{
-				Assembly = RulesCompiler.CreateProjectRulesAssembly(ProjectFile, BuildConfiguration.bUsePrecompiled, BuildConfiguration.bSkipRulesCompile, BuildConfiguration.bForceRulesCompile);
+				Assembly = RulesCompiler.CreateProjectRulesAssembly(ProjectFile, BuildConfiguration.bUsePrecompiled, BuildConfiguration.bSkipRulesCompile, BuildConfiguration.bForceRulesCompile, Logger);
 			}
 			else
 			{
-				Assembly = RulesCompiler.CreateEngineRulesAssembly(BuildConfiguration.bUsePrecompiled, BuildConfiguration.bSkipRulesCompile, BuildConfiguration.bForceRulesCompile);
+				Assembly = RulesCompiler.CreateEngineRulesAssembly(BuildConfiguration.bUsePrecompiled, BuildConfiguration.bSkipRulesCompile, BuildConfiguration.bForceRulesCompile, Logger);
 			}
 
 			// Write information about these targets
-			WriteTargetInfo(ProjectFile, Assembly, OutputFile, Arguments);
-			Log.TraceInformation("Written {0}", OutputFile);
+			WriteTargetInfo(ProjectFile, Assembly, OutputFile, Arguments, Logger);
+			Logger.LogInformation("Written {OutputFile}", OutputFile);
 			return 0;
 		}
 
@@ -89,7 +91,8 @@
 		/// <param name="Assembly">The rules assembly for this target</param>
 		/// <param name="OutputFile">Output file to write to</param>
 		/// <param name="Arguments"></param>
-		public static void WriteTargetInfo(FileReference? ProjectFile, RulesAssembly Assembly, FileReference OutputFile, CommandLineArguments Arguments)
+		/// <param name="Logger">Logger for output</param>
+		public static void WriteTargetInfo(FileReference? ProjectFile, RulesAssembly Assembly, FileReference OutputFile, CommandLineArguments Arguments, ILogger Logger)
 		{
 			// Construct all the targets in this assembly
 			List<string> TargetNames = new List<string>();
@@ -115,12 +118,12 @@
 					try
 					{
 						string Architecture = UEBuildPlatform.GetBuildPlatform(BuildHostPlatform.Current.Platform).GetDefaultArchitecture(ProjectFile);
-						TargetRules = Assembly.CreateTargetRules(TargetName, BuildHostPlatform.Current.Platform, UnrealTargetConfiguration.Development, Architecture, ProjectFile, Arguments);
+						TargetRules = Assembly.CreateTargetRules(TargetName, BuildHostPlatform.Current.Platform, UnrealTargetConfiguration.Development, Architecture, ProjectFile, Arguments, Logger);
 					}
 					catch (Exception Ex)
 					{
-						Log.TraceWarning("Unable to construct target rules for {0}", TargetName);
-						Log.TraceVerbose(ExceptionUtils.FormatException(Ex));
+						Logger.LogWarning("Unable to construct target rules for {TargetName}", TargetName);
+						Logger.LogDebug("{Ex}", ExceptionUtils.FormatException(Ex));
 						continue;
 					}
 
diff --git a/Engine/Source/Programs/UnrealBuildTool/Modes/SetupPlatformsMode.cs b/Engine/Source/Programs/UnrealBuildTool/Modes/SetupPlatformsMode.cs
--- a/Engine/Source/Programs/UnrealBuildTool/Modes/SetupPlatformsMode.cs
+++ b/Engine/Source/Programs/UnrealBuildTool/Modes/SetupPlatformsMode.cs
@@ -6,6 +6,7 @@
 using System.Text;
 using System.Threading.Tasks;
 using EpicGames.Core;
+using Microsoft.Extensions.Logging;
 
 namespace UnrealBuildTool
 {
@@ -20,7 +21,8 @@
 		/// </summary>
 		/// <param name="Arguments">Command line arguments</param>
 		/// <returns>Exit code</returns>
-		public override int Execute(CommandLineArguments Arguments)
+		/// <param name="Logger"></param>
+		public override int Execute(CommandLineArguments Arguments, ILogger Logger)
 		{
 			// Output a warning if there are any arguments that are still unused
 			Arguments.CheckAllArgumentsUsed();
diff --git a/Engine/Source/Programs/UnrealBuildTool/Modes/TestMode.cs b/Engine/Source/Programs/UnrealBuildTool/Modes/TestMode.cs
--- a/Engine/Source/Programs/UnrealBuildTool/Modes/TestMode.cs
+++ b/Engine/Source/Programs/UnrealBuildTool/Modes/TestMode.cs
@@ -3,6 +3,7 @@
 using System.Collections.Generic;
 using EpicGames.Core;
 using UnrealBuildBase;
+using Microsoft.Extensions.Logging;
 
 namespace UnrealBuildTool
 {
@@ -16,7 +17,8 @@
 		/// Main entry point
 		/// </summary>
 		/// <param name="Arguments">Command-line arguments</param>
-		public override int Execute(CommandLineArguments Arguments)
+		/// <param name="Logger"></param>
+		public override int Execute(CommandLineArguments Arguments, ILogger Logger)
 		{
 			Arguments.ApplyTo(this);
 
@@ -26,9 +28,9 @@
 			XmlConfig.ApplyTo(BuildConfiguration);
 
 			// Parse all the targets being built
-			List<TargetDescriptor> TargetDescriptors = TargetDescriptor.ParseCommandLine(Arguments, false, false, false);
+			List<TargetDescriptor> TargetDescriptors = TargetDescriptor.ParseCommandLine(Arguments, false, false, false, Logger);
 
-			BuildTests(TargetDescriptors, BuildConfiguration);
+			BuildTests(TargetDescriptors, BuildConfiguration, Logger);
 
 			return 0;
 		}
@@ -43,7 +45,8 @@
 		/// </summary>
 		/// <param name="TargetDescriptors">Target descriptors</param>
 		/// <param name="BuildConfiguration">Current build configuration</param>
-		public static void BuildTests(List<TargetDescriptor> TargetDescriptors, BuildConfiguration BuildConfiguration)
+		/// <param name="Logger"></param>
+		public static void BuildTests(List<TargetDescriptor> TargetDescriptors, BuildConfiguration BuildConfiguration, ILogger Logger)
 		{
 			List<TargetDescriptor> TestTargetDescriptors = new List<TargetDescriptor>();
 
@@ -55,9 +58,9 @@
 				TestTargetDescriptors.Add(TestsTargetDescriptor);
 			}
 
-			using (ISourceFileWorkingSet WorkingSet = SourceFileWorkingSet.Create(Unreal.RootDirectory, new HashSet<DirectoryReference>()))
+			using (ISourceFileWorkingSet WorkingSet = SourceFileWorkingSet.Create(Unreal.RootDirectory, new HashSet<DirectoryReference>(), Logger))
 			{
-				BuildMode.Build(TestTargetDescriptors, BuildConfiguration, WorkingSet, BuildOptions.None, null);
+				BuildMode.Build(TestTargetDescriptors, BuildConfiguration, WorkingSet, BuildOptions.None, null, Logger);
 			}
 		}
 	}
diff --git a/Engine/Source/Programs/UnrealBuildTool/Modes/UnrealHeaderToolMode.cs b/Engine/Source/Programs/UnrealBuildTool/Modes/UnrealHeaderToolMode.cs
--- a/Engine/Source/Programs/UnrealBuildTool/Modes/UnrealHeaderToolMode.cs
+++ b/Engine/Source/Programs/UnrealBuildTool/Modes/UnrealHeaderToolMode.cs
@@ -11,6 +11,7 @@
 using System.Reflection;
 using System.Text;
 using UnrealBuildBase;
+using Microsoft.Extensions.Logging;
 
 namespace UnrealBuildTool.Modes
 {
@@ -534,7 +535,7 @@
 			public bool External;
 		}
 
-		private static bool RunScriptTest(UhtTables Tables, IUhtConfig Config, UhtGlobalOptions Options, string TestDirectory, string TestOutputDirectory, string Script)
+		private static bool RunScriptTest(UhtTables Tables, IUhtConfig Config, UhtGlobalOptions Options, string TestDirectory, string TestOutputDirectory, string Script, ILogger Logger)
 		{
 			string InPath = Path.Combine(TestDirectory, Script);
 			string OutPath = Path.Combine(TestOutputDirectory, Script);
@@ -743,7 +744,7 @@
 				// We trim the ends because it is too easy to leave off the ending CRLF in the script file.
 				if (ScriptFragments[ConsoleIndex].Body.ToString().TrimEnd() != SBConsole.ToString().TrimEnd())
 				{
-					Log.TraceError("Console output failed to match");
+					Logger.LogError("Console output failed to match");
 					bSuccess = false;
 				}
 
@@ -754,19 +755,19 @@
 					{
 						if (ScriptFragments[Index].Body.ToString().TrimEnd() != KVP.Value.TrimEnd())
 						{
-							Log.TraceError($"Output \"{KVP.Key}\" failed to match");
+							Logger.LogError("Output \"{Key}\" failed to match", KVP.Key);
 							bSuccess = false;
 						}
 						OutputFragments.Remove(KVP.Key);
 					}
 					else
 					{
-						Log.TraceError($"Output \"{KVP.Key}\" not found in test script");
+						Logger.LogError("Output \"{Key}\" not found in test script", KVP.Key);
 					}
 				}
 				foreach (KeyValuePair<string, int> KVP in OutputFragments)
 				{
-					Log.TraceError($"Output \"{KVP.Key}\" in test script but not generated");
+					Logger.LogError("Output \"{Key}\" in test script but not generated", KVP.Key);
 				}
 
 				// Create the complete output.  Output includes all of the source fragments and console fragments
@@ -802,42 +803,42 @@
 				}
 				catch (Exception E)
 				{
-					Log.TraceError("Unable to write test result to \"{0}\"", E.Message);
+					Logger.LogError(E, "Unable to write test result to \"{Ex}\"", E.Message);
 				}
 			}
 
 			if (bSuccess)
 			{
-				Log.TraceInformation("Test {0} succeeded", InPath);
+				Logger.LogInformation("Test {InPath} succeeded", InPath);
 			}
 			else
 			{
-				Log.TraceError("Test {0} failed", InPath);
+				Logger.LogError("Test {InPath} failed", InPath);
 			}
 			return bSuccess;
 		}
 
-		private static bool RunScriptTests(UhtTables Tables, IUhtConfig Config, UhtGlobalOptions Options, string TestDirectory, string TestOutputDirectory, List<string> Scripts)
+		private static bool RunScriptTests(UhtTables Tables, IUhtConfig Config, UhtGlobalOptions Options, string TestDirectory, string TestOutputDirectory, List<string> Scripts, ILogger Logger)
 		{
 			bool bResult = true;
 			foreach (string Script in Scripts)
 			{
-				bResult &= RunScriptTest(Tables, Config, Options, TestDirectory, TestOutputDirectory, Script);
+				bResult &= RunScriptTest(Tables, Config, Options, TestDirectory, TestOutputDirectory, Script, Logger);
 			}
 			return bResult;
 		}
 
-		private static bool RunDirectoryTests(UhtTables Tables, IUhtConfig Config, UhtGlobalOptions Options, string TestDirectory, string TestOutputDirectory, List<string> Directories)
+		private static bool RunDirectoryTests(UhtTables Tables, IUhtConfig Config, UhtGlobalOptions Options, string TestDirectory, string TestOutputDirectory, List<string> Directories, ILogger Logger)
 		{
 			bool bResult = true;
 			foreach (string Directory in Directories)
 			{
-				bResult &= RunTests(Tables, Config, Options, Path.Combine(TestDirectory, Directory), Path.Combine(TestOutputDirectory, Directory));
+				bResult &= RunTests(Tables, Config, Options, Path.Combine(TestDirectory, Directory), Path.Combine(TestOutputDirectory, Directory), Logger);
 			}
 			return bResult;
 		}
 
-		private static bool RunTests(UhtTables Tables, IUhtConfig Config, UhtGlobalOptions Options, string TestDirectory, string TestOutputDirectory)
+		private static bool RunTests(UhtTables Tables, IUhtConfig Config, UhtGlobalOptions Options, string TestDirectory, string TestOutputDirectory, ILogger Logger)
 		{
 			// Create output directory
 			Directory.CreateDirectory(TestOutputDirectory);
@@ -864,11 +865,11 @@
 			Manifests.Sort(StringComparer.OrdinalIgnoreCase);
 
 			return
-				RunScriptTests(Tables, Config, Options, TestDirectory, TestOutputDirectory, Scripts) &&
-				RunDirectoryTests(Tables, Config, Options, TestDirectory, TestOutputDirectory, Directories);
+				RunScriptTests(Tables, Config, Options, TestDirectory, TestOutputDirectory, Scripts, Logger) &&
+				RunDirectoryTests(Tables, Config, Options, TestDirectory, TestOutputDirectory, Directories, Logger);
 		}
 
-		public static bool RunTests(UhtTables Tables, IUhtConfig Config, UhtGlobalOptions Options)
+		public static bool RunTests(UhtTables Tables, IUhtConfig Config, UhtGlobalOptions Options, ILogger Logger)
 		{
 			DirectoryReference EngineSourceProgramDirectory = DirectoryReference.Combine(Unreal.EngineDirectory, "Source", "Programs");
 			string TestDirectory = FileReference.Combine(EngineSourceProgramDirectory, "UnrealBuildTool.Tests", "UHT").FullName;
@@ -883,11 +884,11 @@
 			{ }
 
 			// Collect a list of all the test scripts
-			Log.TraceInformation("Running tests in {0}", TestDirectory);
-			Log.TraceInformation("Output can be compared in {0}", TestOutputDirectory);
+			Logger.LogInformation("Running tests in {TestDirectory}", TestDirectory);
+			Logger.LogInformation("Output can be compared in {TestOutputDirectory}", TestOutputDirectory);
 
 			// Run the tests on the directory
-			return RunTests(Tables, Config, Options, TestDirectory, TestOutputDirectory);
+			return RunTests(Tables, Config, Options, TestDirectory, TestOutputDirectory, Logger);
 		}
 	}
 
@@ -977,7 +978,8 @@
 		/// </summary>
 		/// <param name="Arguments">Command line arguments</param>
 		/// <returns>Exit code</returns>
-		public override int Execute(CommandLineArguments Arguments)
+		/// <param name="Logger"></param>
+		public override int Execute(CommandLineArguments Arguments, ILogger Logger)
 		{
 			try
 			{
@@ -1039,7 +1041,7 @@
 				// If we are running test scripts
 				if (Options.bTest)
 				{
-					return UhtTestHarness.RunTests(Tables, Config, Options) ? (int)CompilationResult.Succeeded : (int)CompilationResult.OtherCompilationError;
+					return UhtTestHarness.RunTests(Tables, Config, Options, Logger) ? (int)CompilationResult.Succeeded : (int)CompilationResult.OtherCompilationError;
 				}
 
 				string? ProjectFile = null;
@@ -1048,17 +1050,17 @@
 				if (TargetArgumentIndex >= 0)
 				{
 					CommandLineArguments LocalArguments = new(new string[] { Arguments[TargetArgumentIndex] });
-					List<TargetDescriptor> TargetDescriptors = TargetDescriptor.ParseCommandLine(LocalArguments, false, false, false);
+					List<TargetDescriptor> TargetDescriptors = TargetDescriptor.ParseCommandLine(LocalArguments, false, false, false, Logger);
 					if (TargetDescriptors.Count == 0)
 					{
-						Log.TraceError("No target descriptors found.");
+						Logger.LogError("No target descriptors found.");
 						return (int)CompilationResult.OtherCompilationError;
 					}
 
 					TargetDescriptor TargetDesc = TargetDescriptors[0];
 
 					// Create the target
-					UEBuildTarget Target = UEBuildTarget.Create(TargetDesc, false, false, false);
+					UEBuildTarget Target = UEBuildTarget.Create(TargetDesc, false, false, false, Logger);
 
 					// Create the makefile for the target and export the module information
 					using ISourceFileWorkingSet WorkingSet = new EmptySourceFileWorkingSet();
@@ -1069,7 +1071,7 @@
 					Arguments.ApplyTo(BuildConfiguration);
 
 					// Create the makefile
-					TargetMakefile Makefile = Target.Build(BuildConfiguration, WorkingSet, TargetDesc, true);
+					TargetMakefile Makefile = Target.Build(BuildConfiguration, WorkingSet, TargetDesc, Logger, true);
 
 					FileReference ModuleInfoFileName = ExternalExecution.GetUHTModuleInfoFileName(Makefile, Target.TargetName);
 					FileReference DepsFileName = ExternalExecution.GetUHTDepsFileName(ModuleInfoFileName);
@@ -1136,8 +1138,8 @@
 			catch (Exception Ex)
 			{
 				// Unhandled exception.
-				Log.TraceError("Unhandled exception: {0}", ExceptionUtils.FormatException(Ex));
-				Log.TraceLog(ExceptionUtils.FormatExceptionDetails(Ex));
+				Logger.LogError(Ex, "Unhandled exception: {Ex}", ExceptionUtils.FormatException(Ex));
+				Logger.LogDebug(Ex, "Unhandled exception: {Ex}", ExceptionUtils.FormatExceptionDetails(Ex));
 				return (int)CompilationResult.OtherCompilationError;
 			}
 		}
diff --git a/Engine/Source/Programs/UnrealBuildTool/Modes/ValidatePlatformsMode.cs b/Engine/Source/Programs/UnrealBuildTool/Modes/ValidatePlatformsMode.cs
--- a/Engine/Source/Programs/UnrealBuildTool/Modes/ValidatePlatformsMode.cs
+++ b/Engine/Source/Programs/UnrealBuildTool/Modes/ValidatePlatformsMode.cs
@@ -6,6 +6,7 @@
 using System.Text;
 using System.Threading.Tasks;
 using EpicGames.Core;
+using Microsoft.Extensions.Logging;
 
 namespace UnrealBuildTool
 {
@@ -38,7 +39,8 @@
 		/// </summary>
 		/// <param name="Arguments">Command line arguments</param>
 		/// <returns>Exit code</returns>
-		public override int Execute(CommandLineArguments Arguments)
+		/// <param name="Logger"></param>
+		public override int Execute(CommandLineArguments Arguments, ILogger Logger)
 		{
 			// Output a message if there are any arguments that are still unused
 			Arguments.ApplyTo(this);
@@ -62,11 +64,11 @@
 
 				if (BuildPlatform != null && BuildPlatform.HasRequiredSDKsInstalled() == SDKStatus.Valid)
 				{
-					Log.TraceInformation("##PlatformValidate: {0} VALID {1}", Platform.ToString(), PlatformSDKString);
+					Logger.LogInformation("##PlatformValidate: {Platform} VALID {PlatformSdkString}", Platform.ToString(), PlatformSDKString);
 				}
 				else
 				{
-					Log.TraceInformation("##PlatformValidate: {0} INVALID {1}", Platform.ToString(), PlatformSDKString);
+					Logger.LogInformation("##PlatformValidate: {Platform} INVALID {PlatformSdkString}", Platform.ToString(), PlatformSDKString);
 				}
 			} 
 			return 0;
diff --git a/Engine/Source/Programs/UnrealBuildTool/Modes/WriteDocumentationMode.cs b/Engine/Source/Programs/UnrealBuildTool/Modes/WriteDocumentationMode.cs
--- a/Engine/Source/Programs/UnrealBuildTool/Modes/WriteDocumentationMode.cs
+++ b/Engine/Source/Programs/UnrealBuildTool/Modes/WriteDocumentationMode.cs
@@ -7,6 +7,7 @@
 using System.Threading.Tasks;
 using EpicGames.Core;
 using UnrealBuildBase;
+using Microsoft.Extensions.Logging;
 
 namespace UnrealBuildTool
 {
@@ -42,7 +43,7 @@
 		/// Entry point for this command
 		/// </summary>
 		/// <returns></returns>
-		public override int Execute(CommandLineArguments Arguments)
+		public override int Execute(CommandLineArguments Arguments, ILogger Logger)
 		{
 			Arguments.ApplyTo(this);
 			Arguments.CheckAllArgumentsUsed();
@@ -50,13 +51,13 @@
 			switch(Type)
 			{
 				case DocumentationType.BuildConfiguration:
-					XmlConfig.WriteDocumentation(OutputFile);
+					XmlConfig.WriteDocumentation(OutputFile, Logger);
 					break;
 				case DocumentationType.ModuleRules:
-					RulesDocumentation.WriteDocumentation(typeof(ModuleRules), OutputFile);
+					RulesDocumentation.WriteDocumentation(typeof(ModuleRules), OutputFile, Logger);
 					break;
 				case DocumentationType.TargetRules:
-					RulesDocumentation.WriteDocumentation(typeof(TargetRules), OutputFile);
+					RulesDocumentation.WriteDocumentation(typeof(TargetRules), OutputFile, Logger);
 					break;
 				default:
 					throw new BuildException("Invalid documentation type: {0}", Type);
