UBT: Use ILogger throughout UBT, rather than legacy EpicGames.Core.Log methods. All output going forwards should use structured logging rather than String.Format style logging.

diff --git a/Engine/Source/Programs/UnrealBuildTool/Modes/WriteMetadataMode.cs b/Engine/Source/Programs/UnrealBuildTool/Modes/WriteMetadataMode.cs
--- a/Engine/Source/Programs/UnrealBuildTool/Modes/WriteMetadataMode.cs
+++ b/Engine/Source/Programs/UnrealBuildTool/Modes/WriteMetadataMode.cs
@@ -10,6 +10,7 @@
 using System.Threading.Tasks;
 using EpicGames.Core;
 using UnrealBuildBase;
+using Microsoft.Extensions.Logging;
 
 namespace UnrealBuildTool
 {
@@ -86,14 +87,15 @@
 		/// </summary>
 		/// <param name="Arguments">Command line arguments</param>
 		/// <returns>Exit code</returns>
-		public override int Execute(CommandLineArguments Arguments)
+		/// <param name="Logger"></param>
+		public override int Execute(CommandLineArguments Arguments, ILogger Logger)
 		{
 			// Acquire a different mutex to the regular UBT instance, since this mode will be called as part of a build. We need the mutex to ensure that building two modular configurations 
 			// in parallel don't clash over writing shared *.modules files (eg. DebugGame and Development editors).
 			string MutexName = SingleInstanceMutex.GetUniqueMutexForPath("UnrealBuildTool_WriteMetadata", Unreal.RootDirectory.FullName);
 			using(new SingleInstanceMutex(MutexName, true))
 			{
-				return ExecuteInternal(Arguments);
+				return ExecuteInternal(Arguments, Logger);
 			}
 		}
 
@@ -101,8 +103,9 @@
 		/// Execute the command, having obtained the appropriate mutex
 		/// </summary>
 		/// <param name="Arguments">Command line arguments</param>
+		/// <param name="Logger">Logger for output</param>
 		/// <returns>Exit code</returns>
-		private int ExecuteInternal(CommandLineArguments Arguments)
+		private int ExecuteInternal(CommandLineArguments Arguments, ILogger Logger)
 		{
 			// Read the target info
 			WriteMetadataTargetInfo TargetInfo = BinaryFormatterUtils.Load<WriteMetadataTargetInfo>(Arguments.GetFileReference("-Input="));
@@ -126,7 +129,7 @@
 			{
 				BuildId = TargetInfo.Receipt.Version.BuildId;
 			}
-			else if (TargetInfo.VersionFile != null && BuildVersion.TryRead(TargetInfo.VersionFile, out BuildVersion? PrevVersion) && CanRecycleBuildId(PrevVersion.BuildId, TargetInfo.FileToManifest))
+			else if (TargetInfo.VersionFile != null && BuildVersion.TryRead(TargetInfo.VersionFile, out BuildVersion? PrevVersion) && CanRecycleBuildId(PrevVersion.BuildId, TargetInfo.FileToManifest, Logger))
 			{
 				BuildId = PrevVersion.BuildId;
 			}
@@ -139,7 +142,7 @@
 			foreach(KeyValuePair<FileReference, ModuleManifest> Pair in TargetInfo.FileToManifest)
 			{
 				ModuleManifest? SourceManifest;
-				if(TryReadManifest(Pair.Key, out SourceManifest) && SourceManifest.BuildId == BuildId)
+				if(TryReadManifest(Pair.Key, Logger, out SourceManifest) && SourceManifest.BuildId == BuildId)
 				{
 					MergeManifests(SourceManifest, Pair.Value);
 				}
@@ -176,7 +179,7 @@
 							string CurrentText = FileReference.ReadAllText(ManifestFile);
 							if (CurrentText != OutputText)
 							{
-								Log.TraceError("Build modifies {0}. This is not permitted. Before:\n    {1}\nAfter:\n    {2}", ManifestFile, CurrentText.Replace("\n", "\n    "), OutputText.Replace("\n", "\n    "));
+								Logger.LogError("Build modifies {File}. This is not permitted. Before:\n    {OldFile}\nAfter:\n    {NewFile}", ManifestFile, CurrentText.Replace("\n", "\n    "), OutputText.Replace("\n", "\n    "));
 							}
 						}
 
@@ -210,13 +213,14 @@
 		/// </summary>
 		/// <param name="BuildId"></param>
 		/// <param name="FileToManifest"></param>
+		/// <param name="Logger"></param>
 		/// <returns></returns>
-		bool CanRecycleBuildId(string? BuildId, Dictionary<FileReference, ModuleManifest> FileToManifest)
+		bool CanRecycleBuildId(string? BuildId, Dictionary<FileReference, ModuleManifest> FileToManifest, ILogger Logger)
 		{
 			foreach (FileReference ManifestFileName in FileToManifest.Keys)
 			{
 				ModuleManifest? Manifest;
-				if (ManifestFileName.IsUnderDirectory(Unreal.EngineDirectory) && TryReadManifest(ManifestFileName, out Manifest) && Manifest.BuildId == BuildId)
+				if (ManifestFileName.IsUnderDirectory(Unreal.EngineDirectory) && TryReadManifest(ManifestFileName, Logger, out Manifest) && Manifest.BuildId == BuildId)
 				{
 					DateTime ManifestTime = FileReference.GetLastWriteTimeUtc(ManifestFileName);
 					foreach (string FileName in Manifest.ModuleNameToFileName.Values)
@@ -236,9 +240,10 @@
 		/// Attempts to read a manifest from the given location
 		/// </summary>
 		/// <param name="ManifestFileName">Path to the manifest</param>
+		/// <param name="Logger"></param>
 		/// <param name="Manifest">If successful, receives the manifest that was read</param>
 		/// <returns>True if the manifest was read correctly, false otherwise</returns>
-		public static bool TryReadManifest(FileReference ManifestFileName, [NotNullWhen(true)] out ModuleManifest? Manifest)
+		public static bool TryReadManifest(FileReference ManifestFileName, ILogger Logger, [NotNullWhen(true)] out ModuleManifest? Manifest)
 		{
 			if(FileReference.Exists(ManifestFileName))
 			{
@@ -249,8 +254,8 @@
 				}
 				catch(Exception Ex)
 				{
-					Log.TraceWarning("Unable to read '{0}'; ignoring.", ManifestFileName);
-					Log.TraceLog(ExceptionUtils.FormatExceptionDetails(Ex));
+					Logger.LogWarning("Unable to read '{ManifestFileName}'; ignoring.", ManifestFileName);
+					Logger.LogDebug("{Ex}", ExceptionUtils.FormatExceptionDetails(Ex));
 				}
 			}
 
diff --git a/Engine/Source/Programs/UnrealBuildTool/Platform/Android/AndroidAARHandler.cs b/Engine/Source/Programs/UnrealBuildTool/Platform/Android/AndroidAARHandler.cs
--- a/Engine/Source/Programs/UnrealBuildTool/Platform/Android/AndroidAARHandler.cs
+++ b/Engine/Source/Programs/UnrealBuildTool/Platform/Android/AndroidAARHandler.cs
@@ -11,6 +11,7 @@
 using EpicGames.Core;
 
 using Ionic.Zip;
+using Microsoft.Extensions.Logging;
 
 namespace UnrealBuildTool
 {
@@ -60,16 +61,17 @@
 		/// Add a new respository path to search for AAR and JAR files
 		/// </summary>
 		/// <param name="RepositoryPath">Directory containing the repository</param>
-		public void AddRepository(string RepositoryPath)
+		/// <param name="Logger">Logger for output</param>
+		public void AddRepository(string RepositoryPath, ILogger Logger)
 		{
 			if (Directory.Exists(RepositoryPath))
 			{
-				Log.TraceInformation("Added repository: {0}", RepositoryPath);
+				Logger.LogInformation("Added repository: {RepositoryPath}", RepositoryPath);
 				Repositories!.Add(RepositoryPath);
 			}
 			else
 			{
-				Log.TraceWarning("AddRepository: Directory {0} not found!", RepositoryPath);
+				Logger.LogWarning("AddRepository: Directory {RepositoryPath} not found!", RepositoryPath);
 			}
 		}
 
@@ -78,7 +80,8 @@
 		/// </summary>
 		/// <param name="RepositoryPath">Root directory containing the repository</param>
 		/// <param name="SearchPattern">Search pattern to match</param>
-		public void AddRepositories(string RepositoryPath, string SearchPattern)
+		/// <param name="Logger">Logger for output</param>
+		public void AddRepositories(string RepositoryPath, string SearchPattern, ILogger Logger)
 		{
 			if (Directory.Exists(RepositoryPath))
 			{
@@ -93,7 +96,7 @@
 					{
 						if (SearchPath.Contains(SearchPattern))
 						{
-							Log.TraceInformation("Added repository: {0}", SearchPath);
+							Logger.LogInformation("Added repository: {SearchPath}", SearchPath);
 							Repositories!.Add(SearchPath);
 						}
 						else
@@ -105,20 +108,20 @@
 			}
 			else
 			{
-				Log.TraceInformation("AddRepositories: Directory {0} not found; ignored", RepositoryPath);
+				Logger.LogInformation("AddRepositories: Directory {RepositoryPath} not found; ignored", RepositoryPath);
 			}
 		}
 
-		public void DumpAAR()
+		public void DumpAAR(ILogger Logger)
 		{
-			Log.TraceInformation("ALL DEPENDENCIES");
+			Logger.LogInformation("ALL DEPENDENCIES");
 			foreach (AndroidAAREntry Entry in AARList!)
 			{
-				Log.TraceInformation("{0}", Entry.Filename);
+				Logger.LogInformation("{EntryFilename}", Entry.Filename);
 			}
 			foreach (AndroidAAREntry Entry in JARList!)
 			{
-				Log.TraceInformation("{0}", Entry.Filename);
+				Logger.LogInformation("{EntryFilename}", Entry.Filename);
 			}
 		}
 
@@ -194,7 +197,7 @@
 		// clean up the version (Maven version info here: https://docs.oracle.com/middleware/1212/core/MAVEN/maven_version.htm)
 		// only going to handle a few cases, not proper ranges (keeps the rightmost valid version which should be highest)
 		// will still return something but will include an error in log, but don't want to throw an exception
-		private string CleanupVersion(string Filename, string InVersion)
+		private string CleanupVersion(string Filename, string InVersion, ILogger Logger)
 		{
 			string WorkVersion = InVersion;
 
@@ -216,7 +219,7 @@
 			// if not left with a possibly valid number, stop
 			if (!HasAnyVersionCharacters(WorkVersion))
 			{
-				Log.TraceError("AAR Dependency file {0} version unknown! {1}", Filename, InVersion);
+				Logger.LogError("AAR Dependency file {Filename} version unknown! {InVersion}", Filename, InVersion);
 				return InVersion;
 			}
 
@@ -230,7 +233,7 @@
 			}
 
 			// give an error, not likely going to work, though
-			Log.TraceError("AAR Dependency file {0} version unknown! {1}", Filename, InVersion);
+			Logger.LogError("AAR Dependency file {Filename} version unknown! {InVersion}", Filename, InVersion);
 			return InVersion;
 		}
 
@@ -240,12 +243,13 @@
 		/// <param name="PackageName">Name of the package the JAR belongs to in repository</param>
 		/// <param name="BaseName">Directory in repository containing the JAR</param>
 		/// <param name="Version">Version of the AAR to use</param>
-		public void AddNewJAR(string PackageName, string BaseName, string Version)
+		/// <param name="Logger">Logger instance</param>
+		public void AddNewJAR(string PackageName, string BaseName, string Version, ILogger Logger)
 		{
 			string? BasePath = FindPackageFile(PackageName, BaseName, Version);
 			if (BasePath == null)
 			{
-				Log.TraceError("AAR: Unable to find package {0}!", PackageName + "/" + BaseName);
+				Logger.LogError("AAR: Unable to find package {Package}!", PackageName + "/" + BaseName);
 				return;
 			}
 			string BaseFilename = Path.Combine(BasePath, BaseName + "-" + Version);
@@ -263,7 +267,7 @@
 						return;
 					}
 
-					Log.TraceInformation("AAR: {0}: {1} newer than {2}", JARList[JARIndex].BaseName, Version, JARList[JARIndex].Version);
+					Logger.LogInformation("AAR: {BaseName}: {Version1} newer than {Version2}", JARList[JARIndex].BaseName, Version, JARList[JARIndex].Version);
 
 					// This is a newer version; remove old one
 					JARList.RemoveAt(JARIndex);
@@ -271,7 +275,7 @@
 				}
 			}
 
-			//Log.TraceInformation("JAR: {0}", BaseName);
+			//Logger.LogInformation("JAR: {BaseName}", BaseName);
 			AndroidAAREntry AAREntry = new AndroidAAREntry(BaseName, Version, BaseFilename);
 			JARList.Add(AAREntry);
 
@@ -286,13 +290,13 @@
 				}
 				catch (Exception e)
 				{
-					Log.TraceError("AAR Dependency file {0} parsing error! {1}", DependencyFilename, e);
+					Logger.LogError("AAR Dependency file {File} parsing error! {Ex}", DependencyFilename, e);
 					return;
 				}
 			}
 			else
 			{
-				Log.TraceError("AAR: Dependency file {0} missing!", DependencyFilename);
+				Logger.LogError("AAR: Dependency file {DependencyFilename} missing!", DependencyFilename);
 				return;
 			}
 
@@ -308,11 +312,11 @@
 			{
 				string DepGroupId = GetElementValue(DependNode, GroupIdName, "");
 				string DepArtifactId = GetElementValue(DependNode, ArtifactIdName, "");
-				string DepVersion = CleanupVersion(DependencyFilename + "." + DepGroupId + "." + DepArtifactId, GetElementValue(DependNode, VersionName, ""));
+				string DepVersion = CleanupVersion(DependencyFilename + "." + DepGroupId + "." + DepArtifactId, GetElementValue(DependNode, VersionName, ""), Logger);
 				string DepScope = GetElementValue(DependNode, ScopeName, "compile");
 				string DepType = GetElementValue(DependNode, TypeName, "jar");
 
-				//Log.TraceInformation("Dependency: {0} {1} {2} {3} {4}", DepGroupId, DepArtifactId, DepVersion, DepScope, DepType);
+				//Logger.LogInformation("Dependency: {DepGroupId} {DepArtifactId} {DepVersion} {DepScope} {DepType}", DepGroupId, DepArtifactId, DepVersion, DepScope, DepType);
 
 				// ignore test scope
 				if (DepScope == "test")
@@ -322,11 +326,11 @@
 
 				if (DepType == "aar")
 				{
-					AddNewAAR(DepGroupId, DepArtifactId, DepVersion);
+					AddNewAAR(DepGroupId, DepArtifactId, DepVersion, Logger);
 				}
 				else if (DepType == "jar")
 				{
-					AddNewJAR(DepGroupId, DepArtifactId, DepVersion);
+					AddNewJAR(DepGroupId, DepArtifactId, DepVersion, Logger);
 				}
 			}
 		}
@@ -337,8 +341,9 @@
 		/// <param name="PackageName">Name of the package the AAR belongs to in repository</param>
 		/// <param name="BaseName">Directory in repository containing the AAR</param>
 		/// <param name="Version">Version of the AAR to use</param>
+		/// <param name="Logger">Logger for output</param>
 		/// <param name="HandleDependencies">Optionally process POM file for dependencies (default)</param>
-		public void AddNewAAR(string PackageName, string BaseName, string Version, bool HandleDependencies = true)
+		public void AddNewAAR(string PackageName, string BaseName, string Version, ILogger Logger, bool HandleDependencies = true)
 		{
 			if (!HandleDependencies)
 			{
@@ -350,7 +355,7 @@
 			string? BasePath = FindPackageFile(PackageName, BaseName, Version);
 			if (BasePath == null)
 			{
-				Log.TraceError("AAR: Unable to find package {0}!", PackageName + "/" + BaseName);
+				Logger.LogError("AAR: Unable to find package {Package}!", PackageName + "/" + BaseName);
 				return;
 			}
 			string BaseFilename = Path.Combine(BasePath, BaseName + "-" + Version);
@@ -368,7 +373,7 @@
 						return;
 					}
 
-					Log.TraceInformation("AAR: {0}: {1} newer than {2}", AARList[AARIndex].BaseName, Version, AARList[AARIndex].Version);
+					Logger.LogInformation("AAR: {BaseName}: {Version1} newer than {Version2}", AARList[AARIndex].BaseName, Version, AARList[AARIndex].Version);
 
 					// This is a newer version; remove old one
 					// @TODO: be smarter about dependency cleanup (newer AAR might not need older dependencies)
@@ -377,7 +382,7 @@
 				}
 			}
 
-			//Log.TraceInformation("AAR: {0}", BaseName);
+			//Logger.LogInformation("AAR: {BaseName}", BaseName);
 			AndroidAAREntry AAREntry = new AndroidAAREntry(BaseName, Version, BaseFilename);
 			AARList.Add(AAREntry);
 
@@ -397,13 +402,13 @@
 				}
 				catch (Exception e)
 				{
-					Log.TraceError("AAR Dependency file {0} parsing error! {1}", DependencyFilename, e);
+					Logger.LogError("AAR Dependency file {File} parsing error! {Ex}", DependencyFilename, e);
 					return;
 				}
 			}
 			else
 			{
-				Log.TraceError("AAR: Dependency file {0} missing!", DependencyFilename);
+				Logger.LogError("AAR: Dependency file {DependencyFilename} missing!", DependencyFilename);
 				return;
 			}
 
@@ -419,11 +424,11 @@
 			{
 				string DepGroupId = GetElementValue(DependNode, GroupIdName, "");
 				string DepArtifactId = GetElementValue(DependNode, ArtifactIdName, "");
-				string DepVersion = CleanupVersion(DependencyFilename + "." + DepGroupId + "." + DepArtifactId, GetElementValue(DependNode, VersionName, ""));
+				string DepVersion = CleanupVersion(DependencyFilename + "." + DepGroupId + "." + DepArtifactId, GetElementValue(DependNode, VersionName, ""), Logger);
 				string DepScope = GetElementValue(DependNode, ScopeName, "compile");
 				string DepType = GetElementValue(DependNode, TypeName, "jar");
 
-				//Log.TraceInformation("Dependency: {0} {1} {2} {3} {4}", DepGroupId, DepArtifactId, DepVersion, DepScope, DepType);
+				//Logger.LogInformation("Dependency: {DepGroupId} {DepArtifactId} {DepVersion} {DepScope} {DepType}", DepGroupId, DepArtifactId, DepVersion, DepScope, DepType);
 
 				// ignore test scope
 				if (DepScope == "test")
@@ -436,12 +441,12 @@
 					// Add dependency
 					AAREntry.AddDependency(DepArtifactId, DepVersion);
 
-					AddNewAAR(DepGroupId, DepArtifactId, DepVersion);
+					AddNewAAR(DepGroupId, DepArtifactId, DepVersion, Logger);
 				}
 				else
 				if (DepType == "jar")
 				{
-					AddNewJAR(DepGroupId, DepArtifactId, DepVersion);
+					AddNewJAR(DepGroupId, DepArtifactId, DepVersion, Logger);
 				}
 			}
 		}
@@ -467,7 +472,8 @@
 		/// Copies the required JAR files to the provided directory
 		/// </summary>
 		/// <param name="DestinationPath">Destination path for JAR files</param>
-		public void CopyJARs(string DestinationPath)
+		/// <param name="Logger">Logger for output</param>
+		public void CopyJARs(string DestinationPath, ILogger Logger)
 		{
 			MakeDirectoryIfRequired(DestinationPath);
 			DestinationPath = Path.Combine(DestinationPath, "libs");
@@ -478,15 +484,15 @@
 				string Filename = Entry.Filename + ".jar";
 				string BaseName = Path.GetFileName(Filename);
 				string TargetPath = Path.Combine(DestinationPath, BaseName);
-                //Log.TraceInformation("Attempting to copy JAR {0} {1} {2}", Filename, BaseName, TargetPath);
+                //Logger.LogInformation("Attempting to copy JAR {Filename} {BaseName} {TargetPath}", Filename, BaseName, TargetPath);
 
                 if (!File.Exists(Filename))
                 {
-                    Log.TraceInformation("JAR doesn't exist! {0}", Filename);
+                    Logger.LogInformation("JAR doesn't exist! {Filename}", Filename);
                 }
                 if (!File.Exists(TargetPath))
 				{
-					Log.TraceInformation("Copying JAR {0}", BaseName);
+					Logger.LogInformation("Copying JAR {BaseName}", BaseName);
 					File.Copy(Filename, TargetPath);
 				}
 			}
@@ -497,13 +503,14 @@
 		/// </summary>
 		/// <param name="DestinationPath">Destination path for AAR files</param>
 		/// <param name="AppPackageName">Name of the package these AARs are being used with</param>
-		public void ExtractAARs(string DestinationPath, string AppPackageName)
+		/// <param name="Logger">Logger for output</param>
+		public void ExtractAARs(string DestinationPath, string AppPackageName, ILogger Logger)
 		{
 			MakeDirectoryIfRequired(DestinationPath);
 			DestinationPath = Path.Combine(DestinationPath, "JavaLibs");
 			MakeDirectoryIfRequired(DestinationPath);
 
-			Log.TraceInformation("Extracting AARs");
+			Logger.LogInformation("Extracting AARs");
 			foreach (AndroidAAREntry Entry in AARList!)
 			{
 				string BaseName = Path.GetFileName(Entry.Filename);
@@ -513,8 +520,8 @@
 				string TargetManifestFileName = Path.Combine(TargetPath, "AndroidManifest.xml");
 				if (!File.Exists(TargetManifestFileName))
 				{
-					Log.TraceInformation("Extracting AAR {0}", BaseName);
-					/*IEnumerable<string> FileNames =*/ UnzipFiles(Entry.Filename + ".aar", TargetPath);
+					Logger.LogInformation("Extracting AAR {BaseName}", BaseName);
+					/*IEnumerable<string> FileNames =*/ UnzipFiles(Entry.Filename + ".aar", TargetPath, Logger);
 
 					// Must have a src directory (even if empty)
 					string SrcDirectory = Path.Combine(TargetPath, "src");
@@ -560,7 +567,7 @@
 							}
 							catch (Exception e)
 							{
-								Log.TraceError("AAR Manifest file {0} parsing error! {1}", ManifestFilename, e);
+								Logger.LogError("AAR Manifest file {ManifestFile} parsing error! {Ex}", ManifestFilename, e);
 							}
 						}
 
@@ -594,8 +601,9 @@
 		/// </summary>
 		/// <param name="ZipFileName">Name of the zip file</param>
 		/// <param name="BaseDirectory">Output directory</param>
+		/// <param name="Logger">Logger for output</param>
 		/// <returns>List of files written</returns>
-		public static IEnumerable<string> UnzipFiles(string ZipFileName, string BaseDirectory)
+		public static IEnumerable<string> UnzipFiles(string ZipFileName, string BaseDirectory, ILogger Logger)
 		{
 			// manually extract the files. There was a problem with the Ionic.Zip library that required this on non-PC at one point,
 			// but that problem is now fixed. Leaving this code as is as we need to return the list of created files and fix up their permissions anyway.
@@ -613,7 +621,7 @@
 						string _ZipName = Path.GetFileNameWithoutExtension(ZipFileName);
 						string NewOutputFileName = Path.Combine(Path.GetDirectoryName(OutputFileName)!,
 							_ZipName + '-' + Path.GetFileNameWithoutExtension(OutputFileName) + '.' + Path.GetExtension(OutputFileName));
-						Log.TraceInformation("Changed FileName {0} => {1}", Entry.FileName, NewOutputFileName);
+						Logger.LogInformation("Changed FileName {EntryFileName} => {NewOutputFileName}", Entry.FileName, NewOutputFileName);
 						OutputFileName = NewOutputFileName;
 					}
 
diff --git a/Engine/Source/Programs/UnrealBuildTool/Platform/Android/AndroidExports.cs b/Engine/Source/Programs/UnrealBuildTool/Platform/Android/AndroidExports.cs
--- a/Engine/Source/Programs/UnrealBuildTool/Platform/Android/AndroidExports.cs
+++ b/Engine/Source/Programs/UnrealBuildTool/Platform/Android/AndroidExports.cs
@@ -6,6 +6,7 @@
 using System.Text;
 using System.Threading.Tasks;
 using EpicGames.Core;
+using Microsoft.Extensions.Logging;
 
 namespace UnrealBuildTool
 {
@@ -89,7 +90,7 @@
 		/// <returns></returns>
 		public static IAndroidToolChain CreateToolChain(FileReference ProjectFile)
 		{
-			return new AndroidToolChain(ProjectFile, false, null, null);
+			return new AndroidToolChain(ProjectFile, false, null, null, Log.Logger);
 		}
 
 		/// <summary>
@@ -98,7 +99,7 @@
 		/// <returns></returns>
 		public static IAndroidToolChain CreateTempToolChain()
 		{
-			return new AndroidToolChain(null, false, null, null);
+			return new AndroidToolChain(null, false, null, null, Log.Logger);
 		}
 
 		/// <summary>
@@ -109,7 +110,7 @@
 		/// <returns></returns>
 		public static IAndroidDeploy CreateDeploymentHandler(FileReference ProjectFile, bool InForcePackageData)
 		{
-			return new UEDeployAndroid(ProjectFile, InForcePackageData);
+			return new UEDeployAndroid(ProjectFile, InForcePackageData, Log.Logger);
 		}
 
 		/// <summary>
@@ -136,10 +137,11 @@
 		/// </summary>
 		/// <param name="SourceFile"></param>
 		/// <param name="TargetFile"></param>
-		public static void StripSymbols(FileReference SourceFile, FileReference TargetFile)
+		/// <param name="Logger">Logger for output</param>
+		public static void StripSymbols(FileReference SourceFile, FileReference TargetFile, ILogger Logger)
 		{
-			AndroidToolChain ToolChain = new AndroidToolChain(null, false, null, null);
-			ToolChain.StripSymbols(SourceFile, TargetFile);
+			AndroidToolChain ToolChain = new AndroidToolChain(null, false, null, null, Logger);
+			ToolChain.StripSymbols(SourceFile, TargetFile, Logger);
 		}
 	}
 }
diff --git a/Engine/Source/Programs/UnrealBuildTool/Platform/Android/AndroidPlatformSDK.cs b/Engine/Source/Programs/UnrealBuildTool/Platform/Android/AndroidPlatformSDK.cs
--- a/Engine/Source/Programs/UnrealBuildTool/Platform/Android/AndroidPlatformSDK.cs
+++ b/Engine/Source/Programs/UnrealBuildTool/Platform/Android/AndroidPlatformSDK.cs
@@ -6,6 +6,7 @@
 using EpicGames.Core;
 using System.Text.RegularExpressions;
 using System.Linq;
+using Microsoft.Extensions.Logging;
 
 ///////////////////////////////////////////////////////////////////
 // If you are looking for supported version numbers, look in the
@@ -16,6 +17,11 @@
 {
 	partial class AndroidPlatformSDK : UEBuildPlatformSDK
 	{
+		public AndroidPlatformSDK(ILogger Logger)
+			: base(Logger)
+		{
+		}
+
 		protected override string? GetInstalledSDKVersion()
 		{
 			string? NDKPath = Environment.GetEnvironmentVariable("NDKROOT");
diff --git a/Engine/Source/Programs/UnrealBuildTool/Platform/Android/AndroidProjectGenerator.cs b/Engine/Source/Programs/UnrealBuildTool/Platform/Android/AndroidProjectGenerator.cs
--- a/Engine/Source/Programs/UnrealBuildTool/Platform/Android/AndroidProjectGenerator.cs
+++ b/Engine/Source/Programs/UnrealBuildTool/Platform/Android/AndroidProjectGenerator.cs
@@ -6,6 +6,7 @@
 using System.Diagnostics;
 using System.IO;
 using EpicGames.Core;
+using Microsoft.Extensions.Logging;
 
 namespace UnrealBuildTool
 {
@@ -35,8 +36,8 @@
 		/// </summary>
 		private bool AGDEInstalled = false;
 
-		public AndroidProjectGenerator(CommandLineArguments Arguments)
-			: base(Arguments)
+		public AndroidProjectGenerator(CommandLineArguments Arguments, ILogger Logger)
+			: base(Arguments, Logger)
 		{
 			if (OperatingSystem.IsWindows())
 			{
@@ -94,7 +95,7 @@
 
 			if (VSDebugCommandLineOptionPresent && VSDebuggingEnabled == false)
 			{
-				Log.TraceWarning("Android SDK tools have to be installed to use this.  Please Install Visual C++ for Cross-Platform Mobile Development https://msdn.microsoft.com/en-us/library/dn707598.aspx");
+				Logger.LogWarning("Android SDK tools have to be installed to use this.  Please Install Visual C++ for Cross-Platform Mobile Development https://msdn.microsoft.com/en-us/library/dn707598.aspx");
 			}
 
 			return VSDebuggingEnabled;
@@ -129,7 +130,7 @@
 			if (String.IsNullOrEmpty(PlatformToolsetVersion))
 			{
 				// future maintainer: add toolset version and verify that the rest of the msbuild path, version, and location in ProgramFiles(x86) is still valid
-				Log.TraceInformation("Android project generation needs to be updated for this version of Visual Studio.");
+				Logger.LogInformation("Android project generation needs to be updated for this version of Visual Studio.");
 				return false;
 			}
 
@@ -177,7 +178,7 @@
 
 			if (!NsightInstalled)
 			{
-				Log.TraceInformation("\nNsight Tegra {0}.{1} found, but Nsight Tegra 1.5 or higher is required for debugging support.", NsightVersion.ProductMajorPart, NsightVersion.ProductMinorPart);
+				Logger.LogInformation("\nNsight Tegra {NsightVersionProductMajorPart}.{NsightVersionProductMinorPart} found, but Nsight Tegra 1.5 or higher is required for debugging support.", NsightVersion.ProductMajorPart, NsightVersion.ProductMinorPart);
 			}
 
 			return NsightInstalled;
@@ -239,7 +240,6 @@
 		/// <param name="InPlatform">  The UnrealTargetPlatform being built</param>
 		/// <param name="ProjectFileFormat"></param>
 		/// <param name="ProjectFileBuilder">String builder for the project file</param>
-		/// <returns>string    The custom property import lines for the project file; Empty string if it doesn't require one</returns>
 		public override void GetAdditionalVisualStudioPropertyGroups(UnrealTargetPlatform InPlatform, VCProjectFileFormat ProjectFileFormat, StringBuilder ProjectFileBuilder)
 		{
 			if (AGDEInstalled || IsVSAndroidSupportInstalled() || !IsNsightInstalled(ProjectFileFormat))
@@ -472,7 +472,7 @@
 								"	</ItemGroup>" + ProjectFileGenerator.NewLine +
 								"</Project>";
 
-			ProjectFileGenerator.WriteFileIfChanged(ProjectFileGenerator.IntermediateProjectFilesPath + "\\" + FileName, FileText);
+			ProjectFileGenerator.WriteFileIfChanged(ProjectFileGenerator.IntermediateProjectFilesPath + "\\" + FileName, FileText, Logger);
 
 		}
 
@@ -505,7 +505,7 @@
 								"	</PropertyGroup>" + ProjectFileGenerator.NewLine +
 								"</Project>";
 
-			ProjectFileGenerator.WriteFileIfChanged(ProjectFileGenerator.IntermediateProjectFilesPath + "\\" + FileName, FileText);
+			ProjectFileGenerator.WriteFileIfChanged(ProjectFileGenerator.IntermediateProjectFilesPath + "\\" + FileName, FileText, Logger);
 
 		}
 
@@ -730,7 +730,7 @@
                                 "	<ImportGroup Label=\"ExtensionTargets\" /> " + ProjectFileGenerator.NewLine +
 								"</Project>";
 			
-			bool Success = ProjectFileGenerator.WriteFileIfChanged(ProjectFileGenerator.IntermediateProjectFilesPath + "\\" + FileName, FileText);
+			bool Success = ProjectFileGenerator.WriteFileIfChanged(ProjectFileGenerator.IntermediateProjectFilesPath + "\\" + FileName, FileText, Logger);
 
 			FileReference ProjectFilePath = FileReference.Combine(ProjectFileGenerator.IntermediateProjectFilesPath, FileName);
 			AndroidDebugProjectFile Project = new AndroidDebugProjectFile(ProjectFilePath, ProjectFile.BaseDir);
@@ -779,8 +779,9 @@
 		/// <param name="SolutionConfiguration">The solution configuration</param>
 		/// <param name="SolutionPlatform">The solution platform</param>
 		/// <param name="PlatformProjectGenerators">Set of platform project generators</param>
+		/// <param name="Logger"></param>
 		/// <returns>Project context matching the given solution context</returns>
-		public override MSBuildProjectContext GetMatchingProjectContext(TargetType SolutionTarget, UnrealTargetConfiguration SolutionConfiguration, UnrealTargetPlatform SolutionPlatform, PlatformProjectGeneratorCollection PlatformProjectGenerators)
+		public override MSBuildProjectContext GetMatchingProjectContext(TargetType SolutionTarget, UnrealTargetConfiguration SolutionConfiguration, UnrealTargetPlatform SolutionPlatform, PlatformProjectGeneratorCollection PlatformProjectGenerators, ILogger Logger)
 		{
 			return new MSBuildProjectContext("Debug", "ARM"){ bBuildByDefault = (SolutionPlatform == UnrealTargetPlatform.Android) };
 		}
diff --git a/Engine/Source/Programs/UnrealBuildTool/Platform/Android/AndroidToolChain.cs b/Engine/Source/Programs/UnrealBuildTool/Platform/Android/AndroidToolChain.cs
--- a/Engine/Source/Programs/UnrealBuildTool/Platform/Android/AndroidToolChain.cs
+++ b/Engine/Source/Programs/UnrealBuildTool/Platform/Android/AndroidToolChain.cs
@@ -9,6 +9,7 @@ using System.Linq;
 using Microsoft.Win32;
 using EpicGames.Core;
 using UnrealBuildBase;
+using Microsoft.Extensions.Logging;
 
 namespace UnrealBuildTool
 {
@@ -119,18 +120,18 @@ namespace UnrealBuildTool
 
 		protected bool bEnableGcSections = true;
 
-		public AndroidToolChain(FileReference? InProjectFile, bool bInUseLdGold, IReadOnlyList<string>? InAdditionalArches, IReadOnlyList<string>? InAdditionalGPUArches)
-			: this(InProjectFile, bInUseLdGold, InAdditionalArches, InAdditionalGPUArches, false, ClangToolChainOptions.None)
+		public AndroidToolChain(FileReference? InProjectFile, bool bInUseLdGold, IReadOnlyList<string>? InAdditionalArches, IReadOnlyList<string>? InAdditionalGPUArches, ILogger InLogger)
+			: this(InProjectFile, bInUseLdGold, InAdditionalArches, InAdditionalGPUArches, false, ClangToolChainOptions.None, InLogger)
 		{
 		}
 
-		public AndroidToolChain(FileReference? InProjectFile, bool bInUseLdGold, IReadOnlyList<string>? InAdditionalArches, IReadOnlyList<string>? InAdditionalGPUArches, ClangToolChainOptions ToolchainOptions)
-			: this(InProjectFile, bInUseLdGold, InAdditionalArches, InAdditionalGPUArches, false, ToolchainOptions)
+		public AndroidToolChain(FileReference? InProjectFile, bool bInUseLdGold, IReadOnlyList<string>? InAdditionalArches, IReadOnlyList<string>? InAdditionalGPUArches, ClangToolChainOptions ToolchainOptions, ILogger InLogger)
+			: this(InProjectFile, bInUseLdGold, InAdditionalArches, InAdditionalGPUArches, false, ToolchainOptions, InLogger)
 		{
 		}
 
-		protected AndroidToolChain(FileReference? InProjectFile, bool bInUseLdGold, IReadOnlyList<string>? InAdditionalArches, IReadOnlyList<string>? InAdditionalGPUArches, bool bAllowMissingNDK, ClangToolChainOptions ToolchainOptions)
-			: base(ToolchainOptions)
+		protected AndroidToolChain(FileReference? InProjectFile, bool bInUseLdGold, IReadOnlyList<string>? InAdditionalArches, IReadOnlyList<string>? InAdditionalGPUArches, bool bAllowMissingNDK, ClangToolChainOptions ToolchainOptions, ILogger InLogger)
+			: base(ToolchainOptions, InLogger)
 		{
 			Options = ToolchainOptions;
 			ProjectFile = InProjectFile;
@@ -150,7 +151,7 @@ namespace UnrealBuildTool
 
 			// by default tools chains don't parse arguments, but we want to be able to check the -architectures flag defined above. This is
 			// only necessary when AndroidToolChain is used during UAT
-			CommandLine.ParseArguments(Environment.GetCommandLineArgs(), this);
+			CommandLine.ParseArguments(Environment.GetCommandLineArgs(), this, Logger);
 
 			if (AdditionalArches.Count == 0 && ArchitectureArg.Count > 0)
 			{
@@ -180,23 +181,23 @@ namespace UnrealBuildTool
 
 			if (Directory.Exists(Path.Combine(NDKPath, "toolchains", "llvm", ArchitecturePathWindows64)))
 			{
-				Log.TraceVerbose("        Found Windows 64 bit versions of toolchain");
+				Logger.LogDebug("        Found Windows 64 bit versions of toolchain");
 				ArchitecturePath = ArchitecturePathWindows64;
 			}
 			else if (Directory.Exists(Path.Combine(NDKPath, "toolchains", "llvm", ArchitecturePathWindows32)))
 			{
-				Log.TraceVerbose("        Found Windows 32 bit versions of toolchain");
+				Logger.LogDebug("        Found Windows 32 bit versions of toolchain");
 				ArchitecturePath = ArchitecturePathWindows32;
 			}
 			else if (Directory.Exists(Path.Combine(NDKPath, "toolchains", "llvm", ArchitecturePathMac)))
 			{
-				Log.TraceVerbose("        Found Mac versions of toolchain");
+				Logger.LogDebug("        Found Mac versions of toolchain");
 				ArchitecturePath = ArchitecturePathMac;
 				ExeExtension = "";
 			}
 			else if (Directory.Exists(Path.Combine(NDKPath, "toolchains", "llvm", ArchitecturePathLinux)))
 			{
-				Log.TraceVerbose("        Found Linux versions of toolchain");
+				Logger.LogDebug("        Found Linux versions of toolchain");
 				ArchitecturePath = ArchitecturePathLinux;
 				ExeExtension = "";
 			}
@@ -232,7 +233,7 @@ namespace UnrealBuildTool
 			ClangPath = Utils.CollapseRelativeDirectories(Path.Combine(NDKPath, @"toolchains/llvm", ArchitecturePath, @"bin/clang++" + ExeExtension));
 
 			// Android (6317467 based on r365631c1) clang version 9.0.8 
-			AndroidClangBuild = Utils.RunLocalProcessAndReturnStdOut(ClangPath, "--version");
+			AndroidClangBuild = Utils.RunLocalProcessAndReturnStdOut(ClangPath, "--version", Logger);
 			try
 			{
 				AndroidClangBuild = Regex.Match(AndroidClangBuild, @"(\w+) based on").Groups[1].ToString();
@@ -243,7 +244,7 @@ namespace UnrealBuildTool
 			}
 			catch
 			{
-				Log.TraceWarning("Failed to retreive build version from {0}", AndroidClangBuild);
+				Logger.LogWarning("Failed to retreive build version from {AndroidClangBuild}", AndroidClangBuild);
 				AndroidClangBuild = "unknown";
 			}
 
@@ -1048,7 +1049,7 @@ namespace UnrealBuildTool
 			}
 		}
 
-		void GenerateEmptyLinkFunctionsForRemovedModules(List<FileItem> SourceFiles, string ModuleName, DirectoryReference OutputDirectory, IActionGraphBuilder Graph)
+		void GenerateEmptyLinkFunctionsForRemovedModules(List<FileItem> SourceFiles, string ModuleName, DirectoryReference OutputDirectory, IActionGraphBuilder Graph, ILogger Logger)
 		{
 			// Only add to Launch module
 			if (!ModuleName.Equals("Launch"))
@@ -1163,7 +1164,7 @@ namespace UnrealBuildTool
 				// Directly added NDK files for NDK extensions
 				ModifySourceFiles(CompileEnvironment, InputFiles, ModuleName);
 				// Deal with dynamic modules removed by architecture
-				GenerateEmptyLinkFunctionsForRemovedModules(InputFiles, ModuleName, OutputDir, Graph);
+				GenerateEmptyLinkFunctionsForRemovedModules(InputFiles, ModuleName, OutputDir, Graph, Logger);
 
 				bHasHandledLaunchModule = true;
 			}
@@ -1567,7 +1568,7 @@ namespace UnrealBuildTool
 
 					Result.GeneratedHeaderFiles.Add(TargetFileItem);
 
-					Log.TraceVerbose("   ISPC Generating Header " + CompileAction.StatusDescription + ": \"" + CompileAction.CommandPath + "\"" + CompileAction.CommandArguments);
+					Logger.LogDebug("   ISPC Generating Header {StatusDescription}: \"{CommandPath}\" {CommandArguments}", CompileAction.StatusDescription, CompileAction.CommandPath, CompileAction.CommandArguments);
 				}
 			}
 
@@ -1772,7 +1773,7 @@ namespace UnrealBuildTool
 
 					Result.ObjectFiles.AddRange(FinalISPCObjFiles);
 
-					Log.TraceVerbose("   ISPC Compiling " + CompileAction.StatusDescription + ": \"" + CompileAction.CommandPath + "\"" + CompileAction.CommandArguments);
+					Logger.LogDebug("   ISPC Compiling {StatusDescription}: \"{CommandPath}\" {CommandArguments}", CompileAction.StatusDescription, CompileAction.CommandPath, CompileAction.CommandArguments);
 				}
 			}
 
@@ -1999,7 +2000,7 @@ namespace UnrealBuildTool
 					SetupActionToExecuteCompilerThroughShell(ref LinkAction, LinkAction.CommandPath.FullName, LinkAction.CommandArguments, "Link");
 				}
 
-				//Log.TraceInformation("Link: {0} {1}", LinkAction.CommandPath.FullName, LinkAction.CommandArguments);
+				//Logger.LogInformation("Link: {LinkActionCommandPathFullName} {LinkActionCommandArguments}", LinkAction.CommandPath.FullName, LinkAction.CommandArguments);
 
 				// Windows can run into an issue with too long of a commandline when clang tries to call ld to link.
 				// To work around this we call clang to just get the command it would execute and generate a
@@ -2105,11 +2106,11 @@ namespace UnrealBuildTool
 			*/
 		}
 
-		public static void OutputReceivedDataEventHandler(Object Sender, DataReceivedEventArgs Line)
+		public static void OutputReceivedDataEventHandler(Object Sender, DataReceivedEventArgs Line, ILogger Logger)
 		{
 			if ((Line != null) && (Line.Data != null))
 			{
-				Log.TraceInformation(Line.Data);
+				Logger.LogInformation("{Output}", Line.Data);
 			}
 		}
 
@@ -2136,7 +2137,7 @@ namespace UnrealBuildTool
 			return StripExe;
 		}
 
-		public void StripSymbols(FileReference SourceFile, FileReference TargetFile)
+		public void StripSymbols(FileReference SourceFile, FileReference TargetFile, ILogger Logger)
 		{
 			if (SourceFile != TargetFile)
 			{
@@ -2149,7 +2150,7 @@ namespace UnrealBuildTool
 			StartInfo.Arguments = " --strip-debug \"" + TargetFile.FullName + "\"";
 			StartInfo.UseShellExecute = false;
 			StartInfo.CreateNoWindow = true;
-			Utils.RunLocalProcessAndLogOutput(StartInfo);
+			Utils.RunLocalProcessAndLogOutput(StartInfo, Logger);
 		}
 
 		protected virtual void SetupActionToExecuteCompilerThroughShell(ref Action CompileOrLinkAction, string CommandPath, string CommandArguments, string CommandDescription)
diff --git a/Engine/Source/Programs/UnrealBuildTool/Platform/Android/UEBuildAndroid.cs b/Engine/Source/Programs/UnrealBuildTool/Platform/Android/UEBuildAndroid.cs
--- a/Engine/Source/Programs/UnrealBuildTool/Platform/Android/UEBuildAndroid.cs
+++ b/Engine/Source/Programs/UnrealBuildTool/Platform/Android/UEBuildAndroid.cs
@@ -5,6 +5,8 @@ using System.Collections.Generic;
 using System.Text;
 using System.IO;
 using EpicGames.Core;
+using Microsoft.Extensions.Logging;
+using Microsoft.Extensions.Logging.Abstractions;
 
 namespace UnrealBuildTool
 {
@@ -129,13 +131,13 @@ namespace UnrealBuildTool
 	{
 		UEBuildPlatformSDK SDK;
 
-		public AndroidPlatform(UnrealTargetPlatform InTargetPlatform, UEBuildPlatformSDK InSDK) 
-			: base(InTargetPlatform, InSDK)
+		public AndroidPlatform(UnrealTargetPlatform InTargetPlatform, UEBuildPlatformSDK InSDK, ILogger InLogger) 
+			: base(InTargetPlatform, InSDK, InLogger)
 		{
 			SDK = InSDK;
 		}
 
-		public AndroidPlatform(AndroidPlatformSDK InSDK) : this(UnrealTargetPlatform.Android, InSDK)
+		public AndroidPlatform(AndroidPlatformSDK InSDK, ILogger InLogger) : this(UnrealTargetPlatform.Android, InSDK, InLogger)
 		{
 		}
 
@@ -216,7 +218,7 @@ namespace UnrealBuildTool
 			}
 		}
 
-		public virtual bool HasSpecificDefaultBuildConfig(UnrealTargetPlatform Platform, DirectoryReference ProjectPath)
+		public virtual bool HasSpecificDefaultBuildConfig(UnrealTargetPlatform Platform, DirectoryReference ProjectPath, ILogger Logger)
 		{
 			string[] BoolKeys = new string[] {
 				"bBuildForArm64", "bBuildForX8664", 
@@ -228,7 +230,7 @@ namespace UnrealBuildTool
 
 			// look up Android specific settings
 			if (!DoProjectSettingsMatchDefault(Platform, ProjectPath, "/Script/AndroidRuntimeSettings.AndroidRuntimeSettings",
-				BoolKeys, null, StringKeys))
+				BoolKeys, null, StringKeys, Logger))
 			{
 				return false;
 			}
@@ -238,7 +240,7 @@ namespace UnrealBuildTool
 		public override bool HasDefaultBuildConfig(UnrealTargetPlatform Platform, DirectoryReference ProjectPath)
 		{
 			// @todo Lumin: This is kinda messy - better way?
-			if (HasSpecificDefaultBuildConfig(Platform, ProjectPath) == false)
+			if (HasSpecificDefaultBuildConfig(Platform, ProjectPath, Logger) == false)
 			{
 				return false;
 			}
@@ -362,12 +364,12 @@ namespace UnrealBuildTool
 		{
 		}
 
-		public virtual void SetUpSpecificEnvironment(ReadOnlyTargetRules Target, CppCompileEnvironment CompileEnvironment, LinkEnvironment LinkEnvironment)
+		public virtual void SetUpSpecificEnvironment(ReadOnlyTargetRules Target, CppCompileEnvironment CompileEnvironment, LinkEnvironment LinkEnvironment, ILogger Logger)
 		{
 			string NDKPath = Environment.GetEnvironmentVariable("NDKROOT")!;
 			NDKPath = NDKPath.Replace("\"", "");
 
-			AndroidToolChain ToolChain = new AndroidToolChain(Target.ProjectFile, false, Target.AndroidPlatform.Architectures, Target.AndroidPlatform.GPUArchitectures);
+			AndroidToolChain ToolChain = new AndroidToolChain(Target.ProjectFile, false, Target.AndroidPlatform.Architectures, Target.AndroidPlatform.GPUArchitectures, Logger);
 
 			// figure out the NDK version
 			string? NDKToolchainVersion = SDK.GetInstalledVersion();
@@ -377,7 +379,7 @@ namespace UnrealBuildTool
 			// PLATFORM_ANDROID_NDK_VERSION is in the form 150100, where 15 is major version, 01 is the letter (1 is 'a'), 00 indicates beta revision if letter is 00
 			CompileEnvironment.Definitions.Add(string.Format("PLATFORM_ANDROID_NDK_VERSION={0}", NDKVersionInt));
 
-			Log.TraceInformation("NDK toolchain: {0}, NDK version: {1}, ClangVersion: {2}", NDKToolchainVersion, NDKVersionInt, ToolChain.GetClangVersionString());
+			Logger.LogInformation("NDK toolchain: {Version}, NDK version: {NdkVersion}, ClangVersion: {ClangVersion}", NDKToolchainVersion, NDKVersionInt, ToolChain.GetClangVersionString());
 
 			CompileEnvironment.Definitions.Add("PLATFORM_DESKTOP=0");
 			CompileEnvironment.Definitions.Add("PLATFORM_CAN_SUPPORT_EDITORONLY_DATA=0");
@@ -439,7 +441,7 @@ namespace UnrealBuildTool
 
 			if (Target.bPGOOptimize || Target.bPGOProfile)
 			{
-				Log.TraceInformation("PGO {0} build", Target.bPGOOptimize ? "optimize" : "profile");
+				Logger.LogInformation("PGO {PgoType} build", Target.bPGOOptimize ? "optimize" : "profile");
 				if(Target.bPGOOptimize)
 				{
 					CompileEnvironment.PGODirectory = Path.Combine(DirectoryReference.FromFile(Target.ProjectFile)!.FullName, "Platforms", "Android", "Build", "PGO");
@@ -448,14 +450,14 @@ namespace UnrealBuildTool
 					LinkEnvironment.PGODirectory = CompileEnvironment.PGODirectory;
 					LinkEnvironment.PGOFilenamePrefix = CompileEnvironment.PGOFilenamePrefix;
 
-					Log.TraceInformation("PGO Dir: {0}", CompileEnvironment.PGODirectory);
-					Log.TraceInformation("PGO Prefix: {0}", CompileEnvironment.PGOFilenamePrefix);
+					Logger.LogInformation("PGO Dir: {PgoDir}", CompileEnvironment.PGODirectory);
+					Logger.LogInformation("PGO Prefix: {PgoPrefix}", CompileEnvironment.PGOFilenamePrefix);
 				}
 			}
 
 			CompileEnvironment.Definitions.Add("INT64_T_TYPES_NOT_LONG_LONG=1");
 
-			SetUpSpecificEnvironment(Target, CompileEnvironment, LinkEnvironment);
+			SetUpSpecificEnvironment(Target, CompileEnvironment, LinkEnvironment, Logger);
 
 			// deliberately not linking stl or stdc++ here (c++_shared is default)
 			LinkEnvironment.SystemLibraries.Add("c");
@@ -516,24 +518,21 @@ namespace UnrealBuildTool
 		{
 			bool bUseLdGold = Target.bUseUnityBuild;
 			ClangToolChainOptions Options = CreateToolChainOptions(Target.AndroidPlatform.TargetRules);
-			return new AndroidToolChain(Target.ProjectFile, bUseLdGold, Target.AndroidPlatform.Architectures, Target.AndroidPlatform.GPUArchitectures, Options);
+			return new AndroidToolChain(Target.ProjectFile, bUseLdGold, Target.AndroidPlatform.Architectures, Target.AndroidPlatform.GPUArchitectures, Options, Logger);
 		}
 		public virtual UEToolChain CreateTempToolChainForProject(FileReference? ProjectFile)
 		{
 			AndroidTargetRules TargetRules = new AndroidTargetRules();
-			CommandLine.ParseArguments(Environment.GetCommandLineArgs(), TargetRules);
+			CommandLine.ParseArguments(Environment.GetCommandLineArgs(), TargetRules, Logger);
 			ClangToolChainOptions Options = CreateToolChainOptions(TargetRules);
-			return new AndroidToolChain(ProjectFile, true, null, null, Options);
+			return new AndroidToolChain(ProjectFile, true, null, null, Options, Logger);
 		}
 
-		/// <summary>
-		/// Deploys the given target
-		/// </summary>
-		/// <param name="Receipt">Receipt for the target being deployed</param>
+		/// <inheritdoc/>
 		public override void Deploy(TargetReceipt Receipt)
 		{
 			// do not package data if building via UBT
-			new UEDeployAndroid(Receipt.ProjectFile, false).PrepTargetForDeployment(Receipt);
+			new UEDeployAndroid(Receipt.ProjectFile, false, Logger).PrepTargetForDeployment(Receipt);
 		}
 	}
 
@@ -545,12 +544,12 @@ namespace UnrealBuildTool
 			get { return UnrealTargetPlatform.Android; }
 		}
 
-		public override void RegisterBuildPlatforms()
+		public override void RegisterBuildPlatforms(ILogger Logger)
 		{
-			AndroidPlatformSDK SDK = new AndroidPlatformSDK();
+			AndroidPlatformSDK SDK = new AndroidPlatformSDK(Logger);
 
 			// Register this build platform
-			UEBuildPlatform.RegisterBuildPlatform(new AndroidPlatform(SDK));
+			UEBuildPlatform.RegisterBuildPlatform(new AndroidPlatform(SDK, Logger), Logger);
 			UEBuildPlatform.RegisterPlatformWithGroup(UnrealTargetPlatform.Android, UnrealPlatformGroup.Android);
 		}
 	}
