UBT: Use ILogger throughout UBT, rather than legacy EpicGames.Core.Log methods. All output going forwards should use structured logging rather than String.Format style logging.

diff --git a/Engine/Source/Programs/UnrealBuildTool/Platform/Android/UEDeployAndroid.cs b/Engine/Source/Programs/UnrealBuildTool/Platform/Android/UEDeployAndroid.cs
--- a/Engine/Source/Programs/UnrealBuildTool/Platform/Android/UEDeployAndroid.cs
+++ b/Engine/Source/Programs/UnrealBuildTool/Platform/Android/UEDeployAndroid.cs
@@ -12,6 +12,7 @@ using System.Xml.Linq;
 using EpicGames.Core;
 using System.Security.Cryptography;
 using UnrealBuildBase;
+using Microsoft.Extensions.Logging;
 
 namespace UnrealBuildTool
 {
@@ -65,7 +66,8 @@ namespace UnrealBuildTool
 		[CommandLine("-ForceAPKGeneration", Value = "true")]
 		public bool ForceAPKGeneration = false;
 
-		public UEDeployAndroid(FileReference? InProjectFile, bool InForcePackageData)
+		public UEDeployAndroid(FileReference? InProjectFile, bool InForcePackageData, ILogger InLogger)
+			: base(InLogger)
 		{
 			ProjectFile = InProjectFile;
 
@@ -73,7 +75,7 @@ namespace UnrealBuildTool
 			bool IniValue = ReadPackageDataInsideApkFromIni(null);
 			bPackageDataInsideApk = InForcePackageData || IniValue == true;
 
-			CommandLine.ParseArguments(Environment.GetCommandLineArgs(), this);
+			CommandLine.ParseArguments(Environment.GetCommandLineArgs(), this, Logger);
 		}
 
 		private UnrealPluginLanguage? UPL = null;
@@ -143,7 +145,7 @@ namespace UnrealBuildTool
 				}
 			}
 
-			UPL = new UnrealPluginLanguage(ProjectFile, inPluginExtraData, NDKArches, "http://schemas.android.com/apk/res/android", "xmlns:android=\"http://schemas.android.com/apk/res/android\"", UnrealTargetPlatform.Android);
+			UPL = new UnrealPluginLanguage(ProjectFile, inPluginExtraData, NDKArches, "http://schemas.android.com/apk/res/android", "xmlns:android=\"http://schemas.android.com/apk/res/android\"", UnrealTargetPlatform.Android, Logger);
 			UPLHashCode = UPL.GetUPLHash();
 //			APL.SetTrace();
 		}
@@ -251,7 +253,7 @@ namespace UnrealBuildTool
 				int SDKLevelInt = GetApiLevelInt(SDKLevel);
 				if (SDKLevelInt < MinimumSDKLevel)
 				{
-					Log.TraceInformation("Requires at least SDK API level {0}, currently set to '{1}'", MinimumSDKLevel, SDKLevel);
+					Logger.LogInformation("Requires at least SDK API level {MinimumSDKLevel}, currently set to '{SDKLevel}'", MinimumSDKLevel, SDKLevel);
 					SDKLevel = ToolChain.GetLargestApiLevel();
 
 					SDKLevelInt = GetApiLevelInt(SDKLevel);
@@ -259,7 +261,7 @@ namespace UnrealBuildTool
 					{
 						SDKLevelInt = MinimumSDKLevel;
 						SDKLevel = "android-" + MinimumSDKLevel.ToString();
-						Log.TraceInformation("Gradle will attempt to download SDK API level {0}", SDKLevelInt);
+						Logger.LogInformation("Gradle will attempt to download SDK API level {SDKLevelInt}", SDKLevelInt);
 					}
 				}
 
@@ -267,10 +269,10 @@ namespace UnrealBuildTool
 				string PlatformsDir = Environment.ExpandEnvironmentVariables("%ANDROID_HOME%/platforms");
 				if (!ValidateSDK(PlatformsDir, SDKLevel))
 				{
-					Log.TraceWarning("The SDK API requested '{0}' not installed in {1}; Gradle will attempt to download it.", SDKLevel, PlatformsDir);
+					Logger.LogWarning("The SDK API requested '{SdkLevel}' not installed in {PlatformsDir}; Gradle will attempt to download it.", SDKLevel, PlatformsDir);
 				}
 
-				Log.TraceInformation("Building Java with SDK API level '{0}'", SDKLevel);
+				Logger.LogInformation("Building Java with SDK API level '{SDKLevel}'", SDKLevel);
 				CachedSDKLevel = SDKLevel;
 			}
 
@@ -366,7 +368,7 @@ namespace UnrealBuildTool
 			if (BestVersionString == null)
 			{
 				BestVersionString = "30.0.3";
-				Log.TraceWarning("Failed to find %ANDROID_HOME%/build-tools subdirectory. Will attempt to use {0}.", BestVersionString);
+				Logger.LogWarning("Failed to find %ANDROID_HOME%/build-tools subdirectory. Will attempt to use {BestVersionString}.", BestVersionString);
 			}
 
 			BestVersion = GetRevisionValue(BestVersionString);
@@ -386,7 +388,7 @@ namespace UnrealBuildTool
 			CachedBuildToolsVersion = BestVersionString;
 			LastAndroidHomePath = HomePath;
 
-			Log.TraceInformation("Building with Build Tools version '{0}'", CachedBuildToolsVersion);
+			Logger.LogInformation("Building with Build Tools version '{CachedBuildToolsVersion}'", CachedBuildToolsVersion);
 
 			return CachedBuildToolsVersion;
 		}
@@ -624,7 +626,7 @@ namespace UnrealBuildTool
 			}
 		}
 
-		private static void DeleteDirectory(string InPath, string SubDirectoryToKeep = "")
+		private static void DeleteDirectory(string InPath, ILogger Logger, string SubDirectoryToKeep = "")
 		{
 			// skip the dir we want to
 			if (String.Compare(Path.GetFileName(InPath), SubDirectoryToKeep, true) == 0)
@@ -654,14 +656,14 @@ namespace UnrealBuildTool
 				}
 				catch (Exception)
 				{
-					Log.TraceInformation("Failed to delete all files in directory {0}. Continuing on...", InPath);
+					Logger.LogInformation("Failed to delete all files in directory {InPath}. Continuing on...", InPath);
 				}
 			}
 
 			string[] Dirs = Directory.GetDirectories(InPath, "*.*", SearchOption.TopDirectoryOnly);
 			foreach (string Dir in Dirs)
 			{
-				DeleteDirectory(Dir, SubDirectoryToKeep);
+				DeleteDirectory(Dir, Logger, SubDirectoryToKeep);
 				// try to delete the directory, but allow it to fail (due to SubDirectoryToKeep still existing)
 				try
 				{
@@ -754,7 +756,7 @@ namespace UnrealBuildTool
 			{
 				if (bLog)
 				{
-					Log.TraceInformation("Copying {0} to {1}", SourceFilename, DestFilename);
+					Logger.LogInformation("Copying {SourceFilename} to {DestFilename}", SourceFilename, DestFilename);
 				}
 
 				if (bDestFileAlreadyExists)
@@ -848,7 +850,7 @@ namespace UnrealBuildTool
 					DestFileInfo.Attributes = DestFileInfo.Attributes & ~FileAttributes.ReadOnly;
 					File.SetLastWriteTimeUtc(DestFilename, File.GetLastWriteTimeUtc(Filename));
 
-					Log.TraceInformation("Copied file {0}.", DestFilename);
+					Logger.LogInformation("Copied file {DestFilename}.", DestFilename);
 				}
 			}
 
@@ -858,7 +860,7 @@ namespace UnrealBuildTool
 			{
 				if (!FilesToKeep.Contains(Filename))
 				{
-					Log.TraceInformation("Deleting unneeded file {0}.", Filename);
+					Logger.LogInformation("Deleting unneeded file {Filename}.", Filename);
 					SafeDeleteFile(Filename);
 				}
 			}
@@ -871,7 +873,7 @@ namespace UnrealBuildTool
 				{
 					if (Directory.Exists(directory) && Directory.GetFiles(directory, "*.*", SearchOption.AllDirectories).Count() == 0)
 					{
-						Log.TraceInformation("Cleaning Directory {0} as empty.", directory);
+						Logger.LogInformation("Cleaning Directory {Directory} as empty.", directory);
 						Directory.Delete(directory, true);
 					}
 				}
@@ -980,7 +982,7 @@ namespace UnrealBuildTool
 					}
 				}
 
-				Log.TraceInformation("GotStoreVersion found v{0}. (bUseChangeListAsStoreVersion={1} IsBuildMachine={2} EngineChangeList={3})", StoreVersion, bUseChangeListAsStoreVersion, IsBuildMachine, EngineChangelist);
+				Logger.LogInformation("GotStoreVersion found v{StoreVersion}. (bUseChangeListAsStoreVersion={bUseChangeListAsStoreVersion} IsBuildMachine={IsBuildMachine} EngineChangeList={EngineChangeList})", StoreVersion, bUseChangeListAsStoreVersion, IsBuildMachine, EngineChangelist);
 
 				CachedStoreVersion = StoreVersion;
 
@@ -1044,7 +1046,7 @@ namespace UnrealBuildTool
 
 		public void WriteJavaOBBDataFile(string FileName, string PackageName, List<string> ObbSources, string CookFlavor, bool bPackageDataInsideApk, string UnrealArch)
 		{
-			Log.TraceInformation("\n==== Writing to OBB data file {0} ====", FileName);
+			Logger.LogInformation("\n==== Writing to OBB data file {FileName} ====", FileName);
 
 			// always must write if file does not exist
 			bool bFileExists = File.Exists(FileName);
@@ -1125,7 +1127,7 @@ namespace UnrealBuildTool
 			}
 			else
 			{
-				Log.TraceInformation("\n==== OBB data file up to date so not writing. ====");
+				Logger.LogInformation("\n==== OBB data file up to date so not writing. ====");
 			}
 		}
 
@@ -1170,7 +1172,7 @@ namespace UnrealBuildTool
 				ShimFileContent.Append("\t}\n");
 
 			ShimFileContent.Append("}\n");
-			Log.TraceInformation("\n==== Writing to shim file {0} ====", ShimFileName);
+			Logger.LogInformation("\n==== Writing to shim file {ShimFileName} ====", ShimFileName);
 
 			// If they aren't the same then dump out the settings
 			if (DestFileContent == null || !DestFileContent.SequenceEqual((ShimFileContent.ToString()).Split('\n')))
@@ -1187,7 +1189,7 @@ namespace UnrealBuildTool
 			}
 			else
 			{
-				Log.TraceInformation("\n==== Shim data file up to date so not writing. ====");
+				Logger.LogInformation("\n==== Shim data file up to date so not writing. ====");
 			}
 
 			// Now we move on to the template files
@@ -1214,7 +1216,7 @@ namespace UnrealBuildTool
 					}
 				}
 
-				Log.TraceInformation("\n==== Writing to template target file {0} ====", template.DestinationFile);
+				Logger.LogInformation("\n==== Writing to template target file {File} ====", template.DestinationFile);
 
 				if (templateDest == null || templateSrc.Length != templateDest.Length || !templateSrc.SequenceEqual(templateDest))
 				{
@@ -1229,7 +1231,7 @@ namespace UnrealBuildTool
 				}
 				else
 				{
-					Log.TraceInformation("\n==== Template target file up to date so not writing. ====");
+					Logger.LogInformation("\n==== Template target file up to date so not writing. ====");
 				}
 			}
 		}
@@ -1256,7 +1258,7 @@ namespace UnrealBuildTool
 			string CrashPropertiesFileName = Path.Combine(UEBuildPath, "assets", "crashlytics-build.properties");
 			MakeDirectoryIfRequired(CrashPropertiesFileName);
 			File.WriteAllText(CrashPropertiesFileName, CrashPropertiesContent.ToString());
-			Log.TraceInformation("==== Write {0}  ====", CrashPropertiesFileName);
+			Logger.LogInformation("==== Write {CrashPropertiesFileName}  ====", CrashPropertiesFileName);
 
 			StringBuilder BuildIDContent = new StringBuilder("");
 			BuildIDContent.Append("<?xml version=\"1.0\" encoding=\"utf-8\" standalone=\"no\"?>\n");
@@ -1273,7 +1275,7 @@ namespace UnrealBuildTool
 			string BuildIDFileName = Path.Combine(UEBuildPath, "res", "values", "com_crashlytics_build_id.xml");
 			MakeDirectoryIfRequired(BuildIDFileName);
 			File.WriteAllText(BuildIDFileName, BuildIDContent.ToString());
-			Log.TraceInformation("==== Write {0}  ====", BuildIDFileName);
+			Logger.LogInformation("==== Write {BuildIDFileName}  ====", BuildIDFileName);
 		}
 
 		private static string GetNDKArch(string UnrealArch)
@@ -1302,7 +1304,7 @@ namespace UnrealBuildTool
 			}
 		}
 
-		private static void StripDebugSymbols(string SourceFileName, string TargetFileName, string UnrealArch, bool bStripAll = false)
+		private static void StripDebugSymbols(string SourceFileName, string TargetFileName, string UnrealArch, ILogger Logger, bool bStripAll = false)
 		{
 			// Copy the file and remove read-only if necessary
 			File.Copy(SourceFileName, TargetFileName, true);
@@ -1324,7 +1326,7 @@ namespace UnrealBuildTool
 			}
 			StartInfo.UseShellExecute = false;
 			StartInfo.CreateNoWindow = true;
-			Utils.RunLocalProcessAndLogOutput(StartInfo);
+			Utils.RunLocalProcessAndLogOutput(StartInfo, Logger);
 		}
 
 		private static string GetPlatformNDKHostName()
@@ -1391,7 +1393,7 @@ namespace UnrealBuildTool
 							}
 							string MaliLibDstPath = Path.Combine(UnrealBuildPath, "libs", NDKArch, "libMGD.so");
 
-							Log.TraceInformation("Copying {0} to {1}", MaliLibSrcPath, MaliLibDstPath);
+							Logger.LogInformation("Copying {MaliLibSrcPath} to {MaliLibDstPath}", MaliLibSrcPath, MaliLibDstPath);
 							File.Copy(MaliLibSrcPath, MaliLibDstPath, true);
 							File.SetLastWriteTimeUtc(MaliLibDstPath, File.GetLastWriteTimeUtc(MaliLibSrcPath));
 
@@ -1399,7 +1401,7 @@ namespace UnrealBuildTool
 							if (File.Exists(MaliVkLayerLibSrcPath))
 							{
 								string MaliVkLayerLibDstPath = Path.Combine(UnrealBuildPath, "libs", NDKArch, "libVkLayerAGA.so");
-								Log.TraceInformation("Copying {0} to {1}", MaliVkLayerLibSrcPath, MaliVkLayerLibDstPath);
+								Logger.LogInformation("Copying {MaliVkLayerLibSrcPath} to {MaliVkLayerLibDstPath}", MaliVkLayerLibSrcPath, MaliVkLayerLibDstPath);
 								File.Copy(MaliVkLayerLibSrcPath, MaliVkLayerLibDstPath, true);
 								File.SetLastWriteTimeUtc(MaliVkLayerLibDstPath, File.GetLastWriteTimeUtc(MaliVkLayerLibSrcPath));
 							}
@@ -1430,8 +1432,8 @@ namespace UnrealBuildTool
 			bool bSupportsVulkan = false;
 			Ini.GetBool("/Script/AndroidRuntimeSettings.AndroidRuntimeSettings", "bSupportsVulkan", out bSupportsVulkan);
 
-			Log.TraceInformation("bBuildForES31: {0}", (bBuildForES31 ? "true" : "false"));
-			Log.TraceInformation("bSupportsVulkan: {0}", (bSupportsVulkan ? "true" : "false"));
+			Logger.LogInformation("bBuildForES31: {bBuildForES31}", (bBuildForES31 ? "true" : "false"));
+			Logger.LogInformation("bSupportsVulkan: {bSupportsVulkan}", (bSupportsVulkan ? "true" : "false"));
 		}
 		
 		void CopyVulkanValidationLayers(string UnrealBuildPath, string UnrealArch, string NDKArch, string Configuration)
@@ -1448,7 +1450,7 @@ namespace UnrealBuildTool
 				string VulkanLayersDir = Environment.ExpandEnvironmentVariables("%NDKROOT%/sources/third_party/vulkan/src/build-android/jniLibs/") + NDKArch;
 				if (Directory.Exists(VulkanLayersDir))
 				{
-					Log.TraceInformation("Copying {0} vulkan layer from {1}", ANDROID_VULKAN_VALIDATION_LAYER, VulkanLayersDir);
+					Logger.LogInformation("Copying {ANDROID_VULKAN_VALIDATION_LAYER} vulkan layer from {VulkanLayersDir}", ANDROID_VULKAN_VALIDATION_LAYER, VulkanLayersDir);
 					string DestDir = Path.Combine(UnrealBuildPath, "libs", NDKArch);
 					Directory.CreateDirectory(DestDir);
 					string SourceFilename = Path.Combine(VulkanLayersDir, ANDROID_VULKAN_VALIDATION_LAYER);
@@ -1467,7 +1469,7 @@ namespace UnrealBuildTool
 				String LayerDir = "";
 				AndroidPlatformSDK.GetPath(Ini, "/Script/AndroidRuntimeSettings.AndroidRuntimeSettings", "DebugVulkanLayerDirectory", out LayerDir);
 
-				Log.TraceInformation("DebugVulkanLayerDirectory {0}", LayerDir);
+				Logger.LogInformation("DebugVulkanLayerDirectory {LayerDir}", LayerDir);
 
 				if (LayerDir != "")
 				{
@@ -1476,7 +1478,7 @@ namespace UnrealBuildTool
 					if (Directory.Exists(VulkanLayersDir))
 					{
 						string DestDir = Path.Combine(UnrealBuildPath, "libs", NDKArch);
-						Log.TraceInformation("Copying Debug vulkan layers from {0} to {1}", VulkanLayersDir, DestDir);
+						Logger.LogInformation("Copying Debug vulkan layers from {VulkanLayersDir} to {DestDir}", VulkanLayersDir, DestDir);
 
 						if (!Directory.Exists(DestDir))
 						{
@@ -1537,12 +1539,12 @@ namespace UnrealBuildTool
 			{
 				string LibDestDir = Path.Combine(UnrealBuildPath, "libs", NDKArch);
 				Directory.CreateDirectory(LibDestDir);
-				Log.TraceInformation("Copying asan lib from {0} to {1}", SanitizerLib, LibDestDir);
+				Logger.LogInformation("Copying asan lib from {SanitizerLib} to {LibDestDir}", SanitizerLib, LibDestDir);
 				File.Copy(SanitizerLib, Path.Combine(LibDestDir, SanitizerFullLibName), true);
 				string WrapDestDir = Path.Combine(UnrealBuildPath, "resources", "lib", NDKArch);
 				Directory.CreateDirectory(WrapDestDir);
 				string WrapDestFilePath = Path.Combine(WrapDestDir, "wrap.sh");
-				Log.TraceInformation("Copying wrap.sh from {0} to {1}", WrapSh, WrapDestFilePath);
+				Logger.LogInformation("Copying wrap.sh from {WrapSh} to {WrapDestFilePath}", WrapSh, WrapDestFilePath);
 				File.Copy(WrapSh, WrapDestFilePath, true);
 			}
 			else
@@ -1551,7 +1553,7 @@ namespace UnrealBuildTool
 			}
 		}
 
-		private static int RunCommandLineProgramAndReturnResult(string WorkingDirectory, string Command, string Params, string? OverrideDesc = null, bool bUseShellExecute = false)
+		private static int RunCommandLineProgramAndReturnResult(string WorkingDirectory, string Command, string Params, ILogger Logger, string? OverrideDesc = null, bool bUseShellExecute = false)
 		{
 			// Process Arguments follow windows conventions in .NET Core
 			// Which means single quotes ' are not considered quotes.
@@ -1566,12 +1568,12 @@ namespace UnrealBuildTool
 
 			if (OverrideDesc == null)
 			{
-				Log.TraceInformation("\nRunning: " + Command + " " + Params);
+				Logger.LogInformation("\nRunning: {Command} {Params}", Command, Params);
 			}
 			else if (OverrideDesc != "")
 			{
-				Log.TraceInformation(OverrideDesc);
-				Log.TraceVerbose("\nRunning: " + Command + " " + Params);
+				Logger.LogInformation("{Message}", OverrideDesc);
+				Logger.LogDebug("\nRunning: {Command} {Params}", Command, Params);
 			}
 
 			ProcessStartInfo StartInfo = new ProcessStartInfo();
@@ -1589,7 +1591,7 @@ namespace UnrealBuildTool
 			return Proc.ExitCode;
 		}
 
-		private static void RunCommandLineProgramWithException(string WorkingDirectory, string Command, string Params, string? OverrideDesc = null, bool bUseShellExecute = false)
+		private static void RunCommandLineProgramWithException(string WorkingDirectory, string Command, string Params, ILogger Logger, string? OverrideDesc = null, bool bUseShellExecute = false)
 		{
 			// Process Arguments follow windows conventions in .NET Core
 			// Which means single quotes ' are not considered quotes.
@@ -1604,12 +1606,12 @@ namespace UnrealBuildTool
 
 			if (OverrideDesc == null)
 			{
-				Log.TraceInformation("\nRunning: " + Command + " " + Params);
+				Logger.LogInformation("\nRunning: {Command} {Params}", Command, Params);
 			}
 			else if (OverrideDesc != "")
 			{
-				Log.TraceInformation(OverrideDesc);
-				Log.TraceVerbose("\nRunning: " + Command + " " + Params);
+				Logger.LogInformation("{Message}", OverrideDesc);
+				Logger.LogDebug("\nRunning: {Command} {Params}", Command, Params);
 			}
 
 			ProcessStartInfo StartInfo = new ProcessStartInfo();
@@ -1755,7 +1757,7 @@ namespace UnrealBuildTool
 			}
 		}
 
-		static void FilterStdOutErr(object sender, DataReceivedEventArgs e)
+		static void FilterStdOutErr(object sender, DataReceivedEventArgs e, ILogger Logger)
 		{
 			if (e.Data != null)
 			{
@@ -1771,11 +1773,11 @@ namespace UnrealBuildTool
 									break;
 
 								case FilterAction.Replace:
-									Log.TraceInformation("{0}", e.Data.Replace(FilterOp.Match, FilterOp.ReplaceWith));
+									Logger.LogInformation("{Output}", e.Data.Replace(FilterOp.Match, FilterOp.ReplaceWith));
 									break;
 
 								case FilterAction.Error:
-									Log.TraceError("{0}", e.Data.Replace(FilterOp.Match, FilterOp.ReplaceWith));
+									Logger.LogError("{Output}", e.Data.Replace(FilterOp.Match, FilterOp.ReplaceWith));
 									break;
 
 								default:
@@ -1785,11 +1787,11 @@ namespace UnrealBuildTool
 						}
 					}
 				}
-				Log.TraceInformation("{0}", e.Data);
+				Logger.LogInformation("{Output}", e.Data);
 			}
 		}
 
-		private static void RunCommandLineProgramWithExceptionAndFiltering(string WorkingDirectory, string Command, string Params, string? OverrideDesc = null, bool bUseShellExecute = false)
+		private static void RunCommandLineProgramWithExceptionAndFiltering(string WorkingDirectory, string Command, string Params, ILogger Logger, string? OverrideDesc = null, bool bUseShellExecute = false)
 		{
 			// Process Arguments follow windows conventions in .NET Core
 			// Which means single quotes ' are not considered quotes.
@@ -1804,12 +1806,12 @@ namespace UnrealBuildTool
 
 			if (OverrideDesc == null)
 			{
-				Log.TraceInformation("\nRunning: " + Command + " " + Params);
+				Logger.LogInformation("\nRunning: {Command} {Params}", Command, Params);
 			}
 			else if (OverrideDesc != "")
 			{
-				Log.TraceInformation(OverrideDesc);
-				Log.TraceVerbose("\nRunning: " + Command + " " + Params);
+				Logger.LogInformation("{Message}", OverrideDesc);
+				Logger.LogDebug("\nRunning: {Command} {Params}", Command, Params);
 			}
 
 			ProcessStartInfo StartInfo = new ProcessStartInfo();
@@ -1824,8 +1826,8 @@ namespace UnrealBuildTool
 
 			Process Proc = new Process();
 			Proc.StartInfo = StartInfo;
-			Proc.OutputDataReceived += FilterStdOutErr;
-			Proc.ErrorDataReceived += FilterStdOutErr;
+			Proc.OutputDataReceived += (s, e) => FilterStdOutErr(s, e, Logger);
+			Proc.ErrorDataReceived += (s, e) => FilterStdOutErr(s, e, Logger);
 
 			Proc.Start();
 			Proc.BeginOutputReadLine();
@@ -2102,7 +2104,7 @@ namespace UnrealBuildTool
 						if (InputFileTime.CompareTo(ApkTime) > 0)
 						{
 							bAllInputsCurrent = false;
-							Log.TraceInformation("{0} is out of date due to newer input file {1}", DestApkName, InputFileName);
+							Logger.LogInformation("{DestApkName} is out of date due to newer input file {InputFileName}", DestApkName, InputFileName);
 							break;
 						}
 					}
@@ -2234,7 +2236,7 @@ namespace UnrealBuildTool
 				{
 					if (!File.Exists(PortraitFilename) && (ResolutionPath == "/res/drawable/"))
 					{
-						Log.TraceWarning("Warning: Downloader screen source image {0} not available, downloader screen will not function properly!", PortraitFilename);
+						Logger.LogWarning("Warning: Downloader screen source image {PortraitFilename} not available, downloader screen will not function properly!", PortraitFilename);
 					}
 				}
 				else
@@ -2248,7 +2250,7 @@ namespace UnrealBuildTool
 				{
 					if (!File.Exists(LandscapeFilename) && (ResolutionPath == "/res/drawable/"))
 					{
-						Log.TraceWarning("Warning: Downloader screen source image {0} not available, downloader screen will not function properly!", LandscapeFilename);
+						Logger.LogWarning("Warning: Downloader screen source image {LandscapeFilename} not available, downloader screen will not function properly!", LandscapeFilename);
 					}
 				}
 				else
@@ -2313,7 +2315,7 @@ namespace UnrealBuildTool
 				{
 					if (!File.Exists(PortraitFilename) && (ResolutionPath == "/res/drawable/"))
 					{
-						Log.TraceWarning("Warning: Splash screen source image {0} not available, splash screen will not function properly!", PortraitFilename);
+						Logger.LogWarning("Warning: Splash screen source image {PortraitFilename} not available, splash screen will not function properly!", PortraitFilename);
 					}
 				}
 				else
@@ -2331,7 +2333,7 @@ namespace UnrealBuildTool
 				{
 					if (!File.Exists(LandscapeFilename) && (ResolutionPath == "/res/drawable/"))
 					{
-						Log.TraceWarning("Warning: Splash screen source image {0} not available, splash screen will not function properly!", LandscapeFilename);
+						Logger.LogWarning("Warning: Splash screen source image {LandscapeFilename} not available, splash screen will not function properly!", LandscapeFilename);
 					}
 				}
 				else
@@ -2441,7 +2443,7 @@ namespace UnrealBuildTool
 					}
 				}
 
-				Log.TraceInformation("Using package name: '{0}'", PackageName);
+				Logger.LogInformation("Using package name: '{PackageName}'", PackageName);
 				CachedPackageName = PackageName;
 			}
 
@@ -2558,7 +2560,7 @@ namespace UnrealBuildTool
 			Ini.GetBool("/Script/AndroidRuntimeSettings.AndroidRuntimeSettings", "bAllowIMU", out bAllowIMU);
 			if (IsPackagingForDaydream(Ini) && bAllowIMU)
 			{
-				Log.TraceInformation("Daydream and IMU both enabled, recommend disabling IMU if not needed.");
+				Logger.LogInformation("Daydream and IMU both enabled, recommend disabling IMU if not needed.");
 			}
 
 			bool bExtractNativeLibs = true;
@@ -2594,7 +2596,7 @@ namespace UnrealBuildTool
 			{
 				if (bPackageForOculusMobile)
 				{
-					Log.TraceInformation("Disabling Package For Oculus Mobile for unsupported architecture {0}", UnrealArch);
+					Logger.LogInformation("Disabling Package For Oculus Mobile for unsupported architecture {UnrealArch}", UnrealArch);
 					bPackageForOculusMobile = false;
 				}
 			}
@@ -2604,7 +2606,7 @@ namespace UnrealBuildTool
 			{
 				if (bShowLaunchImage)
 				{
-					Log.TraceInformation("Disabling Show Launch Image for Oculus Mobile enabled application");
+					Logger.LogInformation("Disabling Show Launch Image for Oculus Mobile enabled application");
 					bShowLaunchImage = false;
 				}
 			}
@@ -2615,7 +2617,7 @@ namespace UnrealBuildTool
 			{
 				if (bShowLaunchImage)
 				{
-					Log.TraceInformation("Disabling Show Launch Image for Daydream enabled application");
+					Logger.LogInformation("Disabling Show Launch Image for Daydream enabled application");
 					bShowLaunchImage = false;
 				}
 			}
@@ -2659,7 +2661,7 @@ namespace UnrealBuildTool
 						bASTCEnabled = true;
 						break;
 					default:
-						Log.TraceWarning("Invalid or unknown CookFlavor used in GenerateManifest: {0}", CookFlavor);
+						Logger.LogWarning("Invalid or unknown CookFlavor used in GenerateManifest: {CookFlavor}", CookFlavor);
 						break;
 				}
 			}
@@ -3120,7 +3122,7 @@ namespace UnrealBuildTool
 					//validate that the AppId matches the .ini value for the GooglePlay AppId, assuming it's valid
 					if (!bInvalidIniAppId &&  IniAppId.CompareTo(XmlAppId) != 0)
 					{
-						Log.TraceInformation("Replacing Google Play AppID in GooglePlayAppID.xml with AndroidRuntimeSettings .ini value");
+						Logger.LogInformation("Replacing Google Play AppID in GooglePlayAppID.xml with AndroidRuntimeSettings .ini value");
 
 						bInvalid = true;
 						ReplacementId = IniAppId;
@@ -3128,7 +3130,7 @@ namespace UnrealBuildTool
 					}					
 					else if(XmlAppId.Length == 0 || !Int64.TryParse(XmlAppId, out Value))
 					{
-						Log.TraceWarning("\nWARNING: GooglePlay Games App ID is invalid! Replacing it with \"1\"");
+						Logger.LogWarning("\nWARNING: GooglePlay Games App ID is invalid! Replacing it with \"1\"");
 
 						//write file with something which will fail but not cause an exception if executed
 						bInvalid = true;
@@ -3157,12 +3159,12 @@ namespace UnrealBuildTool
 				// if we don't have an appID to use from the config, write file with something which will fail but not cause an exception if executed
 				if (bInvalidIniAppId)
 				{
-					Log.TraceWarning("\nWARNING: Creating GooglePlayAppID.xml using a Google Play AppID of \"1\" because there was no valid AppID in AndroidRuntimeSettings!");
+					Logger.LogWarning("\nWARNING: Creating GooglePlayAppID.xml using a Google Play AppID of \"1\" because there was no valid AppID in AndroidRuntimeSettings!");
 					NewAppId = "1";
 				}
 				else
 				{
-					Log.TraceInformation("Creating GooglePlayAppID.xml with AndroidRuntimeSettings .ini value");
+					Logger.LogInformation("Creating GooglePlayAppID.xml with AndroidRuntimeSettings .ini value");
 					NewAppId = IniAppId;
 				}
 
@@ -3257,17 +3259,17 @@ namespace UnrealBuildTool
 		{
 			if (bDisallowPackageInAPK)
 			{
-				Log.TraceInformation("APK contains data.");
+				Logger.LogInformation("APK contains data.");
 				return false;
 			}
 			else if (!String.IsNullOrEmpty(Environment.GetEnvironmentVariable("uebp_LOCAL_ROOT")))
 			{
-				Log.TraceInformation("On build machine.");
+				Logger.LogInformation("On build machine.");
 				return true;
 			}
 			else
 			{
-				Log.TraceInformation("Looking for OBB.");
+				Logger.LogInformation("Looking for OBB.");
 				return File.Exists(OBBLocation);
 			}
 		}
@@ -3300,11 +3302,11 @@ namespace UnrealBuildTool
 
 			if (bFound)
 			{
-				Log.TraceInformation("\nUnable to apply subst, using gradlew.bat directly (all drive letters in use!)");
+				Logger.LogInformation("\nUnable to apply subst, using gradlew.bat directly (all drive letters in use!)");
 				return false;
 			}
 
-			Log.TraceInformation("\nCreating rungradle.bat to work around commandline length limit (using unused drive letter {0})", UnusedDriveLetter);
+			Logger.LogInformation("\nCreating rungradle.bat to work around commandline length limit (using unused drive letter {UnusedDriveLetter})", UnusedDriveLetter);
 
 			// make sure rungradle.bat isn't read-only
 			if (File.Exists(RunGradleBatFilename))
@@ -3367,7 +3369,7 @@ namespace UnrealBuildTool
 			// directory must exist
 			if (!Directory.Exists(LicensePath))
 			{
-				Log.TraceInformation("Directory doesn't exist {0}", LicensePath);
+				Logger.LogInformation("Directory doesn't exist {LicensePath}", LicensePath);
 				return false;
 			}
 
@@ -3375,7 +3377,7 @@ namespace UnrealBuildTool
 			string LicenseFilename = Path.Combine(LicensePath, "android-sdk-license");
 			if (!File.Exists(LicenseFilename))
 			{
-				Log.TraceInformation("File doesn't exist {0}", LicenseFilename);
+				Logger.LogInformation("File doesn't exist {LicenseFilename}", LicenseFilename);
 				return false;
 			}
 
@@ -3404,13 +3406,13 @@ namespace UnrealBuildTool
 			if ((ToolChain.BuildWithSanitizer() != AndroidToolChain.ClangSanitizer.None) && (MinSDKVersion < 27))
 			{
 				MinSDKVersion = 27;
-				Log.TraceInformation("Fixing minSdkVersion; requires minSdkVersion of {0} for Clang's Sanitizers", MinSDKVersion);
+				Logger.LogInformation("Fixing minSdkVersion; requires minSdkVersion of {MinVer} for Clang's Sanitizers", MinSDKVersion);
 			}
 
 			if (bEnableBundle && MinSDKVersion < MinimumSDKLevelForBundle)
 			{
 				MinSDKVersion = MinimumSDKLevelForBundle;
-				Log.TraceInformation("Fixing minSdkVersion; requires minSdkVersion of {0} for App Bundle support", MinimumSDKLevelForBundle);
+				Logger.LogInformation("Fixing minSdkVersion; requires minSdkVersion of {MinVer} for App Bundle support", MinimumSDKLevelForBundle);
 			}
 
 			// Make sure minSdkVersion is at least 13 (need this for appcompat-v13 used by AndroidPermissions)
@@ -3418,7 +3420,7 @@ namespace UnrealBuildTool
 			if (MinSDKVersion < MinimumSDKLevelForGradle)
 			{
 				MinSDKVersion = MinimumSDKLevelForGradle;
-				Log.TraceInformation("Fixing minSdkVersion; requires minSdkVersion of {0} with Gradle based on active plugins", MinimumSDKLevelForGradle);
+				Logger.LogInformation("Fixing minSdkVersion; requires minSdkVersion of {MinVer} with Gradle based on active plugins", MinimumSDKLevelForGradle);
 			}
 
 			// 64-bit targets must be android-21 or higher
@@ -3435,12 +3437,12 @@ namespace UnrealBuildTool
 				if (MinSDKVersion < 21)
 				{
 					MinSDKVersion = 21;
-					Log.TraceInformation("Fixing minSdkVersion; NDK level above 19 requires minSdkVersion of 21 (arch={0})", Arch.Substring(1));
+					Logger.LogInformation("Fixing minSdkVersion; NDK level above 19 requires minSdkVersion of 21 (arch={Arch})", Arch.Substring(1));
 				}
 
 				if (MinSDKVersion < NDKLevelInt)
                 {
-                    Log.TraceInformation("Fixing minSdkVersion; NDK level is {0} which is above minSdkVersion {1}.", NDKLevelInt, MinSDKVersion);
+                    Logger.LogInformation("Fixing minSdkVersion; NDK level is {NdkLevel} which is above minSdkVersion {MinSdkVersion}.", NDKLevelInt, MinSDKVersion);
 					MinSDKVersion = NDKLevelInt;
 				}
             }
@@ -3549,7 +3551,7 @@ namespace UnrealBuildTool
 				if (GetTargetOculusMobileDevices().Contains("Go"))
 				{
 					bDisableV2Signing = true;
-					Log.TraceInformation("Disabling v2Signing for Oculus Go");
+					Logger.LogInformation("Disabling v2Signing for Oculus Go");
 				}
 
 				string KeyAlias, KeyStore, KeyStorePassword, KeyPassword;
@@ -3675,7 +3677,8 @@ namespace UnrealBuildTool
 		private void MakeApk(AndroidToolChain ToolChain, string ProjectName, TargetType InTargetType, string ProjectDirectory, string OutputPath, string EngineDirectory, bool bForDistribution, string CookFlavor, 
 			UnrealTargetConfiguration Configuration, bool bMakeSeparateApks, bool bIncrementalPackage, bool bDisallowPackagingDataInApk, bool bDisallowExternalFilesDir, bool bSkipGradleBuild)
 		{
-			Log.TraceInformation("\n===={0}====PREPARING TO MAKE APK=================================================================", DateTime.Now.ToString());
+			Logger.LogInformation("");
+			Logger.LogInformation("===={Time}====PREPARING TO MAKE APK=================================================================", DateTime.Now.ToString());
 
 			// Get list of all architecture and GPU targets for build
 			List<string> Arches = ToolChain.GetAllArchitectures();
@@ -3710,8 +3713,9 @@ namespace UnrealBuildTool
 					}
 					*/
 				}
-				Log.TraceInformation("APK generation not needed for project {0} with {1}", ProjectName, OutputPath);
-				Log.TraceInformation("\n===={0}====COMPLETED MAKE APK=======================================================================", DateTime.Now.ToString());
+				Logger.LogInformation("APK generation not needed for project {ProjectName} with {OutputPath}", ProjectName, OutputPath);
+				Logger.LogInformation("");
+				Logger.LogInformation("===={Time}====COMPLETED MAKE APK=======================================================================", DateTime.Now.ToString());
 				return;
 			}
 
@@ -3794,7 +3798,7 @@ namespace UnrealBuildTool
 				string BuildTypeContents = File.ReadAllText(BuildTypeFilename);
 				if (BuildTypeID != BuildTypeContents)
 				{
-					Log.TraceInformation("Build type changed, forcing clean");
+					Logger.LogInformation("Build type changed, forcing clean");
 					bCreateFromScratch = true;
 				}
 			}
@@ -3803,7 +3807,7 @@ namespace UnrealBuildTool
 			if (File.Exists(Path.Combine(IntermediateAndroidPath, "arm64", "jni", "arm64-v8a", "libUE4.so")) ||
 				File.Exists(Path.Combine(IntermediateAndroidPath, "x64", "jni", "x86_64", "libUE4.so")))
 			{
-				Log.TraceInformation("Old version of library .so found, forcing clean");
+				Logger.LogInformation("Old version of library .so found, forcing clean");
 				bCreateFromScratch = true;
 			}
 
@@ -3815,15 +3819,15 @@ namespace UnrealBuildTool
 				string PreviousPluginListContents = File.ReadAllText(PluginListFilename);
 				if (PluginListContents != PreviousPluginListContents)
 				{
-					Log.TraceInformation("Active UPL files changed, forcing clean");
+					Logger.LogInformation("Active UPL files changed, forcing clean");
 					bCreateFromScratch = true;
 				}
 			}
 
 			if (bCreateFromScratch)
 			{
-				Log.TraceInformation("Cleaning {0}", IntermediateAndroidPath);
-				DeleteDirectory(IntermediateAndroidPath);
+				Logger.LogInformation("Cleaning {IntermediateAndroidPath}", IntermediateAndroidPath);
+				DeleteDirectory(IntermediateAndroidPath, Logger);
 				Directory.CreateDirectory(IntermediateAndroidPath);
 			}
 			
@@ -3897,7 +3901,7 @@ namespace UnrealBuildTool
 				string ImmediateBaseDir = Path.Combine(UnrealBuildPath, "src");
 				IEnumerable<string> files = Directory.EnumerateFiles(CleanUpBaseDir, "*.java", SearchOption.AllDirectories);
 
-				Log.TraceInformation("Cleaning up files based on template dir {0}", TemplateDestinationBase);
+				Logger.LogInformation("Cleaning up files based on template dir {TemplateDestinationBase}", TemplateDestinationBase);
 
 				// Make a set of files that are okay to clean up
 				HashSet<string> cleanFiles = new HashSet<string>();
@@ -3923,14 +3927,14 @@ namespace UnrealBuildTool
 						if (!cleanFiles.Contains(Path.GetFileName(filename)))
 							continue;
 
-						Log.TraceInformation("Cleaning up file {0}", filename);
+						Logger.LogInformation("Cleaning up file {File}", filename);
 						SafeDeleteFile(filename, false);
 
 						// Check to see if this file also exists in our target destination, and if so delete it too
 						string DestFilename = Path.Combine(ImmediateBaseDir, Utils.MakePathRelativeTo(filename, CleanUpBaseDir));
 						if (File.Exists(DestFilename))
 						{
-							Log.TraceInformation("Cleaning up file {0}", DestFilename);
+							Logger.LogInformation("Cleaning up file {DestFilename}", DestFilename);
 							SafeDeleteFile(DestFilename, false);
 						}
 					}
@@ -3944,7 +3948,7 @@ namespace UnrealBuildTool
 					{
 						if (Directory.Exists(directory) && Directory.GetFiles(directory, "*.*", SearchOption.AllDirectories).Count() == 0)
 						{
-							Log.TraceInformation("Cleaning Directory {0} as empty.", directory);
+							Logger.LogInformation("Cleaning Directory {Directory} as empty.", directory);
 							Directory.Delete(directory, true);
 						}
 					}
@@ -3962,7 +3966,7 @@ namespace UnrealBuildTool
 					{
 						if (Directory.Exists(directory) && Directory.GetFiles(directory, "*.*", SearchOption.AllDirectories).Count() == 0)
 						{
-							Log.TraceInformation("Cleaning Directory {0} as empty.", directory);
+							Logger.LogInformation("Cleaning Directory {Directory} as empty.", directory);
 							Directory.Delete(directory, true);
 						}
 					}
@@ -3991,7 +3995,7 @@ namespace UnrealBuildTool
 			if (CheckApplicationName(Path.Combine(IntermediateAndroidPath, ArchRemapping[NDKArches[0]]), ProjectName, out ApplicationDisplayName))
 			{
 				bBuildSettingsMatch = false;
-				Log.TraceInformation("Application display name is different than last build, forcing repackage.");
+				Logger.LogInformation("Application display name is different than last build, forcing repackage.");
 			}
 
 			// if the manifest matches, look at other settings stored in a file
@@ -4003,7 +4007,7 @@ namespace UnrealBuildTool
 					if (PreviousBuildSettings != CurrentBuildSettings)
 					{
 						bBuildSettingsMatch = false;
-						Log.TraceInformation("Previous .apk file(s) were made with different build settings, forcing repackage.");
+						Logger.LogInformation("Previous .apk file(s) were made with different build settings, forcing repackage.");
 					}
 				}
 			}
@@ -4035,7 +4039,7 @@ namespace UnrealBuildTool
 				if (!bCommandLineMatch)
 				{
 					bBuildSettingsMatch = false;
-					Log.TraceInformation("Previous .apk file(s) were made with different stage/apk command line, forcing repackage.");
+					Logger.LogInformation("Previous .apk file(s) were made with different stage/apk command line, forcing repackage.");
 				}
 			}
 
@@ -4065,7 +4069,7 @@ namespace UnrealBuildTool
 			int BuildListComboTotal = BuildList.Count();
 			if (BuildListComboTotal == 0)
 			{
-				Log.TraceInformation("Output .apk file(s) are up to date (dependencies and build settings are up to date)");
+				Logger.LogInformation("Output .apk file(s) are up to date (dependencies and build settings are up to date)");
 				return;
 			}
 
@@ -4092,22 +4096,22 @@ namespace UnrealBuildTool
 				string Manifest = build.Item2;
 				string NDKArch = GetNDKArch(Arch);
 
-				Log.TraceInformation("\n===={0}====PREPARING NATIVE CODE====={1}============================================================", DateTime.Now.ToString(), Arch);
+				Logger.LogInformation("\n===={Time}====PREPARING NATIVE CODE====={Arch}============================================================", DateTime.Now.ToString(), Arch);
 
 				string UnrealBuildPath = Path.Combine(IntermediateAndroidPath, Arch.Substring(1).Replace("-", "_"));
 
 				// If we are packaging for Amazon then we need to copy the  file to the correct location
-				Log.TraceInformation("bPackageDataInsideApk = {0}", bPackageDataInsideApk);
+				Logger.LogInformation("bPackageDataInsideApk = {bPackageDataInsideApk}", bPackageDataInsideApk);
 				if (bPackageDataInsideApk)
 				{
-					Log.TraceInformation("Obb location {0}", ObbFileLocation);
+					Logger.LogInformation("Obb location {ObbFileLocation}", ObbFileLocation);
 					string ObbFileDestination = UnrealBuildPath + "/assets";
-					Log.TraceInformation("Obb destination location {0}", ObbFileDestination);
+					Logger.LogInformation("Obb destination location {ObbFileDestination}", ObbFileDestination);
 					if (File.Exists(ObbFileLocation))
 					{
 						Directory.CreateDirectory(UnrealBuildPath);
 						Directory.CreateDirectory(ObbFileDestination);
-						Log.TraceInformation("Obb file exists...");
+						Logger.LogInformation("Obb file exists...");
 						string DestFileName = Path.Combine(ObbFileDestination, "main.obb.png"); // Need a rename to turn off compression
 						string SrcFileName = ObbFileLocation;
 						CopyIfDifferent(SrcFileName, DestFileName, true, false);
@@ -4204,7 +4208,7 @@ namespace UnrealBuildTool
 					FinalSOName = OutputPath;
 					if (!File.Exists(FinalSOName))
 					{
-						Log.TraceWarning("Did not find compiled .so [{0}]", FinalSOName);
+						Logger.LogWarning("Did not find compiled .so [{FinalSOName}]", FinalSOName);
 					}
 				}
 				else
@@ -4231,13 +4235,13 @@ namespace UnrealBuildTool
 
 					// clear out libs directory like ndk-build would have
 					string LibsDir = Path.Combine(UnrealBuildPath, "libs");
-					DeleteDirectory(LibsDir);
+					DeleteDirectory(LibsDir, Logger);
 					MakeDirectoryIfRequired(LibsDir);
 
 					// check to see if libUnreal.so needs to be copied
 					if (BuildListComboTotal > 1 || FilesAreDifferent(SourceSOName, FinalSOName))
 					{
-						Log.TraceInformation("\nCopying new .so {0} file to jni folder...", SourceSOName);
+						Logger.LogInformation("\nCopying new .so {SourceSOName} file to jni folder...", SourceSOName);
 						Directory.CreateDirectory(JniDir);
 						// copy the binary to the standard .so location
 						File.Copy(SourceSOName, FinalSOName, true);
@@ -4272,7 +4276,7 @@ namespace UnrealBuildTool
 
 				CreateAdditonalBuildPathFiles(NDKArch, UnrealBuildPath, AdditionalBuildPathFilesDoc);
 
-				Log.TraceInformation("\n===={0}====PERFORMING FINAL APK PACKAGE OPERATION====={1}===========================================", DateTime.Now.ToString(), Arch);
+				Logger.LogInformation("\n===={Time}====PERFORMING FINAL APK PACKAGE OPERATION====={Arch}===========================================", DateTime.Now.ToString(), Arch);
 
 				// check if any plugins want to increase the required compile SDK version
 				string CompileSDKMin = UPL.ProcessPluginNode(NDKArch, "minimumSDKAPI", "");
@@ -4303,7 +4307,7 @@ namespace UnrealBuildTool
 					if (bUpdatedCompileSDK)
 					{
 						CompileSDKVersion = CompileSDKVersionInt.ToString();
-						Log.TraceInformation("Building Java with SDK API Level 'android-{0}' due to enabled plugin requirements", CompileSDKVersion);
+						Logger.LogInformation("Building Java with SDK API Level 'android-{CompileSDKVersion}' due to enabled plugin requirements", CompileSDKVersion);
 					}
 				}
 
@@ -4365,7 +4369,7 @@ namespace UnrealBuildTool
 					if (!RuntimePlatform.IsWindows)
 					{
 						// fix permissions for Mac/Linux
-						RunCommandLineProgramWithException(UnrealBuildGradlePath, "/bin/sh", string.Format("-c 'chmod 0755 \"{0}\"'", GradleScriptPath.Replace("'", "'\"'\"'")), "Fix gradlew permissions");
+						RunCommandLineProgramWithException(UnrealBuildGradlePath, "/bin/sh", string.Format("-c 'chmod 0755 \"{0}\"'", GradleScriptPath.Replace("'", "'\"'\"'")), Logger, "Fix gradlew permissions");
 					}
 					else
 					{
@@ -4398,11 +4402,11 @@ namespace UnrealBuildTool
 							string ShellExecutable = RuntimePlatform.IsWindows ? "cmd.exe" : "/bin/sh";
 							string ShellParametersBegin = RuntimePlatform.IsWindows ? "/c " : "-c '";
 							string ShellParametersEnd = RuntimePlatform.IsWindows ? "" : "'";
-							RunCommandLineProgramWithExceptionAndFiltering(UnrealBuildGradlePath, ShellExecutable, ShellParametersBegin + "\"" + GradleScriptPath + "\" " + GradleOptions + ShellParametersEnd, "Making .apk with Gradle...");
+							RunCommandLineProgramWithExceptionAndFiltering(UnrealBuildGradlePath, ShellExecutable, ShellParametersBegin + "\"" + GradleScriptPath + "\" " + GradleOptions + ShellParametersEnd, Logger, "Making .apk with Gradle...");
 
 							if (GradleSecondCallOptions != "")
 							{
-								RunCommandLineProgramWithExceptionAndFiltering(UnrealBuildGradlePath, ShellExecutable, ShellParametersBegin + "\"" + GradleScriptPath + "\" " + GradleSecondCallOptions + ShellParametersEnd, "Additional Gradle steps...");
+								RunCommandLineProgramWithExceptionAndFiltering(UnrealBuildGradlePath, ShellExecutable, ShellParametersBegin + "\"" + GradleScriptPath + "\" " + GradleSecondCallOptions + ShellParametersEnd, Logger, "Additional Gradle steps...");
 							}
 
 							// For build machine run a clean afterward to clean up intermediate files (does not remove final APK)
@@ -4412,15 +4416,15 @@ namespace UnrealBuildTool
 								//RunCommandLineProgramWithException(UnrealBuildGradlePath, ShellExecutable, ShellParametersBegin + "\"" + GradleScriptPath + "\" " + GradleOptions + ShellParametersEnd, "Listing all tasks...");
 
 								GradleOptions = "clean";
-								RunCommandLineProgramWithExceptionAndFiltering(UnrealBuildGradlePath, ShellExecutable, ShellParametersBegin + "\"" + GradleScriptPath + "\" " + GradleOptions + ShellParametersEnd, "Cleaning Gradle intermediates...");
+								RunCommandLineProgramWithExceptionAndFiltering(UnrealBuildGradlePath, ShellExecutable, ShellParametersBegin + "\"" + GradleScriptPath + "\" " + GradleOptions + ShellParametersEnd, Logger, "Cleaning Gradle intermediates...");
 							}
 						}
 						else
 						{
-							Log.TraceInformation("=============================================================================================");
-							Log.TraceInformation("Android Studio project found, skipping Gradle; complete creation of APK in Android Studio!!!!");
-							Log.TraceInformation("Delete '{0} if you want to have UnrealBuildTool run Gradle for future runs.", GradleAppImlFilename);
-							Log.TraceInformation("=============================================================================================");
+							Logger.LogInformation("=============================================================================================");
+							Logger.LogInformation("Android Studio project found, skipping Gradle; complete creation of APK in Android Studio!!!!");
+							Logger.LogInformation("Delete '{GradleAppImlFilename} if you want to have UnrealBuildTool run Gradle for future runs.", GradleAppImlFilename);
+							Logger.LogInformation("=============================================================================================");
 						}
 					}
 				}
@@ -4438,7 +4442,7 @@ namespace UnrealBuildTool
 					string SymbolSODirectory = Path.Combine(DestApkDirectory, ProjectName + "_Symbols_v" + StoreVersion + "/" + ProjectName + Arch);
 					string SymbolifiedSOPath = Path.Combine(SymbolSODirectory, Path.GetFileName(FinalSOName));
 					MakeDirectoryIfRequired(SymbolifiedSOPath);
-					Log.TraceInformation("Writing symbols to {0}", SymbolifiedSOPath);
+					Logger.LogInformation("Writing symbols to {SymbolifiedSOPath}", SymbolifiedSOPath);
 
 					File.Copy(FinalSOName, SymbolifiedSOPath, true);
 				}
@@ -4453,7 +4457,7 @@ namespace UnrealBuildTool
 				string UnrealGradleDest = Path.Combine(IntermediateAndroidPath, "gradle");
 
 				// start fresh each time for now
-				DeleteDirectory(UnrealGradleDest);
+				DeleteDirectory(UnrealGradleDest, Logger);
 
 				// make sure destination exists
 				Directory.CreateDirectory(UnrealGradleDest);
@@ -4472,7 +4476,7 @@ namespace UnrealBuildTool
 
 					if (!Directory.Exists(UnrealBuildGradlePath))
 					{
-						Log.TraceInformation("Source directory missing: {0}", UnrealBuildGradlePath);
+						Logger.LogInformation("Source directory missing: {GradlePath}", UnrealBuildGradlePath);
 						bCombinedBundleOK = false;
 						break;
 					}
@@ -4510,7 +4514,7 @@ namespace UnrealBuildTool
 							DestFileInfo.Attributes = DestFileInfo.Attributes & ~FileAttributes.ReadOnly;
 							File.SetLastWriteTimeUtc(DestFilename, File.GetLastWriteTimeUtc(Filename));
 
-							Log.TraceInformation("Copied file {0}.", DestFilename);
+							Logger.LogInformation("Copied file {DestFilename}.", DestFilename);
 							continue;
 						}
 
@@ -4582,7 +4586,7 @@ namespace UnrealBuildTool
 
 							// differed too much
 							bCombinedBundleOK = false;
-							Log.TraceInformation("AndroidManifest.xml files differ too much to combine for single AAB: '{0}' != '{1}'", Filename, DestFilename);
+							Logger.LogInformation("AndroidManifest.xml files differ too much to combine for single AAB: '{Filename}' != '{DestFilename}'", Filename, DestFilename);
 							break;
 						}
 
@@ -4619,7 +4623,7 @@ namespace UnrealBuildTool
 
 							// differed too much
 							bCombinedBundleOK = false;
-							Log.TraceInformation("buildAdditions.gradle files differ too much to combine for single AAB: '{0}' != '{1}'", Filename, DestFilename);
+							Logger.LogInformation("buildAdditions.gradle files differ too much to combine for single AAB: '{Filename}' != '{DestFilename}'", Filename, DestFilename);
 							break;
 						}
 
@@ -4677,13 +4681,13 @@ namespace UnrealBuildTool
 
 							// differed too much
 							bCombinedBundleOK = false;
-							Log.TraceInformation("gradle.properties files differ too much to combine for single AAB: '{0}' != '{1}'", Filename, DestFilename);
+							Logger.LogInformation("gradle.properties files differ too much to combine for single AAB: '{Filename}' != '{DestFilename}'", Filename, DestFilename);
 							break;
 						}
 
 						// there are unknown differences, cannot make a single AAB
 						bCombinedBundleOK = false;
-						Log.TraceInformation("Gradle projects differ too much to combine for single AAB: '{0}' != '{1}'", Filename, DestFilename);
+						Logger.LogInformation("Gradle projects differ too much to combine for single AAB: '{Filename}' != '{DestFilename}'", Filename, DestFilename);
 						break;
 					}
 				}
@@ -4744,7 +4748,7 @@ namespace UnrealBuildTool
 					if (!RuntimePlatform.IsWindows)
 					{
 						// fix permissions for Mac/Linux
-						RunCommandLineProgramWithException(UnrealBuildGradlePath, "/bin/sh", string.Format("-c 'chmod 0755 \"{0}\"'", GradleScriptPath.Replace("'", "'\"'\"'")), "Fix gradlew permissions");
+						RunCommandLineProgramWithException(UnrealBuildGradlePath, "/bin/sh", string.Format("-c 'chmod 0755 \"{0}\"'", GradleScriptPath.Replace("'", "'\"'\"'")), Logger, "Fix gradlew permissions");
 					}
 					else
 					{
@@ -4771,11 +4775,11 @@ namespace UnrealBuildTool
 					string ShellExecutable = RuntimePlatform.IsWindows ? "cmd.exe" : "/bin/sh";
 					string ShellParametersBegin = RuntimePlatform.IsWindows ? "/c " : "-c '";
 					string ShellParametersEnd = RuntimePlatform.IsWindows ? "" : "'";
-					RunCommandLineProgramWithExceptionAndFiltering(UnrealBuildGradlePath, ShellExecutable, ShellParametersBegin + "\"" + GradleScriptPath + "\" " + GradleOptions + ShellParametersEnd, "Making .aab with Gradle...");
+					RunCommandLineProgramWithExceptionAndFiltering(UnrealBuildGradlePath, ShellExecutable, ShellParametersBegin + "\"" + GradleScriptPath + "\" " + GradleOptions + ShellParametersEnd, Logger, "Making .aab with Gradle...");
 
 					if (GradleSecondCallOptions != "")
 					{
-						RunCommandLineProgramWithExceptionAndFiltering(UnrealBuildGradlePath, ShellExecutable, ShellParametersBegin + "\"" + GradleScriptPath + "\" " + GradleSecondCallOptions + ShellParametersEnd, "Additional Gradle steps...");
+						RunCommandLineProgramWithExceptionAndFiltering(UnrealBuildGradlePath, ShellExecutable, ShellParametersBegin + "\"" + GradleScriptPath + "\" " + GradleSecondCallOptions + ShellParametersEnd, Logger, "Additional Gradle steps...");
 					}
 
 					// For build machine run a clean afterward to clean up intermediate files (does not remove final APK)
@@ -4785,7 +4789,7 @@ namespace UnrealBuildTool
 						//RunCommandLineProgramWithException(UnrealBuildGradlePath, ShellExecutable, ShellParametersBegin + "\"" + GradleScriptPath + "\" " + GradleOptions + ShellParametersEnd, "Listing all tasks...");
 
 						GradleOptions = "clean";
-						RunCommandLineProgramWithExceptionAndFiltering(UnrealBuildGradlePath, ShellExecutable, ShellParametersBegin + "\"" + GradleScriptPath + "\" " + GradleOptions + ShellParametersEnd, "Cleaning Gradle intermediates...");
+						RunCommandLineProgramWithExceptionAndFiltering(UnrealBuildGradlePath, ShellExecutable, ShellParametersBegin + "\"" + GradleScriptPath + "\" " + GradleOptions + ShellParametersEnd, Logger, "Cleaning Gradle intermediates...");
 					}
 				}
 				else
@@ -4797,7 +4801,7 @@ namespace UnrealBuildTool
 						string Manifest = build.Item2;
 						string NDKArch = GetNDKArch(Arch);
 
-						Log.TraceInformation("\n===={0}====GENERATING BUNDLE====={1}================================================================", DateTime.Now.ToString(), Arch);
+						Logger.LogInformation("\n===={Time}====GENERATING BUNDLE====={Arch}================================================================", DateTime.Now.ToString(), Arch);
 
 						string UnrealBuildPath = Path.Combine(IntermediateAndroidPath, Arch.Substring(1).Replace("-", "_"));
 						string UnrealBuildGradlePath = Path.Combine(UnrealBuildPath, "gradle");
@@ -4806,7 +4810,7 @@ namespace UnrealBuildTool
 						if (!RuntimePlatform.IsWindows)
 						{
 							// fix permissions for Mac/Linux
-							RunCommandLineProgramWithException(UnrealBuildGradlePath, "/bin/sh", string.Format("-c 'chmod 0755 \"{0}\"'", GradleScriptPath.Replace("'", "'\"'\"'")), "Fix gradlew permissions");
+							RunCommandLineProgramWithException(UnrealBuildGradlePath, "/bin/sh", string.Format("-c 'chmod 0755 \"{0}\"'", GradleScriptPath.Replace("'", "'\"'\"'")), Logger, "Fix gradlew permissions");
 						}
 						else
 						{
@@ -4833,11 +4837,11 @@ namespace UnrealBuildTool
 						string ShellExecutable = RuntimePlatform.IsWindows ? "cmd.exe" : "/bin/sh";
 						string ShellParametersBegin = RuntimePlatform.IsWindows ? "/c " : "-c '";
 						string ShellParametersEnd = RuntimePlatform.IsWindows ? "" : "'";
-						RunCommandLineProgramWithExceptionAndFiltering(UnrealBuildGradlePath, ShellExecutable, ShellParametersBegin + "\"" + GradleScriptPath + "\" " + GradleOptions + ShellParametersEnd, "Making .aab with Gradle...");
+						RunCommandLineProgramWithExceptionAndFiltering(UnrealBuildGradlePath, ShellExecutable, ShellParametersBegin + "\"" + GradleScriptPath + "\" " + GradleOptions + ShellParametersEnd, Logger, "Making .aab with Gradle...");
 
 						if (GradleSecondCallOptions != "")
 						{
-							RunCommandLineProgramWithExceptionAndFiltering(UnrealBuildGradlePath, ShellExecutable, ShellParametersBegin + "\"" + GradleScriptPath + "\" " + GradleSecondCallOptions + ShellParametersEnd, "Additional Gradle steps...");
+							RunCommandLineProgramWithExceptionAndFiltering(UnrealBuildGradlePath, ShellExecutable, ShellParametersBegin + "\"" + GradleScriptPath + "\" " + GradleSecondCallOptions + ShellParametersEnd, Logger, "Additional Gradle steps...");
 						}
 
 						// For build machine run a clean afterward to clean up intermediate files (does not remove final APK)
@@ -4847,13 +4851,13 @@ namespace UnrealBuildTool
 							//RunCommandLineProgramWithException(UnrealBuildGradlePath, ShellExecutable, ShellParametersBegin + "\"" + GradleScriptPath + "\" " + GradleOptions + ShellParametersEnd, "Listing all tasks...");
 
 							GradleOptions = "clean";
-							RunCommandLineProgramWithExceptionAndFiltering(UnrealBuildGradlePath, ShellExecutable, ShellParametersBegin + "\"" + GradleScriptPath + "\" " + GradleOptions + ShellParametersEnd, "Cleaning Gradle intermediates...");
+							RunCommandLineProgramWithExceptionAndFiltering(UnrealBuildGradlePath, ShellExecutable, ShellParametersBegin + "\"" + GradleScriptPath + "\" " + GradleOptions + ShellParametersEnd, Logger, "Cleaning Gradle intermediates...");
 						}
 					}
 				}
 			}
 
-			Log.TraceInformation("\n===={0}====COMPLETED MAKE APK=======================================================================", DateTime.Now.ToString());
+			Logger.LogInformation("\n===={Time}====COMPLETED MAKE APK=======================================================================", DateTime.Now.ToString());
 		}
 
 		private List<string> CollectPluginDataPaths(TargetReceipt Receipt)
@@ -4861,7 +4865,7 @@ namespace UnrealBuildTool
 			List<string> PluginExtras = new List<string>();
 			if (Receipt == null)
 			{
-				Log.TraceInformation("Receipt is NULL");
+				Logger.LogInformation("Receipt is NULL");
 				return PluginExtras;
 			}
 
@@ -4874,7 +4878,7 @@ namespace UnrealBuildTool
 				if (PluginExtras.FirstOrDefault(x => x == PluginPath) == null)
 				{
 					PluginExtras.Add(PluginPath);
-					Log.TraceInformation("AndroidPlugin: {0}", PluginPath);
+					Logger.LogInformation("AndroidPlugin: {PluginPath}", PluginPath);
 				}
 			}
 			return PluginExtras;
@@ -4882,7 +4886,7 @@ namespace UnrealBuildTool
 
 		public override bool PrepTargetForDeployment(TargetReceipt Receipt)
 		{
-			//Log.TraceInformation("$$$$$$$$$$$$$$ PrepTargetForDeployment $$$$$$$$$$$$$$$$$");
+			//Logger.LogInformation("$$$$$$$$$$$$$$ PrepTargetForDeployment $$$$$$$$$$$$$$$$$");
 
 			DirectoryReference ProjectDirectory = DirectoryReference.FromFile(Receipt.ProjectFile) ?? Unreal.EngineDirectory;
 			string TargetName = (Receipt.ProjectFile == null ? Receipt.TargetName : Receipt.ProjectFile.GetFileNameWithoutAnyExtensions());
@@ -4916,9 +4920,9 @@ namespace UnrealBuildTool
 			// if we made any non-standard .apk files, the generated debugger settings may be wrong
 			if (ShouldMakeSeparateApks() && (OutputPaths.Count > 1 || !OutputPaths[0].FullName.Contains("-armv7")))
 			{
-				Log.TraceInformation("================================================================================================================================");
-				Log.TraceInformation("Non-default apk(s) have been made: If you are debugging, you will need to manually select one to run in the debugger properties!");
-				Log.TraceInformation("================================================================================================================================");
+				Logger.LogInformation("================================================================================================================================");
+				Logger.LogInformation("Non-default apk(s) have been made: If you are debugging, you will need to manually select one to run in the debugger properties!");
+				Logger.LogInformation("================================================================================================================================");
 			}
 			return true;
 		}
@@ -4935,7 +4939,7 @@ namespace UnrealBuildTool
 			PackageInfoSource[2] = GetVersionDisplayName(bIsEmbedded);
 			PackageInfoSource[3] = string.Format("name='com.epicgames.unreal.GameActivity.AppType' value='{0}'", InTargetType == TargetType.Game ? "" : InTargetType.ToString());
 
-			Log.TraceInformation("Writing packageInfo pkgName:{0} storeVersion:{1} versionDisplayName:{2} to {3}", PackageInfoSource[0], PackageInfoSource[1], PackageInfoSource[2], DestPackageNameFileName);
+			Logger.LogInformation("Writing packageInfo pkgName:{PkgName} storeVersion:{StoreVer} versionDisplayName:{Version} to {DestFile}", PackageInfoSource[0], PackageInfoSource[1], PackageInfoSource[2], DestPackageNameFileName);
 
 			string DestDirectory = Path.GetDirectoryName(DestPackageNameFileName)!;
 			if (!Directory.Exists(DestDirectory))
@@ -4966,7 +4970,7 @@ namespace UnrealBuildTool
 
 		public bool PrepForUATPackageOrDeploy(FileReference ProjectFile, string ProjectName, DirectoryReference ProjectDirectory, string ExecutablePath, string EngineDirectory, bool bForDistribution, string CookFlavor, UnrealTargetConfiguration Configuration, bool bIsDataDeploy, bool bSkipGradleBuild)
 		{
-			//Log.TraceInformation("$$$$$$$$$$$$$$ PrepForUATPackageOrDeploy $$$$$$$$$$$$$$$$$");
+			//Logger.LogInformation("$$$$$$$$$$$$$$ PrepForUATPackageOrDeploy $$$$$$$$$$$$$$$$$");
 
 			TargetType Type = TargetType.Game;
 			if (CookFlavor.EndsWith("Client"))
@@ -4981,9 +4985,9 @@ namespace UnrealBuildTool
 			// note that we cannot allow the data packaged into the APK if we are doing something like Launch On that will not make an obb
 			// file and instead pushes files directly via deploy
 			AndroidTargetRules TargetRules = new AndroidTargetRules();
-			CommandLine.ParseArguments(Environment.GetCommandLineArgs(), TargetRules);
+			CommandLine.ParseArguments(Environment.GetCommandLineArgs(), TargetRules, Logger);
 			ClangToolChainOptions Options = AndroidPlatform.CreateToolChainOptions(TargetRules);
-			AndroidToolChain ToolChain = new AndroidToolChain(ProjectFile, false, null, null, Options);
+			AndroidToolChain ToolChain = new AndroidToolChain(ProjectFile, false, null, null, Options, Logger);
 
 			SavePackageInfo(ProjectName, ProjectDirectory.FullName, Type, bSkipGradleBuild);
 
@@ -4992,11 +4996,11 @@ namespace UnrealBuildTool
 			return true;
 		}
 
-		public static void OutputReceivedDataEventHandler(Object Sender, DataReceivedEventArgs Line)
+		public static void OutputReceivedDataEventHandler(Object Sender, DataReceivedEventArgs Line, ILogger Logger)
 		{
 			if ((Line != null) && (Line.Data != null))
 			{
-				Log.TraceInformation(Line.Data);
+				Logger.LogInformation("{Output}", Line.Data);
 			}
 		}
 
@@ -5159,7 +5163,8 @@ namespace UnrealBuildTool
 
 			if (TemplateDest == null || TemplateSrc.Length != TemplateDest.Length || !TemplateSrc.SequenceEqual(TemplateDest))
 			{
-				Log.TraceInformation("\n==== Writing new GameActivity.java file to {0} ====", DestFilename);
+				Logger.LogInformation("");
+				Logger.LogInformation("==== Writing new GameActivity.java file to {DestFile} ====", DestFilename);
 				File.WriteAllLines(DestFilename, TemplateSrc);
 			}
 		}
@@ -5209,7 +5214,8 @@ namespace UnrealBuildTool
 
 			if (TemplateDest == null || TemplateSrc.Length != TemplateDest.Length || !TemplateSrc.SequenceEqual(TemplateDest))
 			{
-				Log.TraceInformation("\n==== Writing new GameApplication.java file to {0} ====", DestFilename);
+				Logger.LogInformation("");
+				Logger.LogInformation("==== Writing new GameApplication.java file to {DestFile} ====", DestFilename);
 				File.WriteAllLines(DestFilename, TemplateSrc);
 			}
 		}
@@ -5286,7 +5292,7 @@ namespace UnrealBuildTool
 					DirectoryPath = DirectoryPath.Replace("$(ENGINEDIR)", EngineDir);
 					DirectoryPath = DirectoryPath.Replace("$(ANDROID_HOME)", AndroidHome);
 					DirectoryPath = DirectoryPath.Replace('\\', Path.DirectorySeparatorChar).Replace('/', Path.DirectorySeparatorChar);
-					AARHandler.AddRepository(DirectoryPath);
+					AARHandler.AddRepository(DirectoryPath, Logger);
 				}
 				else if (Trimmed.StartsWith("repositories "))
 				{
@@ -5294,7 +5300,7 @@ namespace UnrealBuildTool
 					DirectoryPath = DirectoryPath.Replace("$(ENGINEDIR)", EngineDir);
 					DirectoryPath = DirectoryPath.Replace("$(ANDROID_HOME)", AndroidHome);
 					DirectoryPath = DirectoryPath.Replace('\\', Path.DirectorySeparatorChar).Replace('/', Path.DirectorySeparatorChar);
-					AARHandler.AddRepositories(DirectoryPath, "m2repository");
+					AARHandler.AddRepositories(DirectoryPath, "m2repository", Logger);
 				}
 				else
 				{
@@ -5304,8 +5310,8 @@ namespace UnrealBuildTool
 						string PackageName = Sections[0].Trim(' ');
 						string BaseName = Sections[1].Trim(' ');
 						string Version = Sections[2].Trim(' ');
-						Log.TraceInformation("AARImports: {0}, {1}, {2}", PackageName, BaseName, Version);
-						AARHandler.AddNewAAR(PackageName, BaseName, Version, HandleDependencies);
+						Logger.LogInformation("AARImports: {PackageName}, {BaseName}, {Version}", PackageName, BaseName, Version);
+						AARHandler.AddNewAAR(PackageName, BaseName, Version, Logger, HandleDependencies);
 					}
 				}
 			}
@@ -5421,7 +5427,7 @@ namespace UnrealBuildTool
 					}
 					catch (Exception e)
 					{
-						Log.TraceError("AAR Manifest file {0} parsing error! {1}", ManifestFilename, e);
+						Logger.LogError(e, "AAR Manifest file {FileName} parsing error! {Ex}", ManifestFilename, e);
 					}
 				}
 
diff --git a/Engine/Source/Programs/UnrealBuildTool/Platform/IOS/IOSExports.cs b/Engine/Source/Programs/UnrealBuildTool/Platform/IOS/IOSExports.cs
--- a/Engine/Source/Programs/UnrealBuildTool/Platform/IOS/IOSExports.cs
+++ b/Engine/Source/Programs/UnrealBuildTool/Platform/IOS/IOSExports.cs
@@ -10,6 +10,7 @@ using System.Threading.Tasks;
 using System.Xml.Linq;
 using EpicGames.Core;
 using UnrealBuildBase;
+using Microsoft.Extensions.Logging;
 
 namespace UnrealBuildTool
 {
@@ -79,11 +80,12 @@ namespace UnrealBuildTool
 		/// <param name="bIsDataDeploy"></param>
 		/// <param name="bCreateStubIPA"></param>
 		/// <param name="BuildReceiptFileName"></param>
+		/// <param name="Logger">Logger for output</param>
 		/// <returns></returns>
-		public static bool PrepForUATPackageOrDeploy(UnrealTargetConfiguration Config, FileReference ProjectFile, string InProjectName, DirectoryReference InProjectDirectory, string InExecutablePath, DirectoryReference InEngineDir, bool bForDistribution, string CookFlavor, bool bIsDataDeploy, bool bCreateStubIPA, FileReference BuildReceiptFileName)
+		public static bool PrepForUATPackageOrDeploy(UnrealTargetConfiguration Config, FileReference ProjectFile, string InProjectName, DirectoryReference InProjectDirectory, string InExecutablePath, DirectoryReference InEngineDir, bool bForDistribution, string CookFlavor, bool bIsDataDeploy, bool bCreateStubIPA, FileReference BuildReceiptFileName, ILogger Logger)
 		{
 			TargetReceipt Receipt = TargetReceipt.Read(BuildReceiptFileName);
-			return new UEDeployIOS().PrepForUATPackageOrDeploy(Config, ProjectFile, InProjectName, InProjectDirectory.FullName, InExecutablePath, InEngineDir.FullName, bForDistribution, CookFlavor, bIsDataDeploy, bCreateStubIPA, Receipt);
+			return new UEDeployIOS(Logger).PrepForUATPackageOrDeploy(Config, ProjectFile, InProjectName, InProjectDirectory.FullName, InExecutablePath, InEngineDir.FullName, bForDistribution, CookFlavor, bIsDataDeploy, bCreateStubIPA, Receipt);
 		}
 
 		/// <summary>
@@ -99,13 +101,14 @@ namespace UnrealBuildTool
 		/// <param name="InEngineDir"></param>
 		/// <param name="AppDirectory"></param>
 		/// <param name="BuildReceiptFileName"></param>
+		/// <param name="Logger">Logger for output</param>
 		/// <param name="bSupportsPortrait"></param>
 		/// <param name="bSupportsLandscape"></param>
 		/// <returns></returns>
-		public static bool GeneratePList(FileReference ProjectFile, UnrealTargetConfiguration Config, DirectoryReference ProjectDirectory, bool bIsUnrealGame, string GameName, bool bIsClient, string ProjectName, DirectoryReference InEngineDir, DirectoryReference AppDirectory, FileReference BuildReceiptFileName, out bool bSupportsPortrait, out bool bSupportsLandscape)
+		public static bool GeneratePList(FileReference ProjectFile, UnrealTargetConfiguration Config, DirectoryReference ProjectDirectory, bool bIsUnrealGame, string GameName, bool bIsClient, string ProjectName, DirectoryReference InEngineDir, DirectoryReference AppDirectory, FileReference BuildReceiptFileName, ILogger Logger, out bool bSupportsPortrait, out bool bSupportsLandscape)
 		{
 			TargetReceipt Receipt = TargetReceipt.Read(BuildReceiptFileName);
-			return new UEDeployIOS().GeneratePList(ProjectFile, Config, ProjectDirectory.FullName, bIsUnrealGame, GameName, bIsClient, ProjectName, InEngineDir.FullName, AppDirectory.FullName, Receipt, out bSupportsPortrait, out bSupportsLandscape);
+			return new UEDeployIOS(Logger).GeneratePList(ProjectFile, Config, ProjectDirectory.FullName, bIsUnrealGame, GameName, bIsClient, ProjectName, InEngineDir.FullName, AppDirectory.FullName, Receipt, out bSupportsPortrait, out bSupportsLandscape);
 		}
 
 		/// <summary>
@@ -114,10 +117,11 @@ namespace UnrealBuildTool
 		/// <param name="PlatformType"></param>
 		/// <param name="SourceFile"></param>
 		/// <param name="TargetFile"></param>
-		public static void StripSymbols(UnrealTargetPlatform PlatformType, FileReference SourceFile, FileReference TargetFile)
+		/// <param name="Logger"></param>
+		public static void StripSymbols(UnrealTargetPlatform PlatformType, FileReference SourceFile, FileReference TargetFile, ILogger Logger)
 		{
 			IOSProjectSettings ProjectSettings = ((IOSPlatform)UEBuildPlatform.GetBuildPlatform(PlatformType)).ReadProjectSettings(null);
-			IOSToolChain ToolChain = new IOSToolChain(null, ProjectSettings, ClangToolChainOptions.None);
+			IOSToolChain ToolChain = new IOSToolChain(null, ProjectSettings, ClangToolChainOptions.None, Logger);
 			ToolChain.StripSymbols(SourceFile, TargetFile);
 		}
 
@@ -128,7 +132,8 @@ namespace UnrealBuildTool
 		/// <param name="Executable"></param>
 		/// <param name="StageDirectory"></param>
 		/// <param name="Platform"></param>
-		public static void GenerateAssetCatalog(FileReference ProjectFile, FileReference Executable, DirectoryReference StageDirectory, UnrealTargetPlatform Platform)
+		/// <param name="Logger"></param>
+		public static void GenerateAssetCatalog(FileReference ProjectFile, FileReference Executable, DirectoryReference StageDirectory, UnrealTargetPlatform Platform, ILogger Logger)
 		{
 			// Determine whether the user has modified icons that require a remote Mac to build.
 			bool bUserImagesExist = false;
@@ -151,8 +156,8 @@ namespace UnrealBuildTool
 			{
 				FileReference OutputFile = FileReference.Combine(StageDirectory, "Assets.car");
 
-				RemoteMac Remote = new RemoteMac(ProjectFile);
-				Remote.RunAssetCatalogTool(Platform, ResourcesDir, OutputFile);
+				RemoteMac Remote = new RemoteMac(ProjectFile, Logger);
+				Remote.RunAssetCatalogTool(Platform, ResourcesDir, OutputFile, Logger);
 			}
 			else
 			{
diff --git a/Engine/Source/Programs/UnrealBuildTool/Platform/IOS/IOSPostBuildSyncMode.cs b/Engine/Source/Programs/UnrealBuildTool/Platform/IOS/IOSPostBuildSyncMode.cs
--- a/Engine/Source/Programs/UnrealBuildTool/Platform/IOS/IOSPostBuildSyncMode.cs
+++ b/Engine/Source/Programs/UnrealBuildTool/Platform/IOS/IOSPostBuildSyncMode.cs
@@ -7,6 +7,7 @@
 using System.Threading.Tasks;
 using EpicGames.Core;
 using UnrealBuildBase;
+using Microsoft.Extensions.Logging;
 
 namespace UnrealBuildTool
 {
@@ -64,14 +65,14 @@
 		[CommandLine("-XmlConfigCache=")]
 		public FileReference? XmlConfigCache = null;
 
-		public override int Execute(CommandLineArguments Arguments)
+		public override int Execute(CommandLineArguments Arguments, ILogger Logger)
 		{
 			Arguments.ApplyTo(this);
 			Arguments.CheckAllArgumentsUsed();
 
 			// Run the PostBuildSync command
 			IOSPostBuildSyncTarget Target = BinaryFormatterUtils.Load<IOSPostBuildSyncTarget>(InputFile!);
-			IOSToolChain.PostBuildSync(Target);
+			IOSToolChain.PostBuildSync(Target, Logger);
 
 			return 0;
 		}
diff --git a/Engine/Source/Programs/UnrealBuildTool/Platform/IOS/IOSProjectGenerator.cs b/Engine/Source/Programs/UnrealBuildTool/Platform/IOS/IOSProjectGenerator.cs
--- a/Engine/Source/Programs/UnrealBuildTool/Platform/IOS/IOSProjectGenerator.cs
+++ b/Engine/Source/Programs/UnrealBuildTool/Platform/IOS/IOSProjectGenerator.cs
@@ -6,6 +6,7 @@
 using System.Diagnostics;
 using System.IO;
 using EpicGames.Core;
+using Microsoft.Extensions.Logging;
 
 namespace UnrealBuildTool
 {
@@ -18,8 +19,9 @@
 		/// Constructor
 		/// </summary>
 		/// <param name="Arguments">Command line arguments passed to the project generator</param>
-		public IOSProjectGenerator(CommandLineArguments Arguments)
-			: base(Arguments)
+		/// <param name="Logger">Logger for output</param>
+		public IOSProjectGenerator(CommandLineArguments Arguments, ILogger Logger)
+			: base(Arguments, Logger)
 		{
 		}
 
diff --git a/Engine/Source/Programs/UnrealBuildTool/Platform/IOS/IOSToolChain.cs b/Engine/Source/Programs/UnrealBuildTool/Platform/IOS/IOSToolChain.cs
--- a/Engine/Source/Programs/UnrealBuildTool/Platform/IOS/IOSToolChain.cs
+++ b/Engine/Source/Programs/UnrealBuildTool/Platform/IOS/IOSToolChain.cs
@@ -16,6 +16,7 @@ using Ionic.Zlib;
 using EpicGames.Core;
 using System.Security.Cryptography.X509Certificates;
 using UnrealBuildBase;
+using Microsoft.Extensions.Logging;
 
 namespace UnrealBuildTool
 {
@@ -48,11 +49,11 @@ namespace UnrealBuildTool
 		public readonly string DevicePlatformName;
 		public readonly string SimulatorPlatformName;
 
-		public IOSToolChainSettings() : this("iPhoneOS", "iPhoneSimulator")
+		public IOSToolChainSettings(ILogger Logger) : this("iPhoneOS", "iPhoneSimulator", Logger)
 		{
 		}
 
-		protected IOSToolChainSettings(string DevicePlatformName, string SimulatorPlatformName) : base(true)
+		protected IOSToolChainSettings(string DevicePlatformName, string SimulatorPlatformName, ILogger Logger) : base(true, Logger)
 		{
 			XmlConfig.ApplyTo(this);
 
@@ -65,7 +66,7 @@ namespace UnrealBuildTool
 			ToolchainDir = XcodeDeveloperDir + "Toolchains/XcodeDefault.xctoolchain/usr/bin/";
 
 			// make sure SDK is selected
-			SelectSDK(BaseSDKDir, DevicePlatformName, ref IOSSDKVersion, true);
+			SelectSDK(BaseSDKDir, DevicePlatformName, ref IOSSDKVersion, true, Logger);
 
 			// convert to float for easy comparison
 			IOSSDKVersionFloat = float.Parse(IOSSDKVersion, System.Globalization.CultureInfo.InvariantCulture);
@@ -89,13 +90,13 @@ namespace UnrealBuildTool
 		public readonly ReadOnlyTargetRules? Target;
 		protected IOSProjectSettings ProjectSettings;
 
-		public IOSToolChain(ReadOnlyTargetRules? Target, IOSProjectSettings InProjectSettings, ClangToolChainOptions ToolchainOptions)
-			: this(Target, InProjectSettings, () => new IOSToolChainSettings(), ToolchainOptions)
+		public IOSToolChain(ReadOnlyTargetRules? Target, IOSProjectSettings InProjectSettings, ClangToolChainOptions ToolchainOptions, ILogger InLogger)
+			: this(Target, InProjectSettings, () => new IOSToolChainSettings(InLogger), ToolchainOptions, InLogger)
 		{
 		}
 
-		protected IOSToolChain(ReadOnlyTargetRules? Target, IOSProjectSettings InProjectSettings, Func<IOSToolChainSettings> InCreateSettings, ClangToolChainOptions ToolchainOptions)
-			: base((Target == null) ? null : Target.ProjectFile, ToolchainOptions)
+		protected IOSToolChain(ReadOnlyTargetRules? Target, IOSProjectSettings InProjectSettings, Func<IOSToolChainSettings> InCreateSettings, ClangToolChainOptions ToolchainOptions, ILogger InLogger)
+			: base((Target == null) ? null : Target.ProjectFile, ToolchainOptions, InLogger)
 		{
 			this.Target = Target;
 			ProjectSettings = InProjectSettings;
@@ -662,7 +663,7 @@ namespace UnrealBuildTool
 									|| SourceFile.AbsolutePath.Contains("SWidget.cpp") || SourceText.Contains("SWidget.cpp") || SourceFile.AbsolutePath.Contains("SCanvas.cpp") || SourceText.Contains("SCanvas.cpp") || SourceFile.AbsolutePath.Contains("ShaderCore.cpp") || SourceText.Contains("ShaderCore.cpp")
 									|| SourceFile.AbsolutePath.Contains("ParticleSystemRender.cpp") || SourceText.Contains("ParticleSystemRender.cpp")))
 								{
-									Log.TraceInformation("Forcing {0} to --O3!", SourceFile.AbsolutePath);
+									Logger.LogInformation("Forcing {SourceFileAbsolutePath} to --O3!", SourceFile.AbsolutePath);
 
 									AllArgs = AllArgs.Replace("-Oz", "-O3");
 								}*/
@@ -695,7 +696,7 @@ namespace UnrealBuildTool
 				{
 					if (Framework.ZipFile != null)
 					{
-						FileItem ExtractedTokenFile = ExtractFramework(Framework, Graph);
+						FileItem ExtractedTokenFile = ExtractFramework(Framework, Graph, Logger);
 						CompileAction.PrerequisiteItems.Add(ExtractedTokenFile);
 					}
 				}
@@ -760,7 +761,7 @@ namespace UnrealBuildTool
 			{
 				if (Framework.ZipFile != null)
 				{
-					FileItem ExtractedTokenFile = ExtractFramework(Framework, Graph);
+					FileItem ExtractedTokenFile = ExtractFramework(Framework, Graph, Logger);
 					LinkAction.PrerequisiteItems.Add(ExtractedTokenFile);
 				}
 			}
@@ -946,7 +947,8 @@ namespace UnrealBuildTool
 		/// <param name="Executable">FileItem describing the executable to generate debug info for</param>
 		/// <param name="bIsForLTOBuild">Was this build made with LTO enabled?</param>
 		/// <param name="Graph">List of actions to be executed. Additional actions will be added to this list.</param>
-		private List<FileItem> GenerateDebugInfo(FileItem Executable, bool bIsForLTOBuild, IActionGraphBuilder Graph)
+		/// <param name="Logger">Logger for output</param>
+		private List<FileItem> GenerateDebugInfo(FileItem Executable, bool bIsForLTOBuild, IActionGraphBuilder Graph, ILogger Logger)
 		{
 			// Make a file item for the source and destination files
 			string FullDestPathRoot = GetdSYMPath(Executable);
@@ -960,7 +962,7 @@ namespace UnrealBuildTool
 			GenDebugAction.WorkingDirectory = GetMacDevSrcRoot();
 			GenDebugAction.CommandPath = BuildHostPlatform.Current.Shell;
 			string ExtraOptions;
-			string DsymutilPath = GetDsymutilPath(out ExtraOptions, bIsForLTOBuild: bIsForLTOBuild);
+			string DsymutilPath = GetDsymutilPath(Logger, out ExtraOptions, bIsForLTOBuild: bIsForLTOBuild);
 			if (ProjectSettings.bGeneratedSYMBundle)
 			{
 				GenDebugAction.CommandArguments = string.Format("-c \"rm -rf \\\"{2}\\\"; \\\"{0}\\\" \\\"{1}\\\" {4} -o \\\"{2}\\\"; cd \\\"{2}/..\\\"; zip -r -y -1 {3}.zip {3}\"",
@@ -970,7 +972,7 @@ namespace UnrealBuildTool
 					Path.GetFileName(FullDestPathRoot),
 					ExtraOptions);
 				GenDebugAction.ProducedItems.Add(ZipOutputFile);
-				Log.TraceInformation("Zip file: " + ZipOutputFile.AbsolutePath);
+				Logger.LogInformation("Zip file: {File}", ZipOutputFile.AbsolutePath);
 			}
 			else
 			{
@@ -1138,7 +1140,7 @@ namespace UnrealBuildTool
 			}
 		}
 
-		FileItem ExtractFramework(UEBuildFramework Framework, IActionGraphBuilder Graph)
+		FileItem ExtractFramework(UEBuildFramework Framework, IActionGraphBuilder Graph, ILogger Logger)
 		{
 			if (Framework.ZipFile == null)
 			{
@@ -1380,7 +1382,7 @@ namespace UnrealBuildTool
 			// For IOS/tvOS, generate the dSYM file if needed or requested
 			if (Target!.IOSPlatform.bGeneratedSYM)
 			{
-				List<FileItem> Files = GenerateDebugInfo(Executable, BinaryLinkEnvironment.bAllowLTCG, Graph);
+				List<FileItem> Files = GenerateDebugInfo(Executable, BinaryLinkEnvironment.bAllowLTCG, Graph, Logger);
 				foreach (FileItem item in Files)
 				{
 					OutputFiles.Add(item);
@@ -1439,9 +1441,9 @@ namespace UnrealBuildTool
 			// Generate the app bundle
 			if (!Target.bDisableLinking)
 			{
-				Log.TraceInformation("Adding PostBuildSync action");
+				Logger.LogInformation("Adding PostBuildSync action");
 
-				List<string> UPLScripts = UEDeployIOS.CollectPluginDataPaths(BinaryLinkEnvironment.AdditionalProperties);
+				List<string> UPLScripts = UEDeployIOS.CollectPluginDataPaths(BinaryLinkEnvironment.AdditionalProperties, Logger);
 				VersionNumber SdkVersion = VersionNumber.Parse(Settings.Value.IOSSDKVersion);
 
 				Dictionary<string, DirectoryReference> FrameworkNameToSourceDir = new Dictionary<string, DirectoryReference>();
@@ -1539,11 +1541,11 @@ namespace UnrealBuildTool
 			//This chunk looks to be required to pipe output to VS giving information on the status of a remote build.
 			public bool OutputReceivedDataEventHandlerEncounteredError = false;
 			public string OutputReceivedDataEventHandlerEncounteredErrorMessage = "";
-			public void OutputReceivedDataEventHandler(Object Sender, DataReceivedEventArgs Line)
+			public void OutputReceivedDataEventHandler(Object Sender, DataReceivedEventArgs Line, ILogger Logger)
 			{
 				if ((Line != null) && (Line.Data != null))
 				{
-					Log.TraceInformation(Line.Data);
+					Logger.LogInformation("{Message}", Line.Data);
 
 					foreach (string ErrorToken in ErrorMessageTokens)
 					{
@@ -1557,25 +1559,25 @@ namespace UnrealBuildTool
 				}
 			}
 
-			public void OutputReceivedDataEventLogger(Object Sender, DataReceivedEventArgs Line)
+			public void OutputReceivedDataEventLogger(Object Sender, DataReceivedEventArgs Line, ILogger Logger)
 			{
 				if ((Line != null) && (Line.Data != null))
 				{
-					Log.TraceInformation(Line.Data);
+					Logger.LogInformation("{Output}", Line.Data);
 				}
 			}
 		}
 
-		private static void GenerateCrashlyticsData(string DsymZip, string ProjectDir, string ProjectName)
+		private static void GenerateCrashlyticsData(string DsymZip, string ProjectDir, string ProjectName, ILogger Logger)
 		{
-			Log.TraceInformation("Generating and uploading Crashlytics Data");
+			Logger.LogInformation("Generating and uploading Crashlytics Data");
 
 			// Clean this folder as it's used for extraction
 			string TempPath = Path.Combine(Unreal.EngineDirectory.FullName, "Intermediate", "Unzipped");
 
 			if (Directory.Exists(TempPath))
 			{
-				Log.TraceInformation("Deleting temp path {0}", TempPath);
+				Logger.LogInformation("Deleting temp path {TempPath}", TempPath);
 				Directory.Delete(TempPath, true);
 			}
 
@@ -1592,8 +1594,8 @@ namespace UnrealBuildTool
 
 				ProcessOutput Output = new ProcessOutput();
 
-				FabricProcess.OutputDataReceived += new DataReceivedEventHandler(Output.OutputReceivedDataEventHandler);
-				FabricProcess.ErrorDataReceived += new DataReceivedEventHandler(Output.OutputReceivedDataEventHandler);
+				FabricProcess.OutputDataReceived += new DataReceivedEventHandler((s, e) => Output.OutputReceivedDataEventHandler(s, e, Logger));
+				FabricProcess.ErrorDataReceived += new DataReceivedEventHandler((s, e) => Output.OutputReceivedDataEventHandler(s, e, Logger));
 
 				Utils.RunLocalProcess(FabricProcess);
 				if (Output.OutputReceivedDataEventHandlerEncounteredError)
@@ -1603,7 +1605,7 @@ namespace UnrealBuildTool
 			}
 		}
 
-		internal static bool GenerateProjectFiles(FileReference? ProjectFile, string[] Arguments)
+		internal static bool GenerateProjectFiles(FileReference? ProjectFile, string[] Arguments, ILogger Logger)
 		{
 			ProjectFileGenerator.bGenerateProjectFiles = true;
 			try
@@ -1611,11 +1613,11 @@ namespace UnrealBuildTool
 				CommandLineArguments CmdLine = new CommandLineArguments(Arguments);
 
 				PlatformProjectGeneratorCollection PlatformProjectGenerators = new PlatformProjectGeneratorCollection();
-				PlatformProjectGenerators.RegisterPlatformProjectGenerator(UnrealTargetPlatform.IOS, new IOSProjectGenerator(CmdLine));
-				PlatformProjectGenerators.RegisterPlatformProjectGenerator(UnrealTargetPlatform.TVOS, new TVOSProjectGenerator(CmdLine));
+				PlatformProjectGenerators.RegisterPlatformProjectGenerator(UnrealTargetPlatform.IOS, new IOSProjectGenerator(CmdLine, Logger), Logger);
+				PlatformProjectGenerators.RegisterPlatformProjectGenerator(UnrealTargetPlatform.TVOS, new TVOSProjectGenerator(CmdLine, Logger), Logger);
 
 				XcodeProjectFileGenerator Generator = new XcodeProjectFileGenerator(ProjectFile, CmdLine);
-				return Generator.GenerateProjectFiles(PlatformProjectGenerators, Arguments);
+				return Generator.GenerateProjectFiles(PlatformProjectGenerators, Arguments, Logger);
 			}
 			finally
 			{
@@ -1628,7 +1630,7 @@ namespace UnrealBuildTool
 			return FileReference.Combine(Executable.Directory, "Payload", TargetName + ".app", TargetName);
 		}
 
-		private static void WriteEntitlements(IOSPostBuildSyncTarget Target)
+		private static void WriteEntitlements(IOSPostBuildSyncTarget Target, ILogger Logger)
 		{
 			string AppName = Target.TargetName;
 			FileReference? MobileProvisionFile;
@@ -1684,7 +1686,7 @@ namespace UnrealBuildTool
 			return ProjectDirectory;
 		}
 
-		private static void GenerateFrameworkWrapperIfNonexistent(IOSPostBuildSyncTarget Target)
+		private static void GenerateFrameworkWrapperIfNonexistent(IOSPostBuildSyncTarget Target, ILogger Logger)
 		{
 			DirectoryReference ProjectDirectory = GetActualProjectDirectory(Target.ProjectFile);
 			DirectoryReference WrapperDirectory = DirectoryReference.Combine(ProjectDirectory, "Wrapper");
@@ -1715,7 +1717,7 @@ namespace UnrealBuildTool
 			}
 		}
 
-		public static void PostBuildSync(IOSPostBuildSyncTarget Target)
+		public static void PostBuildSync(IOSPostBuildSyncTarget Target, ILogger Logger)
 		{
 			ConfigHierarchy Ini = ConfigCache.ReadHierarchy(ConfigHierarchyType.Engine, GetActualProjectDirectory(Target.ProjectFile), UnrealTargetPlatform.IOS);
 			string? BundleID;
@@ -1749,7 +1751,7 @@ namespace UnrealBuildTool
 
 					PathToDsymZip = Path.Combine(Target.OutputPath.Directory.ParentDirectory!.FullName, Target.OutputPath.GetFileName() + ".dSYM.zip");
 
-					GenerateFrameworkWrapperIfNonexistent(Target);
+					GenerateFrameworkWrapperIfNonexistent(Target, Logger);
 
 					// do not perform any of the .app creation below
 					bPerformFullAppCreation = false;
@@ -1760,7 +1762,7 @@ namespace UnrealBuildTool
 
 			if (!Target.bSkipCrashlytics)
 			{
-				GenerateCrashlyticsData(PathToDsymZip, Target.ProjectDirectory.FullName, AppName);
+				GenerateCrashlyticsData(PathToDsymZip, Target.ProjectDirectory.FullName, AppName, Logger);
 			}
 
 			// only make the app if needed
@@ -1927,7 +1929,7 @@ namespace UnrealBuildTool
 						Writer.WriteLine(String.Format("rm -rf \"{0}\"", FrameworkPayloadDirectory));
 
 						// Build the framework wrapper
-						CmdLine = new IOSToolChainSettings().XcodeDeveloperDir + "usr/bin/xcodebuild" +
+						CmdLine = new IOSToolChainSettings(Logger).XcodeDeveloperDir + "usr/bin/xcodebuild" +
 							" -project \"" + WrapperProject + "\"" +
 								" -configuration \"" + ConfigName + "\"" +
 							" -scheme '" + SchemeName + "'" +
@@ -1940,7 +1942,7 @@ namespace UnrealBuildTool
 					else
 					{
 						// code sign the project
-						CmdLine = new IOSToolChainSettings().XcodeDeveloperDir + "usr/bin/xcodebuild" +
+						CmdLine = new IOSToolChainSettings(Logger).XcodeDeveloperDir + "usr/bin/xcodebuild" +
 							" -workspace \"" + XcodeWorkspaceDir + "\"" +
 								" -configuration \"" + ConfigName + "\"" +
 							" -scheme '" + SchemeName + "'" +
@@ -1962,11 +1964,11 @@ namespace UnrealBuildTool
 				{
 				    if (AppName == "UnrealGame" || AppName == "UnrealClient" || Target.ProjectFile == null || Target.ProjectFile.IsUnderDirectory(Unreal.EngineDirectory))
 				    {
-					    GenerateProjectFiles(Target.ProjectFile, new string[] { "-platforms=" + (Target.Platform == UnrealTargetPlatform.IOS ? "IOS" : "TVOS"), "-NoIntellIsense", (Target.Platform == UnrealTargetPlatform.IOS ? "-iosdeployonly" : "-tvosdeployonly"), "-ignorejunk", (Target.bForDistribution ? "-distribution" : "-development"), "-bundleID=" + BundleID, "-includetemptargets", "-appname=" + AppName });
+					    GenerateProjectFiles(Target.ProjectFile, new string[] { "-platforms=" + (Target.Platform == UnrealTargetPlatform.IOS ? "IOS" : "TVOS"), "-NoIntellIsense", (Target.Platform == UnrealTargetPlatform.IOS ? "-iosdeployonly" : "-tvosdeployonly"), "-ignorejunk", (Target.bForDistribution ? "-distribution" : "-development"), "-bundleID=" + BundleID, "-includetemptargets", "-appname=" + AppName }, Logger);
 				    }
 				    else
 				    {
-					    GenerateProjectFiles(Target.ProjectFile, new string[] { "-platforms=" + (Target.Platform == UnrealTargetPlatform.IOS ? "IOS" : "TVOS"), "-NoIntellIsense", (Target.Platform == UnrealTargetPlatform.IOS ? "-iosdeployonly" : "-tvosdeployonly"), "-ignorejunk", (Target.bForDistribution ? "-distribution" : "-development"), String.Format("-project={0}", Target.ProjectFile), "-game", "-bundleID=" + BundleID, "-includetemptargets" });
+					    GenerateProjectFiles(Target.ProjectFile, new string[] { "-platforms=" + (Target.Platform == UnrealTargetPlatform.IOS ? "IOS" : "TVOS"), "-NoIntellIsense", (Target.Platform == UnrealTargetPlatform.IOS ? "-iosdeployonly" : "-tvosdeployonly"), "-ignorejunk", (Target.bForDistribution ? "-distribution" : "-development"), String.Format("-project={0}", Target.ProjectFile), "-game", "-bundleID=" + BundleID, "-includetemptargets" }, Logger);
 				    }
 				    // Make sure it exists
 				    if (!DirectoryReference.Exists(XcodeWorkspaceDir!))
@@ -1976,14 +1978,14 @@ namespace UnrealBuildTool
 				}
 
 				// ensure the plist, entitlements, and provision files are properly copied
-				UEDeployIOS DeployHandler = (Target.Platform == UnrealTargetPlatform.IOS ? new UEDeployIOS() : new UEDeployTVOS());
+				UEDeployIOS DeployHandler = (Target.Platform == UnrealTargetPlatform.IOS ? new UEDeployIOS(Logger) : new UEDeployTVOS(Logger));
 				DeployHandler.ForDistribution = Target.bForDistribution;
 				DeployHandler.PrepTargetForDeployment(Target.ProjectFile, Target.TargetName, Target.Platform, Target.Configuration, Target.UPLScripts, Target.bCreateStubIPA, BundleID, Target.bBuildAsFramework);
 
-				Log.TraceInformation("Executing {0}", SignProjectScript);
+				Logger.LogInformation("Executing {Script}", SignProjectScript);
 
 				// write the entitlements file (building remotely)
-				WriteEntitlements(Target);
+				WriteEntitlements(Target, Logger);
 
 				Process SignProcess = new Process();
 				SignProcess.StartInfo.WorkingDirectory = RemoteShadowDirectoryMac;
@@ -1993,8 +1995,8 @@ namespace UnrealBuildTool
 
 				ProcessOutput Output = new ProcessOutput();
 
-				SignProcess.OutputDataReceived += new DataReceivedEventHandler(Output.OutputReceivedDataEventHandler);
-				SignProcess.ErrorDataReceived += new DataReceivedEventHandler(Output.OutputReceivedDataEventHandler);
+				SignProcess.OutputDataReceived += new DataReceivedEventHandler((s, e) => Output.OutputReceivedDataEventHandler(s, e, Logger));
+				SignProcess.ErrorDataReceived += new DataReceivedEventHandler((s, e) => Output.OutputReceivedDataEventHandler(s, e, Logger));
 
 				Output.OutputReceivedDataEventHandlerEncounteredError = false;
 				Output.OutputReceivedDataEventHandlerEncounteredErrorMessage = "";
@@ -2011,7 +2013,7 @@ namespace UnrealBuildTool
 						CleanWriter.WriteLine("security list-keychain -s login.keychain");
 					}
 
-					Log.TraceInformation("Executing {0}", CleanProjectScript);
+					Logger.LogInformation("Executing {Script}", CleanProjectScript);
 
 					Process CleanProcess = new Process();
 					CleanProcess.StartInfo.WorkingDirectory = RemoteShadowDirectoryMac;
@@ -2020,8 +2022,8 @@ namespace UnrealBuildTool
 
 					ProcessOutput CleanOutput = new ProcessOutput();
 
-					SignProcess.OutputDataReceived += new DataReceivedEventHandler(CleanOutput.OutputReceivedDataEventLogger);
-					SignProcess.ErrorDataReceived += new DataReceivedEventHandler(CleanOutput.OutputReceivedDataEventLogger);
+					SignProcess.OutputDataReceived += new DataReceivedEventHandler((s, e) => CleanOutput.OutputReceivedDataEventLogger(s, e, Logger));
+					SignProcess.ErrorDataReceived += new DataReceivedEventHandler((s, e) => CleanOutput.OutputReceivedDataEventLogger(s, e, Logger));
 
 					Utils.RunLocalProcess(CleanProcess);
 				}
@@ -2055,12 +2057,12 @@ namespace UnrealBuildTool
 			else
 			{
 				// ensure the plist, entitlements, and provision files are properly copied
-				UEDeployIOS DeployHandler = (Target.Platform == UnrealTargetPlatform.IOS ? new UEDeployIOS() : new UEDeployTVOS());
+				UEDeployIOS DeployHandler = (Target.Platform == UnrealTargetPlatform.IOS ? new UEDeployIOS(Logger) : new UEDeployTVOS(Logger));
 				DeployHandler.ForDistribution = Target.bForDistribution;
 				DeployHandler.PrepTargetForDeployment(Target.ProjectFile, Target.TargetName, Target.Platform, Target.Configuration, Target.UPLScripts, Target.bCreateStubIPA, BundleID, Target.bBuildAsFramework);
 
 				// write the entitlements file (building on Mac)
-				WriteEntitlements(Target);
+				WriteEntitlements(Target, Logger);
 			}
 
 			{
@@ -2077,7 +2079,7 @@ namespace UnrealBuildTool
 					// For now, this is hard coded, but we need to loop over all modules, and copy bundled assets that need it
 					string LocalDest = LocalFrameworkAssets + "/" + Pair.Key;
 
-					Log.TraceInformation("Copying bundled asset... LocalSource: {0}, LocalDest: {1}", Pair.Value, LocalDest);
+					Logger.LogInformation("Copying bundled asset... LocalSource: {Source}, LocalDest: {Target}", Pair.Value, LocalDest);
 
 					string ResultsText;
 					RunExecutableAndWait("cp", String.Format("-R -L \"{0}\" \"{1}\"", Pair.Value, LocalDest), out ResultsText);
