UBT: Use ILogger throughout UBT, rather than legacy EpicGames.Core.Log methods. All output going forwards should use structured logging rather than String.Format style logging.

diff --git a/Engine/Source/Programs/UnrealBuildTool/Platform/IOS/UEBuildIOS.cs b/Engine/Source/Programs/UnrealBuildTool/Platform/IOS/UEBuildIOS.cs
--- a/Engine/Source/Programs/UnrealBuildTool/Platform/IOS/UEBuildIOS.cs
+++ b/Engine/Source/Programs/UnrealBuildTool/Platform/IOS/UEBuildIOS.cs
@@ -8,6 +8,7 @@ using System.IO;
 using System.Linq;
 using EpicGames.Core;
 using UnrealBuildBase;
+using Microsoft.Extensions.Logging;
 
 namespace UnrealBuildTool
 {
@@ -452,12 +453,12 @@ namespace UnrealBuildTool
 			get => MobileProvisionFile?.GetFileName();
 		}
 
-		public IOSProvisioningData(IOSProjectSettings ProjectSettings, bool bForDistribution)
-			: this(ProjectSettings, false, bForDistribution)
+		public IOSProvisioningData(IOSProjectSettings ProjectSettings, bool bForDistribution, ILogger Logger)
+			: this(ProjectSettings, false, bForDistribution, Logger)
 		{
 		}
 
-		protected IOSProvisioningData(IOSProjectSettings ProjectSettings, bool bIsTVOS, bool bForDistribtion)
+		protected IOSProvisioningData(IOSProjectSettings ProjectSettings, bool bIsTVOS, bool bForDistribution, ILogger Logger)
 		{
             SigningCertificate = ProjectSettings.SigningCertificate;
             string? MobileProvision = ProjectSettings.MobileProvision;
@@ -470,12 +471,12 @@ namespace UnrealBuildTool
                 // verify the certificate
                 Process IPPProcess = new Process();
 
-				string IPPCmd = "certificates " + ((ProjectFile != null) ? ("\"" + ProjectFile.ToString() + "\"") : "Engine") + " -bundlename " + ProjectSettings.BundleIdentifier + (bForDistribtion ? " -distribution" : "");
+				string IPPCmd = "certificates " + ((ProjectFile != null) ? ("\"" + ProjectFile.ToString() + "\"") : "Engine") + " -bundlename " + ProjectSettings.BundleIdentifier + (bForDistribution ? " -distribution" : "");
 				IPPProcess.StartInfo.StandardOutputEncoding = Encoding.UTF8;
 				IPPProcess.StartInfo.StandardErrorEncoding = Encoding.UTF8;
 				IPPProcess.StartInfo.WorkingDirectory = Unreal.EngineDirectory.ToString();
-				IPPProcess.OutputDataReceived += new DataReceivedEventHandler(IPPDataReceivedHandler);
-				IPPProcess.ErrorDataReceived += new DataReceivedEventHandler(IPPDataReceivedHandler);
+				IPPProcess.OutputDataReceived += new DataReceivedEventHandler((s, e) => IPPDataReceivedHandler(s, e, Logger));
+				IPPProcess.ErrorDataReceived += new DataReceivedEventHandler((s, e) => IPPDataReceivedHandler(s, e, Logger));
 
 				if (BuildHostPlatform.Current.Platform == UnrealTargetPlatform.Mac)
                 {
@@ -488,12 +489,12 @@ namespace UnrealBuildTool
 					IPPProcess.StartInfo.Arguments = IPPCmd;
                 }
 
-				Log.TraceInformation("Getting certifcate information via {0} {1}", IPPProcess.StartInfo.FileName, IPPProcess.StartInfo.Arguments);
+				Logger.LogInformation("Getting certifcate information via {IPPProcessStartInfoFileName} {IPPProcessStartInfoArguments}", IPPProcess.StartInfo.FileName, IPPProcess.StartInfo.Arguments);
                 Utils.RunLocalProcess(IPPProcess);
             }
             else
             {
-                SigningCertificate = bForDistribtion ? "iPhone Distribution" : "iPhone Developer";
+                SigningCertificate = bForDistribution ? "iPhone Distribution" : "iPhone Developer";
                 bHaveCertificate = true;
             }
 
@@ -522,16 +523,16 @@ namespace UnrealBuildTool
                 SigningCertificate = "";
                 MobileProvision = "";
 				MobileProvisionFile = null;
-                Log.TraceInformation("Provision not specified or not found for " + ((ProjectFile != null) ? ProjectFile.GetFileNameWithoutAnyExtensions() : "UnrealGame") + ", searching for compatible match...");
+                Logger.LogInformation("Provision not specified or not found for {Project}, searching for compatible match...", ((ProjectFile != null) ? ProjectFile.GetFileNameWithoutAnyExtensions() : "UnrealGame"));
                 Process IPPProcess = new Process();
 
 				IPPProcess.StartInfo.StandardOutputEncoding = Encoding.UTF8;
 				IPPProcess.StartInfo.StandardErrorEncoding = Encoding.UTF8;
-				IPPProcess.OutputDataReceived += new DataReceivedEventHandler(IPPDataReceivedHandler);
-				IPPProcess.ErrorDataReceived += new DataReceivedEventHandler(IPPDataReceivedHandler);
+				IPPProcess.OutputDataReceived += new DataReceivedEventHandler((s, e) => IPPDataReceivedHandler(s, e, Logger));
+				IPPProcess.ErrorDataReceived += new DataReceivedEventHandler((s, e) => IPPDataReceivedHandler(s, e, Logger));
 				IPPProcess.StartInfo.WorkingDirectory = Unreal.EngineDirectory.ToString();
 
-				string IPPCmd = "signing_match " + ((ProjectFile != null) ? ("\"" + ProjectFile.ToString() + "\"") : "Engine") + " -bundlename " + ProjectSettings.BundleIdentifier + (bIsTVOS ? " -tvos" : "") + (bForDistribtion ? " -distribution" : "");
+				string IPPCmd = "signing_match " + ((ProjectFile != null) ? ("\"" + ProjectFile.ToString() + "\"") : "Engine") + " -bundlename " + ProjectSettings.BundleIdentifier + (bIsTVOS ? " -tvos" : "") + (bForDistribution ? " -distribution" : "");
 
 				if (BuildHostPlatform.Current.Platform == UnrealTargetPlatform.Mac)
 				{
@@ -544,14 +545,14 @@ namespace UnrealBuildTool
 					IPPProcess.StartInfo.Arguments = IPPCmd;
 				}
 
-				Log.TraceInformation("Getting signing information via {0} {1}", IPPProcess.StartInfo.FileName, IPPProcess.StartInfo.Arguments);
+				Logger.LogInformation("Getting signing information via {IPPProcessStartInfoFileName} {IPPProcessStartInfoArguments}", IPPProcess.StartInfo.FileName, IPPProcess.StartInfo.Arguments);
 
 				Utils.RunLocalProcess(IPPProcess);
 				if(MobileProvisionFile != null)
 				{
-					Log.TraceInformation("Provision found for " + ((ProjectFile != null) ? ProjectFile.GetFileNameWithoutAnyExtensions() : "UnrealGame") + ", Provision: " + MobileProvisionFile + " Certificate: " + SigningCertificate);
+					Logger.LogInformation("Provision found for {Project}, Provision: {Provision}, Certificate: {Certificate}", ((ProjectFile != null) ? ProjectFile.GetFileNameWithoutAnyExtensions() : "UnrealGame"), MobileProvisionFile, SigningCertificate);
 				}
-            }
+			}
 
             // add to the dictionary
             SigningCertificate = SigningCertificate.Replace("\"", "");
@@ -559,11 +560,11 @@ namespace UnrealBuildTool
             // read the provision to get the UUID
 			if(MobileProvisionFile == null)
 			{
-				Log.TraceInformation("No matching provision file was discovered for {0}. Please ensure you have a compatible provision installed.", ProjectFile);
+				Logger.LogInformation("No matching provision file was discovered for {ProjectFile}. Please ensure you have a compatible provision installed.", ProjectFile);
 			}
 			else if(!FileReference.Exists(MobileProvisionFile))
 			{
-				Log.TraceInformation("Selected mobile provision for {0} ({1}) was not found. Please ensure you have a compatible provision installed.", ProjectFile, MobileProvisionFile);
+				Logger.LogInformation("Selected mobile provision for {ProjectFile} ({MobileProvisionFile}) was not found. Please ensure you have a compatible provision installed.", ProjectFile, MobileProvisionFile);
 			}
 			else
             {
@@ -646,23 +647,23 @@ namespace UnrealBuildTool
 				{
 					MobileProvision = null;
 					SigningCertificate = null;
-					Log.TraceInformation("Failed to parse the mobile provisioning profile.");
+					Logger.LogInformation("Failed to parse the mobile provisioning profile.");
 				}
             }
 		}
 
-        void IPPDataReceivedHandler(Object Sender, DataReceivedEventArgs Line)
+        void IPPDataReceivedHandler(Object Sender, DataReceivedEventArgs Line, ILogger Logger)
         {
             if ((Line != null) && (Line.Data != null))
             {
 				if (Line.Data.StartsWith("IPP WARNING:"))
 				{
 					// Don't output IPP warnings to the console as they may not be warnings relevant to the build and could cause build failures.
-					Log.TraceLog("{0}", Line.Data);
+					Logger.LogDebug("{LineData}", Line.Data);
 				}
 				else
 				{
-					Log.TraceInformation("{0}", Line.Data);
+					Logger.LogInformation("{LineData}", Line.Data);
 				}
 
 				if (!string.IsNullOrEmpty(SigningCertificate))
@@ -700,13 +701,13 @@ namespace UnrealBuildTool
 		// by default, use an empty architecture (which is really just a modifer to the platform for some paths/names)
 		public static string IOSArchitecture = "";
 
-		public IOSPlatform(UEBuildPlatformSDK InSDK)
-			: this(InSDK, UnrealTargetPlatform.IOS)
+		public IOSPlatform(UEBuildPlatformSDK InSDK, ILogger Logger)
+			: this(InSDK, UnrealTargetPlatform.IOS, Logger)
 		{
 		}
 
-		protected IOSPlatform(UEBuildPlatformSDK InSDK, UnrealTargetPlatform TargetPlatform)
-			: base(TargetPlatform, InSDK)
+		protected IOSPlatform(UEBuildPlatformSDK InSDK, UnrealTargetPlatform TargetPlatform, ILogger Logger)
+			: base(TargetPlatform, InSDK, Logger)
 		{
 		}
 
@@ -902,7 +903,7 @@ namespace UnrealBuildTool
 
 		protected virtual IOSProvisioningData CreateProvisioningData(IOSProjectSettings ProjectSettings, bool bForDistribution)
 		{
-			return new IOSProvisioningData(ProjectSettings, bForDistribution);
+			return new IOSProvisioningData(ProjectSettings, bForDistribution, Logger);
 		}
 
 		public override string[] GetDebugInfoExtensions(ReadOnlyTargetRules InTarget, UEBuildBinaryType InBinaryType)
@@ -935,7 +936,7 @@ namespace UnrealBuildTool
 			return true;
 		}
 
-		public bool HasCustomIcons(DirectoryReference ProjectDirectoryName)
+		public bool HasCustomIcons(DirectoryReference ProjectDirectoryName, ILogger Logger)
 		{
 			string IconDir = Path.Combine(ProjectDirectoryName.FullName, "Build", "IOS", "Resources", "Graphics");
 			if(Directory.Exists(IconDir))
@@ -944,7 +945,7 @@ namespace UnrealBuildTool
 				{
 					if (f.Contains("Icon") && Path.GetExtension(f).Contains(".png"))
 					{
-						Log.TraceInformation("Requiring custom build because project {0} has custom icons", Path.GetFileName(ProjectDirectoryName.FullName));
+						Logger.LogInformation("Requiring custom build because project {Project} has custom icons", Path.GetFileName(ProjectDirectoryName.FullName));
 						return true;
 					}
 				}
@@ -970,14 +971,14 @@ namespace UnrealBuildTool
 			};
 
 			// check for custom icons
-			if (HasCustomIcons(ProjectDirectoryName))
+			if (HasCustomIcons(ProjectDirectoryName, Logger))
 			{
 				return false;
 			}
 
 			// look up iOS specific settings
 			if (!DoProjectSettingsMatchDefault(Platform, ProjectDirectoryName, "/Script/IOSRuntimeSettings.IOSRuntimeSettings",
-					BoolKeys, null, StringKeys))
+					BoolKeys, null, StringKeys, Logger))
 			{
 				return false;
 			}
@@ -993,7 +994,7 @@ namespace UnrealBuildTool
 		public override bool RequiresBuild(UnrealTargetPlatform Platform, DirectoryReference ProjectDirectoryName)
 		{
 			// check for custom icons
-			return HasCustomIcons(ProjectDirectoryName);
+			return HasCustomIcons(ProjectDirectoryName, Logger);
 		}
 
 		public override bool ShouldCompileMonolithicBinary(UnrealTargetPlatform InPlatform)
@@ -1105,7 +1106,7 @@ namespace UnrealBuildTool
 			IOSProjectSettings ProjectSettings = ((IOSPlatform)UEBuildPlatform.GetBuildPlatform(Target.Platform)).ReadProjectSettings(Target.ProjectFile);
 			if (!ProjectFileGenerator.bGenerateProjectFiles)
 			{
-				Log.TraceInformation("Compiling against OS Version {0} [minimum allowed at runtime]", ProjectSettings.RuntimeVersion);
+				Logger.LogInformation("Compiling against OS Version {RuntimeVersion} [minimum allowed at runtime]", ProjectSettings.RuntimeVersion);
 			}
 
 			CompileEnvironment.Definitions.Add("PLATFORM_IOS=1");
@@ -1164,7 +1165,7 @@ namespace UnrealBuildTool
 				string OodleDllPath = DirectoryReference.Combine(ProjectDir, "Binaries/ThirdParty/Oodle/Mac/libUnrealPakPlugin.dylib").FullName;
 				if (File.Exists(OodleDllPath))
 				{
-					Log.TraceVerbose("        Registering custom oodle compressor for {0}", UnrealTargetPlatform.IOS.ToString());
+					Logger.LogDebug("        Registering custom oodle compressor for {Platform}", UnrealTargetPlatform.IOS.ToString());
 					CompileEnvironment.Definitions.Add("REGISTER_OODLE_CUSTOM_COMPRESSOR=1");
 				}
 			}
@@ -1217,13 +1218,10 @@ namespace UnrealBuildTool
 			}
 
 			IOSProjectSettings ProjectSettings = ReadProjectSettings(Target.ProjectFile);
-			return new IOSToolChain(Target, ProjectSettings, Options);
+			return new IOSToolChain(Target, ProjectSettings, Options, Logger);
 		}
 
-		/// <summary>
-		/// Deploys the given target
-		/// </summary>
-		/// <param name="Receipt">Receipt for the target being deployed</param>
+		/// <inheritdoc/>
 		public override void Deploy(TargetReceipt Receipt)
 		{
 			if (Receipt.HasValueForAdditionalProperty("CompileAsDll", "true"))
@@ -1232,7 +1230,7 @@ namespace UnrealBuildTool
 			}
 			else
 			{
-				new UEDeployIOS().PrepTargetForDeployment(Receipt);
+				new UEDeployIOS(Logger).PrepTargetForDeployment(Receipt);
 			}
 		}
 	}
@@ -1248,12 +1246,12 @@ namespace UnrealBuildTool
 		/// <summary>
 		/// Register the platform with the UEBuildPlatform class
 		/// </summary>
-		public override void RegisterBuildPlatforms()
+		public override void RegisterBuildPlatforms(ILogger Logger)
 		{
-			ApplePlatformSDK SDK = new IOSPlatformSDK();
+			ApplePlatformSDK SDK = new IOSPlatformSDK(Logger);
 
 			// Register this build platform for IOS
-			UEBuildPlatform.RegisterBuildPlatform(new IOSPlatform(SDK));
+			UEBuildPlatform.RegisterBuildPlatform(new IOSPlatform(SDK, Logger), Logger);
 			UEBuildPlatform.RegisterPlatformWithGroup(UnrealTargetPlatform.IOS, UnrealPlatformGroup.Apple);
 			UEBuildPlatform.RegisterPlatformWithGroup(UnrealTargetPlatform.IOS, UnrealPlatformGroup.IOS);
 		}
diff --git a/Engine/Source/Programs/UnrealBuildTool/Platform/IOS/UEDeployIOS.cs b/Engine/Source/Programs/UnrealBuildTool/Platform/IOS/UEDeployIOS.cs
--- a/Engine/Source/Programs/UnrealBuildTool/Platform/IOS/UEDeployIOS.cs
+++ b/Engine/Source/Programs/UnrealBuildTool/Platform/IOS/UEDeployIOS.cs
@@ -11,12 +11,14 @@
 using System.Xml.Linq;
 using EpicGames.Core;
 using UnrealBuildBase;
+using Microsoft.Extensions.Logging;
 
 namespace UnrealBuildTool
 {
 	class UEDeployIOS : UEBuildDeploy
 	{
-		public UEDeployIOS()
+		public UEDeployIOS(ILogger InLogger)
+			: base(InLogger)
 		{
 		}
 
@@ -196,7 +198,7 @@
 			return result;
 		}
 
-		public static string GetMinimumOSVersion(string MinVersion)
+		public static string GetMinimumOSVersion(string MinVersion, ILogger Logger)
 		{
 			string MinVersionToReturn = "";
 			switch (MinVersion)
@@ -208,14 +210,14 @@
 					MinVersionToReturn = "15.0";
 					break;
 				default:
-					Log.TraceWarning("MinimumiOSVersion {0} specified in ini file is no longer supported, defaulting to 14.0", MinVersion);
+					Logger.LogWarning("MinimumiOSVersion {MinVersion} specified in ini file is no longer supported, defaulting to 14.0", MinVersion);
 					MinVersionToReturn = "14.0";
 					break;
 			}
 			return MinVersionToReturn;
 		}
 
-		public static bool GenerateIOSPList(FileReference? ProjectFile, UnrealTargetConfiguration Config, string ProjectDirectory, bool bIsUnrealGame, string GameName, bool bIsClient, string ProjectName, string InEngineDir, string AppDirectory, UnrealPluginLanguage? UPL, string? BundleID, bool bBuildAsFramework, out bool bSupportsPortrait, out bool bSupportsLandscape)
+		public static bool GenerateIOSPList(FileReference? ProjectFile, UnrealTargetConfiguration Config, string ProjectDirectory, bool bIsUnrealGame, string GameName, bool bIsClient, string ProjectName, string InEngineDir, string AppDirectory, UnrealPluginLanguage? UPL, string? BundleID, bool bBuildAsFramework, ILogger Logger, out bool bSupportsPortrait, out bool bSupportsLandscape)
 		{
 			// generate the Info.plist for future use
 			string BuildDirectory = ProjectDirectory + "/Build/IOS";
@@ -330,7 +332,7 @@
 			// minimum iOS version
 			string MinVersionSetting = "";
 			Ini.GetString("/Script/IOSRuntimeSettings.IOSRuntimeSettings", "MinimumiOSVersion", out MinVersionSetting);
-			string MinVersion  = GetMinimumOSVersion(MinVersionSetting);
+			string MinVersion  = GetMinimumOSVersion(MinVersionSetting, Logger);
 
 			// Get Facebook Support details
 			bool bEnableFacebookSupport = true;
@@ -505,13 +507,13 @@
 				{
 					string outputStoryboard = LaunchStoryboard + "c";
 					string argsStoryboard = "--compile " + outputStoryboard + " " + LaunchStoryboard;
-					string stdOutLaunchScreen = Utils.RunLocalProcessAndReturnStdOut("ibtool", argsStoryboard);
+					string stdOutLaunchScreen = Utils.RunLocalProcessAndReturnStdOut("ibtool", argsStoryboard, Logger);
 
-					Log.TraceInformation("LaunchScreen Storyboard compilation results : " + stdOutLaunchScreen);
+					Logger.LogInformation("LaunchScreen Storyboard compilation results : {Results}", stdOutLaunchScreen);
 				}
 				else
 				{
-					Log.TraceWarning("Custom Launchscreen compilation storyboard only compatible on Mac for now");
+					Logger.LogWarning("Custom Launchscreen compilation storyboard only compatible on Mac for now");
 				}
 			}
 
@@ -686,7 +688,7 @@
 
  		public bool GeneratePList(FileReference ProjectFile, UnrealTargetConfiguration Config, string ProjectDirectory, bool bIsUnrealGame, string GameName, bool bIsClient, string ProjectName, string InEngineDir, string AppDirectory, TargetReceipt Receipt, out bool bSupportsPortrait, out bool bSupportsLandscape)
 		{
-			List<string> UPLScripts = CollectPluginDataPaths(Receipt.AdditionalProperties);
+			List<string> UPLScripts = CollectPluginDataPaths(Receipt.AdditionalProperties, Logger);
 			VersionNumber? SdkVersion = GetSdkVersion(Receipt);
 			bool bBuildAsFramework = GetCompileAsDll(Receipt);
 			return GeneratePList(ProjectFile, Config, ProjectDirectory, bIsUnrealGame, GameName, bIsClient, ProjectName, InEngineDir, AppDirectory, UPLScripts, "", bBuildAsFramework, out bSupportsPortrait, out bSupportsLandscape);
@@ -719,12 +721,12 @@
 
 			string RelativeEnginePath = Unreal.EngineDirectory.MakeRelativeTo(DirectoryReference.GetCurrentDirectory());
 
-			UnrealPluginLanguage UPL = new UnrealPluginLanguage(ProjectFile, UPLScripts, ProjectArches, "", "", UnrealTargetPlatform.IOS);
+			UnrealPluginLanguage UPL = new UnrealPluginLanguage(ProjectFile, UPLScripts, ProjectArches, "", "", UnrealTargetPlatform.IOS, Logger);
 
 			// Passing in true for distribution is not ideal here but given the way that ios packaging happens and this call chain it seems unavoidable for now, maybe there is a way to correctly pass it in that I can't find?
 			UPL.Init(ProjectArches, true, RelativeEnginePath, BundlePath, ProjectDirectory, Config.ToString(), false);
 
-			return GenerateIOSPList(ProjectFile, Config, ProjectDirectory, bIsUnrealGame, GameName, bIsClient, ProjectName, InEngineDir, AppDirectory, UPL, BundleID, bBuildAsFramework, out bSupportsPortrait, out bSupportsLandscape);
+			return GenerateIOSPList(ProjectFile, Config, ProjectDirectory, bIsUnrealGame, GameName, bIsClient, ProjectName, InEngineDir, AppDirectory, UPL, BundleID, bBuildAsFramework, Logger, out bSupportsPortrait, out bSupportsLandscape);
 		}
 
 		protected virtual void CopyCloudResources(string InEngineDir, string AppDirectory)
@@ -732,7 +734,7 @@
 			CopyFiles(InEngineDir + "/Build/IOS/Cloud", AppDirectory, "*.*", true);
 		}
 
-		protected virtual void CopyCustomLaunchScreenResources(string InEngineDir, string AppDirectory, string BuildDirectory)
+		protected virtual void CopyCustomLaunchScreenResources(string InEngineDir, string AppDirectory, string BuildDirectory, ILogger Logger)
 		{
 			if (Directory.Exists(BuildDirectory + "/Resources/Interface/LaunchScreen.storyboardc"))
 			{
@@ -741,7 +743,7 @@
 			}
 			else
 			{
-				Log.TraceWarning("Custom LaunchScreen Storyboard is checked but no compiled Storyboard could be found. Custom Storyboard compilation is only Mac compatible for now. Fallback to default Launchscreen");
+				Logger.LogWarning("Custom LaunchScreen Storyboard is checked but no compiled Storyboard could be found. Custom Storyboard compilation is only Mac compatible for now. Fallback to default Launchscreen");
 				CopyStandardLaunchScreenResources(InEngineDir, AppDirectory, BuildDirectory);
 			}
 		}
@@ -761,11 +763,11 @@
 
 		}
 
-		protected virtual void CopyLaunchScreenResources(string InEngineDir, string AppDirectory, string BuildDirectory)
+		protected virtual void CopyLaunchScreenResources(string InEngineDir, string AppDirectory, string BuildDirectory, ILogger Logger)
 		{
 			if (VersionUtilities.bCustomLaunchscreenStoryboard)
 			{
-				CopyCustomLaunchScreenResources(InEngineDir, AppDirectory, BuildDirectory);
+				CopyCustomLaunchScreenResources(InEngineDir, AppDirectory, BuildDirectory, Logger);
 			}
 			else
 			{
@@ -774,16 +776,16 @@
 
 			if (!File.Exists(AppDirectory + "/LaunchScreen.storyboardc/LaunchScreen.nib"))
 			{
-				Log.TraceError("Launchscreen.storyboard ViewController needs an ID named LaunchScreen");
+				Logger.LogError("Launchscreen.storyboard ViewController needs an ID named LaunchScreen");
 			}
 		}
 
-		protected virtual void CopyLocalizationsResources(string InEngineDir, string AppDirectory, string BuildDirectory, string IntermediateDir)
+		protected virtual void CopyLocalizationsResources(string InEngineDir, string AppDirectory, string BuildDirectory, string IntermediateDir, ILogger Logger)
 		{
 			string LocalizationsPath = BuildDirectory + "/Resources/Localizations";
 			if (Directory.Exists(LocalizationsPath))
 			{
-				Log.TraceInformation("Copy localizations from Resources {0}, {1}", LocalizationsPath, AppDirectory);
+				Logger.LogInformation("Copy localizations from Resources {LocalizationsPath}, {AppDirectory}", LocalizationsPath, AppDirectory);
 				CopyFolder(BuildDirectory + "/Resources/Localizations", AppDirectory, true, delegate(string InFilename)
 				{
 					if (Path.GetFileName(InFilename).Equals(".DS_Store")) return false;
@@ -794,13 +796,13 @@
 
 		public bool PrepForUATPackageOrDeploy(UnrealTargetConfiguration Config, FileReference ProjectFile, string InProjectName, string InProjectDirectory, string InExecutablePath, string InEngineDir, bool bForDistribution, string CookFlavor, bool bIsDataDeploy, bool bCreateStubIPA, TargetReceipt Receipt)
 		{
-			List<string> UPLScripts = CollectPluginDataPaths(Receipt.AdditionalProperties);
+			List<string> UPLScripts = CollectPluginDataPaths(Receipt.AdditionalProperties, Logger);
 			VersionNumber? SdkVersion = GetSdkVersion(Receipt);
 			bool bBuildAsFramework = GetCompileAsDll(Receipt);
 			return PrepForUATPackageOrDeploy(Config, ProjectFile, InProjectName, InProjectDirectory, InExecutablePath, InEngineDir, bForDistribution, CookFlavor, bIsDataDeploy, bCreateStubIPA, UPLScripts, "", bBuildAsFramework);
 		}
 
-		void CopyAllProvisions(string ProvisionDir)
+		void CopyAllProvisions(string ProvisionDir, ILogger Logger)
 		{
 			try
 			{
@@ -813,11 +815,11 @@
 				string LocalProvisionFolder = Path.Combine(Environment.GetFolderPath(Environment.SpecialFolder.Personal), "Library/MobileDevice/Provisioning Profiles");
 				if (!Directory.Exists(LocalProvisionFolder))
 				{
-					Log.TraceLog("Local Provision Folder {0} not found, attempting to create...", LocalProvisionFolder);
+					Logger.LogDebug("Local Provision Folder {LocalProvisionFolder} not found, attempting to create...", LocalProvisionFolder);
 					Directory.CreateDirectory(LocalProvisionFolder);
 					if(Directory.Exists(LocalProvisionFolder))
 					{
-						Log.TraceLog("Local Provision Folder {0} created successfully.", LocalProvisionFolder);
+						Logger.LogDebug("Local Provision Folder {LocalProvisionFolder} created successfully.", LocalProvisionFolder);
 					}
 					else
 					{
@@ -844,7 +846,7 @@
 			}
 			catch (Exception Ex)
 			{
-				Log.TraceError(Ex.ToString());
+				Logger.LogError("{Message}", Ex.ToString());
 				throw;
 			}
 		}
@@ -937,12 +939,12 @@
 			if (!File.Exists(ProvisionWithPrefix) || Environment.GetEnvironmentVariable("IsBuildMachine") == "1")
 			{
 				// copy all provisions from the game directory, the engine directory, notforlicensees directory, and, if defined, the ProvisionDirectory.
-				CopyAllProvisions(BuildDirectory);
-				CopyAllProvisions(InEngineDir + "/Build/IOS");
+				CopyAllProvisions(BuildDirectory, Logger);
+				CopyAllProvisions(InEngineDir + "/Build/IOS", Logger);
 				string? ProvisionDirectory = Environment.GetEnvironmentVariable("ProvisionDirectory");
 				if (!string.IsNullOrWhiteSpace(ProvisionDirectory))
 				{
-					CopyAllProvisions(ProvisionDirectory);
+					CopyAllProvisions(ProvisionDirectory, Logger);
 				}
 			}
 
@@ -994,16 +996,16 @@
 			if (SubDir == GetTargetPlatformName())
 			{
 				string BuildDirectoryFortvOS = InProjectDirectory + "/Build/IOS";
-				CopyLaunchScreenResources(InEngineDir, AppDirectory, BuildDirectoryFortvOS);
+				CopyLaunchScreenResources(InEngineDir, AppDirectory, BuildDirectoryFortvOS, Logger);
 			}
 			else
 			{
-				CopyLaunchScreenResources(InEngineDir, AppDirectory, BuildDirectory);
+				CopyLaunchScreenResources(InEngineDir, AppDirectory, BuildDirectory, Logger);
 			}
 
 			if (!bCreateStubIPA)
 			{
-				CopyLocalizationsResources(InEngineDir, AppDirectory, BuildDirectory, IntermediateDirectory);
+				CopyLocalizationsResources(InEngineDir, AppDirectory, BuildDirectory, IntermediateDirectory, Logger);
 
 				CopyCloudResources(InProjectDirectory, AppDirectory);
 
@@ -1030,7 +1032,7 @@
 
 		public override bool PrepTargetForDeployment(TargetReceipt Receipt)
 		{
-			List<string> UPLScripts = CollectPluginDataPaths(Receipt.AdditionalProperties);
+			List<string> UPLScripts = CollectPluginDataPaths(Receipt.AdditionalProperties, Logger);
 			bool bBuildAsFramework = GetCompileAsDll(Receipt);
 			return PrepTargetForDeployment(Receipt.ProjectFile, Receipt.TargetName, Receipt.Platform, Receipt.Configuration, UPLScripts, false, "", bBuildAsFramework);
 		}
@@ -1069,13 +1071,13 @@
 			return true;
 		}
 
-		public static List<string> CollectPluginDataPaths(List<ReceiptProperty> ReceiptProperties)
+		public static List<string> CollectPluginDataPaths(List<ReceiptProperty> ReceiptProperties, ILogger Logger)
 		{
 			List<string> PluginExtras = new List<string>();
 			if (ReceiptProperties == null)
 			{
-				Log.TraceInformation("Receipt is NULL");
-				//Log.TraceInformation("Receipt is NULL");
+				Logger.LogInformation("Receipt is NULL");
+				//Logger.LogInformation("Receipt is NULL");
 				return PluginExtras;
 			}
 
@@ -1088,7 +1090,7 @@
 				if (PluginExtras.FirstOrDefault(x => x == PluginPath) == null)
 				{
 					PluginExtras.Add(PluginPath);
-					Log.TraceInformation("IOSPlugin: {0}", PluginPath);
+					Logger.LogInformation("IOSPlugin: {PluginPath}", PluginPath);
 				}
 			}
 			return PluginExtras;
diff --git a/Engine/Source/Programs/UnrealBuildTool/Platform/Linux/LinuxCommon.cs b/Engine/Source/Programs/UnrealBuildTool/Platform/Linux/LinuxCommon.cs
--- a/Engine/Source/Programs/UnrealBuildTool/Platform/Linux/LinuxCommon.cs
+++ b/Engine/Source/Programs/UnrealBuildTool/Platform/Linux/LinuxCommon.cs
@@ -8,12 +8,13 @@
 using System.Linq;
 using Microsoft.Win32;
 using EpicGames.Core;
+using Microsoft.Extensions.Logging;
 
 namespace UnrealBuildTool
 {
 	class LinuxCommon
 	{
-		public static string? Which(string name)
+		public static string? Which(string name, ILogger Logger)
 		{
 			Process proc = new Process();
 			proc.StartInfo.FileName = "/bin/sh";
@@ -27,7 +28,7 @@
 			proc.WaitForExit();
 
 			string? path = proc.StandardOutput.ReadLine();
-			Log.TraceVerbose(String.Format("which {0} result: ({1}) {2}", name, proc.ExitCode, path));
+			Logger.LogDebug("which {Name} result: ({ExitCode}) {Path}", name, proc.ExitCode, path);
 
 			if (proc.ExitCode == 0 && String.IsNullOrEmpty(proc.StandardError.ReadToEnd()))
 			{
@@ -36,7 +37,7 @@
 			return null;
 		}
 
-		public static string? WhichClang()
+		public static string? WhichClang(ILogger Logger)
 		{
 			string? InternalSDKPath = UEBuildPlatform.GetSDK(UnrealTargetPlatform.Linux)?.GetInternalSDKPath();
 			if (!String.IsNullOrEmpty(InternalSDKPath))
@@ -48,7 +49,7 @@
 			string? ClangPath;
 			foreach (string ClangName in ClangNames)
 			{
-				ClangPath = Which(ClangName);
+				ClangPath = Which(ClangName, Logger);
 				if (!String.IsNullOrEmpty(ClangPath))
 				{
 					return ClangPath;
diff --git a/Engine/Source/Programs/UnrealBuildTool/Platform/Linux/LinuxExports.cs b/Engine/Source/Programs/UnrealBuildTool/Platform/Linux/LinuxExports.cs
--- a/Engine/Source/Programs/UnrealBuildTool/Platform/Linux/LinuxExports.cs
+++ b/Engine/Source/Programs/UnrealBuildTool/Platform/Linux/LinuxExports.cs
@@ -1,28 +1,30 @@
-// Copyright Epic Games, Inc. All Rights Reserved.
-
-using System;
-using System.Collections.Generic;
-using System.Linq;
-using System.Text;
-using System.Threading.Tasks;
-using EpicGames.Core;
-
-namespace UnrealBuildTool
-{
-	/// <summary>
-	/// Public Linux functions exposed to UAT
-	/// </summary>
-	public class LinuxExports
-	{
-		/// <summary>
-		/// 
-		/// </summary>
-		/// <param name="SourceFile"></param>
-		/// <param name="TargetFile"></param>
-		public static void StripSymbols(FileReference SourceFile, FileReference TargetFile)
-		{
-			LinuxToolChain ToolChain = new LinuxToolChain(LinuxPlatform.DefaultHostArchitecture, new LinuxPlatformSDK());
-			ToolChain.StripSymbols(SourceFile, TargetFile);
-		}
-	}
-}
+// Copyright Epic Games, Inc. All Rights Reserved.
+
+using System;
+using System.Collections.Generic;
+using System.Linq;
+using System.Text;
+using System.Threading.Tasks;
+using EpicGames.Core;
+using Microsoft.Extensions.Logging;
+
+namespace UnrealBuildTool
+{
+	/// <summary>
+	/// Public Linux functions exposed to UAT
+	/// </summary>
+	public class LinuxExports
+	{
+		/// <summary>
+		/// 
+		/// </summary>
+		/// <param name="SourceFile"></param>
+		/// <param name="TargetFile"></param>
+		/// <param name="Logger">Logger for output</param>
+		public static void StripSymbols(FileReference SourceFile, FileReference TargetFile, ILogger Logger)
+		{
+			LinuxToolChain ToolChain = new LinuxToolChain(LinuxPlatform.DefaultHostArchitecture, new LinuxPlatformSDK(Logger), ClangToolChainOptions.None, Logger);
+			ToolChain.StripSymbols(SourceFile, TargetFile, Logger);
+		}
+	}
+}
diff --git a/Engine/Source/Programs/UnrealBuildTool/Platform/Linux/LinuxPlatformSDK.cs b/Engine/Source/Programs/UnrealBuildTool/Platform/Linux/LinuxPlatformSDK.cs
--- a/Engine/Source/Programs/UnrealBuildTool/Platform/Linux/LinuxPlatformSDK.cs
+++ b/Engine/Source/Programs/UnrealBuildTool/Platform/Linux/LinuxPlatformSDK.cs
@@ -6,6 +6,7 @@
 using EpicGames.Core;
 using System.Text.RegularExpressions;
 using UnrealBuildBase;
+using Microsoft.Extensions.Logging;
 
 ///////////////////////////////////////////////////////////////////
 // If you are looking for supported version numbers, look in the
@@ -16,6 +17,11 @@
 {
 	partial class LinuxPlatformSDK : UEBuildPlatformSDK
 	{
+		public LinuxPlatformSDK(ILogger Logger)
+			: base(Logger)
+		{
+		}
+
 		public override string GetAutoSDKPlatformName()
 		{
 			return TargetPlatformName;
@@ -255,7 +261,7 @@
 
 			if (ForceUseSystemCompiler())
 			{
-				if (!String.IsNullOrEmpty(LinuxCommon.WhichClang()))
+				if (!String.IsNullOrEmpty(LinuxCommon.WhichClang(Logger)))
 				{
 					return SDKStatus.Valid;
 				}
diff --git a/Engine/Source/Programs/UnrealBuildTool/Platform/Linux/LinuxProjectGenerator.cs b/Engine/Source/Programs/UnrealBuildTool/Platform/Linux/LinuxProjectGenerator.cs
--- a/Engine/Source/Programs/UnrealBuildTool/Platform/Linux/LinuxProjectGenerator.cs
+++ b/Engine/Source/Programs/UnrealBuildTool/Platform/Linux/LinuxProjectGenerator.cs
@@ -6,6 +6,7 @@
 using System.Diagnostics;
 using System.IO;
 using EpicGames.Core;
+using Microsoft.Extensions.Logging;
 
 namespace UnrealBuildTool
 {
@@ -18,8 +19,9 @@
 		/// Constructor
 		/// </summary>
 		/// <param name="Arguments">Command line arguments passed to the project generator</param>
-		public LinuxProjectGenerator(CommandLineArguments Arguments)
-			: base(Arguments)
+		/// <param name="Logger">Logger for output</param>
+		public LinuxProjectGenerator(CommandLineArguments Arguments, ILogger Logger)
+			: base(Arguments, Logger)
 		{
 		}
 
diff --git a/Engine/Source/Programs/UnrealBuildTool/Platform/Linux/LinuxToolChain.cs b/Engine/Source/Programs/UnrealBuildTool/Platform/Linux/LinuxToolChain.cs
--- a/Engine/Source/Programs/UnrealBuildTool/Platform/Linux/LinuxToolChain.cs
+++ b/Engine/Source/Programs/UnrealBuildTool/Platform/Linux/LinuxToolChain.cs
@@ -9,6 +9,7 @@ using System.Linq;
 using Microsoft.Win32;
 using EpicGames.Core;
 using UnrealBuildBase;
+using Microsoft.Extensions.Logging;
 
 namespace UnrealBuildTool
 {
@@ -35,8 +36,8 @@ namespace UnrealBuildTool
 		/** Toolchain information to print during the build. */
 		protected string? ToolchainInfo;
 
-		public LinuxToolChain(string InArchitecture, LinuxPlatformSDK InSDK, ClangToolChainOptions InOptions = ClangToolChainOptions.None)
-			: this(UnrealTargetPlatform.Linux, InArchitecture, InSDK, InOptions)
+		public LinuxToolChain(string InArchitecture, LinuxPlatformSDK InSDK, ClangToolChainOptions InOptions, ILogger InLogger)
+			: this(UnrealTargetPlatform.Linux, InArchitecture, InSDK, InOptions, InLogger)
 		{
 			MultiArchRoot = PlatformSDK.GetSDKLocation();
 			BaseLinuxPath = PlatformSDK.GetBaseLinuxPathForArchitecture(InArchitecture);
@@ -57,12 +58,12 @@ namespace UnrealBuildTool
 				ToolchainInfo = "system toolchain";
 
 				// use native linux toolchain
-				ClangPath = LinuxCommon.WhichClang();
-				ArPath = LinuxCommon.Which("ar");
-				LlvmArPath = LinuxCommon.Which("llvm-ar");
-				RanlibPath = LinuxCommon.Which("ranlib");
-				StripPath = LinuxCommon.Which("strip");
-				ObjcopyPath = LinuxCommon.Which("objcopy");
+				ClangPath = LinuxCommon.WhichClang(Logger);
+				ArPath = LinuxCommon.Which("ar", Logger);
+				LlvmArPath = LinuxCommon.Which("llvm-ar", Logger);
+				RanlibPath = LinuxCommon.Which("ranlib", Logger);
+				StripPath = LinuxCommon.Which("strip", Logger);
+				ObjcopyPath = LinuxCommon.Which("objcopy", Logger);
 
 				// When compiling on Linux, use a faster way to relink circularly dependent libraries.
 				// Race condition between actions linking to the .so and action overwriting it is avoided thanks to inodes
@@ -81,7 +82,7 @@ namespace UnrealBuildTool
 				if (String.IsNullOrEmpty(MultiArchRoot))
 				{
 					MultiArchRoot = BaseLinuxPath;
-					Log.TraceInformation("Using LINUX_ROOT (deprecated, consider LINUX_MULTIARCH_ROOT)");
+					Logger.LogInformation("Using LINUX_ROOT (deprecated, consider LINUX_MULTIARCH_ROOT)");
 				}
 
 				BaseLinuxPath = BaseLinuxPath.Replace("\"", "").Replace('\\', '/');
@@ -127,8 +128,8 @@ namespace UnrealBuildTool
 			}
 		}
 
-		public LinuxToolChain(UnrealTargetPlatform InPlatform, string InArchitecture, LinuxPlatformSDK InSDK, ClangToolChainOptions InOptions = ClangToolChainOptions.None)
-			: base(InOptions)
+		public LinuxToolChain(UnrealTargetPlatform InPlatform, string InArchitecture, LinuxPlatformSDK InSDK, ClangToolChainOptions InOptions, ILogger InLogger)
+			: base(InOptions, InLogger)
 		{
 			Architecture = InArchitecture;
 			PlatformSDK = InSDK;
@@ -948,17 +949,17 @@ namespace UnrealBuildTool
 			return Explanation;
 		}
 
-		protected void PrintBuildDetails(CppCompileEnvironment CompileEnvironment)
+		protected void PrintBuildDetails(CppCompileEnvironment CompileEnvironment, ILogger Logger)
 		{
-			Log.TraceInformation("------- Build details --------");
-			Log.TraceInformation("Using {0}.", ToolchainInfo);
-			Log.TraceInformation("Using clang ({0}) version '{1}' (string), {2} (major), {3} (minor), {4} (patch)",
+			Logger.LogInformation("------- Build details --------");
+			Logger.LogInformation("Using {ToolchainInfo}.", ToolchainInfo);
+			Logger.LogInformation("Using clang ({ClangPath}) version '{ClangVersionString}' (string), {ClangVersionMajor} (major), {ClangVersionMinor} (minor), {ClangVersionPatch} (patch)",
 				ClangPath, ClangVersionString, ClangVersionMajor, ClangVersionMinor, ClangVersionPatch);
 
 			// inform the user which C++ library the engine is going to be compiled against - important for compatibility with third party code that uses STL
-			Log.TraceInformation("Using {0} standard C++ library.", ShouldUseLibcxx(CompileEnvironment.Architecture) ? "bundled libc++" : "compiler default (most likely libstdc++)");
-			Log.TraceInformation("Using {0}", UsingLld(CompileEnvironment.Architecture) ? "lld linker" : "default linker (ld)");
-			Log.TraceInformation("Using {0}", !String.IsNullOrEmpty(LlvmArPath) ? String.Format("llvm-ar : {0}", LlvmArPath) : String.Format("ar and ranlib: {0}, {1}", GetArPath(CompileEnvironment.Architecture), GetRanlibPath(CompileEnvironment.Architecture)));
+			Logger.LogInformation("Using {Lib} standard C++ library.", ShouldUseLibcxx(CompileEnvironment.Architecture) ? "bundled libc++" : "compiler default (most likely libstdc++)");
+			Logger.LogInformation("Using {Linker}", UsingLld(CompileEnvironment.Architecture) ? "lld linker" : "default linker (ld)");
+			Logger.LogInformation("Using {LlvmAr}", !String.IsNullOrEmpty(LlvmArPath) ? String.Format("llvm-ar : {0}", LlvmArPath) : String.Format("ar and ranlib: {0}, {1}", GetArPath(CompileEnvironment.Architecture), GetRanlibPath(CompileEnvironment.Architecture)));
 
 			if (Options.HasFlag(ClangToolChainOptions.EnableAddressSanitizer) ||
 				Options.HasFlag(ClangToolChainOptions.EnableThreadSanitizer) ||
@@ -973,41 +974,41 @@ namespace UnrealBuildTool
 				SanitizerInfo += Options.HasFlag(ClangToolChainOptions.EnableUndefinedBehaviorSanitizer) ? StaticOrShared + " linked UndefinedBehaviorSanitizer" : "";
 				SanitizerInfo += Options.HasFlag(ClangToolChainOptions.EnableMemorySanitizer) ? StaticOrShared + " linked MemorySanitizer" : "";
 
-				Log.TraceInformation(SanitizerInfo);
+				Logger.LogInformation("{SanitizerInfo}", SanitizerInfo);
 			}
 
 			// Also print other once-per-build information
 			if (bUseFixdeps)
 			{
-				Log.TraceInformation("Using old way to relink circularly dependent libraries (with a FixDeps step).");
+				Logger.LogInformation("Using old way to relink circularly dependent libraries (with a FixDeps step).");
 			}
 			else
 			{
-				Log.TraceInformation("Using fast way to relink  circularly dependent libraries (no FixDeps).");
+				Logger.LogInformation("Using fast way to relink  circularly dependent libraries (no FixDeps).");
 			}
 
 			if (CompileEnvironment.bPGOOptimize)
 			{
-				Log.TraceInformation("Using PGO (profile guided optimization).");
-				Log.TraceInformation("  Directory for PGO data files='{0}'", CompileEnvironment.PGODirectory);
-				Log.TraceInformation("  Prefix for PGO data files='{0}'", CompileEnvironment.PGOFilenamePrefix);
+				Logger.LogInformation("Using PGO (profile guided optimization).");
+				Logger.LogInformation("  Directory for PGO data files='{CompileEnvironmentPGODirectory}'", CompileEnvironment.PGODirectory);
+				Logger.LogInformation("  Prefix for PGO data files='{CompileEnvironmentPGOFilenamePrefix}'", CompileEnvironment.PGOFilenamePrefix);
 			}
 
 			if (CompileEnvironment.bPGOProfile)
 			{
-				Log.TraceInformation("Using PGI (profile guided instrumentation).");
+				Logger.LogInformation("Using PGI (profile guided instrumentation).");
 			}
 
 			if (CompileEnvironment.bAllowLTCG)
 			{
-				Log.TraceInformation("Using LTO (link-time optimization).");
+				Logger.LogInformation("Using LTO (link-time optimization).");
 			}
 
 			if (bSuppressPIE)
 			{
-				Log.TraceInformation("Compiler is set up to generate position independent executables by default, but we're suppressing it.");
+				Logger.LogInformation("Compiler is set up to generate position independent executables by default, but we're suppressing it.");
 			}
-			Log.TraceInformation("------------------------------");
+			Logger.LogInformation("------------------------------");
 		}
 
 		public override CPPOutput CompileCPPFiles(CppCompileEnvironment CompileEnvironment, List<FileItem> InputFiles, DirectoryReference OutputDir, string ModuleName, IActionGraphBuilder Graph)
@@ -1020,7 +1021,7 @@ namespace UnrealBuildTool
 
 			if (!bHasPrintedBuildDetails)
 			{
-				PrintBuildDetails(CompileEnvironment);
+				PrintBuildDetails(CompileEnvironment, Logger);
 
 				bHasPrintedBuildDetails = true;
 			}
@@ -1217,7 +1218,7 @@ namespace UnrealBuildTool
 		/// <summary>
 		/// Creates an action to archive all the .o files into single .a file
 		/// </summary>
-		public FileItem CreateArchiveAndIndex(LinkEnvironment LinkEnvironment, IActionGraphBuilder Graph)
+		public FileItem CreateArchiveAndIndex(LinkEnvironment LinkEnvironment, IActionGraphBuilder Graph, ILogger Logger)
 		{
 			// Create an archive action
 			Action ArchiveAction = Graph.CreateAction(ActionType.Link);
@@ -1286,11 +1287,11 @@ namespace UnrealBuildTool
 			return OutputFile;
 		}
 
-		public FileItem? FixDependencies(LinkEnvironment LinkEnvironment, FileItem Executable, IActionGraphBuilder Graph)
+		public FileItem? FixDependencies(LinkEnvironment LinkEnvironment, FileItem Executable, IActionGraphBuilder Graph, ILogger Logger)
 		{
 			if (bUseFixdeps)
 			{
-				Log.TraceVerbose("Adding postlink step");
+				Logger.LogDebug("Adding postlink step");
 
 				bool bUseCmdExe = BuildHostPlatform.Current.ShellType == ShellType.Cmd;
 				FileReference ShellBinary = BuildHostPlatform.Current.Shell;
@@ -1364,7 +1365,7 @@ namespace UnrealBuildTool
 
 			if (LinkEnvironment.bIsBuildingLibrary || bBuildImportLibraryOnly)
 			{
-				return CreateArchiveAndIndex(LinkEnvironment, Graph);
+				return CreateArchiveAndIndex(LinkEnvironment, Graph, Logger);
 			}
 
 			// Create an action that invokes the linker.
@@ -1670,7 +1671,7 @@ namespace UnrealBuildTool
 
 			string LinkScriptName = string.Format((bUseCmdExe ? "Link-{0}.link.bat" : "Link-{0}.link.sh"), OutputFile.Location.GetFileName());
 			string LinkScriptFullPath = Path.Combine(LinkEnvironment.LocalShadowDirectory!.FullName, LinkScriptName);
-			Log.TraceVerbose("Creating link script: {0}", LinkScriptFullPath);
+			Logger.LogDebug("Creating link script: {LinkScriptFullPath}", LinkScriptFullPath);
 			Directory.CreateDirectory(Path.GetDirectoryName(LinkScriptFullPath)!);
 			using (StreamWriter LinkWriter = File.CreateText(LinkScriptFullPath))
 			{
@@ -1740,7 +1741,7 @@ namespace UnrealBuildTool
 					if (!bHasWipedFixDepsScript)
 					{
 						bHasWipedFixDepsScript = true;
-						Log.TraceVerbose("Creating script: {0}", FixDepsScriptPath);
+						Logger.LogDebug("Creating script: {FixDepsScriptPath}", FixDepsScriptPath);
 						Directory.CreateDirectory(Path.GetDirectoryName(FixDepsScriptPath)!);
 						using (StreamWriter Writer = File.CreateText(FixDepsScriptPath))
 						{
@@ -1831,7 +1832,7 @@ namespace UnrealBuildTool
 					string RelinkScriptName = string.Format((bUseCmdExe ? "Relink-{0}.bat" : "Relink-{0}.sh"), OutputFile.Location.GetFileName());
 					string RelinkScriptFullPath = Path.Combine(LinkEnvironment.LocalShadowDirectory.FullName, RelinkScriptName);
 
-					Log.TraceVerbose("Creating script: {0}", RelinkScriptFullPath);
+					Logger.LogDebug("Creating script: {RelinkScriptFullPath}", RelinkScriptFullPath);
 					Directory.CreateDirectory(Path.GetDirectoryName(RelinkScriptFullPath)!);
 					using (StreamWriter RelinkWriter = File.CreateText(RelinkScriptFullPath))
 					{
@@ -1918,7 +1919,7 @@ namespace UnrealBuildTool
 					return OutputFiles;
 				}
 
-				FileItem? FixDepsOutputFile = FixDependencies(BinaryLinkEnvironment, Executable, Graph);
+				FileItem? FixDepsOutputFile = FixDependencies(BinaryLinkEnvironment, Executable, Graph, Logger);
 				if (FixDepsOutputFile != null)
 				{
 					OutputFiles.Add(FixDepsOutputFile);
@@ -1936,7 +1937,7 @@ namespace UnrealBuildTool
 			return OutputFiles;
 		}
 
-		public void StripSymbols(FileReference SourceFile, FileReference TargetFile)
+		public void StripSymbols(FileReference SourceFile, FileReference TargetFile, ILogger Logger)
 		{
 			if (SourceFile != TargetFile)
 			{
@@ -1949,7 +1950,7 @@ namespace UnrealBuildTool
 			StartInfo.Arguments = "--strip-debug \"" + TargetFile.FullName + "\"";
 			StartInfo.UseShellExecute = false;
 			StartInfo.CreateNoWindow = true;
-			Utils.RunLocalProcessAndLogOutput(StartInfo);
+			Utils.RunLocalProcessAndLogOutput(StartInfo, Logger);
 		}
 	}
 }
diff --git a/Engine/Source/Programs/UnrealBuildTool/Platform/Linux/UEBuildLinux.cs b/Engine/Source/Programs/UnrealBuildTool/Platform/Linux/UEBuildLinux.cs
--- a/Engine/Source/Programs/UnrealBuildTool/Platform/Linux/UEBuildLinux.cs
+++ b/Engine/Source/Programs/UnrealBuildTool/Platform/Linux/UEBuildLinux.cs
@@ -8,6 +8,7 @@ using System.IO;
 using EpicGames.Core;
 using System.Text.RegularExpressions;
 using UnrealBuildBase;
+using Microsoft.Extensions.Logging;
 
 namespace UnrealBuildTool
 {
@@ -162,14 +163,14 @@ namespace UnrealBuildTool
 		/// <summary>
 		/// Constructor
 		/// </summary>
-		public LinuxPlatform(LinuxPlatformSDK InSDK) 
-			: this(UnrealTargetPlatform.Linux, InSDK)
+		public LinuxPlatform(LinuxPlatformSDK InSDK, ILogger Logger) 
+			: this(UnrealTargetPlatform.Linux, InSDK, Logger)
 		{
 			SDK = InSDK;
 		}
 
-		public LinuxPlatform(UnrealTargetPlatform UnrealTarget, LinuxPlatformSDK InSDK)
-			: base(UnrealTarget, InSDK)
+		public LinuxPlatform(UnrealTargetPlatform UnrealTarget, LinuxPlatformSDK InSDK, ILogger Logger)
+			: base(UnrealTarget, InSDK, Logger)
 		{
 			SDK = InSDK;
 		}
@@ -452,12 +453,7 @@ namespace UnrealBuildTool
 				(Target.Platform == UnrealTargetPlatform.LinuxArm64 ? "1" : "0"));
 		}
 
-		/// <summary>
-		/// Setup the target environment for building
-		/// </summary>
-		/// <param name="Target">Settings for the target being compiled</param>
-		/// <param name="CompileEnvironment">The compile environment for this target</param>
-		/// <param name="LinkEnvironment">The link environment for this target</param>
+		/// <inheritdoc/>
 		public override void SetUpEnvironment(ReadOnlyTargetRules Target, CppCompileEnvironment CompileEnvironment, LinkEnvironment LinkEnvironment)
 		{
 			// During the native builds, check the system includes as well (check toolchain when cross-compiling?)
@@ -616,21 +612,23 @@ namespace UnrealBuildTool
 				Log.ColorConsoleOutput = false;
 			}
 
-			return new LinuxToolChain(Target.Architecture, SDK, Options);
+			return new LinuxToolChain(Target.Architecture, SDK, Options, Logger);
 		}
 
-		/// <summary>
-		/// Deploys the given target
-		/// </summary>
-		/// <param name="Receipt">Receipt for the target being deployed</param>
+		/// <inheritdoc/>
 		public override void Deploy(TargetReceipt Receipt)
 		{
-			new UEDeployLinux().PrepTargetForDeployment(Receipt);
+			new UEDeployLinux(Logger).PrepTargetForDeployment(Receipt);
 		}
 	}
 
 	class UEDeployLinux: UEBuildDeploy
 	{
+		public UEDeployLinux(ILogger InLogger)
+			: base(InLogger)
+		{
+		}
+
 		public override bool PrepTargetForDeployment(TargetReceipt Receipt)
 		{
 			return base.PrepTargetForDeployment(Receipt);
@@ -647,18 +645,18 @@ namespace UnrealBuildTool
 		/// <summary>
 		/// Register the platform with the UEBuildPlatform class
 		/// </summary>
-		public override void RegisterBuildPlatforms()
+		public override void RegisterBuildPlatforms(ILogger Logger)
 		{
-			LinuxPlatformSDK SDK = new LinuxPlatformSDK();
-			LinuxPlatformSDK SDKArm64 = new LinuxPlatformSDK();
+			LinuxPlatformSDK SDK = new LinuxPlatformSDK(Logger);
+			LinuxPlatformSDK SDKArm64 = new LinuxPlatformSDK(Logger);
 
 			// Register this build platform for Linux x86-64 and Arm64
-			UEBuildPlatform.RegisterBuildPlatform(new LinuxPlatform(UnrealTargetPlatform.Linux, SDK));
+			UEBuildPlatform.RegisterBuildPlatform(new LinuxPlatform(UnrealTargetPlatform.Linux, SDK, Logger), Logger);
 			UEBuildPlatform.RegisterPlatformWithGroup(UnrealTargetPlatform.Linux, UnrealPlatformGroup.Linux);
 			UEBuildPlatform.RegisterPlatformWithGroup(UnrealTargetPlatform.Linux, UnrealPlatformGroup.Unix);
 			UEBuildPlatform.RegisterPlatformWithGroup(UnrealTargetPlatform.Linux, UnrealPlatformGroup.Desktop);
 
-			UEBuildPlatform.RegisterBuildPlatform(new LinuxPlatform(UnrealTargetPlatform.LinuxArm64, SDKArm64));
+			UEBuildPlatform.RegisterBuildPlatform(new LinuxPlatform(UnrealTargetPlatform.LinuxArm64, SDKArm64, Logger), Logger);
 			UEBuildPlatform.RegisterPlatformWithGroup(UnrealTargetPlatform.LinuxArm64, UnrealPlatformGroup.Linux);
 			UEBuildPlatform.RegisterPlatformWithGroup(UnrealTargetPlatform.LinuxArm64, UnrealPlatformGroup.Unix);
 			UEBuildPlatform.RegisterPlatformWithGroup(UnrealTargetPlatform.LinuxArm64, UnrealPlatformGroup.Desktop);
diff --git a/Engine/Source/Programs/UnrealBuildTool/Platform/Mac/MacExports.cs b/Engine/Source/Programs/UnrealBuildTool/Platform/Mac/MacExports.cs
--- a/Engine/Source/Programs/UnrealBuildTool/Platform/Mac/MacExports.cs
+++ b/Engine/Source/Programs/UnrealBuildTool/Platform/Mac/MacExports.cs
@@ -6,6 +6,7 @@ using System.Linq;
 using System.Text;
 using System.Threading.Tasks;
 using EpicGames.Core;
+using Microsoft.Extensions.Logging;
 
 namespace UnrealBuildTool
 {
@@ -104,7 +105,7 @@ namespace UnrealBuildTool
 			{
 				if (!IsRunningUnderRosettaVar.HasValue)
 				{
-					string TranslatedOutput = Utils.RunLocalProcessAndReturnStdOut("/usr/sbin/sysctl", "sysctl");
+					string TranslatedOutput = Utils.RunLocalProcessAndReturnStdOut("/usr/sbin/sysctl", "sysctl", null);
 					IsRunningUnderRosettaVar = TranslatedOutput.Contains("sysctl.proc_translated: 1");
 				}
 
@@ -123,7 +124,7 @@ namespace UnrealBuildTool
 				if (!IsRunningOnAppleArchitectureVar.HasValue)
 				{
 					// On an m1 mac this appears to be where the brand is.
-					string BrandOutput = Utils.RunLocalProcessAndReturnStdOut("/usr/sbin/sysctl", "-n machdep.cpu.brand_string");
+					string BrandOutput = Utils.RunLocalProcessAndReturnStdOut("/usr/sbin/sysctl", "-n machdep.cpu.brand_string", null);
 					IsRunningOnAppleArchitectureVar = BrandOutput.Contains("Apple") || IsRunningUnderRosetta;
 				}
 
@@ -136,9 +137,10 @@ namespace UnrealBuildTool
 		/// </summary>
 		/// <param name="SourceFile">The input file</param>
 		/// <param name="TargetFile">The output file</param>
-		public static void StripSymbols(FileReference SourceFile, FileReference TargetFile)
+		/// <param name="Logger"></param>
+		public static void StripSymbols(FileReference SourceFile, FileReference TargetFile, ILogger Logger)
 		{
-			MacToolChain ToolChain = new MacToolChain(null, ClangToolChainOptions.None);
+			MacToolChain ToolChain = new MacToolChain(null, ClangToolChainOptions.None, Logger);
 			ToolChain.StripSymbols(SourceFile, TargetFile);
 		}		
 	}
diff --git a/Engine/Source/Programs/UnrealBuildTool/Platform/Mac/MacPlatformSDK.Versions.cs b/Engine/Source/Programs/UnrealBuildTool/Platform/Mac/MacPlatformSDK.Versions.cs
--- a/Engine/Source/Programs/UnrealBuildTool/Platform/Mac/MacPlatformSDK.Versions.cs
+++ b/Engine/Source/Programs/UnrealBuildTool/Platform/Mac/MacPlatformSDK.Versions.cs
@@ -1,15 +1,21 @@
-// Copyright Epic Games, Inc. All Rights Reserved.
-
-using EpicGames.Core;
-
-namespace UnrealBuildTool
-{
-	internal partial class MacPlatformSDK : ApplePlatformSDK
-	{
-		protected override void GetValidSoftwareVersionRange(out string MinVersion, out string? MaxVersion)
-		{
-			MinVersion = "10.15.7";
-			MaxVersion = null;
-		}
-	}
+// Copyright Epic Games, Inc. All Rights Reserved.
+
+using EpicGames.Core;
+using Microsoft.Extensions.Logging;
+
+namespace UnrealBuildTool
+{
+	internal partial class MacPlatformSDK : ApplePlatformSDK
+	{
+		public MacPlatformSDK(ILogger Logger)
+			: base(Logger)
+		{
+		}
+
+		protected override void GetValidSoftwareVersionRange(out string MinVersion, out string? MaxVersion)
+		{
+			MinVersion = "10.15.7";
+			MaxVersion = null;
+		}
+	}
 }
\ No newline at end of file
diff --git a/Engine/Source/Programs/UnrealBuildTool/Platform/Mac/MacProjectGenerator.cs b/Engine/Source/Programs/UnrealBuildTool/Platform/Mac/MacProjectGenerator.cs
--- a/Engine/Source/Programs/UnrealBuildTool/Platform/Mac/MacProjectGenerator.cs
+++ b/Engine/Source/Programs/UnrealBuildTool/Platform/Mac/MacProjectGenerator.cs
@@ -6,6 +6,7 @@
 using System.Diagnostics;
 using System.IO;
 using EpicGames.Core;
+using Microsoft.Extensions.Logging;
 
 namespace UnrealBuildTool
 {
@@ -18,8 +19,9 @@
 		/// Constructor
 		/// </summary>
 		/// <param name="Arguments">Command line arguments passed to the project generator</param>
-		public MacProjectGenerator(CommandLineArguments Arguments)
-			: base(Arguments)
+		/// <param name="Logger">Logger for output</param>
+		public MacProjectGenerator(CommandLineArguments Arguments, ILogger Logger)
+			: base(Arguments, Logger)
 		{
 		}
 
diff --git a/Engine/Source/Programs/UnrealBuildTool/Platform/Mac/MacToolChain.cs b/Engine/Source/Programs/UnrealBuildTool/Platform/Mac/MacToolChain.cs
--- a/Engine/Source/Programs/UnrealBuildTool/Platform/Mac/MacToolChain.cs
+++ b/Engine/Source/Programs/UnrealBuildTool/Platform/Mac/MacToolChain.cs
@@ -13,6 +13,7 @@ using EpicGames.Core;
 using System.Globalization;
 using System.Text.RegularExpressions;
 using UnrealBuildBase;
+using Microsoft.Extensions.Logging;
 
 namespace UnrealBuildTool
 {
@@ -48,12 +49,13 @@ namespace UnrealBuildTool
 		/// Constructor
 		/// </summary>
 		/// <param name="bVerbose">Whether to output verbose logging</param>
-		public MacToolChainSettings(bool bVerbose) : base(bVerbose)
+		/// <param name="Logger">Logger for output</param>
+		public MacToolChainSettings(bool bVerbose, ILogger Logger) : base(bVerbose, Logger)
 		{
 			BaseSDKDir = XcodeDeveloperDir + "Platforms/MacOSX.platform/Developer/SDKs";
 			ToolchainDir = XcodeDeveloperDir + "Toolchains/XcodeDefault.xctoolchain/usr/bin/";
 
-			SelectSDK(BaseSDKDir, "MacOSX", ref MacOSSDKVersion, bVerbose);
+			SelectSDK(BaseSDKDir, "MacOSX", ref MacOSSDKVersion, bVerbose, Logger);
 
 			// convert to float for easy comparison
 			if(String.IsNullOrWhiteSpace(MacOSSDKVersion))
@@ -72,12 +74,12 @@ namespace UnrealBuildTool
 	/// </summary>
 	class MacToolChain : AppleToolChain
 	{
-		public MacToolChain(FileReference? InProjectFile, ClangToolChainOptions InOptions)
-			: base(InProjectFile, InOptions)
+		public MacToolChain(FileReference? InProjectFile, ClangToolChainOptions InOptions, ILogger InLogger)
+			: base(InProjectFile, InOptions, InLogger)
 		{
 		}
 
-		public static Lazy<MacToolChainSettings> SettingsPrivate = new Lazy<MacToolChainSettings>(() => new MacToolChainSettings(false));
+		public static Lazy<MacToolChainSettings> SettingsPrivate = new Lazy<MacToolChainSettings>(() => new MacToolChainSettings(false, Log.Logger));
 
 		public static MacToolChainSettings Settings
 		{
@@ -126,7 +128,7 @@ namespace UnrealBuildTool
 			// validation, because sometimes this is called from a shell script and quoting messes up		
 			if (!Target.Architecture.All(C => char.IsLetterOrDigit(C) || C == '_' || C == '+'))
 			{
-				Log.TraceError("Architecture '{0}' contains invalid characters", Target.Architecture);
+				throw new BuildException($"Architecture '{Target.Architecture}' contains invalid characters");
 			}			
 
 			SetupXcodePaths(true);
@@ -608,7 +610,7 @@ namespace UnrealBuildTool
 			return GameName;
 		}
 
-		private void AddLibraryPathToRPaths(string Library, string ExeAbsolutePath, ref List<string> RPaths, ref string LinkCommand, bool bIsBuildingAppBundle)
+		private void AddLibraryPathToRPaths(string Library, string ExeAbsolutePath, ref List<string> RPaths, ref string LinkCommand, bool bIsBuildingAppBundle, ILogger Logger)
 		{
 			string LibraryFullPath = Path.GetFullPath(Library);
  			string LibraryDir = Path.GetDirectoryName(LibraryFullPath)!;
@@ -666,7 +668,7 @@ namespace UnrealBuildTool
 					}
 					else
 					{
-						Log.TraceWarning("Unexpected third party dylib location when generating RPATH entries: {0}. Skipping.", LibraryFullPath);
+						Logger.LogWarning("Unexpected third party dylib location when generating RPATH entries: {LibraryFullPath}. Skipping.", LibraryFullPath);
 					}
 
 					// For staged code-based games we need additional entry if the game is not stored directly in the engine's root directory
@@ -773,7 +775,7 @@ namespace UnrealBuildTool
 						LinkCommand += string.Format(" \"{0}\"", Path.GetFullPath(AdditionalLibrary));
 					}
 
-					AddLibraryPathToRPaths(AdditionalLibrary, AbsolutePath, ref RPaths, ref LinkCommand, bIsBuildingAppBundle);
+					AddLibraryPathToRPaths(AdditionalLibrary, AbsolutePath, ref RPaths, ref LinkCommand, bIsBuildingAppBundle, Logger);
 				}
 
 				foreach (string AdditionalLibrary in LinkEnvironment.DelayLoadDLLs)
@@ -786,7 +788,7 @@ namespace UnrealBuildTool
 
 					LinkCommand += string.Format(" -weak_library \"{0}\"", Path.GetFullPath(AdditionalLibrary));
 
-					AddLibraryPathToRPaths(AdditionalLibrary, AbsolutePath, ref RPaths, ref LinkCommand, bIsBuildingAppBundle);
+					AddLibraryPathToRPaths(AdditionalLibrary, AbsolutePath, ref RPaths, ref LinkCommand, bIsBuildingAppBundle, Logger);
 				}
 			}
 
@@ -819,7 +821,7 @@ namespace UnrealBuildTool
 				foreach (KeyValuePair<string, bool> Framework in AllFrameworks)
 				{
 					LinkCommand += AddFrameworkToLinkCommand(Framework.Key, Framework.Value ? "-weak_framework" : "-framework");
-					AddLibraryPathToRPaths(Framework.Key, AbsolutePath, ref RPaths, ref LinkCommand, bIsBuildingAppBundle);
+					AddLibraryPathToRPaths(Framework.Key, AbsolutePath, ref RPaths, ref LinkCommand, bIsBuildingAppBundle, Logger);
 				}
 			}
 
@@ -995,7 +997,7 @@ namespace UnrealBuildTool
 						}
 						else
 						{
-							Log.TraceWarning("Found {0} Target.cs files for {1} in alldir search of directory {2}", TargetFiles.Length, GameName, Directory.GetCurrentDirectory());
+							Logger.LogWarning("Found {NumFiles} Target.cs files for {GameName} in alldir search of directory {Dir}", TargetFiles.Length, GameName, Directory.GetCurrentDirectory());
 						}
 					}
 					else
@@ -1117,7 +1119,8 @@ namespace UnrealBuildTool
 		/// <param name="MachOBinary">FileItem describing the executable or dylib to generate debug info for</param>
 		/// <param name="LinkEnvironment"></param>
 		/// <param name="Graph">List of actions to be executed. Additional actions will be added to this list.</param>
-		public FileItem GenerateDebugInfo(FileItem MachOBinary, LinkEnvironment LinkEnvironment, IActionGraphBuilder Graph)
+		/// <param name="Logger">Logger for output</param>
+		public FileItem GenerateDebugInfo(FileItem MachOBinary, LinkEnvironment LinkEnvironment, IActionGraphBuilder Graph, ILogger Logger)
 		{
 			string BinaryPath = MachOBinary.AbsolutePath;
 			if (BinaryPath.Contains(".app"))
@@ -1150,7 +1153,7 @@ namespace UnrealBuildTool
 			// a problem where dsymutil would exit with an error saying the input file did not exist.
 			// Note that the source and dest are switched from a copy command
 			string ExtraOptions;
-			string DsymutilPath = GetDsymutilPath(out ExtraOptions, bIsForLTOBuild: false);
+			string DsymutilPath = GetDsymutilPath(Logger, out ExtraOptions, bIsForLTOBuild: false);
 
 			string ArgumentString = "-c \"";
 			ArgumentString += string.Format("for i in {{1..30}}; ");
@@ -1353,8 +1356,10 @@ namespace UnrealBuildTool
 
 		private List<FileItem> DebugInfoFiles = new List<FileItem>();
 
-		public override void FinalizeOutput(ReadOnlyTargetRules Target, TargetMakefile Makefile)
+		public override void FinalizeOutput(ReadOnlyTargetRules Target, TargetMakefileBuilder MakefileBuilder)
 		{
+			TargetMakefile Makefile = MakefileBuilder.Makefile;
+
 			// Re-add any .dSYM files that may have been stripped out.
 			List<string> OutputFiles = Makefile.OutputItems.Select(Item => Path.ChangeExtension(Item.FullName, ".dSYM")).Distinct().ToList();
 			foreach (FileItem DebugItem in DebugInfoFiles)
@@ -1389,7 +1394,7 @@ namespace UnrealBuildTool
 				// We want dsyms to be created after all dylib dependencies are fixed. If FixDylibDependencies action was not created yet, save the info for later.
 				if (FixDylibOutputFile != null)
 				{
-					DebugInfoFiles.Add(GenerateDebugInfo(Executable, BinaryLinkEnvironment, Graph));
+					DebugInfoFiles.Add(GenerateDebugInfo(Executable, BinaryLinkEnvironment, Graph, Logger));
 				}
 				else
 				{
@@ -1414,7 +1419,7 @@ namespace UnrealBuildTool
 			// Add dsyms that we couldn't add before FixDylibDependencies action was created
 			foreach (FileItem Exe in ExecutablesThatNeedDsyms)
 			{
-				DebugInfoFiles.Add(GenerateDebugInfo(Exe, BinaryLinkEnvironment, Graph));
+				DebugInfoFiles.Add(GenerateDebugInfo(Exe, BinaryLinkEnvironment, Graph, Logger));
 			}
 			ExecutablesThatNeedDsyms.Clear();
 
diff --git a/Engine/Source/Programs/UnrealBuildTool/Platform/Mac/UEBuildMac.cs b/Engine/Source/Programs/UnrealBuildTool/Platform/Mac/UEBuildMac.cs
--- a/Engine/Source/Programs/UnrealBuildTool/Platform/Mac/UEBuildMac.cs
+++ b/Engine/Source/Programs/UnrealBuildTool/Platform/Mac/UEBuildMac.cs
@@ -5,6 +5,7 @@ using System;
 using System.Collections.Generic;
 using EpicGames.Core;
 using System.Linq;
+using Microsoft.Extensions.Logging;
 
 namespace UnrealBuildTool
 {
@@ -100,7 +101,7 @@ namespace UnrealBuildTool
 
 	class MacPlatform : UEBuildPlatform
 	{
-		public MacPlatform(UEBuildPlatformSDK InSDK) : base(UnrealTargetPlatform.Mac, InSDK)
+		public MacPlatform(UEBuildPlatformSDK InSDK, ILogger InLogger) : base(UnrealTargetPlatform.Mac, InSDK, InLogger)
 		{
 		}
 
@@ -115,7 +116,7 @@ namespace UnrealBuildTool
 		}
 
 		public override void ResetTarget(TargetRules Target)
-		{			
+		{
 		}
 
 		public override void ValidateTarget(TargetRules Target)
@@ -399,16 +400,13 @@ namespace UnrealBuildTool
 				Options |= ClangToolChainOptions.OutputDylib;
 			}
 
-			return new MacToolChain(Target.ProjectFile, Options);
+			return new MacToolChain(Target.ProjectFile, Options, Logger);
 		}
 
-		/// <summary>
-		/// Deploys the given target
-		/// </summary>
-		/// <param name="Receipt">Receipt for the target being deployed</param>
+		/// <inheritdoc/>
 		public override void Deploy(TargetReceipt Receipt)
 		{
-			new UEDeployMac().PrepTargetForDeployment(Receipt);
+			new UEDeployMac(Logger).PrepTargetForDeployment(Receipt);
 		}
 	}
 
@@ -422,12 +420,12 @@ namespace UnrealBuildTool
 		/// <summary>
 		/// Register the platform with the UEBuildPlatform class
 		/// </summary>
-		public override void RegisterBuildPlatforms()
+		public override void RegisterBuildPlatforms(ILogger Logger)
 		{
-			MacPlatformSDK SDK = new MacPlatformSDK();
+			MacPlatformSDK SDK = new MacPlatformSDK(Logger);
 
 			// Register this build platform for Mac
-			UEBuildPlatform.RegisterBuildPlatform(new MacPlatform(SDK));
+			UEBuildPlatform.RegisterBuildPlatform(new MacPlatform(SDK, Logger), Logger);
 			UEBuildPlatform.RegisterPlatformWithGroup(UnrealTargetPlatform.Mac, UnrealPlatformGroup.Apple);
 			UEBuildPlatform.RegisterPlatformWithGroup(UnrealTargetPlatform.Mac, UnrealPlatformGroup.Desktop);
 		}
diff --git a/Engine/Source/Programs/UnrealBuildTool/Platform/Mac/UEDeployMac.cs b/Engine/Source/Programs/UnrealBuildTool/Platform/Mac/UEDeployMac.cs
--- a/Engine/Source/Programs/UnrealBuildTool/Platform/Mac/UEDeployMac.cs
+++ b/Engine/Source/Programs/UnrealBuildTool/Platform/Mac/UEDeployMac.cs
@@ -9,14 +9,20 @@
 using System.IO;
 using Ionic.Zip;
 using EpicGames.Core;
+using Microsoft.Extensions.Logging;
 
 namespace UnrealBuildTool
 {
 	class UEDeployMac : UEBuildDeploy
 	{
+		public UEDeployMac(ILogger InLogger) 
+			: base(InLogger)
+		{
+		}
+
 		public override bool PrepTargetForDeployment(TargetReceipt Receipt)
 		{
-			Log.TraceInformation("Deploying now!");
+			Logger.LogInformation("Deploying now!");
 			return base.PrepTargetForDeployment(Receipt);
 		}
 
diff --git a/Engine/Source/Programs/UnrealBuildTool/Platform/TVOS/TVOSExports.cs b/Engine/Source/Programs/UnrealBuildTool/Platform/TVOS/TVOSExports.cs
--- a/Engine/Source/Programs/UnrealBuildTool/Platform/TVOS/TVOSExports.cs
+++ b/Engine/Source/Programs/UnrealBuildTool/Platform/TVOS/TVOSExports.cs
@@ -6,6 +6,7 @@
 using System.Text;
 using System.Threading.Tasks;
 using EpicGames.Core;
+using Microsoft.Extensions.Logging;
 
 namespace UnrealBuildTool
 {
@@ -75,11 +76,12 @@
 		/// <param name="bIsDataDeploy"></param>
 		/// <param name="bCreateStubIPA"></param>
 		/// <param name="BuildReceiptFileName"></param>
+		/// <param name="Logger"></param>
 		/// <returns></returns>
-		public static bool PrepForUATPackageOrDeploy(UnrealTargetConfiguration Config, FileReference ProjectFile, string InProjectName, DirectoryReference InProjectDirectory, string InExecutablePath, DirectoryReference InEngineDir, bool bForDistribution, string CookFlavor, bool bIsDataDeploy, bool bCreateStubIPA, FileReference BuildReceiptFileName)
+		public static bool PrepForUATPackageOrDeploy(UnrealTargetConfiguration Config, FileReference ProjectFile, string InProjectName, DirectoryReference InProjectDirectory, string InExecutablePath, DirectoryReference InEngineDir, bool bForDistribution, string CookFlavor, bool bIsDataDeploy, bool bCreateStubIPA, FileReference BuildReceiptFileName, ILogger Logger)
 		{
 			TargetReceipt Receipt = TargetReceipt.Read(BuildReceiptFileName);
-			return new UEDeployTVOS().PrepForUATPackageOrDeploy(Config, ProjectFile, InProjectName, InProjectDirectory.FullName, InExecutablePath, InEngineDir.FullName, bForDistribution, CookFlavor, bIsDataDeploy, bCreateStubIPA, Receipt);
+			return new UEDeployTVOS(Logger).PrepForUATPackageOrDeploy(Config, ProjectFile, InProjectName, InProjectDirectory.FullName, InExecutablePath, InEngineDir.FullName, bForDistribution, CookFlavor, bIsDataDeploy, bCreateStubIPA, Receipt);
 		}
 
 		/// <summary>
@@ -95,13 +97,14 @@
 		/// <param name="InEngineDir"></param>
 		/// <param name="AppDirectory"></param>
 		/// <param name="BuildReceiptFileName"></param>
+		/// <param name="Logger"></param>
 		/// <param name="bSupportsPortrait"></param>
 		/// <param name="bSupportsLandscape"></param>
 		/// <returns></returns>
-		public static bool GeneratePList(FileReference ProjectFile, UnrealTargetConfiguration Config, DirectoryReference ProjectDirectory, bool bIsUnrealGame, string GameName, bool bIsClient, string ProjectName, DirectoryReference InEngineDir, DirectoryReference AppDirectory, FileReference BuildReceiptFileName, out bool bSupportsPortrait, out bool bSupportsLandscape)
+		public static bool GeneratePList(FileReference ProjectFile, UnrealTargetConfiguration Config, DirectoryReference ProjectDirectory, bool bIsUnrealGame, string GameName, bool bIsClient, string ProjectName, DirectoryReference InEngineDir, DirectoryReference AppDirectory, FileReference BuildReceiptFileName, ILogger Logger, out bool bSupportsPortrait, out bool bSupportsLandscape)
 		{
 			TargetReceipt Receipt = TargetReceipt.Read(BuildReceiptFileName);
-			return new UEDeployTVOS().GeneratePList(ProjectFile, Config, ProjectDirectory.FullName, bIsUnrealGame, GameName, bIsClient, ProjectName, InEngineDir.FullName, AppDirectory.FullName, Receipt, out bSupportsPortrait, out bSupportsLandscape);
+			return new UEDeployTVOS(Logger).GeneratePList(ProjectFile, Config, ProjectDirectory.FullName, bIsUnrealGame, GameName, bIsClient, ProjectName, InEngineDir.FullName, AppDirectory.FullName, Receipt, out bSupportsPortrait, out bSupportsLandscape);
 		}
 	}
 }
diff --git a/Engine/Source/Programs/UnrealBuildTool/Platform/TVOS/TVOSProjectGenerator.cs b/Engine/Source/Programs/UnrealBuildTool/Platform/TVOS/TVOSProjectGenerator.cs
--- a/Engine/Source/Programs/UnrealBuildTool/Platform/TVOS/TVOSProjectGenerator.cs
+++ b/Engine/Source/Programs/UnrealBuildTool/Platform/TVOS/TVOSProjectGenerator.cs
@@ -1,31 +1,33 @@
-// Copyright Epic Games, Inc. All Rights Reserved.
-
-using System;
-using System.Collections.Generic;
-using System.Text;
-using System.Diagnostics;
-using System.IO;
-using EpicGames.Core;
-
-namespace UnrealBuildTool
-{
-	class TVOSProjectGenerator : IOSProjectGenerator
-    {
-		/// <summary>
-		/// Constructor
-		/// </summary>
-		/// <param name="Arguments">Command line arguments passed to the project generator</param>
-		public TVOSProjectGenerator(CommandLineArguments Arguments)
-			: base(Arguments)
-		{
-		}
-
-		/// <summary>
-		/// Enumerate all the platforms that this generator supports
-		/// </summary>
-		public override IEnumerable<UnrealTargetPlatform> GetPlatforms()
-		{
-			yield return UnrealTargetPlatform.TVOS;
-		}
-    }
-}
+// Copyright Epic Games, Inc. All Rights Reserved.
+
+using System;
+using System.Collections.Generic;
+using System.Text;
+using System.Diagnostics;
+using System.IO;
+using EpicGames.Core;
+using Microsoft.Extensions.Logging;
+
+namespace UnrealBuildTool
+{
+	class TVOSProjectGenerator : IOSProjectGenerator
+    {
+		/// <summary>
+		/// Constructor
+		/// </summary>
+		/// <param name="Arguments">Command line arguments passed to the project generator</param>
+		/// <param name="Logger">Logger for output</param>
+		public TVOSProjectGenerator(CommandLineArguments Arguments, ILogger Logger)
+			: base(Arguments, Logger)
+		{
+		}
+
+		/// <summary>
+		/// Enumerate all the platforms that this generator supports
+		/// </summary>
+		public override IEnumerable<UnrealTargetPlatform> GetPlatforms()
+		{
+			yield return UnrealTargetPlatform.TVOS;
+		}
+    }
+}
diff --git a/Engine/Source/Programs/UnrealBuildTool/Platform/TVOS/TVOSToolChain.cs b/Engine/Source/Programs/UnrealBuildTool/Platform/TVOS/TVOSToolChain.cs
--- a/Engine/Source/Programs/UnrealBuildTool/Platform/TVOS/TVOSToolChain.cs
+++ b/Engine/Source/Programs/UnrealBuildTool/Platform/TVOS/TVOSToolChain.cs
@@ -13,20 +13,21 @@ using System.Text;
 using Ionic.Zip;
 using Ionic.Zlib;
 using EpicGames.Core;
+using Microsoft.Extensions.Logging;
 
 namespace UnrealBuildTool
 {
 	class TVOSToolChainSettings : IOSToolChainSettings
 	{
-		public TVOSToolChainSettings() : base("AppleTVOS", "AppleTVSimulator")
+		public TVOSToolChainSettings(ILogger Logger) : base("AppleTVOS", "AppleTVSimulator", Logger)
 		{
 		}
 	}
 
 	class TVOSToolChain : IOSToolChain
 	{
-		public TVOSToolChain(ReadOnlyTargetRules InTarget, TVOSProjectSettings InProjectSettings)
-			: base(InTarget, InProjectSettings, () => new TVOSToolChainSettings(), ClangToolChainOptions.None)
+		public TVOSToolChain(ReadOnlyTargetRules InTarget, TVOSProjectSettings InProjectSettings, ILogger InLogger)
+			: base(InTarget, InProjectSettings, () => new TVOSToolChainSettings(InLogger), ClangToolChainOptions.None, InLogger)
 		{
 		}
 
diff --git a/Engine/Source/Programs/UnrealBuildTool/Platform/TVOS/UEBuildTVOS.cs b/Engine/Source/Programs/UnrealBuildTool/Platform/TVOS/UEBuildTVOS.cs
--- a/Engine/Source/Programs/UnrealBuildTool/Platform/TVOS/UEBuildTVOS.cs
+++ b/Engine/Source/Programs/UnrealBuildTool/Platform/TVOS/UEBuildTVOS.cs
@@ -3,6 +3,7 @@
 using System;
 using System.Collections.Generic;
 using EpicGames.Core;
+using Microsoft.Extensions.Logging;
 
 namespace UnrealBuildTool
 {
@@ -24,16 +25,16 @@
 
 	class TVOSProvisioningData : IOSProvisioningData
 	{
-		public TVOSProvisioningData(TVOSProjectSettings ProjectSettings, bool bForDistribution)
-			: base(ProjectSettings, true, bForDistribution)
+		public TVOSProvisioningData(TVOSProjectSettings ProjectSettings, bool bForDistribution, ILogger Logger)
+			: base(ProjectSettings, true, bForDistribution, Logger)
 		{
 		}
 	}
 
 	class TVOSPlatform : IOSPlatform
     {
-		public TVOSPlatform(UEBuildPlatformSDK InSDK)
-			: base(InSDK, UnrealTargetPlatform.TVOS)
+		public TVOSPlatform(UEBuildPlatformSDK InSDK, ILogger Logger)
+			: base(InSDK, UnrealTargetPlatform.TVOS, Logger)
 		{
 		}
 
@@ -77,7 +78,7 @@
 
 		protected override IOSProvisioningData CreateProvisioningData(IOSProjectSettings ProjectSettings, bool bForDistribution)
 		{
-			return new TVOSProvisioningData((TVOSProjectSettings)ProjectSettings, bForDistribution);
+			return new TVOSProvisioningData((TVOSProjectSettings)ProjectSettings, bForDistribution, Logger);
 		}
 
 		public override void ModifyModuleRulesForOtherPlatform(string ModuleName, ModuleRules Rules, ReadOnlyTargetRules Target)
@@ -100,7 +101,7 @@
 				}
 			}
 		}
-    
+
 		/// <summary>
 		/// Setup the target environment for building
 		/// </summary>
@@ -125,12 +126,12 @@
 		public override UEToolChain CreateToolChain(ReadOnlyTargetRules Target)
 		{
 			TVOSProjectSettings ProjectSettings = ((TVOSPlatform)UEBuildPlatform.GetBuildPlatform(UnrealTargetPlatform.TVOS)).ReadProjectSettings(Target.ProjectFile);
-			return new TVOSToolChain(Target, ProjectSettings);
+			return new TVOSToolChain(Target, ProjectSettings, Logger);
 		}
 
 		public override void Deploy(TargetReceipt Receipt)
 		{
-			new UEDeployTVOS().PrepTargetForDeployment(Receipt);
+			new UEDeployTVOS(Logger).PrepTargetForDeployment(Receipt);
 		}
 	}
 
@@ -144,12 +145,12 @@
 		/// <summary>
 		/// Register the platform with the UEBuildPlatform class
 		/// </summary>
-		public override void RegisterBuildPlatforms()
+		public override void RegisterBuildPlatforms(ILogger Logger)
 		{
-			ApplePlatformSDK SDK = new IOSPlatformSDK();
+			ApplePlatformSDK SDK = new IOSPlatformSDK(Logger);
 		
 			// Register this build platform for IOS
-			UEBuildPlatform.RegisterBuildPlatform(new TVOSPlatform(SDK));
+			UEBuildPlatform.RegisterBuildPlatform(new TVOSPlatform(SDK, Logger), Logger);
 			UEBuildPlatform.RegisterPlatformWithGroup(UnrealTargetPlatform.TVOS, UnrealPlatformGroup.Apple);
 			UEBuildPlatform.RegisterPlatformWithGroup(UnrealTargetPlatform.TVOS, UnrealPlatformGroup.IOS);
 		}
diff --git a/Engine/Source/Programs/UnrealBuildTool/Platform/TVOS/UEDeployTVOS.cs b/Engine/Source/Programs/UnrealBuildTool/Platform/TVOS/UEDeployTVOS.cs
--- a/Engine/Source/Programs/UnrealBuildTool/Platform/TVOS/UEDeployTVOS.cs
+++ b/Engine/Source/Programs/UnrealBuildTool/Platform/TVOS/UEDeployTVOS.cs
@@ -10,13 +10,15 @@
 using System.Xml.Linq;
 using EpicGames.Core;
 using UnrealBuildBase;
+using Microsoft.Extensions.Logging;
 
 namespace UnrealBuildTool
 {
 	class UEDeployTVOS : UEDeployIOS
 	{
 
-        public UEDeployTVOS()
+        public UEDeployTVOS(ILogger InLogger)
+			: base(InLogger)
         {
         }
 
@@ -25,7 +27,7 @@
 			return "TVOS";
 		}
 
-		public static bool GenerateTVOSPList(string ProjectDirectory, bool bIsUnrealGame, string GameName, bool bIsClient, string ProjectName, string InEngineDir, string AppDirectory, UnrealPluginLanguage? UPL, string? BundleID)
+		public static bool GenerateTVOSPList(string ProjectDirectory, bool bIsUnrealGame, string GameName, bool bIsClient, string ProjectName, string InEngineDir, string AppDirectory, UnrealPluginLanguage? UPL, string? BundleID, ILogger Logger)
 		{
 			// @todo tvos: THIS!
 
@@ -74,7 +76,7 @@
 			// minimum iOS version
 			string MinVersionSetting = "";
 			Ini.GetString("/Script/IOSRuntimeSettings.IOSRuntimeSettings", "MinimumiOSVersion", out MinVersionSetting);
-			string MinVersion = GetMinimumOSVersion(MinVersionSetting);
+			string MinVersion = GetMinimumOSVersion(MinVersionSetting, Logger);
 
 			// extra plist data
 			string ExtraData = "";
@@ -204,7 +206,7 @@
 		public override bool GeneratePList(FileReference? ProjectFile, UnrealTargetConfiguration Config, string ProjectDirectory, bool bIsUnrealGame, string GameName, bool bIsClient, string ProjectName, string InEngineDir, string AppDirectory, List<string> UPLScripts, string? BundleID, bool bBuildAsFramework, out bool bSupportsPortrait, out bool bSupportsLandscape)
 		{
 			bSupportsLandscape = bSupportsPortrait = true;
-			return GenerateTVOSPList(ProjectDirectory, bIsUnrealGame, GameName, bIsClient, ProjectName, InEngineDir, AppDirectory, null, BundleID);
+			return GenerateTVOSPList(ProjectDirectory, bIsUnrealGame, GameName, bIsClient, ProjectName, InEngineDir, AppDirectory, null, BundleID, Logger);
 		}
 	}
 }
