UBT: Use ILogger throughout UBT, rather than legacy EpicGames.Core.Log methods. All output going forwards should use structured logging rather than String.Format style logging.

diff --git a/Engine/Source/Programs/UnrealBuildTool/Platform/Windows/MicrosoftPlatformSDK.cs b/Engine/Source/Programs/UnrealBuildTool/Platform/Windows/MicrosoftPlatformSDK.cs
--- a/Engine/Source/Programs/UnrealBuildTool/Platform/Windows/MicrosoftPlatformSDK.cs
+++ b/Engine/Source/Programs/UnrealBuildTool/Platform/Windows/MicrosoftPlatformSDK.cs
@@ -14,6 +14,7 @@ using System.Runtime.InteropServices;
 using System.Buffers.Binary;
 using UnrealBuildBase;
 using System.Runtime.Versioning;
+using Microsoft.Extensions.Logging;
 
 ///////////////////////////////////////////////////////////////////
 // If you are looking for supported version numbers, look in the
@@ -24,6 +25,11 @@ namespace UnrealBuildTool
 {
 	internal partial class MicrosoftPlatformSDK : UEBuildPlatformSDK
 	{
+		public MicrosoftPlatformSDK(ILogger Logger)
+			: base(Logger)
+		{
+		}
+
 		[SupportedOSPlatform("windows")]
 		protected override string? GetInstalledSDKVersion()
 		{
@@ -34,7 +40,7 @@ namespace UnrealBuildTool
 
 			// use the PreferredWindowsSdkVersions array to find the SDK version that UBT will use to build with - 
 			VersionNumber? WindowsVersion;
-			if (TryGetWindowsSdkDir(null, out WindowsVersion, out _))
+			if (TryGetWindowsSdkDir(null, Logger, out WindowsVersion, out _))
 			{
 				return WindowsVersion.ToString();
 			}
@@ -87,15 +93,10 @@ namespace UnrealBuildTool
 		/// </summary>
 		private static IReadOnlyDictionary<VersionNumber, DirectoryReference>? CachedUniversalCrtDirs;
 
-		/// <summary>
-		/// Visual Studio installations
-		/// </summary>
-		public static IReadOnlyList<VisualStudioInstallation> VisualStudioInstallations => CachedVisualStudioInstallations.Value;
-
 		/// <summary>
 		/// Cache of Visual Studio installation directories
 		/// </summary>
-		private static Lazy<IReadOnlyList<VisualStudioInstallation>> CachedVisualStudioInstallations = new Lazy<IReadOnlyList<VisualStudioInstallation>>(FindVisualStudioInstallations);
+		private static IReadOnlyList<VisualStudioInstallation>? CachedVisualStudioInstallations;
 
 		/// <summary>
 		/// Cache of DIA SDK installation directories
@@ -111,14 +112,15 @@ namespace UnrealBuildTool
 		/// <summary>
 		/// Finds all the installed Windows SDK versions
 		/// </summary>
+		/// <param name="Logger">Logger for output</param>
 		/// <returns>Map of version number to Windows SDK directories</returns>
 		[SupportedOSPlatform("windows")]
-		public static IReadOnlyDictionary<VersionNumber, DirectoryReference> FindWindowsSdkDirs()
+		public static IReadOnlyDictionary<VersionNumber, DirectoryReference> FindWindowsSdkDirs(ILogger Logger)
 		{
 			// Update the cache of install directories, if it's not set
 			if (CachedWindowsSdkDirs == null)
 			{
-				UpdateCachedWindowsSdks();
+				UpdateCachedWindowsSdks(Logger);
 			}
 			return CachedWindowsSdkDirs!;
 		}
@@ -126,13 +128,14 @@ namespace UnrealBuildTool
 		/// <summary>
 		/// Finds all the installed Universal CRT versions
 		/// </summary>
+		/// <param name="Logger">Logger for output</param>
 		/// <returns>Map of version number to universal CRT directories</returns>
 		[SupportedOSPlatform("windows")]
-		public static IReadOnlyDictionary<VersionNumber, DirectoryReference> FindUniversalCrtDirs()
+		public static IReadOnlyDictionary<VersionNumber, DirectoryReference> FindUniversalCrtDirs(ILogger Logger)
 		{
 			if (CachedUniversalCrtDirs == null)
 			{
-				UpdateCachedWindowsSdks();
+				UpdateCachedWindowsSdks(Logger);
 			}
 			return CachedUniversalCrtDirs!;
 		}
@@ -141,17 +144,18 @@ namespace UnrealBuildTool
 		/// Determines the directory containing the Windows SDK toolchain
 		/// </summary>
 		/// <param name="DesiredVersion">The desired Windows SDK version. This may be "Latest", a specific version number, or null. If null, the function will look for DefaultWindowsSdkVersion. Failing that, it will return the latest version.</param>
+		/// <param name="Logger">Logger for output</param>
 		/// <param name="OutSdkVersion">Receives the version number of the selected Windows SDK</param>
 		/// <param name="OutSdkDir">Receives the root directory for the selected SDK</param>
 		/// <param name="MinVersion">Optional minimum required version. Ignored if DesiredVesrion is specified</param>
 		/// <param name="MaxVersion">Optional maximum required version. Ignored if DesiredVesrion is specified</param>
 		/// <returns>True if the toolchain directory was found correctly</returns>
 		[SupportedOSPlatform("windows")]
-		public static bool TryGetWindowsSdkDir(string? DesiredVersion, [NotNullWhen(true)] out VersionNumber? OutSdkVersion, [NotNullWhen(true)] out DirectoryReference? OutSdkDir, VersionNumber? MinVersion = null, VersionNumber? MaxVersion = null)
+		public static bool TryGetWindowsSdkDir(string? DesiredVersion, ILogger Logger, [NotNullWhen(true)] out VersionNumber? OutSdkVersion, [NotNullWhen(true)] out DirectoryReference? OutSdkDir, VersionNumber? MinVersion = null, VersionNumber? MaxVersion = null)
 		{
 			// Get a map of Windows SDK versions to their root directories
 			/*IReadOnlyDictionary<VersionNumber, DirectoryReference> WindowsSdkDirs =*/
-			FindWindowsSdkDirs();
+			FindWindowsSdkDirs(Logger);
 
 			// Figure out which version number to look for
 			VersionNumber? WindowsSdkVersion = null;
@@ -206,7 +210,7 @@ namespace UnrealBuildTool
 		/// Updates the CachedWindowsSdkDirs and CachedUniversalCrtDirs variables
 		/// </summary>
 		[SupportedOSPlatform("windows")]
-		private static void UpdateCachedWindowsSdks()
+		private static void UpdateCachedWindowsSdks(ILogger Logger)
 		{
 			Dictionary<VersionNumber, DirectoryReference> WindowsSdkDirs = new Dictionary<VersionNumber, DirectoryReference>();
 			Dictionary<VersionNumber, DirectoryReference> UniversalCrtDirs = new Dictionary<VersionNumber, DirectoryReference>();
@@ -217,7 +221,7 @@ namespace UnrealBuildTool
 			{
 				if (FileReference.Exists(FileReference.Combine(InstallDir_8_1, "Include", "um", "windows.h")))
 				{
-					Log.TraceLog("Found Windows 8.1 SDK at {0}", InstallDir_8_1);
+					Logger.LogDebug("Found Windows 8.1 SDK at {InstallDir_8_1}", InstallDir_8_1);
 					VersionNumber Version_8_1 = new VersionNumber(8, 1);
 					WindowsSdkDirs[Version_8_1] = InstallDir_8_1;
 				}
@@ -225,7 +229,7 @@ namespace UnrealBuildTool
 
 			// Find all the root directories for Windows 10 SDKs
 			List<DirectoryReference> InstallDirs_10 = new List<DirectoryReference>();
-			EnumerateSdkRootDirs(InstallDirs_10);
+			EnumerateSdkRootDirs(InstallDirs_10, Logger);
 
 			// Enumerate all the Windows 10 SDKs
 			foreach (DirectoryReference InstallDir_10 in InstallDirs_10.Distinct())
@@ -240,12 +244,12 @@ namespace UnrealBuildTool
 						{
 							if (FileReference.Exists(FileReference.Combine(IncludeDir, "um", "windows.h")))
 							{
-								Log.TraceLog("Found Windows 10 SDK version {0} at {1}", IncludeVersion, InstallDir_10);
+								Logger.LogDebug("Found Windows 10 SDK version {IncludeVersion} at {InstallDir_10}", IncludeVersion, InstallDir_10);
 								WindowsSdkDirs[IncludeVersion] = InstallDir_10;
 							}
 							if (FileReference.Exists(FileReference.Combine(IncludeDir, "ucrt", "corecrt.h")))
 							{
-								Log.TraceLog("Found Universal CRT version {0} at {1}", IncludeVersion, InstallDir_10);
+								Logger.LogDebug("Found Universal CRT version {IncludeVersion} at {InstallDir_10}", IncludeVersion, InstallDir_10);
 								UniversalCrtDirs[IncludeVersion] = InstallDir_10;
 							}
 						}
@@ -261,18 +265,19 @@ namespace UnrealBuildTool
 		/// Enumerates all the Windows 10 SDK root directories
 		/// </summary>
 		/// <param name="RootDirs">Receives all the Windows 10 sdk root directories</param>
+		/// <param name="Logger">Logger for output</param>
 		[SupportedOSPlatform("windows")]
-		public static void EnumerateSdkRootDirs(List<DirectoryReference> RootDirs)
+		public static void EnumerateSdkRootDirs(List<DirectoryReference> RootDirs, ILogger Logger)
 		{
 			DirectoryReference? RootDir;
 			if (TryReadInstallDirRegistryKey32("Microsoft\\Windows Kits\\Installed Roots", "KitsRoot10", out RootDir))
 			{
-				Log.TraceLog("Found Windows 10 SDK root at {0} (1)", RootDir);
+				Logger.LogDebug("Found Windows 10 SDK root at {RootDir} (1)", RootDir);
 				RootDirs.Add(RootDir);
 			}
 			if (TryReadInstallDirRegistryKey32("Microsoft\\Microsoft SDKs\\Windows\\v10.0", "InstallationFolder", out RootDir))
 			{
-				Log.TraceLog("Found Windows 10 SDK root at {0} (2)", RootDir);
+				Logger.LogDebug("Found Windows 10 SDK root at {RootDir} (2)", RootDir);
 				RootDirs.Add(RootDir);
 			}
 
@@ -282,7 +287,7 @@ namespace UnrealBuildTool
 				DirectoryReference RootDirAutoSdk = DirectoryReference.Combine(HostAutoSdkDir, "Win64", "Windows Kits", "10");
 				if (DirectoryReference.Exists(RootDirAutoSdk))
 				{
-					Log.TraceLog("Found Windows 10 AutoSDK root at {0}", RootDirAutoSdk);
+					Logger.LogDebug("Found Windows 10 AutoSDK root at {RootDirAutoSdk}", RootDirAutoSdk);
 					RootDirs.Add(RootDirAutoSdk);
 				}
 			}
@@ -455,10 +460,10 @@ namespace UnrealBuildTool
 		/// </summary>
 		/// <returns>Path to MSBuild</returns>
 		[SupportedOSPlatform("windows")]
-		public static FileReference GetMsBuildToolPath()
+		public static FileReference GetMsBuildToolPath(ILogger Logger)
 		{
 			FileReference? Location;
-			if (!TryGetMsBuildPath(out Location))
+			if (!TryGetMsBuildPath(Logger, out Location))
 			{
 				throw new BuildException("Unable to find installation of MSBuild.");
 			}
@@ -470,17 +475,18 @@ namespace UnrealBuildTool
 		/// </summary>
 		/// <param name="Compiler">Compiler to check for</param>
 		/// <param name="Architecture">Architecture the compiler must support</param>
+		/// <param name="Logger">Logger for output</param>
 		/// <returns>True if the given compiler is installed</returns>
-		public static bool HasCompiler(WindowsCompiler Compiler, WindowsArchitecture Architecture)
+		public static bool HasCompiler(WindowsCompiler Compiler, WindowsArchitecture Architecture, ILogger Logger)
 		{
-			return FindToolChainInstallations(Compiler).Where(x => (x.Architecture == Architecture)).Count() > 0;
+			return FindToolChainInstallations(Compiler, Logger).Where(x => (x.Architecture == Architecture)).Count() > 0;
 		}
 
-		public static bool HasValidCompiler(WindowsCompiler Compiler, WindowsArchitecture Architecture)
+		public static bool HasValidCompiler(WindowsCompiler Compiler, WindowsArchitecture Architecture, ILogger Logger)
 		{
 			// since this is static, we get the instance for it
 			MicrosoftPlatformSDK SDK = (MicrosoftPlatformSDK)GetSDKForPlatform("Win64")!;
-			return FindToolChainInstallations(Compiler).Where(x => (x.Version >= MinimumVisualCppVersion && x.Architecture == Architecture)).Count() > 0;
+			return FindToolChainInstallations(Compiler, Logger).Where(x => (x.Version >= MinimumVisualCppVersion && x.Architecture == Architecture)).Count() > 0;
 		}
 
 		/// <summary>
@@ -489,14 +495,15 @@ namespace UnrealBuildTool
 		/// <param name="Compiler">Major version of the compiler to use</param>
 		/// <param name="CompilerVersion">The minimum compiler version to use</param>
 		/// <param name="Architecture">Architecture that is required</param>
+		/// <param name="Logger">Logger for output</param>
 		/// <param name="OutToolChainVersion">Receives the chosen toolchain version</param>
 		/// <param name="OutToolChainDir">Receives the directory containing the toolchain</param>
 		/// <param name="OutRedistDir">Receives the optional directory containing redistributable components</param>
 		/// <returns>True if the toolchain directory was found correctly</returns>
-		public static bool TryGetToolChainDir(WindowsCompiler Compiler, string? CompilerVersion, WindowsArchitecture Architecture, [NotNullWhen(true)] out VersionNumber? OutToolChainVersion, [NotNullWhen(true)] out DirectoryReference? OutToolChainDir, out DirectoryReference? OutRedistDir)
+		public static bool TryGetToolChainDir(WindowsCompiler Compiler, string? CompilerVersion, WindowsArchitecture Architecture, ILogger Logger, [NotNullWhen(true)] out VersionNumber? OutToolChainVersion, [NotNullWhen(true)] out DirectoryReference? OutToolChainDir, out DirectoryReference? OutRedistDir)
 		{
 			// Find all the installed toolchains
-			List<ToolChainInstallation> ToolChains = FindToolChainInstallations(Compiler);
+			List<ToolChainInstallation> ToolChains = FindToolChainInstallations(Compiler, Logger);
 
 			// Figure out the actual version number that we want
 			ToolChainInstallation? ToolChain = null;
@@ -594,7 +601,7 @@ namespace UnrealBuildTool
 			return ToolChain;
 		}
 
-		static List<ToolChainInstallation> FindToolChainInstallations(WindowsCompiler Compiler)
+		static List<ToolChainInstallation> FindToolChainInstallations(WindowsCompiler Compiler, ILogger Logger)
 		{
 			List<ToolChainInstallation>? ToolChains;
 			if (!CachedToolChainInstallations.TryGetValue(Compiler, out ToolChains))
@@ -606,25 +613,25 @@ namespace UnrealBuildTool
 					{
 						// Check for a manual installation to the default directory
 						DirectoryReference ManualInstallDir = DirectoryReference.Combine(DirectoryReference.GetSpecialFolder(Environment.SpecialFolder.ProgramFiles)!, "LLVM");
-						AddClangToolChain(ManualInstallDir, ToolChains);
+						AddClangToolChain(ManualInstallDir, ToolChains, Logger);
 
 						// Check for a manual installation to a custom directory
 						string? LlvmPath = Environment.GetEnvironmentVariable("LLVM_PATH");
 						if (!String.IsNullOrEmpty(LlvmPath))
 						{
-							AddClangToolChain(new DirectoryReference(LlvmPath), ToolChains);
+							AddClangToolChain(new DirectoryReference(LlvmPath), ToolChains, Logger);
 						}
 
 						// Check for installations bundled with Visual Studio 2019
-						foreach (VisualStudioInstallation Installation in FindVisualStudioInstallations(WindowsCompiler.VisualStudio2019))
+						foreach (VisualStudioInstallation Installation in FindVisualStudioInstallations(WindowsCompiler.VisualStudio2019, Logger))
 						{
-							AddClangToolChain(DirectoryReference.Combine(Installation.BaseDir, "VC", "Tools", "Llvm"), ToolChains);
+							AddClangToolChain(DirectoryReference.Combine(Installation.BaseDir, "VC", "Tools", "Llvm"), ToolChains, Logger);
 						}
 
 						// Check for installations bundled with Visual Studio 2022
-						foreach (VisualStudioInstallation Installation in FindVisualStudioInstallations(WindowsCompiler.VisualStudio2022))
+						foreach (VisualStudioInstallation Installation in FindVisualStudioInstallations(WindowsCompiler.VisualStudio2022, Logger))
 						{
-							AddClangToolChain(DirectoryReference.Combine(Installation.BaseDir, "VC", "Tools", "Llvm"), ToolChains);
+							AddClangToolChain(DirectoryReference.Combine(Installation.BaseDir, "VC", "Tools", "Llvm"), ToolChains, Logger);
 						}
 
 						// Check for AutoSDK paths
@@ -636,7 +643,7 @@ namespace UnrealBuildTool
 							{
 								foreach (DirectoryReference ToolChainDir in DirectoryReference.EnumerateDirectories(ClangBaseDir))
 								{
-									AddClangToolChain(ToolChainDir, ToolChains);
+									AddClangToolChain(ToolChainDir, ToolChains, Logger);
 								}
 							}
 						}
@@ -644,10 +651,10 @@ namespace UnrealBuildTool
 					else if (Compiler == WindowsCompiler.Intel)
 					{
 						DirectoryReference InstallDir = DirectoryReference.Combine(DirectoryReference.GetSpecialFolder(Environment.SpecialFolder.ProgramFiles)!, "Intel", "oneAPI", "compiler");
-						FindIntelOneApiToolChains(InstallDir, ToolChains);
+						FindIntelOneApiToolChains(InstallDir, ToolChains, Logger);
 
 						InstallDir = DirectoryReference.Combine(DirectoryReference.GetSpecialFolder(Environment.SpecialFolder.ProgramFilesX86)!, "Intel", "oneAPI", "compiler");
-						FindIntelOneApiToolChains(InstallDir, ToolChains);
+						FindIntelOneApiToolChains(InstallDir, ToolChains, Logger);
 
 						// Check for AutoSDK paths
 						DirectoryReference? AutoSdkDir;
@@ -656,19 +663,19 @@ namespace UnrealBuildTool
 							DirectoryReference IntelBaseDir = DirectoryReference.Combine(AutoSdkDir, "Win64", "Intel");
 							if (DirectoryReference.Exists(IntelBaseDir))
 							{
-								FindIntelOneApiToolChains(IntelBaseDir, ToolChains);
+								FindIntelOneApiToolChains(IntelBaseDir, ToolChains, Logger);
 							}
 						}
 					}
 					else if (Compiler.IsMSVC())
 					{
 						// Enumerate all the manually installed toolchains
-						List<VisualStudioInstallation> Installations = FindVisualStudioInstallations(Compiler);
+						List<VisualStudioInstallation> Installations = FindVisualStudioInstallations(Compiler, Logger);
 						foreach (VisualStudioInstallation Installation in Installations)
 						{
 							DirectoryReference ToolChainBaseDir = DirectoryReference.Combine(Installation.BaseDir, "VC", "Tools", "MSVC");
 							DirectoryReference RedistBaseDir = DirectoryReference.Combine(Installation.BaseDir, "VC", "Redist", "MSVC");
-							FindVisualStudioToolChains(ToolChainBaseDir, RedistBaseDir, Installation.bPreview, ToolChains);
+							FindVisualStudioToolChains(ToolChainBaseDir, RedistBaseDir, Installation.bPreview, ToolChains, Logger);
 						}
 
 						// Enumerate all the AutoSDK toolchains
@@ -685,10 +692,10 @@ namespace UnrealBuildTool
 							if (!string.IsNullOrEmpty(VSDir))
 							{
 								DirectoryReference ReleaseBaseDir = DirectoryReference.Combine(PlatformDir, "Win64", VSDir);
-								FindVisualStudioToolChains(ReleaseBaseDir, null, false, ToolChains);
+								FindVisualStudioToolChains(ReleaseBaseDir, null, false, ToolChains, Logger);
 
 								DirectoryReference PreviewBaseDir = DirectoryReference.Combine(PlatformDir, "Win64", $"{VSDir}-Preview");
-								FindVisualStudioToolChains(PreviewBaseDir, null, true, ToolChains);
+								FindVisualStudioToolChains(PreviewBaseDir, null, true, ToolChains, Logger);
 							}
 						}
 					}
@@ -707,8 +714,13 @@ namespace UnrealBuildTool
 		/// Visual Studio, and vice versa.
 		/// </summary>
 		/// <returns>List of directories containing Visual Studio installations</returns>
-		private static IReadOnlyList<VisualStudioInstallation> FindVisualStudioInstallations()
+		public static IReadOnlyList<VisualStudioInstallation> FindVisualStudioInstallations(ILogger Logger)
 		{
+			if (CachedVisualStudioInstallations != null)
+			{
+				return CachedVisualStudioInstallations;
+			}
+
 			List<VisualStudioInstallation> Installations = new List<VisualStudioInstallation>();
 			if (BuildHostPlatform.Current.Platform == UnrealTargetPlatform.Win64)
 			{
@@ -765,7 +777,7 @@ namespace UnrealBuildTool
 						DirectoryReference BaseDir = new DirectoryReference(Instance.GetInstallationPath());
 						Installations.Add(new VisualStudioInstallation(Compiler, Version, BaseDir, bExpress, bPreview));
 
-						Log.TraceLog("Found Visual Studio installation: {0} (Product={1}, Version={2})", BaseDir, ProductId, Version);
+						Logger.LogDebug("Found Visual Studio installation: {BaseDir} (Product={ProductId}, Version={Version})", BaseDir, ProductId, Version);
 					}
 
 					Installations = Installations.OrderBy(x => x.bExpress).ThenBy(x => x.bPreview).ThenByDescending(x => x.Version).ToList();
@@ -775,6 +787,8 @@ namespace UnrealBuildTool
 					throw new BuildException(Ex, "Error while enumerating Visual Studio toolchains");
 				}
 			}
+			CachedVisualStudioInstallations = Installations;
+
 			return Installations;
 		}
 
@@ -783,9 +797,9 @@ namespace UnrealBuildTool
 		/// Visual Studio, and vice versa.
 		/// </summary>
 		/// <returns>List of directories containing Visual Studio installations</returns>
-		public static List<VisualStudioInstallation> FindVisualStudioInstallations(WindowsCompiler Compiler)
+		public static List<VisualStudioInstallation> FindVisualStudioInstallations(WindowsCompiler Compiler, ILogger Logger)
 		{
-			return VisualStudioInstallations.Where(x => x.Compiler == Compiler).ToList();
+			return FindVisualStudioInstallations(Logger).Where(x => x.Compiler == Compiler).ToList();
 		}
 
 		/// <summary>
@@ -795,7 +809,8 @@ namespace UnrealBuildTool
 		/// <param name="OptionalRedistDir">Optional directory for redistributable components (DLLs etc)</param>
 		/// <param name="bPreview">Whether this is a preview installation</param>
 		/// <param name="ToolChains">Map of tool chain version to installation info</param>
-		static void FindVisualStudioToolChains(DirectoryReference BaseDir, DirectoryReference? OptionalRedistDir, bool bPreview, List<ToolChainInstallation> ToolChains)
+		/// <param name="Logger">Logger for output</param>
+		static void FindVisualStudioToolChains(DirectoryReference BaseDir, DirectoryReference? OptionalRedistDir, bool bPreview, List<ToolChainInstallation> ToolChains, ILogger Logger)
 		{
 			if (DirectoryReference.Exists(BaseDir))
 			{
@@ -806,7 +821,7 @@ namespace UnrealBuildTool
 					{
 						DirectoryReference? RedistDir = FindVisualStudioRedistForToolChain(ToolChainDir, OptionalRedistDir, Version);
 
-						AddVisualCppToolChain(Version, bPreview, ToolChainDir, RedistDir, ToolChains);
+						AddVisualCppToolChain(Version, bPreview, ToolChainDir, RedistDir, ToolChains, Logger);
 					}
 				}
 			}
@@ -857,7 +872,8 @@ namespace UnrealBuildTool
 		/// <param name="ToolChainDir"></param>
 		/// <param name="RedistDir"></param>
 		/// <param name="ToolChains"></param>
-		static void AddVisualCppToolChain(VersionNumber Version, bool bPreview, DirectoryReference ToolChainDir, DirectoryReference? RedistDir, List<ToolChainInstallation> ToolChains)
+		/// <param name="Logger"></param>
+		static void AddVisualCppToolChain(VersionNumber Version, bool bPreview, DirectoryReference ToolChainDir, DirectoryReference? RedistDir, List<ToolChainInstallation> ToolChains, ILogger Logger)
 		{
 			bool Is64Bit = Has64BitToolChain(ToolChainDir);
 
@@ -881,13 +897,13 @@ namespace UnrealBuildTool
 				Error = $"UnrealBuildTool has banned the MSVC {Banned} toolchains due to compiler issues. Please install a different toolchain such as {PreferredVisualCppVersions.Select(x => x.Min).Max()} from the Visual Studio installer.";
 			}
 
-			Log.TraceLog("Found Visual Studio toolchain: {0} (Family={1}, FamilyRank={2}, Version={3}, Is64Bit={4}, Preview={5}, Architecture={6}, Error={7}, Redist={8})", ToolChainDir, Family, FamilyRank, Version, Is64Bit, bPreview, WindowsArchitecture.x64.ToString(), Error != null, RedistDir);
+			Logger.LogDebug("Found Visual Studio toolchain: {ToolChainDir} (Family={Family}, FamilyRank={FamilyRank}, Version={Version}, Is64Bit={Is64Bit}, Preview={Preview}, Architecture={Arch}, Error={Error}, Redist={RedistDir})", ToolChainDir, Family, FamilyRank, Version, Is64Bit, bPreview, WindowsArchitecture.x64.ToString(), Error != null, RedistDir);
 			ToolChains.Add(new ToolChainInstallation(Family, FamilyRank, Version, Is64Bit, bPreview, WindowsArchitecture.x64, Error, ToolChainDir, RedistDir));
 
 			bool HasArm64 = HasArm64ToolChain(ToolChainDir);
 			if (HasArm64)
 			{
-				Log.TraceLog("Found Visual Studio toolchain: {0} (Family={1}, FamilyRank={2}, Version={3}, Is64Bit={4}, Preview={5}, Architecture={6}, Error={7}, Redist={8})", ToolChainDir, Family, FamilyRank, Version, Is64Bit, bPreview, WindowsArchitecture.ARM64.ToString(), Error != null, RedistDir);
+				Logger.LogDebug("Found Visual Studio toolchain: {ToolChainDir} (Family={Family}, FamilyRank={FamilyRank}, Version={Version}, Is64Bit={Is64Bit}, Preview={Preview}, Architecture={Arch}, Error={Error}, Redist={RedistDir})", ToolChainDir, Family, FamilyRank, Version, Is64Bit, bPreview, WindowsArchitecture.ARM64.ToString(), Error != null, RedistDir);
 				ToolChains.Add(new ToolChainInstallation(Family, FamilyRank, Version, Is64Bit, bPreview, WindowsArchitecture.ARM64, Error, ToolChainDir, RedistDir));
 			}
 		}
@@ -897,7 +913,8 @@ namespace UnrealBuildTool
 		/// </summary>
 		/// <param name="ToolChainDir"></param>
 		/// <param name="ToolChains"></param>
-		static void AddClangToolChain(DirectoryReference ToolChainDir, List<ToolChainInstallation> ToolChains)
+		/// <param name="Logger"></param>
+		static void AddClangToolChain(DirectoryReference ToolChainDir, List<ToolChainInstallation> ToolChains, ILogger Logger)
 		{
 			FileReference CompilerFile = FileReference.Combine(ToolChainDir, "bin", "clang-cl.exe");
 			if (FileReference.Exists(CompilerFile))
@@ -914,7 +931,7 @@ namespace UnrealBuildTool
 					Error = $"UnrealBuildTool requires at minimum the Clang {MinimumClangVersion} toolchain. Please install a later toolchain such as {PreferredClangVersions.Select(x => x.Min).Max()} from LLVM.";
 				}
 
-				Log.TraceLog("Found Clang toolchain: {0} (Version={1}, Is64Bit={2}, Rank={3})", ToolChainDir, Version, Is64Bit, Rank);
+				Logger.LogDebug("Found Clang toolchain: {ToolChainDir} (Version={Version}, Is64Bit={Is64Bit}, Rank={Rank})", ToolChainDir, Version, Is64Bit, Rank);
 				ToolChains.Add(new ToolChainInstallation(Version, Rank, Version, Is64Bit, false, WindowsArchitecture.x64, null, ToolChainDir, null));
 			}
 		}
@@ -925,7 +942,8 @@ namespace UnrealBuildTool
 		/// </summary>
 		/// <param name="ToolChainDir"></param>
 		/// <param name="ToolChains"></param>
-		static void AddIntelOneApiToolChain(DirectoryReference ToolChainDir, List<ToolChainInstallation> ToolChains)
+		/// <param name="Logger"></param>
+		static void AddIntelOneApiToolChain(DirectoryReference ToolChainDir, List<ToolChainInstallation> ToolChains, ILogger Logger)
 		{
 			FileReference CompilerFile = FileReference.Combine(ToolChainDir, "windows", "bin", "icx.exe");
 			if (FileReference.Exists(CompilerFile))
@@ -948,7 +966,7 @@ namespace UnrealBuildTool
 					Error = $"UnrealBuildTool requires at minimum the Intel OneAPI {MinimumIntelOneApiVersion} toolchain. Please install a later toolchain such as {PreferredIntelOneApiVersions.Select(x => x.Min).Max()} from Intel.";
 				}
 
-				Log.TraceLog("Found Intel OneAPI toolchain: {0} (Version={1}, Is64Bit={2}, Rank={3})", ToolChainDir, Version, Is64Bit, Rank);
+				Logger.LogDebug("Found Intel OneAPI toolchain: {ToolChainDir} (Version={Version}, Is64Bit={Is64Bit}, Rank={Rank})", ToolChainDir, Version, Is64Bit, Rank);
 				ToolChains.Add(new ToolChainInstallation(Version, Rank, Version, Is64Bit, false, WindowsArchitecture.x64, Error, ToolChainDir, null));
 			}
 		}
@@ -958,13 +976,14 @@ namespace UnrealBuildTool
 		/// </summary>
 		/// <param name="BaseDir">Base directory to search</param>
 		/// <param name="ToolChains">Map of tool chain version to installation info</param>
-		static void FindIntelOneApiToolChains(DirectoryReference BaseDir, List<ToolChainInstallation> ToolChains)
+		/// <param name="Logger">Logger for output</param>
+		static void FindIntelOneApiToolChains(DirectoryReference BaseDir, List<ToolChainInstallation> ToolChains, ILogger Logger)
 		{
 			if (DirectoryReference.Exists(BaseDir))
 			{
 				foreach (DirectoryReference ToolChainDir in DirectoryReference.EnumerateDirectories(BaseDir))
 				{
-					AddIntelOneApiToolChain(ToolChainDir, ToolChains);
+					AddIntelOneApiToolChain(ToolChainDir, ToolChains, Logger);
 				}
 			}
 		}
@@ -1076,22 +1095,24 @@ namespace UnrealBuildTool
 		/// Determines if an IDE for the given compiler is installed.
 		/// </summary>
 		/// <param name="Compiler">Compiler to check for</param>
+		/// <param name="Logger">Logger for output</param>
 		/// <returns>True if the given compiler is installed</returns>
-		public static bool HasIDE(WindowsCompiler Compiler)
+		public static bool HasIDE(WindowsCompiler Compiler, ILogger Logger)
 		{
-			return FindVisualStudioInstallations(Compiler).Count > 0;
+			return FindVisualStudioInstallations(Compiler, Logger).Count > 0;
 		}
 
 		/// <summary>
 		/// Gets the path to MSBuild. This mirrors the logic in GetMSBuildPath.bat.
 		/// </summary>
+		/// <param name="Logger">Logger for output</param>
 		/// <param name="OutLocation">On success, receives the path to the MSBuild executable.</param>
 		/// <returns>True on success.</returns>
 		[SupportedOSPlatform("windows")]
-		public static bool TryGetMsBuildPath([NotNullWhen(true)] out FileReference? OutLocation)
+		public static bool TryGetMsBuildPath(ILogger Logger, [NotNullWhen(true)] out FileReference? OutLocation)
 		{
 			// Get the Visual Studio 2019 install directory
-			List<DirectoryReference> InstallDirs2019 = MicrosoftPlatformSDK.FindVisualStudioInstallations(WindowsCompiler.VisualStudio2019).ConvertAll(x => x.BaseDir);
+			List<DirectoryReference> InstallDirs2019 = MicrosoftPlatformSDK.FindVisualStudioInstallations(WindowsCompiler.VisualStudio2019, Logger).ConvertAll(x => x.BaseDir);
 			foreach (DirectoryReference InstallDir in InstallDirs2019)
 			{
 				FileReference MsBuildLocation = FileReference.Combine(InstallDir, "MSBuild", "Current", "Bin", "MSBuild.exe");
@@ -1103,7 +1124,7 @@ namespace UnrealBuildTool
 			}
 
 			// Get the Visual Studio 2022 install directory
-			List<DirectoryReference> InstallDirs2022 = MicrosoftPlatformSDK.FindVisualStudioInstallations(WindowsCompiler.VisualStudio2022).ConvertAll(x => x.BaseDir);
+			List<DirectoryReference> InstallDirs2022 = MicrosoftPlatformSDK.FindVisualStudioInstallations(WindowsCompiler.VisualStudio2022, Logger).ConvertAll(x => x.BaseDir);
 			foreach (DirectoryReference InstallDir in InstallDirs2022)
 			{
 				FileReference MsBuildLocation = FileReference.Combine(InstallDir, "MSBuild", "Current", "Bin", "MSBuild.exe");
@@ -1215,7 +1236,7 @@ namespace UnrealBuildTool
 					}
 				}
 
-				List<DirectoryReference> VisualStudioDirs = MicrosoftPlatformSDK.FindVisualStudioInstallations(Compiler).ConvertAll(x => x.BaseDir);
+				List<DirectoryReference> VisualStudioDirs = MicrosoftPlatformSDK.FindVisualStudioInstallations(Compiler, Log.Logger).ConvertAll(x => x.BaseDir);
 				foreach (DirectoryReference VisualStudioDir in VisualStudioDirs)
 				{
 					DirectoryReference DiaSdkDir = DirectoryReference.Combine(VisualStudioDir, "DIA SDK");
diff --git a/Engine/Source/Programs/UnrealBuildTool/Platform/Windows/PVSToolChain.cs b/Engine/Source/Programs/UnrealBuildTool/Platform/Windows/PVSToolChain.cs
--- a/Engine/Source/Programs/UnrealBuildTool/Platform/Windows/PVSToolChain.cs
+++ b/Engine/Source/Programs/UnrealBuildTool/Platform/Windows/PVSToolChain.cs
@@ -12,6 +12,7 @@ using System.Threading.Tasks;
 using System.Xml;
 using System.Xml.Serialization;
 using EpicGames.Core;
+using Microsoft.Extensions.Logging;
 using UnrealBuildBase;
 
 namespace UnrealBuildTool
@@ -263,12 +264,13 @@ namespace UnrealBuildTool
 		/// </summary>
 		/// <param name="Arguments">List of command line arguments</param>
 		/// <returns>Always zero, or throws an exception</returns>
-		public override int Execute(CommandLineArguments Arguments)
+		/// <param name="Logger"></param>
+		public override int Execute(CommandLineArguments Arguments, ILogger Logger)
 		{
 			Arguments.ApplyTo(this);
 			Arguments.CheckAllArgumentsUsed();
 
-			Log.TraceInformation("{0}", OutputFile!.GetFileName());
+			Logger.LogInformation("{File}", OutputFile!.GetFileName());
 
 			// Read the input files
 			string[] InputFileLines = FileReference.ReadAllLines(InputFileList!);
@@ -331,7 +333,7 @@ namespace UnrealBuildTool
 					}
 				}
 			}
-			Log.TraceInformation("Written {0} {1} to {2}.", UniqueItems.Count, (UniqueItems.Count == 1)? "diagnostic" : "diagnostics", OutputFile.FullName);
+			Logger.LogInformation("Written {NumItems} {Noun} to {File}.", UniqueItems.Count, (UniqueItems.Count == 1)? "diagnostic" : "diagnostics", OutputFile.FullName);
 			return 0;
 		}
 	}
@@ -347,11 +349,12 @@ namespace UnrealBuildTool
 		UnrealTargetPlatform Platform;
 		Version AnalyzerVersion;
 
-		public PVSToolChain(ReadOnlyTargetRules Target)
+		public PVSToolChain(ReadOnlyTargetRules Target, ILogger Logger)
+			: base(Logger)
 		{
 			this.Target = Target;
 			Platform = Target.Platform;
-			InnerToolChain = new VCToolChain(Target);
+			InnerToolChain = new VCToolChain(Target, Logger);
 
 			AnalyzerFile = FileReference.Combine(Unreal.RootDirectory, "Engine", "Restricted", "NoRedist", "Extras", "ThirdPartyNotUE", "PVS-Studio", "PVS-Studio.exe");
 			if (!FileReference.Exists(AnalyzerFile))
@@ -381,7 +384,7 @@ namespace UnrealBuildTool
 				if (!String.IsNullOrEmpty(ApplicationSettings.UserName) && !String.IsNullOrEmpty(ApplicationSettings.SerialNumber))
 				{
 					LicenseFile = FileReference.Combine(Unreal.EngineDirectory, "Intermediate", "PVS", "PVS-Studio.lic");
-					Utils.WriteFileIfChanged(LicenseFile, String.Format("{0}\n{1}\n", ApplicationSettings.UserName, ApplicationSettings.SerialNumber));
+					Utils.WriteFileIfChanged(LicenseFile, String.Format("{0}\n{1}\n", ApplicationSettings.UserName, ApplicationSettings.SerialNumber), Logger);
 				}
 			}
 			else
@@ -518,14 +521,14 @@ namespace UnrealBuildTool
 				FileItem? SourceFileItem = PreprocessAction.SourceFile;
 				if (SourceFileItem == null)
 				{
-					Log.TraceWarning("Unable to find source file from command producing: {0}", String.Join(", ", PreprocessActions[Idx].ProducedItems.Select(x => x.Location.GetFileName())));
+					Logger.LogWarning("Unable to find source file from command producing: {File}", String.Join(", ", PreprocessActions[Idx].ProducedItems.Select(x => x.Location.GetFileName())));
 					continue;
 				}
 
 				FileItem? PreprocessedFileItem = PreprocessAction.PreprocessedFile;
 				if (PreprocessedFileItem == null)
 				{
-					Log.TraceWarning("Unable to find preprocessed output file from {0}", SourceFileItem.Location.GetFileName());
+					Logger.LogWarning("Unable to find preprocessed output file from {File}", SourceFileItem.Location.GetFileName());
 					continue;
 				}
 
@@ -613,7 +616,7 @@ namespace UnrealBuildTool
 			throw new BuildException("Unable to link with PVS toolchain.");
 		}
 
-		public override void FinalizeOutput(ReadOnlyTargetRules Target, TargetMakefile Makefile)
+		public override void FinalizeOutput(ReadOnlyTargetRules Target, TargetMakefileBuilder MakefileBuilder)
 		{
 			FileReference OutputFile;
 			if (Target.ProjectFile == null)
@@ -625,14 +628,15 @@ namespace UnrealBuildTool
 				OutputFile = FileReference.Combine(Target.ProjectFile.Directory, "Saved", "PVS-Studio", String.Format("{0}.pvslog", Target.Name));
 			}
 
+			TargetMakefile Makefile = MakefileBuilder.Makefile;
 			List<FileReference> InputFiles = Makefile.OutputItems.Select(x => x.Location).Where(x => x.HasExtension(".pvslog")).ToList();
 
 			// Collect the sourcefile items off of the Compile action added in CompileCPPFiles so that in SingleFileCompile mode the PVSGather step is also not filtered out
 			List<FileItem> CompileSourceFiles = Makefile.Actions.OfType<VCCompileAction>().Select(x => x.SourceFile!).ToList();
 
-			FileItem InputFileListItem = Makefile.CreateIntermediateTextFile(OutputFile.ChangeExtension(".input"), InputFiles.Select(x => x.FullName));
+			FileItem InputFileListItem = MakefileBuilder.CreateIntermediateTextFile(OutputFile.ChangeExtension(".input"), InputFiles.Select(x => x.FullName));
 
-			Action AnalyzeAction = Makefile.CreateAction(ActionType.Compile);
+			Action AnalyzeAction = MakefileBuilder.CreateAction(ActionType.Compile);
 			AnalyzeAction.ActionType = ActionType.PostBuildStep;
 			AnalyzeAction.CommandDescription = "Process PVS-Studio Results";
 			AnalyzeAction.CommandPath = Unreal.DotnetPath;
diff --git a/Engine/Source/Programs/UnrealBuildTool/Platform/Windows/UEBuildWindows.cs b/Engine/Source/Programs/UnrealBuildTool/Platform/Windows/UEBuildWindows.cs
--- a/Engine/Source/Programs/UnrealBuildTool/Platform/Windows/UEBuildWindows.cs
+++ b/Engine/Source/Programs/UnrealBuildTool/Platform/Windows/UEBuildWindows.cs
@@ -10,6 +10,7 @@
 using System.Diagnostics.CodeAnalysis;
 using UnrealBuildBase;
 using System.Runtime.Versioning;
+using Microsoft.Extensions.Logging;
 
 namespace UnrealBuildTool
 {
@@ -427,7 +428,7 @@
 			{
 				try
 				{
-					return MicrosoftPlatformSDK.FindVisualStudioInstallations(Environment!.Compiler).Select(x => x.BaseDir.FullName).FirstOrDefault();
+					return MicrosoftPlatformSDK.FindVisualStudioInstallations(Environment!.Compiler, Target.Logger).Select(x => x.BaseDir.FullName).FirstOrDefault();
 				}
 				catch(Exception) // Find function will throw if there is no visual studio installed! This can happen w/ clang builds
 				{
@@ -765,8 +766,9 @@
 		/// </summary>
 		/// <param name="InPlatform">Creates a windows platform with the given enum value</param>
 		/// <param name="InSDK">The installed Windows SDK</param>
-		public WindowsPlatform(UnrealTargetPlatform InPlatform, MicrosoftPlatformSDK InSDK)
-			: base(InPlatform, InSDK)
+		/// <param name="InLogger">Logger instance</param>
+		public WindowsPlatform(UnrealTargetPlatform InPlatform, MicrosoftPlatformSDK InSDK, ILogger InLogger)
+			: base(InPlatform, InSDK, InLogger)
 		{
 			SDK = InSDK;
 		}
@@ -788,7 +790,7 @@
 		[SupportedOSPlatform("windows")]
 		protected virtual VCEnvironment CreateVCEnvironment(TargetRules Target)
 		{
-			return VCEnvironment.Create(Target.WindowsPlatform.Compiler, Platform, Target.WindowsPlatform.Architecture, Target.WindowsPlatform.CompilerVersion, Target.WindowsPlatform.WindowsSdkVersion, null);
+			return VCEnvironment.Create(Target.WindowsPlatform.Compiler, Platform, Target.WindowsPlatform.Architecture, Target.WindowsPlatform.CompilerVersion, Target.WindowsPlatform.WindowsSdkVersion, null, Logger);
 		}
 
 		/// <summary>
@@ -822,23 +824,23 @@
 			else if (Target.WindowsPlatform.StaticAnalyzer != WindowsStaticAnalyzer.None && 
 			         Target.WindowsPlatform.StaticAnalyzerOutputType != WindowsStaticAnalyzerOutputType.Text)
 			{
-				Log.TraceInformation("Defaulting static analyzer output type to text");
+				Logger.LogInformation("Defaulting static analyzer output type to text");
 			}
 
 			// Set the compiler version if necessary
 			if (Target.WindowsPlatform.Compiler == WindowsCompiler.Default)
 			{
-				if (Target.WindowsPlatform.StaticAnalyzer == WindowsStaticAnalyzer.PVSStudio && HasCompiler(WindowsCompiler.VisualStudio2019, Target.WindowsPlatform.Architecture))
+				if (Target.WindowsPlatform.StaticAnalyzer == WindowsStaticAnalyzer.PVSStudio && HasCompiler(WindowsCompiler.VisualStudio2019, Target.WindowsPlatform.Architecture, Logger))
 				{
 					Target.WindowsPlatform.Compiler = WindowsCompiler.VisualStudio2019;
 				}
-				if (Target.WindowsPlatform.StaticAnalyzer == WindowsStaticAnalyzer.PVSStudio && HasCompiler(WindowsCompiler.VisualStudio2022, Target.WindowsPlatform.Architecture))
+				if (Target.WindowsPlatform.StaticAnalyzer == WindowsStaticAnalyzer.PVSStudio && HasCompiler(WindowsCompiler.VisualStudio2022, Target.WindowsPlatform.Architecture, Logger))
 				{
 					Target.WindowsPlatform.Compiler = WindowsCompiler.VisualStudio2022;
 				}
 				else
 				{
-					Target.WindowsPlatform.Compiler = GetDefaultCompiler(Target.ProjectFile, Target.WindowsPlatform.Architecture);
+					Target.WindowsPlatform.Compiler = GetDefaultCompiler(Target.ProjectFile, Target.WindowsPlatform.Architecture, Logger);
 				}
 			}
 
@@ -920,12 +922,12 @@
 		/// Gets the default compiler which should be used, if it's not set explicitly by the target, command line, or config file.
 		/// </summary>
 		/// <returns>The default compiler version</returns>
-		internal static WindowsCompiler GetDefaultCompiler(FileReference? ProjectFile, WindowsArchitecture Architecture)
+		internal static WindowsCompiler GetDefaultCompiler(FileReference? ProjectFile, WindowsArchitecture Architecture, ILogger Logger)
 		{
 			// If there's no specific compiler set, try to pick the matching compiler for the selected IDE
 			if (ProjectFileGeneratorSettings.Format != null)
 			{
-				foreach(ProjectFileFormat Format in ProjectFileGeneratorSettings.ParseFormatList(ProjectFileGeneratorSettings.Format))
+				foreach(ProjectFileFormat Format in ProjectFileGeneratorSettings.ParseFormatList(ProjectFileGeneratorSettings.Format, Logger))
 				{
 					if (Format == ProjectFileFormat.VisualStudio2019)
 					{
@@ -967,33 +969,33 @@
 			}
 
 			// Second, default based on what's installed, test for 2019 first
-			if (MicrosoftPlatformSDK.HasValidCompiler(WindowsCompiler.VisualStudio2019, Architecture))
+			if (MicrosoftPlatformSDK.HasValidCompiler(WindowsCompiler.VisualStudio2019, Architecture, Logger))
 			{
 				return WindowsCompiler.VisualStudio2019;
 			}
-			else if (MicrosoftPlatformSDK.HasValidCompiler(WindowsCompiler.VisualStudio2022, Architecture))
+			else if (MicrosoftPlatformSDK.HasValidCompiler(WindowsCompiler.VisualStudio2022, Architecture, Logger))
 			{
 				return WindowsCompiler.VisualStudio2022;
 			}
 
 			// If we do have a Visual Studio installation, but we're missing just the C++ parts, warn about that.
-			if (TryGetVSInstallDirs(WindowsCompiler.VisualStudio2019) != null)
+			if (TryGetVSInstallDirs(WindowsCompiler.VisualStudio2019, Logger) != null)
 			{
 				string ToolSetWarning = Architecture == WindowsArchitecture.x64 ?
 					"MSVC v142 - VS 2019 C++ x64/x86 build tools (Latest)" :
 					"MSVC v142 - VS 2019 C++ ARM64 build tools (Latest)";
-				Log.TraceWarning("Visual Studio 2019 is installed, but is missing the C++ toolchain. Please verify that the \"{0}\" component is selected in the Visual Studio 2019 installation options.", ToolSetWarning);
+				Logger.LogWarning("Visual Studio 2019 is installed, but is missing the C++ toolchain. Please verify that the \"{Component}\" component is selected in the Visual Studio 2019 installation options.", ToolSetWarning);
 			}
-			else if (TryGetVSInstallDirs(WindowsCompiler.VisualStudio2022) != null)
+			else if (TryGetVSInstallDirs(WindowsCompiler.VisualStudio2022, Logger) != null)
 			{
 				string ToolSetWarning = Architecture == WindowsArchitecture.x64 ?
 					"MSVC v143 - VS 2022 C++ x64/x86 build tools (Latest)" :
 					"MSVC v143 - VS 2022 C++ ARM64 build tools (Latest)";
-				Log.TraceWarning("Visual Studio 2022 is installed, but is missing the C++ toolchain. Please verify that the \"{0}\" component is selected in the Visual Studio 2022 installation options.", ToolSetWarning);
+				Logger.LogWarning("Visual Studio 2022 is installed, but is missing the C++ toolchain. Please verify that the \"{Component}\" component is selected in the Visual Studio 2022 installation options.", ToolSetWarning);
 			}
 			else
 			{
-				Log.TraceWarning("No Visual C++ installation was found. Please download and install Visual Studio 2019 or 2022 with C++ components.");
+				Logger.LogWarning("No Visual C++ installation was found. Please download and install Visual Studio 2019 or 2022 with C++ components.");
 			}
 
 			// Finally, default to VS2019 anyway
@@ -1015,10 +1017,11 @@
 		/// Visual Studio.
 		/// </summary>
 		/// <param name="Compiler">Version of the toolchain to look for.</param>
+		/// <param name="Logger">Logger for output</param>
 		/// <returns>True if the directory was found, false otherwise.</returns>
-		public static IEnumerable<DirectoryReference>? TryGetVSInstallDirs(WindowsCompiler Compiler)
+		public static IEnumerable<DirectoryReference>? TryGetVSInstallDirs(WindowsCompiler Compiler, ILogger Logger)
 		{
-			List<VisualStudioInstallation> Installations = MicrosoftPlatformSDK.FindVisualStudioInstallations(Compiler);
+			List<VisualStudioInstallation> Installations = MicrosoftPlatformSDK.FindVisualStudioInstallations(Compiler, Logger);
 			if(Installations.Count == 0)
 			{
 				return null;
@@ -1032,10 +1035,11 @@
 		/// </summary>
 		/// <param name="Compiler">Compiler to check for</param>
 		/// <param name="Architecture">Architecture the compiler must support</param>
+		/// <param name="Logger">Logger for output</param>
 		/// <returns>True if the given compiler is installed</returns>
-		public static bool HasCompiler(WindowsCompiler Compiler, WindowsArchitecture Architecture)
+		public static bool HasCompiler(WindowsCompiler Compiler, WindowsArchitecture Architecture, ILogger Logger)
 		{
-			return MicrosoftPlatformSDK.HasCompiler(Compiler, Architecture);
+			return MicrosoftPlatformSDK.HasCompiler(Compiler, Architecture, Logger);
 		}
 
 		/// <summary>
@@ -1044,13 +1048,14 @@
 		/// <param name="Compiler">Major version of the compiler to use</param>
 		/// <param name="CompilerVersion">The minimum compiler version to use</param>
 		/// <param name="Architecture">Architecture that is required</param>
+		/// <param name="Logger">Logger for output</param>
 		/// <param name="OutToolChainVersion">Receives the chosen toolchain version</param>
 		/// <param name="OutToolChainDir">Receives the directory containing the toolchain</param>
 		/// <param name="OutRedistDir">Receives the optional directory containing redistributable components</param>
 		/// <returns>True if the toolchain directory was found correctly</returns>
-		public static bool TryGetToolChainDir(WindowsCompiler Compiler, string? CompilerVersion, WindowsArchitecture Architecture, [NotNullWhen(true)] out VersionNumber? OutToolChainVersion, [NotNullWhen(true)] out DirectoryReference? OutToolChainDir, out DirectoryReference? OutRedistDir)
+		public static bool TryGetToolChainDir(WindowsCompiler Compiler, string? CompilerVersion, WindowsArchitecture Architecture, ILogger Logger, [NotNullWhen(true)] out VersionNumber? OutToolChainVersion, [NotNullWhen(true)] out DirectoryReference? OutToolChainDir, out DirectoryReference? OutRedistDir)
 		{
-			return MicrosoftPlatformSDK.TryGetToolChainDir(Compiler, CompilerVersion, Architecture, out OutToolChainVersion, out OutToolChainDir, out OutRedistDir);
+			return MicrosoftPlatformSDK.TryGetToolChainDir(Compiler, CompilerVersion, Architecture, Logger, out OutToolChainVersion, out OutToolChainDir, out OutRedistDir);
 		}
 
 		public static string GetArchitectureSubpath(WindowsArchitecture arch)
@@ -1081,9 +1086,9 @@
 
 
 		[SupportedOSPlatform("windows")]
-		public static bool TryGetWindowsSdkDir(string? DesiredVersion, [NotNullWhen(true)] out VersionNumber? OutSdkVersion, [NotNullWhen(true)] out DirectoryReference? OutSdkDir)
+		public static bool TryGetWindowsSdkDir(string? DesiredVersion, ILogger Logger, [NotNullWhen(true)] out VersionNumber? OutSdkVersion, [NotNullWhen(true)] out DirectoryReference? OutSdkDir)
 		{
-			return MicrosoftPlatformSDK.TryGetWindowsSdkDir(DesiredVersion, out OutSdkVersion, out OutSdkDir);
+			return MicrosoftPlatformSDK.TryGetWindowsSdkDir(DesiredVersion, Logger, out OutSdkVersion, out OutSdkDir);
 		}
 
 
@@ -1455,11 +1460,11 @@
 		{
 			if (Target.WindowsPlatform.StaticAnalyzer == WindowsStaticAnalyzer.PVSStudio)
 			{
-				return new PVSToolChain(Target);
+				return new PVSToolChain(Target, Logger);
 			}
 			else
 			{
-				return new VCToolChain(Target);
+				return new VCToolChain(Target, Logger);
 			}
 		}
 
@@ -1484,12 +1489,17 @@
 		/// <param name="Receipt">Receipt for the target being deployed</param>
 		public override void Deploy(TargetReceipt Receipt)
 		{
-			new UEDeployWindows().PrepTargetForDeployment(Receipt);
+			new UEDeployWindows(Logger).PrepTargetForDeployment(Receipt);
 		}
 	}
 
 	class UEDeployWindows : UEBuildDeploy
 	{
+		public UEDeployWindows(ILogger InLogger)
+			: base(InLogger)
+		{
+		}
+
 		public override bool PrepTargetForDeployment(TargetReceipt Receipt)
 		{
 			return base.PrepTargetForDeployment(Receipt);
@@ -1506,12 +1516,12 @@
 		/// <summary>
 		/// Register the platform with the UEBuildPlatform class
 		/// </summary>
-		public override void RegisterBuildPlatforms()
+		public override void RegisterBuildPlatforms(ILogger Logger)
 		{
-			MicrosoftPlatformSDK SDK = new MicrosoftPlatformSDK();
+			MicrosoftPlatformSDK SDK = new MicrosoftPlatformSDK(Logger);
 
 			// Register this build platform for Win64 (no more Win32)
-			UEBuildPlatform.RegisterBuildPlatform(new WindowsPlatform(UnrealTargetPlatform.Win64, SDK));
+			UEBuildPlatform.RegisterBuildPlatform(new WindowsPlatform(UnrealTargetPlatform.Win64, SDK, Logger), Logger);
 			UEBuildPlatform.RegisterPlatformWithGroup(UnrealTargetPlatform.Win64, UnrealPlatformGroup.Windows);
 			UEBuildPlatform.RegisterPlatformWithGroup(UnrealTargetPlatform.Win64, UnrealPlatformGroup.Microsoft);
 			UEBuildPlatform.RegisterPlatformWithGroup(UnrealTargetPlatform.Win64, UnrealPlatformGroup.Desktop);
diff --git a/Engine/Source/Programs/UnrealBuildTool/Platform/Windows/VCCompileAction.cs b/Engine/Source/Programs/UnrealBuildTool/Platform/Windows/VCCompileAction.cs
--- a/Engine/Source/Programs/UnrealBuildTool/Platform/Windows/VCCompileAction.cs
+++ b/Engine/Source/Programs/UnrealBuildTool/Platform/Windows/VCCompileAction.cs
@@ -4,6 +4,7 @@ using System;
 using System.Collections.Generic;
 using EpicGames.Core;
 using UnrealBuildBase;
+using Microsoft.Extensions.Logging;
 
 namespace UnrealBuildTool
 {
@@ -383,15 +384,16 @@ namespace UnrealBuildTool
 		/// Writes the response file with the action's arguments
 		/// </summary>
 		/// <param name="Graph">The graph builder</param>
-		public void WriteResponseFile(IActionGraphBuilder Graph)
+		/// <param name="Logger">Logger for output</param>
+		public void WriteResponseFile(IActionGraphBuilder Graph, ILogger Logger)
 		{
 			if (ResponseFile != null)
 			{
-				Graph.CreateIntermediateTextFile(ResponseFile, GetCompilerArguments());
+				Graph.CreateIntermediateTextFile(ResponseFile, GetCompilerArguments(Logger));
 			}
 		}
 
-		public List<string> GetCompilerArguments()
+		public List<string> GetCompilerArguments(ILogger Logger)
 		{
 			List<string> Arguments = new List<string>();
 			if (SourceFile != null)
@@ -433,10 +435,10 @@ namespace UnrealBuildTool
 
 			if (PreprocessedFile != null)
 			{
-				VCToolChain.AddPreprocessedFile(Arguments, PreprocessedFile);
+				VCToolChain.AddPreprocessedFile(Arguments, PreprocessedFile, Logger);
 
 				// this is parsed by external tools wishing to open this file directly.
-				Log.TraceInformation("PreProcessPath: " + PreprocessedFile);
+				Logger.LogInformation("PreProcessPath: {File}", PreprocessedFile);
 			}
 
 			if (ObjectFile != null)
diff --git a/Engine/Source/Programs/UnrealBuildTool/Platform/Windows/VCEnvironment.cs b/Engine/Source/Programs/UnrealBuildTool/Platform/Windows/VCEnvironment.cs
--- a/Engine/Source/Programs/UnrealBuildTool/Platform/Windows/VCEnvironment.cs
+++ b/Engine/Source/Programs/UnrealBuildTool/Platform/Windows/VCEnvironment.cs
@@ -11,6 +11,7 @@
 using EpicGames.Core;
 using UnrealBuildBase;
 using System.Runtime.Versioning;
+using Microsoft.Extensions.Logging;
 
 namespace UnrealBuildTool
 {
@@ -103,8 +104,9 @@
 		/// Constructor
 		/// </summary>
 		/// <param name="Params">Main constructor parameters</param>
+		/// <param name="Logger">Logger for output</param>
 		[SupportedOSPlatform("windows")]
-		public VCEnvironment(VCEnvironmentParameters Params)
+		public VCEnvironment(VCEnvironmentParameters Params, ILogger Logger)
 		{
 			this.Compiler = Params.Compiler;
 			this.CompilerDir = Params.CompilerDir;
@@ -132,7 +134,7 @@
 			ResourceCompilerPath = GetResourceCompilerToolPath(Params.Platform, WindowsSdkDir, WindowsSdkVersion);
 
 			// Get all the system include paths
-			SetupEnvironment(Params.Platform);
+			SetupEnvironment(Params.Platform, Logger);
 		}
 
 		/// <summary>
@@ -310,7 +312,7 @@
 		/// Sets up the standard compile environment for the toolchain
 		/// </summary>
 		[SupportedOSPlatform("windows")]
-		private void SetupEnvironment(UnrealTargetPlatform Platform)
+		private void SetupEnvironment(UnrealTargetPlatform Platform, ILogger Logger)
 		{
 			string ArchFolder = WindowsExports.GetArchitectureSubpath(Architecture);
 
@@ -323,7 +325,7 @@
 			// If we're on >= Visual Studio 2015 and using pre-Windows 10 SDK, we need to find a Windows 10 SDK and add the UCRT include paths
 			if(ToolChain.IsMSVC() && WindowsSdkVersion < new VersionNumber(10))
 			{
-				KeyValuePair<VersionNumber, DirectoryReference> Pair = MicrosoftPlatformSDK.FindUniversalCrtDirs().OrderByDescending(x => x.Key).FirstOrDefault();
+				KeyValuePair<VersionNumber, DirectoryReference> Pair = MicrosoftPlatformSDK.FindUniversalCrtDirs(Logger).OrderByDescending(x => x.Key).FirstOrDefault();
 				if(Pair.Key == null || Pair.Key < new VersionNumber(10))
 				{
 					throw new BuildException("{0} requires the Universal CRT to be installed.", WindowsPlatform.GetCompilerName(ToolChain));
@@ -390,20 +392,21 @@
 		/// <param name="CompilerVersion">The specific toolchain version to use</param>
 		/// <param name="WindowsSdkVersion">Version of the Windows SDK to use</param>
 		/// <param name="SuppliedSdkDirectoryForVersion">If specified, this is the SDK directory to use, otherwise, attempt to look up via registry. If specified, the WindowsSdkVersion is used directly</param>
+		/// <param name="Logger">Logger for output</param>
 		/// <returns>New environment object with paths for the given settings</returns>
 		[SupportedOSPlatform("windows")]
-		public static VCEnvironment Create(WindowsCompiler Compiler, UnrealTargetPlatform Platform, WindowsArchitecture Architecture, string? CompilerVersion, string? WindowsSdkVersion, string? SuppliedSdkDirectoryForVersion)
+		public static VCEnvironment Create(WindowsCompiler Compiler, UnrealTargetPlatform Platform, WindowsArchitecture Architecture, string? CompilerVersion, string? WindowsSdkVersion, string? SuppliedSdkDirectoryForVersion, ILogger Logger)
 		{
-			return Create( new VCEnvironmentParameters(Compiler, Platform, Architecture, CompilerVersion, WindowsSdkVersion, SuppliedSdkDirectoryForVersion) );
+			return Create( new VCEnvironmentParameters(Compiler, Platform, Architecture, CompilerVersion, WindowsSdkVersion, SuppliedSdkDirectoryForVersion, Logger), Logger );
 		}
 
 		/// <summary>
 		/// Creates an environment with the given parameters
 		/// </summary>
 		[SupportedOSPlatform("windows")]
-		public static VCEnvironment Create( VCEnvironmentParameters Params)
+		public static VCEnvironment Create( VCEnvironmentParameters Params, ILogger Logger)
 		{
-			return new VCEnvironment(Params);
+			return new VCEnvironment(Params, Logger);
 		}
 
 
@@ -445,7 +448,7 @@
 		public VersionNumber WindowsSdkVersion;
 
 		/// <summary>Optional directory for redistributable items (DLLs etc)</summary>
-		public DirectoryReference? RedistDir;	   
+		public DirectoryReference? RedistDir;
 
 		/// <summary>
 		/// Creates VC environment construction parameters with the given settings
@@ -456,15 +459,16 @@
 		/// <param name="CompilerVersion">The specific toolchain version to use</param>
 		/// <param name="WindowsSdkVersion">Version of the Windows SDK to use</param>
 		/// <param name="SuppliedSdkDirectoryForVersion">If specified, this is the SDK directory to use, otherwise, attempt to look up via registry. If specified, the WindowsSdkVersion is used directly</param>
+		/// <param name="Logger">Logger for output</param>
 		/// <returns>Creation parameters for VC environment</returns>
 		[SupportedOSPlatform("windows")]
-		public VCEnvironmentParameters (WindowsCompiler Compiler, UnrealTargetPlatform Platform, WindowsArchitecture Architecture, string? CompilerVersion, string? WindowsSdkVersion, string? SuppliedSdkDirectoryForVersion)
+		public VCEnvironmentParameters (WindowsCompiler Compiler, UnrealTargetPlatform Platform, WindowsArchitecture Architecture, string? CompilerVersion, string? WindowsSdkVersion, string? SuppliedSdkDirectoryForVersion, ILogger Logger)
 		{
 			// Get the compiler version info
 			VersionNumber? SelectedCompilerVersion;
 			DirectoryReference? SelectedCompilerDir;
 			DirectoryReference? SelectedRedistDir;
-			if(!WindowsPlatform.TryGetToolChainDir(Compiler, CompilerVersion, Architecture, out SelectedCompilerVersion, out SelectedCompilerDir, out SelectedRedistDir))
+			if(!WindowsPlatform.TryGetToolChainDir(Compiler, CompilerVersion, Architecture, Logger, out SelectedCompilerVersion, out SelectedCompilerDir, out SelectedRedistDir))
 			{
 				throw new BuildException("{0}{1} {2} must be installed in order to build this target.", WindowsPlatform.GetCompilerName(Compiler), String.IsNullOrEmpty(CompilerVersion)? "" : String.Format(" ({0})", CompilerVersion), Architecture.ToString());
 			}
@@ -475,11 +479,11 @@
 			DirectoryReference? SelectedToolChainDir;
 			if(Compiler.IsClang())
 			{
-				if (WindowsPlatform.TryGetToolChainDir(WindowsCompiler.VisualStudio2019, null, Architecture, out SelectedToolChainVersion, out SelectedToolChainDir, out SelectedRedistDir))
+				if (WindowsPlatform.TryGetToolChainDir(WindowsCompiler.VisualStudio2019, null, Architecture, Logger, out SelectedToolChainVersion, out SelectedToolChainDir, out SelectedRedistDir))
 				{
 					ToolChain = WindowsCompiler.VisualStudio2019;
 				}
-				else if (WindowsPlatform.TryGetToolChainDir(WindowsCompiler.VisualStudio2022, null, Architecture, out SelectedToolChainVersion, out SelectedToolChainDir, out SelectedRedistDir))
+				else if (WindowsPlatform.TryGetToolChainDir(WindowsCompiler.VisualStudio2022, null, Architecture, Logger, out SelectedToolChainVersion, out SelectedToolChainDir, out SelectedRedistDir))
 				{
 					ToolChain = WindowsCompiler.VisualStudio2022;
 				}
@@ -510,7 +514,7 @@
 			}
 			else
 			{
-				if (!WindowsPlatform.TryGetWindowsSdkDir(WindowsSdkVersion, out SelectedWindowsSdkVersion, out SelectedWindowsSdkDir))
+				if (!WindowsPlatform.TryGetWindowsSdkDir(WindowsSdkVersion, Logger, out SelectedWindowsSdkVersion, out SelectedWindowsSdkDir))
 				{
 					throw new BuildException("Windows SDK{0} must be installed in order to build this target.", String.IsNullOrEmpty(WindowsSdkVersion) ? "" : String.Format(" ({0})", WindowsSdkVersion));
 				}
diff --git a/Engine/Source/Programs/UnrealBuildTool/Platform/Windows/VCManifestGenerator.cs b/Engine/Source/Programs/UnrealBuildTool/Platform/Windows/VCManifestGenerator.cs
--- a/Engine/Source/Programs/UnrealBuildTool/Platform/Windows/VCManifestGenerator.cs
+++ b/Engine/Source/Programs/UnrealBuildTool/Platform/Windows/VCManifestGenerator.cs
@@ -12,6 +12,7 @@
 using System.Diagnostics;
 using UnrealBuildBase;
 using System.Diagnostics.CodeAnalysis;
+using Microsoft.Extensions.Logging;
 
 namespace UnrealBuildTool
 {
@@ -78,11 +79,17 @@
 		protected Dictionary<string,string>? ManifestFiles; // Dst, Src
 
 		/// <summary>
+		/// Logger for output
+		/// </summary>
+		protected readonly ILogger Logger;
+
+		/// <summary>
 		/// Create a manifest generator for the given platform variant.
 		/// </summary>
-		public VCManifestGenerator( UnrealTargetPlatform InPlatform )
+		public VCManifestGenerator( UnrealTargetPlatform InPlatform, ILogger InLogger )
 		{
-			this.Platform = InPlatform;
+			Platform = InPlatform;
+			Logger = InLogger;
 		}
 
 		/// <summary>
@@ -102,7 +109,7 @@
 		/// <summary>
 		/// Attempts to create the given directory
 		/// </summary>
-		protected static bool CreateCheckDirectory(string TargetDirectory)
+		protected static bool CreateCheckDirectory(string TargetDirectory, ILogger Logger)
 		{
 			if (!Directory.Exists(TargetDirectory))
 			{
@@ -112,12 +119,12 @@
 				}
 				catch (Exception)
 				{
-					Log.TraceError("Could not create directory {0}.", TargetDirectory);
+					Logger.LogError("Could not create directory {TargetDir}.", TargetDirectory);
 					return false;
 				}
 				if (!Directory.Exists(TargetDirectory))
 				{
-					Log.TraceError("Path {0} does not exist or is not a directory.", TargetDirectory);
+					Logger.LogError("Path {TargetDir} does not exist or is not a directory.", TargetDirectory);
 					return false;
 				}
 			}
@@ -162,13 +169,14 @@
 
 			if (ExitCode == 0)
 			{
-				Log.TraceVerbose(ProcessOutput.ToString());
+				Logger.LogDebug("Output", ProcessOutput.ToString());
 				return true;
 			}
 			else
 			{
-				Log.TraceInformation(ProcessOutput.ToString());
-				Log.TraceError(Path.GetFileName(PriExecutable) + " returned an error.\nExit code:" + ExitCode );
+				Logger.LogInformation("Output", ProcessOutput.ToString());
+				Logger.LogError("{File} returned an error.", Path.GetFileName(PriExecutable));
+				Logger.LogError("Exit code: {Code}", ExitCode);
 				return false;
 			}
         }
@@ -221,8 +229,8 @@
 			}
 			if (CompletedVersionString == null || CompletedVersionString.Length <= 0)
 			{
-				Log.TraceError("Invalid package version {0}. Package versions must be in the format #.#.#.# where # is a number 0-65535.", InVersionNumber);
-				Log.TraceError("Consider setting [{0}]:PackageVersion to provide a specific value.", IniSection_PlatformTargetSettings);
+				Logger.LogError("Invalid package version {Ver}. Package versions must be in the format #.#.#.# where # is a number 0-65535.", InVersionNumber);
+				Logger.LogError("Consider setting [{IniSection}]:PackageVersion to provide a specific value.", IniSection_PlatformTargetSettings);
 			}
 			return CompletedVersionString;
 		}
@@ -240,7 +248,7 @@
 			}
 			if (ReturnVal == null || ReturnVal.Length <= 0)
 			{
-				Log.TraceError("Invalid application ID {0}. Application IDs must only contain letters and numbers. And they must begin with a letter.", InApplicationId);
+				Logger.LogError("Invalid application ID {AppId}. Application IDs must only contain letters and numbers. And they must begin with a letter.", InApplicationId);
 			}
 			return ReturnVal;
 		}
@@ -309,7 +317,7 @@
 				return "#" + R.ToString("X2") + G.ToString("X2") + B.ToString("X2");
 			}
 
-			Log.TraceWarning("Failed to parse color config value. Using default.");
+			Logger.LogWarning("Failed to parse color config value. Using default.");
 			return DefaultValue;
 		}
 
@@ -325,13 +333,13 @@
 			var StaleResourceFiles = TargetResourceInstances.Where(X => !ManifestFiles!.ContainsKey(X)).ToList();
 			if (StaleResourceFiles.Any())
 			{
-				Log.TraceVerbose("Removing stale manifest resource files...");
+				Logger.LogDebug("Removing stale manifest resource files...");
 				foreach (string StaleResourceFile in StaleResourceFiles)
 				{
 					// try to delete the file & the directory that contains it
 					try
 					{
-						Log.TraceVerbose($"\tremoving {Utils.MakePathRelativeTo(StaleResourceFile, OutputPath!)}");
+						Logger.LogDebug("    removing {Path}", Utils.MakePathRelativeTo(StaleResourceFile, OutputPath!));
 						FileUtils.ForceDeleteFile(StaleResourceFile);
 						if (!Directory.EnumerateFileSystemEntries(Path.GetDirectoryName(StaleResourceFile)!).Any())
 						{
@@ -340,7 +348,7 @@
 					}
 					catch (Exception E)
 					{
-						Log.TraceError($"\tCould not remove {StaleResourceFile} - {E.Message}.");
+						Logger.LogError("    Could not remove {StaleResourceFile} - {Message}.", StaleResourceFile, E.Message);
 					}
 				}
 			}
@@ -408,7 +416,7 @@
 			}
 
 			// If the target resource folder doesn't exist yet, create it
-			if (!CreateCheckDirectory(TargetPath))
+			if (!CreateCheckDirectory(TargetPath, Logger))
 			{
 				return false;
 			}
@@ -429,9 +437,9 @@
 			foreach (string SourceResourceFile in SourceResourceInstances)
 			{
 				string TargetResourcePath = Path.Combine(TargetPath, SourceResourceFile.Substring(SourcePath.Length + 1));
-				if (!CreateCheckDirectory(Path.GetDirectoryName(TargetResourcePath)!))
+				if (!CreateCheckDirectory(Path.GetDirectoryName(TargetResourcePath)!, Logger))
 				{
-					Log.TraceError("Unable to create intermediate directory {0}.", Path.GetDirectoryName(TargetResourcePath));
+					Logger.LogError("Unable to create intermediate directory {IntDir}.", Path.GetDirectoryName(TargetResourcePath));
 					continue;
 				}
 				AddFileReference(SourceResourceFile, TargetResourcePath, bIsGeneratedFile:false);
@@ -493,7 +501,7 @@
 				return UpdatedFiles;
 			}
 
-			Log.TraceVerbose("Updating manifest resource files...");
+			Logger.LogDebug("Updating manifest resource files...");
 
 			// copy over any new or updated files
 			foreach ( var ManifestFilePair in ManifestFiles! )
@@ -516,14 +524,21 @@
 					}
 					catch (Exception E)
 					{
-						Log.TraceError($"\tCould not replace file {TargetPath} - {E.Message}");
+						Logger.LogError("    Could not replace file {TargetPath} - {Message}", TargetPath, E.Message);
 					}
 				}
 
 				// copy new version
 				try
 				{
-					Log.TraceVerbose($"\t{(bFileExists ? "updating" : "adding")} {Utils.MakePathRelativeTo(TargetPath, OutputPath!)}");
+					if (bFileExists)
+					{
+						Logger.LogDebug("    updating {Path}", Utils.MakePathRelativeTo(TargetPath, OutputPath!));
+					}
+					else
+					{
+						Logger.LogDebug("    adding {Path}", Utils.MakePathRelativeTo(TargetPath, OutputPath!));
+					}
 
 					Directory.CreateDirectory(Path.GetDirectoryName(TargetPath)!);
 					File.Copy(SourcePath, TargetPath);
@@ -534,7 +549,7 @@
 				}
 				catch (Exception E)
 				{
-					Log.TraceError($"\tUnable to copy file {TargetPath} - {E.Message}");
+					Logger.LogError("    Unable to copy file {TargetPath} - {Message}", TargetPath, E.Message);
 				}
 			}
 
@@ -556,7 +571,7 @@
 				Dictionary<string, string>? Values;
 				if (!ConfigHierarchy.TryParse(DefaultCultureScratchValue, out Values))
 				{
-					Log.TraceError("Invalid default culture string resources: \"{0}\". Unable to add resource entry.", DefaultCultureScratchValue);
+					Logger.LogError("Invalid default culture string resources: \"{Culture}\". Unable to add resource entry.", DefaultCultureScratchValue);
 					return "";
 				}
 
@@ -582,7 +597,7 @@
 						|| !SeparatedCultureValues.ContainsKey("CultureStringResources")
 						|| !SeparatedCultureValues.ContainsKey("CultureId"))
 					{
-						Log.TraceError("Invalid per-culture resource: \"{0}\". Unable to add resource entry.", CultureCombinedValues);
+						Logger.LogError("Invalid per-culture resource: \"{Culture}\". Unable to add resource entry.", CultureCombinedValues);
 						continue;
 					}
 
@@ -592,7 +607,7 @@
 						Dictionary<string, string>? CultureStringResources;
 						if (!ConfigHierarchy.TryParse(SeparatedCultureValues["CultureStringResources"], out CultureStringResources))
 						{
-							Log.TraceError("Invalid culture string resources: \"{0}\". Unable to add resource entry.", CultureCombinedValues);
+							Logger.LogError("Invalid culture string resources: \"{Culture}\". Unable to add resource entry.", CultureCombinedValues);
 							continue;
 						}
 
@@ -663,7 +678,7 @@
 			// Check that we have a valid number of cultures
 			if (CulturesToStage!.Count < 1 || CulturesToStage.Count >= MaxResourceEntries)
 			{
-				Log.TraceWarning("Incorrect number of cultures to stage. There must be between 1 and {0} cultures selected.", MaxResourceEntries);
+				Logger.LogWarning("Incorrect number of cultures to stage. There must be between 1 and {MaxCultures} cultures selected.", MaxResourceEntries);
 			}
 
 			// Create the culture list. This list is unordered except that the default language must be first which we already took care of above.
@@ -676,15 +691,15 @@
 		/// <summary>
 		/// Get the package identity name string
 		/// </summary>
-		protected string GetIdentityPackageName( string? TargetName)
+		protected string GetIdentityPackageName(string? TargetName)
 		{
             // Read the PackageName from config
 			var DefaultName = (ProjectFile != null) ? ProjectFile.GetFileNameWithoutAnyExtensions() : (TargetName ?? "DefaultUEProject");
             var PackageName = Regex.Replace(GetConfigString("PackageName", "ProjectName", DefaultName), "[^-.A-Za-z0-9]", "");
             if (string.IsNullOrWhiteSpace(PackageName))
             {
-                Log.TraceError("Invalid package name {0}. Package names must only contain letters, numbers, dash, and period and must be at least one character long.", PackageName);
-                Log.TraceError("Consider using the setting [{0}]:PackageName to provide a specific value.", IniSection_PlatformTargetSettings);
+                Logger.LogError("Invalid package name {Name}. Package names must only contain letters, numbers, dash, and period and must be at least one character long.", PackageName);
+                Logger.LogError("Consider using the setting [{IniSection}]:PackageName to provide a specific value.", IniSection_PlatformTargetSettings);
             }
 
 			// If specified in the project settings append the users machine name onto the package name to allow sharing of devkits without stomping of deploys
@@ -795,23 +810,23 @@
 			// Check parameter values are valid.
 			if (InTargetConfigs.Count != InExecutables.Count)
 			{
-				Log.TraceError("The number of target configurations ({0}) and executables ({1}) passed to manifest generation do not match.", InTargetConfigs.Count, InExecutables.Count);
+				Logger.LogError("The number of target configurations ({NumConfigs}) and executables ({NumExes}) passed to manifest generation do not match.", InTargetConfigs.Count, InExecutables.Count);
 				return null;
 			}
 			if (InTargetConfigs.Count < 1)
 			{
-				Log.TraceError("The number of target configurations is zero, so we cannot generate a manifest.");
+				Logger.LogError("The number of target configurations is zero, so we cannot generate a manifest.");
 				return null;
 			}
 
-			if (!CreateCheckDirectory(InOutputPath))
+			if (!CreateCheckDirectory(InOutputPath, Logger))
 			{
-				Log.TraceError("Failed to create output directory \"{0}\".", InOutputPath);
+				Logger.LogError("Failed to create output directory \"{OutputDir}\".", InOutputPath);
 				return null;
 			}
-			if (!CreateCheckDirectory(InIntermediatePath))
+			if (!CreateCheckDirectory(InIntermediatePath, Logger))
 			{
-				Log.TraceError("Failed to create intermediate directory \"{0}\".", InIntermediatePath);
+				Logger.LogError("Failed to create intermediate directory \"{IntDir}\".", InIntermediatePath);
 				return null;
 			}
 
@@ -833,7 +848,7 @@
 				GameIni.GetString("/Script/UnrealEd.ProjectPackagingSettings", "DefaultCulture", out DefaultCulture);
 				if (CulturesToStageWithDuplicates == null || CulturesToStageWithDuplicates.Count < 1)
 				{
-					Log.TraceError("At least one culture must be selected to stage.");
+					Logger.LogError("At least one culture must be selected to stage.");
 					return null;
 				}
 
@@ -842,12 +857,12 @@
 			if (DefaultCulture == null || DefaultCulture.Length < 1)
 			{
 				DefaultCulture = CulturesToStage[0];
-				Log.TraceWarning("A default culture must be selected to stage. Using {0}.", DefaultCulture);
+				Logger.LogWarning("A default culture must be selected to stage. Using {DefaultCulture}.", DefaultCulture);
 			}
 			if (!CulturesToStage.Contains(DefaultCulture))
 			{
 				DefaultCulture = CulturesToStage[0];
-				Log.TraceWarning("The default culture must be one of the staged cultures. Using {0}.", DefaultCulture);
+				Logger.LogWarning("The default culture must be one of the staged cultures. Using {DefaultCulture}.", DefaultCulture);
 			}
 
 			List<string>? PerCultureValues;
@@ -858,7 +873,7 @@
 					Dictionary<string, string>? SeparatedCultureValues;
 					if (!ConfigHierarchy.TryParse(CultureCombinedValues, out SeparatedCultureValues))
 					{
-						Log.TraceWarning("Invalid per-culture resource value: {0}", CultureCombinedValues);
+						Logger.LogWarning("Invalid per-culture resource value: {Culture}", CultureCombinedValues);
 						continue;
 					}
 
@@ -877,15 +892,15 @@
 			// Only warn if shipping, we can run without translated cultures they're just needed for cert
 			else if (InTargetConfigs.Contains(UnrealTargetConfiguration.Shipping))
 			{
-				Log.TraceInformation("Staged culture mappings not setup in the editor. See Per Culture Resources in the {0} Target Settings.", Platform.ToString() );
+				Logger.LogInformation("Staged culture mappings not setup in the editor. See Per Culture Resources in the {Platform} Target Settings.", Platform.ToString() );
 			}
 
 			// Clean out the resources intermediate path so that we know there are no stale binary files.
 			string IntermediateResourceDirectory = Path.Combine(IntermediatePath, BuildResourceSubPath);
 			FileUtils.ForceDeleteDirectory(IntermediateResourceDirectory);
-			if (!CreateCheckDirectory(IntermediateResourceDirectory))
+			if (!CreateCheckDirectory(IntermediateResourceDirectory, Logger))
 			{
-				Log.TraceError("Could not create directory {0}.", IntermediateResourceDirectory);
+				Logger.LogError("Could not create directory {IntDir}.", IntermediateResourceDirectory);
 				return null;
 			}
 
@@ -899,14 +914,14 @@
 			{
 				string IntermediateStringResourcePath = Path.Combine(IntermediateResourceDirectory, Culture);
 				string IntermediateStringResourceFile = Path.Combine(IntermediateStringResourcePath, "resources.resw");
-				if (!CreateCheckDirectory(IntermediateStringResourcePath))
+				if (!CreateCheckDirectory(IntermediateStringResourcePath, Logger))
 				{
-					Log.TraceWarning("Culture {0} resources not staged.", Culture);
+					Logger.LogWarning("Culture {Culture} resources not staged.", Culture);
 					CulturesToStage.Remove(Culture);
 					if (Culture == DefaultCulture)
 					{
 						DefaultCulture = CulturesToStage[0];
-						Log.TraceWarning("Default culture skipped. Using {0} as default culture.", DefaultCulture);
+						Logger.LogWarning("Default culture skipped. Using {DefaultCulture} as default culture.", DefaultCulture);
 					}
 					continue;
 				}
@@ -996,7 +1011,7 @@
 				}
 				catch (Exception E)
 				{
-					Log.TraceError($"cannot remove old pri file: {TargetResourceIndexFile} - {E.Message}");
+					Logger.LogError("cannot remove old pri file: {TargetResourceIndexFile} - {Message}", TargetResourceIndexFile, E.Message);
 				}
 
 				// Generate the Package Resource Index
@@ -1007,7 +1022,7 @@
 					MakePriCommandLine += " /indexName \"" + IdentityName + "\"";
 				}
 
-				Log.TraceVerbose($"\tgenerating {Utils.MakePathRelativeTo(TargetResourceIndexFile, OutputPath!)}");
+				Logger.LogDebug("    generating {Path}", Utils.MakePathRelativeTo(TargetResourceIndexFile, OutputPath!));
 				RunMakePri(MakePriCommandLine);
 				UpdatedFilePaths.Add(TargetResourceIndexFile);
 			}
@@ -1015,7 +1030,7 @@
 			// Report if nothing was changed
 			if (!bHadStaleResources && !UpdatedFilePaths.Any())
 			{
-				Log.TraceVerbose($"Manifest resource files are up to date");
+				Logger.LogDebug($"Manifest resource files are up to date");
 			}
 
 			return UpdatedFilePaths;
diff --git a/Engine/Source/Programs/UnrealBuildTool/Platform/Windows/VCToolChain.cs b/Engine/Source/Programs/UnrealBuildTool/Platform/Windows/VCToolChain.cs
--- a/Engine/Source/Programs/UnrealBuildTool/Platform/Windows/VCToolChain.cs
+++ b/Engine/Source/Programs/UnrealBuildTool/Platform/Windows/VCToolChain.cs
@@ -11,6 +11,7 @@ using System.Text;
 using EpicGames.Core;
 using UnrealBuildBase;
 using System.Runtime.Versioning;
+using Microsoft.Extensions.Logging;
 
 namespace UnrealBuildTool
 {
@@ -31,15 +32,16 @@ namespace UnrealBuildTool
 		/// </summary>
 		private static int MaxPathWarningLength = 260;
 
-		public VCToolChain(ReadOnlyTargetRules Target)
+		public VCToolChain(ReadOnlyTargetRules Target, ILogger Logger)
+			: base(Logger)
 		{
 			this.Target = Target;
 			this.EnvVars = Target.WindowsPlatform.Environment!;
 
-			Log.TraceLog("Compiler: {0}", EnvVars.CompilerPath);
-			Log.TraceLog("Linker: {0}", EnvVars.LinkerPath);
-			Log.TraceLog("Library Manager: {0}", EnvVars.LibraryManagerPath);
-			Log.TraceLog("Resource Compiler: {0}", EnvVars.ResourceCompilerPath);
+			Logger.LogDebug("Compiler: {Path}", EnvVars.CompilerPath);
+			Logger.LogDebug("Linker: {Path}", EnvVars.LinkerPath);
+			Logger.LogDebug("Library Manager: {Path}", EnvVars.LibraryManagerPath);
+			Logger.LogDebug("Resource Compiler: {Path}", EnvVars.ResourceCompilerPath);
 
 			if (Target.WindowsPlatform.ObjSrcMapFile != null)
 			{
@@ -210,7 +212,7 @@ namespace UnrealBuildTool
 			Arguments.Add($"/Fp\"{UsingPchFilePath}\"");
 		}
 
-		public static void AddPreprocessedFile(List<string> Arguments, FileItem PreprocessedFile)
+		public static void AddPreprocessedFile(List<string> Arguments, FileItem PreprocessedFile, ILogger Logger)
 		{
 			string PreprocessedFileString = NormalizeCommandLinePath(PreprocessedFile);
 			Arguments.Add("/P"); // Preprocess
@@ -218,7 +220,7 @@ namespace UnrealBuildTool
 			Arguments.Add($"/Fi\"{PreprocessedFileString}\""); // Preprocess to a file
 
 			// this is parsed by external tools wishing to open this file directly.
-			Log.TraceInformation("PreProcessPath: " + PreprocessedFile);
+			Logger.LogInformation("PreProcessPath: {Path}", PreprocessedFile);
 		}
 
 		public static void AddObjectFile(List<string> Arguments, FileItem ObjectFile)
@@ -1380,7 +1382,7 @@ namespace UnrealBuildTool
 
 					if (!ProjectFileGenerator.bGenerateProjectFiles)
 					{
-						CompileDepsAction.WriteResponseFile(Graph);
+						CompileDepsAction.WriteResponseFile(Graph, Logger);
 					}
 
 					CompileAction.ActionType = ActionType.CompileModuleInterface;
@@ -1418,7 +1420,7 @@ namespace UnrealBuildTool
 
 				if (!ProjectFileGenerator.bGenerateProjectFiles)
 				{
-					CompileAction.WriteResponseFile(Graph);
+					CompileAction.WriteResponseFile(Graph, Logger);
 				}
 
 				// When compiling with SN-DBS, modules that contain a #import must be built locally
@@ -1465,10 +1467,12 @@ namespace UnrealBuildTool
 			return ParseTimingInfoAction;
 		}
 
-		public override void FinalizeOutput(ReadOnlyTargetRules Target, TargetMakefile Makefile)
+		public override void FinalizeOutput(ReadOnlyTargetRules Target, TargetMakefileBuilder MakefileBuilder)
 		{
 			if (Target.bPrintToolChainTimingInfo || Target.WindowsPlatform.bCompilerTrace)
 			{
+				TargetMakefile Makefile = MakefileBuilder.Makefile;
+
 				List<IExternalAction> ParseTimingActions = Makefile.Actions.Where(x => x.ActionType == ActionType.ParseTimingInfo).ToList();
 				List<FileItem> TimingJsonFiles = ParseTimingActions.SelectMany(a => a.ProducedItems.Where(i => i.HasExtension(".cta"))).ToList();
 				Makefile.OutputItems.AddRange(TimingJsonFiles);
@@ -1492,7 +1496,7 @@ namespace UnrealBuildTool
 						$"-CompileTimingFile={ExpectedCompileTimeFile}",
 					};
 
-					Action AggregateTimingInfoAction = Makefile.CreateRecursiveAction<AggregateParsedTimingInfo>(ActionType.ParseTimingInfo, string.Join(" ", ActionArgs));
+					Action AggregateTimingInfoAction = MakefileBuilder.CreateRecursiveAction<AggregateParsedTimingInfo>(ActionType.ParseTimingInfo, string.Join(" ", ActionArgs));
 					AggregateTimingInfoAction.WorkingDirectory = Unreal.EngineSourceDirectory;
 					AggregateTimingInfoAction.StatusDescription = $"Aggregating {TimingJsonFiles.Count} Timing File(s)";
 					AggregateTimingInfoAction.bCanExecuteRemotely = false;
@@ -2015,7 +2019,7 @@ namespace UnrealBuildTool
 			}
 
 			// Allow the toolchain to adjust/process the link arguments
-			ModifyFinalLinkArguments(LinkEnvironment, Arguments, bBuildImportLibraryOnly );
+			ModifyFinalLinkArguments(LinkEnvironment, Arguments, bBuildImportLibraryOnly);
 
 			// Create a response file for the linker, unless we're generating IntelliSense data
 			FileReference ResponseFileName = GetResponseFileName(LinkEnvironment, OutputFile);
@@ -2062,8 +2066,8 @@ namespace UnrealBuildTool
 			// Allow remote linking. Note that this may be overriden by the executor (eg. XGE.bAllowRemoteLinking)
 			LinkAction.bCanExecuteRemotely = true;
 
-			Log.TraceVerbose("     Linking: " + LinkAction.StatusDescription);
-			Log.TraceVerbose("     Command: " + LinkAction.CommandArguments);
+			Logger.LogDebug("     Linking: {StatusDescription}", LinkAction.StatusDescription);
+			Logger.LogDebug("     Command: {CommandArguments}", LinkAction.CommandArguments);
 
 			return OutputFile;
 		}
@@ -2074,7 +2078,7 @@ namespace UnrealBuildTool
 			{
 				// The linker expects the .pgd and any .pgc files to be in the output directory.
 				// Copy the files there and make them writable...
-				Log.TraceInformation("...copying the profile guided optimization files to output directory...");
+				Logger.LogInformation("...copying the profile guided optimization files to output directory...");
 
 				string[] PGDFiles = Directory.GetFiles(LinkEnvironment.PGODirectory!, "*.pgd");
 				string[] PGCFiles = Directory.GetFiles(LinkEnvironment.PGODirectory!, "*.pgc");
@@ -2085,13 +2089,13 @@ namespace UnrealBuildTool
 				}
 				else if (PGDFiles.Length == 0)
 				{
-					Log.TraceWarning("No .pgd files found in \"{0}\".", LinkEnvironment.PGODirectory);
+					Logger.LogWarning("No .pgd files found in \"{PgoDir}\".", LinkEnvironment.PGODirectory);
 					return false;
 				}
 
 				if (PGCFiles.Length == 0)
 				{
-					Log.TraceWarning("No .pgc files found in \"{0}\".", LinkEnvironment.PGODirectory);
+					Logger.LogWarning("No .pgc files found in \"{PgoDir}\".", LinkEnvironment.PGODirectory);
 					return false;
 				}
 
@@ -2101,7 +2105,7 @@ namespace UnrealBuildTool
 				// Copy the .pgd to the linker output directory, renaming it to match the PGO filename prefix.
 				string PGDFile = PGDFiles.First();
 				string DestPGDFile = Path.Combine(LinkEnvironment.OutputDirectory.FullName, LinkEnvironment.PGOFilenamePrefix + ".pgd");
-				Log.TraceInformation("{0} -> {1}", PGDFile, DestPGDFile);
+				Logger.LogInformation("{Source} -> {Target}", PGDFile, DestPGDFile);
 				File.Copy(PGDFile, DestPGDFile, true);
 				File.SetAttributes(DestPGDFile, FileAttributes.Normal);
 
@@ -2112,7 +2116,7 @@ namespace UnrealBuildTool
 					string DestFileName = string.Format("{0}!{1}.pgc", LinkEnvironment.PGOFilenamePrefix, ++PGCFileIndex);
 					string DestFilePath = Path.Combine(LinkEnvironment.OutputDirectory.FullName, DestFileName);
 
-					Log.TraceInformation("{0} -> {1}", SrcFilePath, DestFilePath);
+					Logger.LogInformation("{Source} -> {Target}", SrcFilePath, DestFilePath);
 					File.Copy(SrcFilePath, DestFilePath, true);
 					File.SetAttributes(DestFilePath, FileAttributes.Normal);
 				}
@@ -2138,7 +2142,7 @@ namespace UnrealBuildTool
 				}
 				else
 				{
-					Log.TraceWarning("PGO Optimize build will be disabled");
+					Logger.LogWarning("PGO Optimize build will be disabled");
 					bPGOOptimize = false;
 				}
 			}
@@ -2192,10 +2196,10 @@ namespace UnrealBuildTool
 		/// Gets the default include paths for the given platform.
 		/// </summary>
 		[SupportedOSPlatform("windows")]
-		public static string GetVCIncludePaths(UnrealTargetPlatform Platform, WindowsCompiler Compiler, string? CompilerVersion)
+		public static string GetVCIncludePaths(UnrealTargetPlatform Platform, WindowsCompiler Compiler, string? CompilerVersion, ILogger Logger)
 		{
 			// Make sure we've got the environment variables set up for this target
-			VCEnvironment EnvVars = VCEnvironment.Create(Compiler, Platform, WindowsArchitecture.x64, CompilerVersion, null, null);
+			VCEnvironment EnvVars = VCEnvironment.Create(Compiler, Platform, WindowsArchitecture.x64, CompilerVersion, null, null, Logger);
 
 			// Also add any include paths from the INCLUDE environment variable.  MSVC is not necessarily running with an environment that
 			// matches what UBT extracted from the vcvars*.bat using SetEnvironmentVariablesFromBatchFile().  We'll use the variables we
diff --git a/Engine/Source/Programs/UnrealBuildTool/Platform/Windows/WindowsExports.cs b/Engine/Source/Programs/UnrealBuildTool/Platform/Windows/WindowsExports.cs
--- a/Engine/Source/Programs/UnrealBuildTool/Platform/Windows/WindowsExports.cs
+++ b/Engine/Source/Programs/UnrealBuildTool/Platform/Windows/WindowsExports.cs
@@ -8,6 +8,7 @@
 using System.Text;
 using System.Threading.Tasks;
 using EpicGames.Core;
+using Microsoft.Extensions.Logging;
 
 namespace UnrealBuildTool
 {
@@ -23,7 +24,7 @@
 		/// <returns>True if successful</returns>
 		public static IEnumerable<DirectoryReference>? TryGetVSInstallDirs(WindowsCompiler Compiler)
 		{
-			return WindowsPlatform.TryGetVSInstallDirs(Compiler);
+			return WindowsPlatform.TryGetVSInstallDirs(Compiler, Log.Logger);
 		}
 
 		/// <summary>
@@ -33,7 +34,7 @@
 		[SupportedOSPlatform("windows")]
 		public static string GetMSBuildToolPath()
 		{
-			return MicrosoftPlatformSDK.GetMsBuildToolPath().FullName;
+			return MicrosoftPlatformSDK.GetMsBuildToolPath(Log.Logger).FullName;
 		}
 
 		/// <summary>
@@ -57,7 +58,7 @@
 		public static bool TryGetWindowsSdkDir(string DesiredVersion, [NotNullWhen(true)] out Version? OutSdkVersion, [NotNullWhen(true)] out DirectoryReference? OutSdkDir)
 		{
 			VersionNumber? vn;
-			if(WindowsPlatform.TryGetWindowsSdkDir(DesiredVersion, out vn, out OutSdkDir))
+			if(WindowsPlatform.TryGetWindowsSdkDir(DesiredVersion, Log.Logger, out vn, out OutSdkDir))
 			{
 				OutSdkVersion = new Version(vn.ToString());
 				return true;
@@ -78,13 +79,13 @@
 			// Add the default directory first
 			VersionNumber? Version;
 			DirectoryReference? DefaultWindowsSdkDir;
-			if (WindowsPlatform.TryGetWindowsSdkDir(null, out Version, out DefaultWindowsSdkDir))
+			if (WindowsPlatform.TryGetWindowsSdkDir(null, Log.Logger, out Version, out DefaultWindowsSdkDir))
 			{
 				WindowsSdkDirs.Add(new KeyValuePair<string, DirectoryReference>(Version.ToString(), DefaultWindowsSdkDir));
 			}
 
 			// Add all the other directories sorted in reverse order
-			IReadOnlyDictionary<VersionNumber, DirectoryReference> WindowsSdkDirPairs = MicrosoftPlatformSDK.FindWindowsSdkDirs();
+			IReadOnlyDictionary<VersionNumber, DirectoryReference> WindowsSdkDirPairs = MicrosoftPlatformSDK.FindWindowsSdkDirs(Log.Logger);
 			foreach(KeyValuePair<VersionNumber, DirectoryReference> Pair in WindowsSdkDirPairs.OrderByDescending(x => x.Key))
 			{
 				if(!WindowsSdkDirs.Any(x => x.Value == Pair.Value))
@@ -100,10 +101,11 @@
 		/// Enumerates all the Windows 10 SDK root directories
 		/// </summary>
 		/// <param name="RootDirs">Receives all the Windows 10 sdk root directories</param>
+		/// <param name="Logger">Logger for output</param>
 		[SupportedOSPlatform("windows")]
-		public static void EnumerateSdkRootDirs(List<DirectoryReference> RootDirs)
+		public static void EnumerateSdkRootDirs(List<DirectoryReference> RootDirs, ILogger Logger)
 		{
-			MicrosoftPlatformSDK.EnumerateSdkRootDirs(RootDirs);
+			MicrosoftPlatformSDK.EnumerateSdkRootDirs(RootDirs, Logger);
 		}
 
 	}
diff --git a/Engine/Source/Programs/UnrealBuildTool/Platform/Windows/WindowsProjectGenerator.cs b/Engine/Source/Programs/UnrealBuildTool/Platform/Windows/WindowsProjectGenerator.cs
--- a/Engine/Source/Programs/UnrealBuildTool/Platform/Windows/WindowsProjectGenerator.cs
+++ b/Engine/Source/Programs/UnrealBuildTool/Platform/Windows/WindowsProjectGenerator.cs
@@ -6,6 +6,7 @@
 using System.Diagnostics;
 using System.IO;
 using EpicGames.Core;
+using Microsoft.Extensions.Logging;
 
 namespace UnrealBuildTool
 {
@@ -18,8 +19,9 @@
 		/// Constructor
 		/// </summary>
 		/// <param name="Arguments">Command line arguments passed to the project generator</param>
-		public WindowsProjectGenerator(CommandLineArguments Arguments)
-			: base(Arguments)
+		/// <param name="Logger">Logger for output</param>
+		public WindowsProjectGenerator(CommandLineArguments Arguments, ILogger Logger)
+			: base(Arguments, Logger)
 		{
 		}
 
diff --git a/Engine/Source/Programs/UnrealBuildTool/Preprocessor/Preprocessor.cs b/Engine/Source/Programs/UnrealBuildTool/Preprocessor/Preprocessor.cs
--- a/Engine/Source/Programs/UnrealBuildTool/Preprocessor/Preprocessor.cs
+++ b/Engine/Source/Programs/UnrealBuildTool/Preprocessor/Preprocessor.cs
@@ -12,6 +12,7 @@
 using EpicGames.Core;
 using System.Diagnostics.CodeAnalysis;
 using UnrealBuildBase;
+using Microsoft.Extensions.Logging;
 
 namespace UnrealBuildTool
 {
@@ -384,9 +385,10 @@
 		/// <param name="Fragments">Lists of fragments that are parsed</param>
 		/// <param name="OuterContext">Outer context information, for error messages</param>
 		/// <param name="SourceFileCache">Cache for source files</param>
+		/// <param name="Logger">Logger for output</param>
 		/// <param name="bShowIncludes">Show all the included files, in order</param>
 		/// <param name="bIgnoreMissingIncludes">Suppress exceptions if an include path can not be resolved</param>
-		public void ParseFile(FileItem File, List<SourceFileFragment> Fragments, PreprocessorContext? OuterContext, SourceFileMetadataCache SourceFileCache, bool bShowIncludes = false, bool bIgnoreMissingIncludes = false)
+		public void ParseFile(FileItem File, List<SourceFileFragment> Fragments, PreprocessorContext? OuterContext, SourceFileMetadataCache SourceFileCache, ILogger Logger, bool bShowIncludes = false, bool bIgnoreMissingIncludes = false)
 		{
 			// If the file has already been included and had a #pragma once directive, don't include it again
 			if(PragmaOnceFiles.Contains(File))
@@ -402,7 +404,7 @@
 			// Output a trace of the included files
 			if(bShowIncludes)
 			{
-				Log.TraceInformation("Note: including file: {0}", File.Location);
+				Logger.LogInformation("Note: including file: {FileLocation}", File.Location);
 			}
 
 			// If the file had a header guard, and the macro is still defined, don't include it again
@@ -429,7 +431,7 @@
 						// Parse the included file
 						if (IncludedFile != null)
 						{
-							ParseFile(IncludedFile, Fragments, Context, SourceFileCache, bShowIncludes, bIgnoreMissingIncludes);
+							ParseFile(IncludedFile, Fragments, Context, SourceFileCache, Logger, bShowIncludes, bIgnoreMissingIncludes);
 						}
 					}
 					Context.MarkupIdx++;
diff --git a/Engine/Source/Programs/UnrealBuildTool/ProjectFiles/CMake/CMakefileGenerator.cs b/Engine/Source/Programs/UnrealBuildTool/ProjectFiles/CMake/CMakefileGenerator.cs
--- a/Engine/Source/Programs/UnrealBuildTool/ProjectFiles/CMake/CMakefileGenerator.cs
+++ b/Engine/Source/Programs/UnrealBuildTool/ProjectFiles/CMake/CMakefileGenerator.cs
@@ -7,6 +7,7 @@
 using EpicGames.Core;
 using UnrealBuildBase;
 using System.Linq;
+using Microsoft.Extensions.Logging;
 
 namespace UnrealBuildTool
 {
@@ -211,7 +212,7 @@
 				return "cmake-definitions" + CMakeExtension;
 			}
 		}
-		
+
 
 		/// <summary>
 		/// Writes the primary project file (e.g. Visual Studio Solution file)
@@ -219,7 +220,8 @@
 		/// <param name="UBTProject">The UnrealBuildTool project</param>
 		/// <param name="PlatformProjectGenerators">The registered platform project generators</param>
 		/// <returns>True if successful</returns>
-		protected override bool WritePrimaryProjectFile(ProjectFile? UBTProject, PlatformProjectGeneratorCollection PlatformProjectGenerators)
+		/// <param name="Logger"></param>
+		protected override bool WritePrimaryProjectFile(ProjectFile? UBTProject, PlatformProjectGeneratorCollection PlatformProjectGenerators, ILogger Logger)
 		{
 			return true;
 		}
@@ -244,7 +246,7 @@
 			}
 		}
 
-		private bool WriteCMakeLists()
+		private bool WriteCMakeLists(ILogger Logger)
 		{
 			string BuildCommand;
 			const string CMakeSectionEnd = " )\n\n";
@@ -289,7 +291,7 @@
 				HostArchitecture = "Linux";
 				BuildCommand = "cd \"" + UnrealRootPath + "\" && bash \"" + UnrealRootPath + "/Engine/Build/BatchFiles/" + HostArchitecture + "/Build.sh\"";
 
-				string? CompilerPath = LinuxCommon.WhichClang();
+				string? CompilerPath = LinuxCommon.WhichClang(Logger);
 				SetCompiler = "set(CMAKE_CXX_COMPILER " + CompilerPath + ")\n\n";
 			}
 			else
@@ -391,7 +393,7 @@
 			}
 
 			// Create SourceFiles, HeaderFiles, and ConfigFiles sections.
-			List<FileReference> AllModuleFiles = DiscoverModules(FindGameProjects(), null);
+			List<FileReference> AllModuleFiles = DiscoverModules(FindGameProjects(Logger), null);
 			foreach (FileReference CurModuleFile in AllModuleFiles)
 			{
 				List<FileReference> FoundFiles = SourceFileSearch.FindModuleSourceFiles(CurModuleFile);
@@ -595,19 +597,19 @@
 			string FullFileName = Path.Combine(PrimaryProjectPath.FullName, ProjectFileName);
 
 			// Write out CMake files
-			bool bWriteMakeList = WriteFileIfChanged(FullFileName, CMakefileContent.ToString());
-			bool bWriteEngineHeaders = WriteFileIfChanged(EngineHeadersFilePath, CMakeEngineHeaderFilesList.ToString());
-			bool bWriteProjectHeaders = WriteFileIfChanged(ProjectHeadersFilePath, CMakeProjectHeaderFilesList.ToString());
-			bool bWriteEngineSources = WriteFileIfChanged(EngineSourcesFilePath, CMakeEngineSourceFilesList.ToString());
-			bool bWriteProjectSources = WriteFileIfChanged(ProjectSourcesFilePath, CMakeProjectSourceFilesList.ToString());
-			bool bWriteIncludes = WriteFileIfChanged(IncludeFilePath, IncludeDirectoriesList.ToString());
-			bool bWriteDefinitions = WriteFileIfChanged(DefinitionsFilePath, PreprocessorDefinitionsList.ToString());
-			bool bWriteEngineConfigs = WriteFileIfChanged(EngineConfigsFilePath, CMakeEngineConfigFilesList.ToString());
-			bool bWriteProjectConfigs = WriteFileIfChanged(ProjectConfigsFilePath, CMakeProjectConfigFilesList.ToString());
-			bool bWriteEngineShaders = WriteFileIfChanged(EngineShadersFilePath, CMakeEngineShaderFilesList.ToString());
-			bool bWriteProjectShaders = WriteFileIfChanged(ProjectShadersFilePath, CMakeProjectShaderFilesList.ToString());
-			bool bWriteEngineCS = WriteFileIfChanged(EngineCSFilePath, CMakeEngineCSFilesList.ToString());
-			bool bWriteProjectCS = WriteFileIfChanged(ProjectCSFilePath, CMakeProjectCSFilesList.ToString());			
+			bool bWriteMakeList = WriteFileIfChanged(FullFileName, CMakefileContent.ToString(), Logger);
+			bool bWriteEngineHeaders = WriteFileIfChanged(EngineHeadersFilePath, CMakeEngineHeaderFilesList.ToString(), Logger);
+			bool bWriteProjectHeaders = WriteFileIfChanged(ProjectHeadersFilePath, CMakeProjectHeaderFilesList.ToString(), Logger);
+			bool bWriteEngineSources = WriteFileIfChanged(EngineSourcesFilePath, CMakeEngineSourceFilesList.ToString(), Logger);
+			bool bWriteProjectSources = WriteFileIfChanged(ProjectSourcesFilePath, CMakeProjectSourceFilesList.ToString(), Logger);
+			bool bWriteIncludes = WriteFileIfChanged(IncludeFilePath, IncludeDirectoriesList.ToString(), Logger);
+			bool bWriteDefinitions = WriteFileIfChanged(DefinitionsFilePath, PreprocessorDefinitionsList.ToString(), Logger);
+			bool bWriteEngineConfigs = WriteFileIfChanged(EngineConfigsFilePath, CMakeEngineConfigFilesList.ToString(), Logger);
+			bool bWriteProjectConfigs = WriteFileIfChanged(ProjectConfigsFilePath, CMakeProjectConfigFilesList.ToString(), Logger);
+			bool bWriteEngineShaders = WriteFileIfChanged(EngineShadersFilePath, CMakeEngineShaderFilesList.ToString(), Logger);
+			bool bWriteProjectShaders = WriteFileIfChanged(ProjectShadersFilePath, CMakeProjectShaderFilesList.ToString(), Logger);
+			bool bWriteEngineCS = WriteFileIfChanged(EngineCSFilePath, CMakeEngineCSFilesList.ToString(), Logger);
+			bool bWriteProjectCS = WriteFileIfChanged(ProjectCSFilePath, CMakeProjectCSFilesList.ToString(), Logger);			
 
 			// Return success flag if all files were written out successfully
 			return bWriteMakeList &&
@@ -726,9 +728,9 @@
 
 		#region ProjectFileGenerator implementation
 
-		protected override bool WriteProjectFiles(PlatformProjectGeneratorCollection PlatformProjectGenerators)
+		protected override bool WriteProjectFiles(PlatformProjectGeneratorCollection PlatformProjectGenerators, ILogger Logger)
 		{
-			return WriteCMakeLists();
+			return WriteCMakeLists(Logger);
 		}
 
 		/// <summary>
@@ -746,9 +748,9 @@
 		/// </summary>
 		protected bool bIncludeTVOSTargets = false;
 
-		protected override void ConfigureProjectFileGeneration(String[] Arguments, ref bool IncludeAllPlatforms)
+		protected override void ConfigureProjectFileGeneration(String[] Arguments, ref bool IncludeAllPlatforms, ILogger Logger)
 		{
-			base.ConfigureProjectFileGeneration(Arguments, ref IncludeAllPlatforms);
+			base.ConfigureProjectFileGeneration(Arguments, ref IncludeAllPlatforms, Logger);
 			// Check for minimal build targets to speed up cmake processing
 			foreach (string CurArgument in Arguments)
 			{
@@ -778,7 +780,7 @@
 			return new CMakefileProjectFile(InitFilePath, BaseDir);
 		}
 
-		public override void CleanProjectFiles(DirectoryReference InPrimaryProjectDirectory, string InPrimaryProjectName, DirectoryReference InIntermediateProjectFilesDirectory)
+		public override void CleanProjectFiles(DirectoryReference InPrimaryProjectDirectory, string InPrimaryProjectName, DirectoryReference InIntermediateProjectFilesDirectory, ILogger Logger)
 		{
 			// Remove Project File
 			FileReference PrimaryProjectFile = FileReference.Combine(InPrimaryProjectDirectory, ProjectFileName);
diff --git a/Engine/Source/Programs/UnrealBuildTool/ProjectFiles/CodeLite/CodeLiteGenerator.cs b/Engine/Source/Programs/UnrealBuildTool/ProjectFiles/CodeLite/CodeLiteGenerator.cs
--- a/Engine/Source/Programs/UnrealBuildTool/ProjectFiles/CodeLite/CodeLiteGenerator.cs
+++ b/Engine/Source/Programs/UnrealBuildTool/ProjectFiles/CodeLite/CodeLiteGenerator.cs
@@ -8,6 +8,7 @@
 using System.Xml.Linq;
 using EpicGames.Core;
 using System.Linq;
+using Microsoft.Extensions.Logging;
 
 namespace UnrealBuildTool
 {
@@ -47,7 +48,7 @@
 				return ".project";
 			}
 		}
-		protected override bool WritePrimaryProjectFile(ProjectFile? UBTProject, PlatformProjectGeneratorCollection PlatformProjectGenerators)
+		protected override bool WritePrimaryProjectFile(ProjectFile? UBTProject, PlatformProjectGeneratorCollection PlatformProjectGenerators, ILogger Logger)
 		{
 			string SolutionFileName = PrimaryProjectName + SolutionExtension;
 			string CodeCompletionFile = PrimaryProjectName + CodeCompletionFileName;
@@ -291,7 +292,7 @@
 			return new CodeLiteProject(InitFilePath, BaseDir, OnlyGameProject);
 		}
 
-		public override void CleanProjectFiles(DirectoryReference InPrimaryProjectDirectory, string InPrimaryProjectName, DirectoryReference InIntermediateProjectFilesDirectory)
+		public override void CleanProjectFiles(DirectoryReference InPrimaryProjectDirectory, string InPrimaryProjectName, DirectoryReference InIntermediateProjectFilesDirectory, ILogger Logger)
 		{
 			// TODO Delete all files here. Not finished yet.
 			string SolutionFileName = InPrimaryProjectName + SolutionExtension;
@@ -324,8 +325,8 @@
 				}
 				catch (Exception Ex)
 				{
-					Log.TraceInformation("Error while trying to clean project files path {0}. Ignored.", InIntermediateProjectFilesDirectory);
-					Log.TraceInformation("\t" + Ex.Message);
+					Logger.LogInformation("Error while trying to clean project files path {InIntermediateProjectFilesDirectory}. Ignored.", InIntermediateProjectFilesDirectory);
+					Logger.LogInformation("\t{Ex}", Ex.Message);
 				}
 			}
 		}
diff --git a/Engine/Source/Programs/UnrealBuildTool/ProjectFiles/CodeLite/CodeLiteProject.cs b/Engine/Source/Programs/UnrealBuildTool/ProjectFiles/CodeLite/CodeLiteProject.cs
--- a/Engine/Source/Programs/UnrealBuildTool/ProjectFiles/CodeLite/CodeLiteProject.cs
+++ b/Engine/Source/Programs/UnrealBuildTool/ProjectFiles/CodeLite/CodeLiteProject.cs
@@ -8,6 +8,7 @@
 using System.Xml.Linq;
 using EpicGames.Core;
 using UnrealBuildBase;
+using Microsoft.Extensions.Logging;
 
 namespace UnrealBuildTool
 {
@@ -35,7 +36,7 @@
 			return diff.ToString();
 		}
 
-		public override bool WriteProjectFile(List<UnrealTargetPlatform> InPlatforms, List<UnrealTargetConfiguration> InConfigurations, PlatformProjectGeneratorCollection PlatformProjectGenerators)
+		public override bool WriteProjectFile(List<UnrealTargetPlatform> InPlatforms, List<UnrealTargetConfiguration> InConfigurations, PlatformProjectGeneratorCollection PlatformProjectGenerators, ILogger Logger)
 		{
 			bool bSuccess = false;
 			string ProjectNameRaw = ProjectFilePath.GetFileNameWithoutExtension();
