UBT: Use ILogger throughout UBT, rather than legacy EpicGames.Core.Log methods. All output going forwards should use structured logging rather than String.Format style logging.

diff --git a/Engine/Source/Programs/UnrealBuildTool/ProjectFiles/Eddie/EddieGenerator.cs b/Engine/Source/Programs/UnrealBuildTool/ProjectFiles/Eddie/EddieGenerator.cs
--- a/Engine/Source/Programs/UnrealBuildTool/ProjectFiles/Eddie/EddieGenerator.cs
+++ b/Engine/Source/Programs/UnrealBuildTool/ProjectFiles/Eddie/EddieGenerator.cs
@@ -5,6 +5,7 @@
 using System.Text;
 using System.IO;
 using EpicGames.Core;
+using Microsoft.Extensions.Logging;
 
 namespace UnrealBuildTool
 {
@@ -23,7 +24,7 @@
 			}
 		}
 		
-		public override void CleanProjectFiles(DirectoryReference InPrimaryProjectDirectory, string InPrimaryProjectName, DirectoryReference InIntermediateProjectFilesPath)
+		public override void CleanProjectFiles(DirectoryReference InPrimaryProjectDirectory, string InPrimaryProjectName, DirectoryReference InIntermediateProjectFilesPath, ILogger Logger)
 		{
 			FileReference PrimaryProjDeleteFilename = FileReference.Combine(InPrimaryProjectDirectory, InPrimaryProjectName + ".wkst");
 			if (FileReference.Exists(PrimaryProjDeleteFilename))
@@ -40,8 +41,8 @@
 				}
 				catch (Exception Ex)
 				{
-					Log.TraceInformation("Error while trying to clean project files path {0}. Ignored.", InIntermediateProjectFilesPath);
-					Log.TraceInformation("\t" + Ex.Message);
+					Logger.LogInformation("Error while trying to clean project files path {InIntermediateProjectFilesPath}. Ignored.", InIntermediateProjectFilesPath);
+					Logger.LogInformation("\t{Ex}", Ex.Message);
 				}
 			}
 		}
@@ -51,7 +52,7 @@
 			return new EddieProjectFile(InitFilePath, BaseDir);
 		}
 		
-		private bool WriteEddieWorkset()
+		private bool WriteEddieWorkset(ILogger Logger)
 		{
 			bool bSuccess = false;
 			
@@ -86,20 +87,20 @@
 			string ProjectName = PrimaryProjectName;
 			string FilePath = PrimaryProjectPath + "/" + ProjectName + ".wkst";
 			
-			bSuccess = WriteFileIfChanged(FilePath, WorksetDataContent.ToString(), new UTF8Encoding());
+			bSuccess = WriteFileIfChanged(FilePath, WorksetDataContent.ToString(), Logger, new UTF8Encoding());
 			
 			return bSuccess;
 		}
 		
-		protected override bool WritePrimaryProjectFile(ProjectFile? UBTProject, PlatformProjectGeneratorCollection PlatformProjectGenerators)
+		protected override bool WritePrimaryProjectFile(ProjectFile? UBTProject, PlatformProjectGeneratorCollection PlatformProjectGenerators, ILogger Logger)
 		{
-			return WriteEddieWorkset();
+			return WriteEddieWorkset(Logger);
 		}
 		
-		protected override void ConfigureProjectFileGeneration(string[] Arguments, ref bool IncludeAllPlatforms)
+		protected override void ConfigureProjectFileGeneration(string[] Arguments, ref bool IncludeAllPlatforms, ILogger Logger)
 		{
 			// Call parent implementation first
-			base.ConfigureProjectFileGeneration(Arguments, ref IncludeAllPlatforms);
+			base.ConfigureProjectFileGeneration(Arguments, ref IncludeAllPlatforms, Logger);
 
 			if (bGeneratingGameProjectFiles)
 			{
diff --git a/Engine/Source/Programs/UnrealBuildTool/ProjectFiles/Eddie/EddieProject.cs b/Engine/Source/Programs/UnrealBuildTool/ProjectFiles/Eddie/EddieProject.cs
--- a/Engine/Source/Programs/UnrealBuildTool/ProjectFiles/Eddie/EddieProject.cs
+++ b/Engine/Source/Programs/UnrealBuildTool/ProjectFiles/Eddie/EddieProject.cs
@@ -9,6 +9,7 @@
 using System.Linq;
 using System.Text;
 using EpicGames.Core;
+using Microsoft.Extensions.Logging;
 
 namespace UnrealBuildTool
 {
@@ -146,7 +147,7 @@
 			}
 		}
 		
-		private void EmitProject(StringBuilder Content, Dictionary<string, EddieFolder> Folders)
+		private void EmitProject(StringBuilder Content, Dictionary<string, EddieFolder> Folders, ILogger Logger)
 		{
 			foreach (KeyValuePair<string, EddieFolder> CurGroup in Folders)
 			{
@@ -163,7 +164,7 @@
                 
                         ProjectFileContent.Append("AddFileGroup \"" + Path.GetFileName(CurGroup.Key) + "\" \"" + (CurGroup.Value.FullPath != null ? CurGroup.Value.FullPath : CurGroup.Value.Path) + "\"" + ProjectFileGenerator.NewLine);
                 
-                        EmitProject(ProjectFileContent, CurGroup.Value.Folders);
+                        EmitProject(ProjectFileContent, CurGroup.Value.Folders, Logger);
                         
                         foreach (EddieSourceFile File in CurGroup.Value.Files)
                         {
@@ -173,13 +174,13 @@
                         
                         ProjectFileContent.Append("EndFileGroup \"" + Path.GetFileName(CurGroup.Key) + "\"" + ProjectFileGenerator.NewLine);
                 
-                        ProjectFileGenerator.WriteFileIfChanged(CurGroup.Value.WorksetPath!, ProjectFileContent.ToString(), new UTF8Encoding());
+                        ProjectFileGenerator.WriteFileIfChanged(CurGroup.Value.WorksetPath!, ProjectFileContent.ToString(), Logger, new UTF8Encoding());
                     
                         Content.Append("AddFile \"" + Path.GetFileName(CurGroup.Key) + "\" \"" + CurGroup.Value.WorksetPath + "\"" + ProjectFileGenerator.NewLine);
                     }
                     else
                     {
-                        EmitProject(Content, CurGroup.Value.Folders);
+                        EmitProject(Content, CurGroup.Value.Folders, Logger);
                         
                         foreach (EddieSourceFile File in CurGroup.Value.Files)
                         {
@@ -192,7 +193,7 @@
 			}
 		}
 		
-		public override bool WriteProjectFile(List<UnrealTargetPlatform> InPlatforms, List<UnrealTargetConfiguration> InConfigurations, PlatformProjectGeneratorCollection PlatformProjectGenerators)
+		public override bool WriteProjectFile(List<UnrealTargetPlatform> InPlatforms, List<UnrealTargetConfiguration> InConfigurations, PlatformProjectGeneratorCollection PlatformProjectGenerators, ILogger Logger)
 		{
 			bool bSuccess = false;
 			
@@ -214,9 +215,9 @@
 			ProjectFileContent.Append("AddWorkset \"" + this.ToString() + ".wkst\" \"" + ProjectFilePath.FullName + "\"" + ProjectFileGenerator.NewLine);
 			
 			ParseSourceFilesIntoGroups();
-			EmitProject(ProjectFileContent, Folders);
+			EmitProject(ProjectFileContent, Folders, Logger);
 			
-			bSuccess = ProjectFileGenerator.WriteFileIfChanged(ProjectFilePath.FullName, ProjectFileContent.ToString(), new UTF8Encoding());
+			bSuccess = ProjectFileGenerator.WriteFileIfChanged(ProjectFilePath.FullName, ProjectFileContent.ToString(), Logger, new UTF8Encoding());
 			
 			return bSuccess;
 		}
diff --git a/Engine/Source/Programs/UnrealBuildTool/ProjectFiles/KDevelop/KDevelopGenerator.cs b/Engine/Source/Programs/UnrealBuildTool/ProjectFiles/KDevelop/KDevelopGenerator.cs
--- a/Engine/Source/Programs/UnrealBuildTool/ProjectFiles/KDevelop/KDevelopGenerator.cs
+++ b/Engine/Source/Programs/UnrealBuildTool/ProjectFiles/KDevelop/KDevelopGenerator.cs
@@ -7,6 +7,7 @@
 using System.Linq;
 using EpicGames.Core;
 using UnrealBuildBase;
+using Microsoft.Extensions.Logging;
 
 namespace UnrealBuildTool
 {
@@ -37,7 +38,7 @@
 			}
 		}
 
-		protected override bool WritePrimaryProjectFile(ProjectFile? UBTProject, PlatformProjectGeneratorCollection PlatformProjectGenerators)
+		protected override bool WritePrimaryProjectFile(ProjectFile? UBTProject, PlatformProjectGeneratorCollection PlatformProjectGenerators, ILogger Logger)
 		{
 			bool bSuccess = true;
 			return bSuccess;
@@ -243,7 +244,8 @@
 		/// Adds the include directory to the list, after converting it to an absolute path to UnrealEngine root directory.
 		/// </summary>
 		/// <param name="FileContent">File content.</param>
-		private void WriteIncludeSection(ref StringBuilder FileContent)
+		/// <param name="Logger"></param>
+		private void WriteIncludeSection(ref StringBuilder FileContent, ILogger Logger)
 		{
 			List<string> IncludeDirectories = new List<string>();
 			List<string> SystemIncludeDirectories = new List<string>();
@@ -259,7 +261,7 @@
 				KDevelopProjectFile? KDevelopProject = CurProject as KDevelopProjectFile;
 				if (KDevelopProject == null)
 				{
-					Log.TraceInformation("KDevelopProject == null");
+					Logger.LogInformation("KDevelopProject == null");
 					continue;
 				}
 
@@ -367,7 +369,8 @@
 		/// Write the defines section to the .kdev4/$ProjectName.kdev4 project file.
 		/// </summary>
 		/// <param name="FileContent">File content.</param>
-		private void WriteDefineSection(ref StringBuilder FileContent)
+		/// <param name="Logger">Logger for output</param>
+		private void WriteDefineSection(ref StringBuilder FileContent, ILogger Logger)
 		{
 			String Key = "";
 			String Value = "";
@@ -379,7 +382,7 @@
 				KDevelopProjectFile? KDevelopProject = CurProject as KDevelopProjectFile;
 				if (KDevelopProject == null)
 				{
-					Log.TraceInformation("KDevelopProject == null");
+					Logger.LogInformation("KDevelopProject == null");
 					continue;
 				}
 
@@ -438,7 +441,7 @@
 		}
 
 		/// Simple Place to call all the Write*Section functions.
-		private bool WriteKDevelopPro()
+		private bool WriteKDevelopPro(ILogger Logger)
 		{
 			// RAKE! Take one KDevelopProjectFileContent and pass
 			// it through each function that writes out the sections.
@@ -458,8 +461,8 @@
 			WriteKDevPrimaryProjectSection(ref KDevelopPrimaryFileContent, PrimaryProjectName);
 
 			WriteCommandSection(ref KDevelopFileContent);
-			WriteIncludeSection(ref IncludesFileContent);
-			WriteDefineSection(ref DefinesFileContent);
+			WriteIncludeSection(ref IncludesFileContent, Logger);
+			WriteDefineSection(ref DefinesFileContent, Logger);
 			WriteExcludeSection(ref KDevelopFileContent);
 
 			// Write the primary kdev file.
@@ -476,18 +479,18 @@
 			string FullDefinesFileName = Path.Combine(FullPrimaryProjectPath, DefinesFileName);
 			string FullIncludesFileName = Path.Combine(FullPrimaryProjectPath, IncludesFileName);
 
-			WriteFileIfChanged(FullDefinesFileName, DefinesFileContent.ToString());
-			WriteFileIfChanged(FullIncludesFileName, IncludesFileContent.ToString());
+			WriteFileIfChanged(FullDefinesFileName, DefinesFileContent.ToString(), Logger);
+			WriteFileIfChanged(FullIncludesFileName, IncludesFileContent.ToString(), Logger);
 
-			return WriteFileIfChanged(FullKDevelopPrimaryFileName, KDevelopPrimaryFileContent.ToString()) &&
-			WriteFileIfChanged(FullKDevelopFileName, KDevelopFileContent.ToString());
+			return WriteFileIfChanged(FullKDevelopPrimaryFileName, KDevelopPrimaryFileContent.ToString(), Logger) &&
+			WriteFileIfChanged(FullKDevelopFileName, KDevelopFileContent.ToString(), Logger);
 		}
 
 		/// ProjectFileGenerator interface
 		//protected override bool WritePrimaryProjectFile( ProjectFile UBTProject )
-		protected override bool WriteProjectFiles(PlatformProjectGeneratorCollection PlatformProjectGenerators)
+		protected override bool WriteProjectFiles(PlatformProjectGeneratorCollection PlatformProjectGenerators, ILogger Logger)
 		{
-			return WriteKDevelopPro();
+			return WriteKDevelopPro(Logger);
 		}
 
 		/// ProjectFileGenerator interface
@@ -503,7 +506,7 @@
 		}
 
 		/// ProjectFileGenerator interface
-		public override void CleanProjectFiles(DirectoryReference InPrimaryProjectDirectory, string InPrimaryProjectName, DirectoryReference InIntermediateProjectFilesDirectory)
+		public override void CleanProjectFiles(DirectoryReference InPrimaryProjectDirectory, string InPrimaryProjectName, DirectoryReference InIntermediateProjectFilesDirectory, ILogger Logger)
 		{
 		}
 	}
diff --git a/Engine/Source/Programs/UnrealBuildTool/ProjectFiles/Make/MakefileGenerator.cs b/Engine/Source/Programs/UnrealBuildTool/ProjectFiles/Make/MakefileGenerator.cs
--- a/Engine/Source/Programs/UnrealBuildTool/ProjectFiles/Make/MakefileGenerator.cs
+++ b/Engine/Source/Programs/UnrealBuildTool/ProjectFiles/Make/MakefileGenerator.cs
@@ -7,6 +7,7 @@
 using EpicGames.Core;
 using UnrealBuildBase;
 using System.Linq;
+using Microsoft.Extensions.Logging;
 
 namespace UnrealBuildTool
 {
@@ -48,13 +49,13 @@
 			}
 		}
 
-		protected override bool WritePrimaryProjectFile(ProjectFile? UBTProject, PlatformProjectGeneratorCollection PlatformProjectGenerators)
+		protected override bool WritePrimaryProjectFile(ProjectFile? UBTProject, PlatformProjectGeneratorCollection PlatformProjectGenerators, ILogger Logger)
 		{
 			bool bSuccess = true;
 			return bSuccess;
 		}
 
-		private bool WriteMakefile()
+		private bool WriteMakefile(ILogger Logger)
 		{
 			string GameProjectFile = "";
 			string BuildCommand = "";
@@ -181,14 +182,14 @@
 
 			MakefileContent.Append("\n.PHONY: $(TARGETS)\n");
 			FileReference FullFileName = FileReference.Combine(PrimaryProjectPath, FileName);
-			return WriteFileIfChanged(FullFileName.FullName, MakefileContent.ToString());
+			return WriteFileIfChanged(FullFileName.FullName, MakefileContent.ToString(), Logger);
 		}
 
 		/// ProjectFileGenerator interface
 		//protected override bool WritePrimaryProjectFile( ProjectFile UBTProject )
-		protected override bool WriteProjectFiles(PlatformProjectGeneratorCollection PlatformProjectGenerators)
+		protected override bool WriteProjectFiles(PlatformProjectGeneratorCollection PlatformProjectGenerators, ILogger Logger)
 		{
-			return WriteMakefile();
+			return WriteMakefile(Logger);
 		}
 
 		/// ProjectFileGenerator interface
@@ -204,7 +205,7 @@
 		}
 
 		/// ProjectFileGenerator interface
-		public override void CleanProjectFiles(DirectoryReference InPrimaryProjectDirectory, string InPrimaryProjectName, DirectoryReference InIntermediateProjectFilesDirectory)
+		public override void CleanProjectFiles(DirectoryReference InPrimaryProjectDirectory, string InPrimaryProjectName, DirectoryReference InIntermediateProjectFilesDirectory, ILogger Logger)
 		{
 		}
 	}
diff --git a/Engine/Source/Programs/UnrealBuildTool/ProjectFiles/PlatformProjectGenerator.cs b/Engine/Source/Programs/UnrealBuildTool/ProjectFiles/PlatformProjectGenerator.cs
--- a/Engine/Source/Programs/UnrealBuildTool/ProjectFiles/PlatformProjectGenerator.cs
+++ b/Engine/Source/Programs/UnrealBuildTool/ProjectFiles/PlatformProjectGenerator.cs
@@ -6,6 +6,7 @@
 using System.Diagnostics;
 using System.IO;
 using EpicGames.Core;
+using Microsoft.Extensions.Logging;
 
 namespace UnrealBuildTool
 {
@@ -14,12 +15,16 @@
 	/// </summary>
 	abstract class PlatformProjectGenerator
 	{
+		protected readonly ILogger Logger;
+
 		/// <summary>
 		/// Constructor
 		/// </summary>
 		/// <param name="Arguments">Command line arguments passed to the project generator</param>
-		public PlatformProjectGenerator(CommandLineArguments Arguments)
+		/// <param name="Logger">Logger for output</param>
+		public PlatformProjectGenerator(CommandLineArguments Arguments, ILogger Logger)
 		{
+			this.Logger = Logger;
 		}
 
 		/// <summary>
diff --git a/Engine/Source/Programs/UnrealBuildTool/ProjectFiles/PlatformProjectGeneratorCollection.cs b/Engine/Source/Programs/UnrealBuildTool/ProjectFiles/PlatformProjectGeneratorCollection.cs
--- a/Engine/Source/Programs/UnrealBuildTool/ProjectFiles/PlatformProjectGeneratorCollection.cs
+++ b/Engine/Source/Programs/UnrealBuildTool/ProjectFiles/PlatformProjectGeneratorCollection.cs
@@ -8,6 +8,7 @@
 using System.Threading.Tasks;
 using EpicGames.Core;
 using UnrealBuildBase;
+using Microsoft.Extensions.Logging;
 
 namespace UnrealBuildTool
 {
@@ -23,7 +24,8 @@
 		/// </summary>
 		/// <param name="InPlatform">  The UnrealTargetPlatform to register with</param>
 		/// <param name="InProjectGenerator">The UEPlatformProjectGenerator instance to use for the InPlatform</param>
-		public void RegisterPlatformProjectGenerator(UnrealTargetPlatform InPlatform, PlatformProjectGenerator InProjectGenerator)
+		/// <param name="Logger">Logger for output</param>
+		public void RegisterPlatformProjectGenerator(UnrealTargetPlatform InPlatform, PlatformProjectGenerator InProjectGenerator, ILogger Logger)
 		{
 			// Make sure the build platform is legal
 			UEBuildPlatform? BuildPlatform;
@@ -31,7 +33,7 @@
 			{
 				if (ProjectGeneratorDictionary.ContainsKey(InPlatform) == true)
 				{
-					Log.TraceInformation("RegisterPlatformProjectGenerator Warning: Registering project generator {0} for {1} when it is already set to {2}",
+					Logger.LogInformation("RegisterPlatformProjectGenerator Warning: Registering project generator {Generator} for {Platform} when it is already set to {ExistingGenerator}",
 						InProjectGenerator.ToString(), InPlatform.ToString(), ProjectGeneratorDictionary[InPlatform].ToString());
 					ProjectGeneratorDictionary[InPlatform] = InProjectGenerator;
 				}
@@ -42,7 +44,7 @@
 			}
 			else
 			{
-				Log.TraceVerbose("Skipping project file generator registration for {0} due to no valid BuildPlatform.", InPlatform.ToString());
+				Logger.LogDebug("Skipping project file generator registration for {Platform} due to no valid BuildPlatform.", InPlatform.ToString());
 			}
 		}
 
@@ -62,7 +64,7 @@
 			{
 				return null;
 			}
-			throw new BuildException("GetPlatformProjectGenerator: No PlatformProjectGenerator found for {0}", InPlatform.ToString());
+			throw new BuildException("GetPlatformProjectGenerator: No PlatformProjectGenerator found for {Platform}", InPlatform.ToString());
 		}
 
 		/// <summary>
@@ -100,7 +102,7 @@
 		{
 			if (ProjectGeneratorDictionary.ContainsKey(InPlatform) == true)
 			{
-				ProjectGeneratorDictionary[InPlatform].GenerateGameProperties(Configuration, VCProjectFileContent, TargetType, RootDirectory, TargetFilePath); ;
+				ProjectGeneratorDictionary[InPlatform].GenerateGameProperties(Configuration, VCProjectFileContent, TargetType, RootDirectory, TargetFilePath);
 			}
 			return true;
 		}
diff --git a/Engine/Source/Programs/UnrealBuildTool/ProjectFiles/Project.cs b/Engine/Source/Programs/UnrealBuildTool/ProjectFiles/Project.cs
--- a/Engine/Source/Programs/UnrealBuildTool/ProjectFiles/Project.cs
+++ b/Engine/Source/Programs/UnrealBuildTool/ProjectFiles/Project.cs
@@ -10,6 +10,7 @@
 using System.Text;
 using EpicGames.Core;
 using UnrealBuildBase;
+using Microsoft.Extensions.Logging;
 
 namespace UnrealBuildTool
 {
@@ -474,8 +475,9 @@
 		/// <param name="InPlatforms">The platforms to write the project files for</param>
 		/// <param name="InConfigurations">The configurations to add to the project files</param>
 		/// <param name="PlatformProjectGenerators">The registered platform project generators</param>
+		/// <param name="Logger">Logger for output</param>
 		/// <returns>True on success</returns>
-		public virtual bool WriteProjectFile(List<UnrealTargetPlatform> InPlatforms, List<UnrealTargetConfiguration> InConfigurations, PlatformProjectGeneratorCollection PlatformProjectGenerators)
+		public virtual bool WriteProjectFile(List<UnrealTargetPlatform> InPlatforms, List<UnrealTargetConfiguration> InConfigurations, PlatformProjectGeneratorCollection PlatformProjectGenerators, ILogger Logger)
 		{
 			throw new BuildException("UnrealBuildTool cannot automatically generate this project type because WriteProjectFile() was not overridden.");
 		}
@@ -486,8 +488,9 @@
 		/// <param name="InPlatforms">The platforms to write the project files for</param>
 		/// <param name="InConfigurations">The configurations to add to the project files</param>
 		/// <param name="PlatformProjectGenerators">The registered platform project generators</param>
+		/// <param name="Logger">Logger for output</param>
 		/// <returns>List of project files written</returns>
-		public virtual List<Tuple<ProjectFile, string>>? WriteDebugProjectFiles(List<UnrealTargetPlatform> InPlatforms, List<UnrealTargetConfiguration> InConfigurations, PlatformProjectGeneratorCollection PlatformProjectGenerators)
+		public virtual List<Tuple<ProjectFile, string>>? WriteDebugProjectFiles(List<UnrealTargetPlatform> InPlatforms, List<UnrealTargetConfiguration> InConfigurations, PlatformProjectGeneratorCollection PlatformProjectGenerators, ILogger Logger)
 		{
 			return null;
 		}
diff --git a/Engine/Source/Programs/UnrealBuildTool/ProjectFiles/ProjectFileGenerator.cs b/Engine/Source/Programs/UnrealBuildTool/ProjectFiles/ProjectFileGenerator.cs
--- a/Engine/Source/Programs/UnrealBuildTool/ProjectFiles/ProjectFileGenerator.cs
+++ b/Engine/Source/Programs/UnrealBuildTool/ProjectFiles/ProjectFileGenerator.cs
@@ -12,6 +12,7 @@ using EpicGames.Core;
 using System.Xml.Linq;
 using System.Xml;
 using UnrealBuildBase;
+using Microsoft.Extensions.Logging;
 
 namespace UnrealBuildTool
 {
@@ -158,8 +159,9 @@ namespace UnrealBuildTool
 		/// Parses a list of project file formats from a string
 		/// </summary>
 		/// <param name="Formats"></param>
+		/// <param name="Logger"></param>
 		/// <returns>Sequence of project file formats</returns>
-		public static IEnumerable<ProjectFileFormat> ParseFormatList(string Formats)
+		public static IEnumerable<ProjectFileFormat> ParseFormatList(string Formats, ILogger Logger)
 		{
 			foreach (string FormatName in Formats.Split('+').Select(x => x.Trim()))
 			{
@@ -170,7 +172,7 @@ namespace UnrealBuildTool
 				}
 				else
 				{
-					Log.TraceError("Invalid project file format '{0}'", FormatName);
+					Logger.LogError("Invalid project file format '{FormatName}'", FormatName);
 				}
 			}
 		}
@@ -433,7 +435,7 @@ namespace UnrealBuildTool
 		/// Adds all rules project files to the solution.
 		/// </summary>
 		void AddRulesModules(Rules.RulesFileType RulesFileType, string ProgramSubDirectory, List<ProjectFile> AddedProjectFiles, 
-			List<FileReference> UnrealProjectFiles, PrimaryProjectFolder RootFolder, PrimaryProjectFolder ProgramsFolder)
+			List<FileReference> UnrealProjectFiles, PrimaryProjectFolder RootFolder, PrimaryProjectFolder ProgramsFolder, ILogger Logger)
 		{
 			List<DirectoryReference> BuildFolders = new List<DirectoryReference>();
 			foreach (FileReference UnrealProjectFile in UnrealProjectFiles)
@@ -454,7 +456,7 @@ namespace UnrealBuildTool
 			{
 				if (FileReference.Exists(ProjectFile))
 				{
-					VCSharpProjectFile Project = new VCSharpProjectFile(ProjectFile);
+					VCSharpProjectFile Project = new VCSharpProjectFile(ProjectFile, Logger);
 
 					Project.ShouldBuildForAllSolutionTargets = false;//true;
 					AddExistingProjectFile(Project, bForceDevelopmentConfiguration: true);
@@ -494,7 +496,7 @@ namespace UnrealBuildTool
 		/// <summary>
 		/// Adds all the DotNet folder to the solution.
 		/// </summary>
-		void AddSharedDotNetModules(List<DirectoryReference> AllEngineDirectories, PrimaryProjectFolder ProgramsFolder)
+		void AddSharedDotNetModules(List<DirectoryReference> AllEngineDirectories, PrimaryProjectFolder ProgramsFolder, ILogger Logger)
 		{
 			foreach (DirectoryReference EngineDir in AllEngineDirectories)
 			{
@@ -511,7 +513,7 @@ namespace UnrealBuildTool
 						PrimaryProjectFolder Folder = ProgramsFolder.AddSubFolder("Shared");
 						foreach (FileInfo ProjectFile in ProjectFiles)
 						{
-							VCSharpProjectFile Project = new VCSharpProjectFile(new FileReference(ProjectFile));
+							VCSharpProjectFile Project = new VCSharpProjectFile(new FileReference(ProjectFile), Logger);
 							Project.ShouldBuildForAllSolutionTargets = false;
 							AddExistingProjectFile(Project, bForceDevelopmentConfiguration: true);
 							Folder.ChildProjects.Add(Project);
@@ -544,7 +546,7 @@ namespace UnrealBuildTool
 		/// <summary>
 		/// Finds all csproj within Engine/Source/Programs, and add them if their .uprogram file exists.
 		/// </summary>
-		void DiscoverCSharpProgramProjects(List<DirectoryReference> AllEngineDirectories, List<FileReference> AllGameProjects, PrimaryProjectFolder ProgramsFolder)
+		void DiscoverCSharpProgramProjects(List<DirectoryReference> AllEngineDirectories, List<FileReference> AllGameProjects, PrimaryProjectFolder ProgramsFolder, ILogger Logger)
 		{
 			List<FileReference> FoundProjects = new List<FileReference>();
 			List<DirectoryReference> ProjectDirs = new List<DirectoryReference>();
@@ -570,7 +572,7 @@ namespace UnrealBuildTool
 				}
 			}
 
-			string[] UnsupportedPlatformNames = Utils.MakeListOfUnsupportedPlatforms(SupportedPlatforms, bIncludeUnbuildablePlatforms: true).ToArray();
+			string[] UnsupportedPlatformNames = Utils.MakeListOfUnsupportedPlatforms(SupportedPlatforms, bIncludeUnbuildablePlatforms: true, Logger).ToArray();
 			foreach (FileReference FoundProject in FoundProjects)
 			{
 				foreach (DirectoryReference ProjectDir in ProjectDirs)
@@ -579,7 +581,7 @@ namespace UnrealBuildTool
 					{
 						if (!FoundProject.ContainsAnyNames(UnsupportedPlatformNames, ProjectDir))
 						{
-							VCSharpProjectFile Project = new VCSharpProjectFile(FoundProject);
+							VCSharpProjectFile Project = new VCSharpProjectFile(FoundProject, Logger);
 
 							Project.ShouldBuildForAllSolutionTargets = true;
 							Project.ShouldBuildByDefaultForSolutionTargets = true;
@@ -589,7 +591,7 @@ namespace UnrealBuildTool
 						}
 						else
 						{
-							Log.TraceLog($"Skipping C# project \"{FoundProject}\" due to unsupported platform");
+							Logger.LogDebug("Skipping C# project \"{FoundProject}\" due to unsupported platform", FoundProject);
 						}
 						break;
 					}
@@ -597,7 +599,7 @@ namespace UnrealBuildTool
 			}
 		}
 
-		private VCSharpProjectFile CreateRulesAssemblyProject(RulesAssembly RulesAssembly, DirectoryReference FSPathBase)
+		private VCSharpProjectFile CreateRulesAssemblyProject(RulesAssembly RulesAssembly, DirectoryReference FSPathBase, ILogger Logger)
 		{
 			DirectoryReference ProjectFilesDirectory = DirectoryReference.Combine(FSPathBase, "Intermediate/Build/BuildRulesProjects", RulesAssembly.GetSimpleAssemblyName()!);
 			DirectoryReference.CreateDirectory(ProjectFilesDirectory);
@@ -642,7 +644,7 @@ namespace UnrealBuildTool
 				Writer.WriteLine("</Project>");
 			}
 
-			VCSharpProjectFile ModuleFilesProject = new VCSharpProjectFile(ModuleFilesProjectLocation);
+			VCSharpProjectFile ModuleFilesProject = new VCSharpProjectFile(ModuleFilesProjectLocation, Logger);
 			return ModuleFilesProject;
 		}
 
@@ -681,8 +683,9 @@ namespace UnrealBuildTool
 		/// <summary>
 		/// Finds the game projects that we're generating project files for
 		/// </summary>
+		/// <param name="Logger"></param>
 		/// <returns>List of project files</returns>
-		public List<FileReference> FindGameProjects()
+		public List<FileReference> FindGameProjects(ILogger Logger)
 		{
 			List<FileReference> ProjectFiles = new List<FileReference>();
 			if (OnlyGameProject != null)
@@ -691,7 +694,7 @@ namespace UnrealBuildTool
 			}
 			else
 			{
-				ProjectFiles.AddRange(NativeProjects.EnumerateProjectFiles());
+				ProjectFiles.AddRange(NativeProjects.EnumerateProjectFiles(Logger));
 			}
 			return ProjectFiles;
 		}
@@ -763,17 +766,18 @@ namespace UnrealBuildTool
 		/// </summary>
 		/// <param name="PlatformProjectGenerators">The registered platform project generators</param>
 		/// <param name="Arguments">Command-line arguments</param>
-		public virtual bool GenerateProjectFiles(PlatformProjectGeneratorCollection PlatformProjectGenerators, String[] Arguments)
+		/// <param name="Logger">Logger for output</param>
+		public virtual bool GenerateProjectFiles(PlatformProjectGeneratorCollection PlatformProjectGenerators, String[] Arguments, ILogger Logger)
 		{
 			bool bSuccess = true;
 
 			// Parse project generator options
 			bool IncludeAllPlatforms = true;
-			ConfigureProjectFileGeneration(Arguments, ref IncludeAllPlatforms);
+			ConfigureProjectFileGeneration(Arguments, ref IncludeAllPlatforms, Logger);
 
 			if (bGeneratingGameProjectFiles || Unreal.IsEngineInstalled())
 			{
-				Log.TraceInformation("Discovering modules, targets and source code for project...");
+				Logger.LogInformation("Discovering modules, targets and source code for project...");
 
 				PrimaryProjectPath = OnlyGameProject!.Directory;
 
@@ -846,21 +850,21 @@ namespace UnrealBuildTool
 			bool bCleanProjectFiles = Arguments.Any(x => x.Equals("-CleanProjects", StringComparison.InvariantCultureIgnoreCase));
 			if (bCleanProjectFiles)
 			{
-				CleanProjectFiles(PrimaryProjectPath, PrimaryProjectName, IntermediateProjectFilesPath);
+				CleanProjectFiles(PrimaryProjectPath, PrimaryProjectName, IntermediateProjectFilesPath, Logger);
 			}
 
 			// Figure out which platforms we should generate project files for.
 			string SupportedPlatformNames;
-			SetupSupportedPlatformsAndConfigurations(IncludeAllPlatforms: IncludeAllPlatforms, SupportedPlatformNames: out SupportedPlatformNames);
+			SetupSupportedPlatformsAndConfigurations(IncludeAllPlatforms: IncludeAllPlatforms, Logger: Logger, SupportedPlatformNames: out SupportedPlatformNames);
 
-			Log.TraceVerbose("Detected supported platforms: " + SupportedPlatformNames);
+			Logger.LogDebug("Detected supported platforms: {Platforms}", SupportedPlatformNames);
 
 			// Build the list of games to generate projects for
-			List<FileReference> AllGameProjects = FindGameProjects();
+			List<FileReference> AllGameProjects = FindGameProjects(Logger);
 
 			// Find all of the target files.  This will filter out any modules or targets that don't
 			// belong to platforms we're generating project files for.
-			List<FileReference> AllTargetFiles = DiscoverTargets(AllGameProjects);
+			List<FileReference> AllTargetFiles = DiscoverTargets(AllGameProjects, Logger);
 
 			// Sort the targets by name. When we have multiple targets of a given type for a project, we'll use the order to determine which goes in the primary project file (so that client names with a suffix will go into their own project).
 			AllTargetFiles = AllTargetFiles.OrderBy(x => x.FullName, StringComparer.OrdinalIgnoreCase).ToList();
@@ -891,7 +895,7 @@ namespace UnrealBuildTool
 			if(IncludeCppSource)
 			{
 				// Setup buildable projects for all targets
-				AddProjectsForAllTargets(PlatformProjectGenerators, AllGameProjects, AllTargetFiles, Arguments, ref EngineProject, GameProjects, ProgramProjects, RulesAssemblies);
+				AddProjectsForAllTargets(PlatformProjectGenerators, AllGameProjects, AllTargetFiles, Arguments, ref EngineProject, GameProjects, ProgramProjects, RulesAssemblies, Logger);
 
 				// Add projects for mods
 				AddProjectsForMods(GameProjects, ModProjects);
@@ -963,7 +967,7 @@ namespace UnrealBuildTool
 
 						if (bIncludeShaderSource)
 						{
-							Log.TraceVerbose("Adding shader source code...");
+							Logger.LogDebug("Adding shader source code...");
 
 							// Find shader source files and generate stub project
 							AddEngineShaderSource(EngineProject);
@@ -971,17 +975,17 @@ namespace UnrealBuildTool
 
 						if (bIncludeBuildSystemFiles)
 						{
-							Log.TraceVerbose("Adding build system files...");
+							Logger.LogDebug("Adding build system files...");
 
 							AddEngineBuildFiles(EngineProject);
 						}
 
 						if (bIncludeDocumentation)
 						{
-							AddEngineDocumentation(EngineProject);
+							AddEngineDocumentation(EngineProject, Logger);
 						}
 
-						List<Tuple<ProjectFile, string>>? NewProjectFiles = EngineProject.WriteDebugProjectFiles(SupportedPlatforms, SupportedConfigurations, PlatformProjectGenerators);
+						List<Tuple<ProjectFile, string>>? NewProjectFiles = EngineProject.WriteDebugProjectFiles(SupportedPlatforms, SupportedConfigurations, PlatformProjectGenerators, Logger);
 
 						if (NewProjectFiles != null)
 						{
@@ -1016,7 +1020,7 @@ namespace UnrealBuildTool
 							RootFolder.AddSubFolder("Games").ChildProjects.Add(CurGameProject);
 						}
 
-						List<Tuple<ProjectFile, string>>? NewProjectFiles = CurGameProject.WriteDebugProjectFiles(SupportedPlatforms, SupportedConfigurations, PlatformProjectGenerators);
+						List<Tuple<ProjectFile, string>>? NewProjectFiles = CurGameProject.WriteDebugProjectFiles(SupportedPlatforms, SupportedConfigurations, PlatformProjectGenerators, Logger);
 
 						if (NewProjectFiles != null)
 						{
@@ -1059,19 +1063,19 @@ namespace UnrealBuildTool
 			}
 
 			// Setup "stub" projects for all modules
-			AddProjectsForAllModules(AllGameProjects, ProgramProjects, ModProjects, AllModuleFiles, AdditionalSearchPaths, bGatherThirdPartySource);
+			AddProjectsForAllModules(AllGameProjects, ProgramProjects, ModProjects, AllModuleFiles, AdditionalSearchPaths, bGatherThirdPartySource, Logger);
 
 			{
 				if (bIncludeDotNetPrograms)
 				{
 					// gather engine directories across extension dirs
-					string[] UnsupportedPlatformNames = Utils.MakeListOfUnsupportedPlatforms(SupportedPlatforms, bIncludeUnbuildablePlatforms: true).ToArray();
+					string[] UnsupportedPlatformNames = Utils.MakeListOfUnsupportedPlatforms(SupportedPlatforms, bIncludeUnbuildablePlatforms: true, Logger).ToArray();
 					List<DirectoryReference> AllEngineDirectories =
 						Unreal.GetExtensionDirs(Unreal.EngineDirectory, "Source/Programs").Where(x =>
 						{
 							if (x.ContainsAnyNames(UnsupportedPlatformNames, Unreal.EngineDirectory))
 							{
-								Log.TraceLog($"Skipping any C# project files in \"{x}\" due to unsupported platform");
+								Logger.LogDebug("Skipping any C# project files in \"{x}\" due to unsupported platform", x);
 								return false;
 							}
 							return true;
@@ -1080,28 +1084,28 @@ namespace UnrealBuildTool
 					PrimaryProjectFolder ProgramsFolder = RootFolder.AddSubFolder("Programs");
 
 					// Add UnrealBuildTool to the primary project
-					AddUnrealBuildToolProject(ProgramsFolder);
+					AddUnrealBuildToolProject(ProgramsFolder, Logger);
 
 					// Add AutomationTool to the primary project
-					VCSharpProjectFile AutomationToolProject = AddSimpleCSharpProject("AutomationTool", bShouldBuildForAllSolutionTargets: true, bForceDevelopmentConfiguration: true);
+					VCSharpProjectFile AutomationToolProject = AddSimpleCSharpProject("AutomationTool", Logger, bShouldBuildForAllSolutionTargets: true, bForceDevelopmentConfiguration: true);
 					if (AutomationToolProject != null)
 						ProgramsFolder.ChildProjects.Add(AutomationToolProject);
 
 					// Add automation.csproj files to the primary project
-					AddRulesModules(Rules.RulesFileType.AutomationModule, "Automation", AutomationProjectFiles, AllGameProjects, RootFolder, ProgramsFolder);
+					AddRulesModules(Rules.RulesFileType.AutomationModule, "Automation", AutomationProjectFiles, AllGameProjects, RootFolder, ProgramsFolder, Logger);
 
 					// Add ubtplugin.csproj files to the primary project
-					AddRulesModules(Rules.RulesFileType.UbtPlugin, "UnrealBuildTool.Plugins", UbtPluginProjectFiles, AllGameProjects, RootFolder, ProgramsFolder);
+					AddRulesModules(Rules.RulesFileType.UbtPlugin, "UnrealBuildTool.Plugins", UbtPluginProjectFiles, AllGameProjects, RootFolder, ProgramsFolder, Logger);
 
 					// Add shared projects
-					AddSharedDotNetModules(AllEngineDirectories, ProgramsFolder);
+					AddSharedDotNetModules(AllEngineDirectories, ProgramsFolder, Logger);
 
 					// Discover C# programs which should additionally be included in the solution.
-					DiscoverCSharpProgramProjects(AllEngineDirectories, AllGameProjects, ProgramsFolder);
+					DiscoverCSharpProgramProjects(AllEngineDirectories, AllGameProjects, ProgramsFolder, Logger);
 
 					foreach (KeyValuePair<RulesAssembly, DirectoryReference> RulesAssemblyEntry in RulesAssemblies)
 					{
-						VCSharpProjectFile RulesProject = CreateRulesAssemblyProject(RulesAssemblyEntry.Key, RulesAssemblyEntry.Value);
+						VCSharpProjectFile RulesProject = CreateRulesAssemblyProject(RulesAssemblyEntry.Key, RulesAssemblyEntry.Value, Logger);
 						AddExistingProjectFile(RulesProject, bForceDevelopmentConfiguration: false);
 						RootFolder.AddSubFolder("Rules").ChildProjects.Add(RulesProject);
 					}
@@ -1162,21 +1166,21 @@ namespace UnrealBuildTool
 				}
 
 				// write out any additional debug information for the solution (such as UnrealVS configuration)
-				WriteDebugSolutionFiles(PlatformProjectGenerators, IntermediateProjectFilesPath);
+				WriteDebugSolutionFiles(PlatformProjectGenerators, IntermediateProjectFilesPath, Logger);
 
 				// Generate IntelliSense data if we need to.  This involves having UBT simulate the action compilation of
 				// the targets so that we can extra the compiler defines, include paths, etc.
-				GenerateIntelliSenseData(Arguments, IntelliSenseTargetFiles);
+				GenerateIntelliSenseData(Arguments, IntelliSenseTargetFiles, Logger);
 
 				// Write the project files
-				WriteProjectFiles(PlatformProjectGenerators);
-				Log.TraceVerbose("Project generation complete ({0} generated, {1} imported)", GeneratedProjectFiles.Count, OtherProjectFiles.Count);
+				WriteProjectFiles(PlatformProjectGenerators, Logger);
+				Logger.LogDebug("Project generation complete ({GeneratedProjectFilesCount} generated, {OtherProjectFilesCount} imported)", GeneratedProjectFiles.Count, OtherProjectFiles.Count);
 
 				// Generate all the target info files for the editor
 				foreach (FileReference ProjectFile in AllGameProjects)
 				{
-					RulesAssembly RulesAssembly = RulesCompiler.CreateProjectRulesAssembly(ProjectFile, false, false, false);
-					QueryTargetsMode.WriteTargetInfo(ProjectFile, RulesAssembly, QueryTargetsMode.GetDefaultOutputFile(ProjectFile), new CommandLineArguments(GetTargetArguments(Arguments)));
+					RulesAssembly RulesAssembly = RulesCompiler.CreateProjectRulesAssembly(ProjectFile, false, false, false, Logger);
+					QueryTargetsMode.WriteTargetInfo(ProjectFile, RulesAssembly, QueryTargetsMode.GetDefaultOutputFile(ProjectFile), new CommandLineArguments(GetTargetArguments(Arguments)), Logger);
 				}
 			}
 
@@ -1220,14 +1224,16 @@ namespace UnrealBuildTool
 		/// <param name="InPrimaryProjectDirectory">The primary project directory</param>
 		/// <param name="InPrimaryProjectName">The name of the primary project</param>
 		/// <param name="InIntermediateProjectFilesDirectory">The intermediate path of project files</param>
-		public abstract void CleanProjectFiles(DirectoryReference InPrimaryProjectDirectory, string InPrimaryProjectName, DirectoryReference InIntermediateProjectFilesDirectory);
+		/// <param name="Logger">Logger for output</param>
+		public abstract void CleanProjectFiles(DirectoryReference InPrimaryProjectDirectory, string InPrimaryProjectName, DirectoryReference InIntermediateProjectFilesDirectory, ILogger Logger);
 
 		/// <summary>
 		/// Configures project generator based on command-line options
 		/// </summary>
 		/// <param name="Arguments">Arguments passed into the program</param>
 		/// <param name="IncludeAllPlatforms">True if all platforms should be included</param>
-		protected virtual void ConfigureProjectFileGeneration(String[] Arguments, ref bool IncludeAllPlatforms)
+		/// <param name="Logger">Logger for output</param>
+		protected virtual void ConfigureProjectFileGeneration(String[] Arguments, ref bool IncludeAllPlatforms, ILogger Logger)
 		{
 			if (PlatformNames != null)
 			{
@@ -1277,7 +1283,7 @@ namespace UnrealBuildTool
 							}
 							else
 							{
-								Log.TraceWarning("ProjectFiles invalid platform specified: {0}", PlatformString);
+								Logger.LogWarning("ProjectFiles invalid platform specified: {PlatformString}", PlatformString);
 							}
 
 							// Append the platform suffix to the solution name
@@ -1564,12 +1570,12 @@ namespace UnrealBuildTool
 		/// Finds all target files (filtering by platform)
 		/// </summary>
 		/// <returns>Filtered list of target files</returns>
-		protected List<FileReference> DiscoverTargets(List<FileReference> AllGameProjects)
+		protected List<FileReference> DiscoverTargets(List<FileReference> AllGameProjects, ILogger Logger)
 		{
 			List<FileReference> AllTargetFiles = new List<FileReference>();
 
 			// Make a list of all platform name strings that we're *not* including in the project files
-			List<string> UnsupportedPlatformNameStrings = Utils.MakeListOfUnsupportedPlatforms(SupportedPlatforms, bIncludeUnbuildablePlatforms: true);
+			List<string> UnsupportedPlatformNameStrings = Utils.MakeListOfUnsupportedPlatforms(SupportedPlatforms, bIncludeUnbuildablePlatforms: true, Logger);
 
 			// Locate all targets (*.Target.cs files)
 			List<FileReference> FoundTargetFiles = Rules.FindAllRulesSourceFiles(Rules.RulesFileType.Target, AllGameProjects.Select(x => x.Directory).ToList(), ForeignPlugins: null, AdditionalSearchPaths: null, bIncludeEngine: bIncludeEngineSource, bIncludeTempTargets: bIncludeTempTargets);
@@ -1810,14 +1816,14 @@ namespace UnrealBuildTool
 		/// <summary>
 		/// Adds UnrealBuildTool to the primary project
 		/// </summary>
-		private void AddUnrealBuildToolProject(PrimaryProjectFolder ProgramsFolder)
+		private void AddUnrealBuildToolProject(PrimaryProjectFolder ProgramsFolder, ILogger Logger)
 		{
 			DirectoryReference ProgramsDirectory = DirectoryReference.Combine(Unreal.EngineSourceDirectory, "Programs");
 
 			FileReference ProjectFileName = FileReference.Combine(ProgramsDirectory, "UnrealBuildTool", "UnrealBuildTool.csproj");
 			if (FileReference.Exists(ProjectFileName))
 			{
-				VCSharpProjectFile UnrealBuildToolProject = new VCSharpProjectFile(ProjectFileName);
+				VCSharpProjectFile UnrealBuildToolProject = new VCSharpProjectFile(ProjectFileName, Logger);
 				UnrealBuildToolProject.ShouldBuildForAllSolutionTargets = true;
 
 				// Store it off as we need it when generating target projects.
@@ -1833,7 +1839,7 @@ namespace UnrealBuildTool
 			FileReference TestsProjectFileName = FileReference.Combine(ProgramsDirectory, "UnrealBuildTool.Tests", "UnrealBuildTool.Tests.csproj");
 			if (FileReference.Exists(TestsProjectFileName))
 			{
-				VCSharpProjectFile UbtTestsProject = new VCSharpProjectFile(TestsProjectFileName);
+				VCSharpProjectFile UbtTestsProject = new VCSharpProjectFile(TestsProjectFileName, Logger);
 				AddExistingProjectFile(UbtTestsProject, bNeedsAllPlatformAndConfigurations: true, bForceDevelopmentConfiguration: true);
 				ProgramsFolder.ChildProjects.Add(UbtTestsProject);
 			}
@@ -1843,11 +1849,12 @@ namespace UnrealBuildTool
 		/// Adds a C# project to the primary project
 		/// </summary>
 		/// <param name="ProjectName">Name of project file to add</param>
+		/// <param name="Logger">Logger for output</param>
 		/// <param name="bShouldBuildForAllSolutionTargets"></param>
 		/// <param name="bForceDevelopmentConfiguration"></param>
 		/// <param name="bShouldBuildByDefaultForSolutionTargets"></param>
 		/// <returns>ProjectFile if the operation was successful, otherwise null.</returns>
-		private VCSharpProjectFile AddSimpleCSharpProject(string ProjectName, bool bShouldBuildForAllSolutionTargets = false, bool bForceDevelopmentConfiguration = false, bool bShouldBuildByDefaultForSolutionTargets = true)
+		private VCSharpProjectFile AddSimpleCSharpProject(string ProjectName, ILogger Logger, bool bShouldBuildForAllSolutionTargets = false, bool bForceDevelopmentConfiguration = false, bool bShouldBuildByDefaultForSolutionTargets = true)
 		{
 			VCSharpProjectFile Project;
 
@@ -1856,7 +1863,7 @@ namespace UnrealBuildTool
 			FileInfo Info = new FileInfo(ProjectFileName.FullName);
 			if (Info.Exists)
 			{
-				Project = new VCSharpProjectFile(ProjectFileName);
+				Project = new VCSharpProjectFile(ProjectFileName, Logger);
 
 				Project.ShouldBuildForAllSolutionTargets = bShouldBuildForAllSolutionTargets;
 				Project.ShouldBuildByDefaultForSolutionTargets = bShouldBuildByDefaultForSolutionTargets;
@@ -1902,13 +1909,14 @@ namespace UnrealBuildTool
 		/// </summary>
 		/// <param name="Arguments">Incoming command-line arguments to UBT</param>
 		/// <param name="Targets">Targets to build for intellisense</param>
+		/// <param name="Logger">Logger for output</param>
 		/// <return>Whether the process was successful or not</return>
-		private void GenerateIntelliSenseData(string[] Arguments, List<Tuple<ProjectFile, ProjectTarget>> Targets)
+		private void GenerateIntelliSenseData(string[] Arguments, List<Tuple<ProjectFile, ProjectTarget>> Targets, ILogger Logger)
 		{
 			if (ShouldGenerateIntelliSenseData() && Targets.Count > 0)
 			{
 				string ProgressInfoText = RuntimePlatform.IsWindows ? "Binding IntelliSense data..." : "Generating data for project indexing...";
-				using (ProgressWriter Progress = new ProgressWriter(ProgressInfoText, true))
+				using (ProgressWriter Progress = new ProgressWriter(ProgressInfoText, true, Logger))
 				{
 					for (int TargetIndex = 0; TargetIndex < Targets.Count; ++TargetIndex)
 					{
@@ -1922,7 +1930,7 @@ namespace UnrealBuildTool
 							continue;
 						}
 
-						Log.TraceVerbose("Found target: " + CurTarget.Name);
+						Logger.LogDebug("Found target: {Target}", CurTarget.Name);
 
 						List<string> NewArguments = new List<string>(Arguments.Length + 4);
 						if (CurTarget.TargetRules!.Type != TargetType.Program)
@@ -1940,14 +1948,14 @@ namespace UnrealBuildTool
 							TargetDescriptor TargetDesc = new TargetDescriptor(CurTarget.UnrealProjectFilePath, CurTarget.Name, IntellisensePlatform, UnrealTargetConfiguration.Development, DefaultArchitecture, new CommandLineArguments(NewArguments.ToArray()));
 
 							// Create the target
-							UEBuildTarget Target = UEBuildTarget.Create(TargetDesc, false, false, bUsePrecompiled);
-							AddTargetForIntellisense(Target);
+							UEBuildTarget Target = UEBuildTarget.Create(TargetDesc, false, false, bUsePrecompiled, Logger);
+							AddTargetForIntellisense(Target, Logger);
 
 							// If the project generator just cares about the result of UEBuildTarget.Create, skip generating the compile environments.
 							if (ShouldGenerateIntelliSenseCompileEnvironments())
 							{
 								// Generate a compile environment for each module in the binary
-								CppCompileEnvironment GlobalCompileEnvironment = Target.CreateCompileEnvironmentForProjectFiles();
+								CppCompileEnvironment GlobalCompileEnvironment = Target.CreateCompileEnvironmentForProjectFiles(Logger);
 								foreach (UEBuildBinary Binary in Target.Binaries)
 								{
 									CppCompileEnvironment BinaryCompileEnvironment = Binary.CreateBinaryCompileEnvironment(GlobalCompileEnvironment);
@@ -1958,7 +1966,7 @@ namespace UnrealBuildTool
 										{
 											Utils.WriteFileIfChangedContext = $"{Target.TargetName} {Target.Configuration} {Target.Platform} {Binary.OutputFilePaths[0].GetFileName()} {Module.Name}";
 
-											CppCompileEnvironment ModuleCompileEnvironment = Module.CreateCompileEnvironmentForIntellisense(Target.Rules, BinaryCompileEnvironment);
+											CppCompileEnvironment ModuleCompileEnvironment = Module.CreateCompileEnvironmentForIntellisense(Target.Rules, BinaryCompileEnvironment, Logger);
 											ProjectFileForIDE.AddModule(Module, ModuleCompileEnvironment);
 
 											Utils.WriteFileIfChangedContext = "";
@@ -1978,7 +1986,7 @@ namespace UnrealBuildTool
 						}
 						catch (Exception Ex)
 						{
-							Log.TraceWarning("Exception while generating include data for {0}: {1}", CurTarget.Name, Ex.ToString());
+							Logger.LogWarning("Exception while generating include data for {Target}: {Ex}", CurTarget.Name, Ex.ToString());
 						}
 
 						// Display progress
@@ -1988,7 +1996,7 @@ namespace UnrealBuildTool
 			}
 		}
 
-		protected virtual void AddTargetForIntellisense(UEBuildTarget Target)
+		protected virtual void AddTargetForIntellisense(UEBuildTarget Target, ILogger Logger)
 		{
 
 		}
@@ -1998,8 +2006,9 @@ namespace UnrealBuildTool
 		/// Selects which platforms and build configurations we want in the project file
 		/// </summary>
 		/// <param name="IncludeAllPlatforms">True if we should include ALL platforms that are supported on this machine.  Otherwise, only desktop platforms will be included.</param>
+		/// <param name="Logger"></param>
 		/// <param name="SupportedPlatformNames">Output string for supported platforms, returned as comma-separated values.</param>
-		protected virtual void SetupSupportedPlatformsAndConfigurations(bool IncludeAllPlatforms, out string SupportedPlatformNames)
+		protected virtual void SetupSupportedPlatformsAndConfigurations(bool IncludeAllPlatforms, ILogger Logger, out string SupportedPlatformNames)
 		{
 			StringBuilder SupportedPlatformsString = new StringBuilder();
 
@@ -2049,7 +2058,7 @@ namespace UnrealBuildTool
 					}
 					catch (Exception)
 					{
-						Log.TraceWarning("Invalid entry found in Configurations: {0}. Must be a member of UnrealTargetConfiguration", ConfigName);
+						Logger.LogWarning("Invalid entry found in Configurations: {ConfigName}. Must be a member of UnrealTargetConfiguration", ConfigName);
 						continue;
 					}
 				}
@@ -2119,12 +2128,13 @@ namespace UnrealBuildTool
 		/// <param name="AllModuleFiles">List of *.Build.cs files for all engine programs and games</param>
 		/// <param name="AdditionalSearchPaths"></param>
 		/// <param name="bGatherThirdPartySource">True to gather source code from third party projects too</param>
-		protected void AddProjectsForAllModules(List<FileReference> AllGames, Dictionary<FileReference, ProjectFile> ProgramProjects, List<ProjectFile> ModProjects, List<FileReference> AllModuleFiles, Dictionary<FileReference, List<DirectoryReference>> AdditionalSearchPaths, bool bGatherThirdPartySource)
+		/// <param name="Logger">Logger for output</param>
+		protected void AddProjectsForAllModules(List<FileReference> AllGames, Dictionary<FileReference, ProjectFile> ProgramProjects, List<ProjectFile> ModProjects, List<FileReference> AllModuleFiles, Dictionary<FileReference, List<DirectoryReference>> AdditionalSearchPaths, bool bGatherThirdPartySource, ILogger Logger)
 		{
 			HashSet<ProjectFile> ProjectsWithPlugins = new HashSet<ProjectFile>();
 			foreach (FileReference CurModuleFile in AllModuleFiles)
 			{
-				Log.TraceVerbose("AddProjectsForAllModules " + CurModuleFile);
+				Logger.LogDebug("AddProjectsForAllModules {File}", CurModuleFile);
 
 				// The module's "base directory" is simply the directory where its xxx.Build.cs file is stored.  We'll always
 				// harvest source files for this module in this base directory directory and all of its sub-directories.
@@ -2166,7 +2176,7 @@ namespace UnrealBuildTool
 
 					if (bGatherThirdPartySource)
 					{
-						Log.TraceInformation("Searching for third-party source files...");
+						Logger.LogInformation("Searching for third-party source files...");
 					}
 
 
@@ -2316,6 +2326,7 @@ namespace UnrealBuildTool
 		/// <param name="GameProjects">Map of game folder name to all of the game projects we created</param>
 		/// <param name="ProgramProjects">Map of program names to all of the program projects we created</param>
 		/// <param name="RulesAssemblies">Map of RuleAssemblies to their base folders</param>
+		/// <param name="Logger">Logger for output</param>
 		private void AddProjectsForAllTargets(
 			PlatformProjectGeneratorCollection PlatformProjectGenerators,
 			List<FileReference> AllGames,
@@ -2324,7 +2335,8 @@ namespace UnrealBuildTool
 			ref ProjectFile? EngineProject,
 			List<ProjectFile> GameProjects,
 			Dictionary<FileReference, ProjectFile> ProgramProjects,
-			Dictionary<RulesAssembly, DirectoryReference> RulesAssemblies)
+			Dictionary<RulesAssembly, DirectoryReference> RulesAssemblies,
+			ILogger Logger)
 		{
 			// Separate the .target.cs files that are platform extension specializations, per target name. These will be added alongside their base target.cs
 			HashSet<FileReference> AllSubTargetFiles;
@@ -2366,7 +2378,7 @@ namespace UnrealBuildTool
 					FileReference? CheckProjectFile = AllGames.FirstOrDefault(x => TargetFilePath.IsUnderDirectory(x.Directory));
 					if (CheckProjectFile == null)
 					{
-						RulesAssembly = RulesCompiler.CreateEngineRulesAssembly(false, false, false);
+						RulesAssembly = RulesCompiler.CreateEngineRulesAssembly(false, false, false, Logger);
 					
 						// Record the Engine assembly, and any parent assemblies (varies e.g. Rules, ProgramRules, MarketplaceRules)
 						if (!RulesAssemblies.ContainsKey(RulesAssembly))
@@ -2380,13 +2392,13 @@ namespace UnrealBuildTool
 					}
 					else
 					{
-						RulesAssembly = RulesCompiler.CreateProjectRulesAssembly(CheckProjectFile, false, false, false);
+						RulesAssembly = RulesCompiler.CreateProjectRulesAssembly(CheckProjectFile, false, false, false, Logger);
 						// Recording of game rule assemblies happens after BaseFolder has been computed
 					}
 
 					// Create target rules for all of the platforms and configuration combinations that we want to enable support for.
 					// Just use the current platform as we only need to recover the target type and both should be supported for all targets...
-					TargetRules TargetRulesObject = RulesAssembly.CreateTargetRules(TargetName, BuildHostPlatform.Current.Platform, UnrealTargetConfiguration.Development, "", CheckProjectFile, new CommandLineArguments(GetTargetArguments(Arguments)));
+					TargetRules TargetRulesObject = RulesAssembly.CreateTargetRules(TargetName, BuildHostPlatform.Current.Platform, UnrealTargetConfiguration.Development, "", CheckProjectFile, new CommandLineArguments(GetTargetArguments(Arguments)), Logger);
 
 					bool IsProgramTarget = false;
 
@@ -2492,7 +2504,7 @@ namespace UnrealBuildTool
 
 					bool bProjectAlreadyExisted;
 					ProjectFile ProjectFile = FindOrAddProject(ProjectFilePath, BaseFolder, IncludeInGeneratedProjects: true, bAlreadyExisted: out bProjectAlreadyExisted);
-					ProjectFile.IsForeignProject = CheckProjectFile != null && !NativeProjects.IsNativeProject(CheckProjectFile);
+					ProjectFile.IsForeignProject = CheckProjectFile != null && !NativeProjects.IsNativeProject(CheckProjectFile, Logger);
 					ProjectFile.IsGeneratedProject = true;
 					ProjectFile.IsStubProject = UnrealBuildTool.IsProjectInstalled();
 					if (TargetRulesObject.bBuildInSolutionByDefault.HasValue)
@@ -2560,7 +2572,7 @@ namespace UnrealBuildTool
 													x => UEBuildPlatform.TryGetBuildPlatform(x, out _) && 
 													(TargetRulesObject.LinkType != TargetLinkType.Modular || !UEBuildPlatform.PlatformRequiresMonolithicBuilds(x, TargetRulesObject.Configuration))
 													).ToArray(),
-							CreateRulesDelegate: (Platform, Configuration) => RulesAssembly.CreateTargetRules(TargetName, Platform, Configuration, "", CheckProjectFile, new CommandLineArguments(GetTargetArguments(Arguments)))
+							CreateRulesDelegate: (Platform, Configuration) => RulesAssembly.CreateTargetRules(TargetName, Platform, Configuration, "", CheckProjectFile, new CommandLineArguments(GetTargetArguments(Arguments)), Logger)
                         );
 
 					ProjectFile.ProjectTargets.Add(ProjectTarget);
@@ -2575,7 +2587,7 @@ namespace UnrealBuildTool
 					}
 
 
-					Log.TraceVerbose("Generating target {0} for {1}", TargetRulesObject.Type.ToString(), ProjectFilePath);
+					Logger.LogDebug("Generating target {Target} for {Project}", TargetRulesObject.Type.ToString(), ProjectFilePath);
 				}
 			}
 		}
@@ -2657,7 +2669,7 @@ namespace UnrealBuildTool
 
 
 		/// Adds engine documentation to the specified project
-		protected void AddEngineDocumentation(ProjectFile EngineProject)
+		protected void AddEngineDocumentation(ProjectFile EngineProject, ILogger Logger)
 		{
 			// NOTE: The project folder added here will actually be collapsed away later if not needed
 			DirectoryReference DocumentationProjectDirectory = DirectoryReference.Combine(Unreal.EngineDirectory, "Documentation");
@@ -2665,7 +2677,7 @@ namespace UnrealBuildTool
 			DirectoryInfo DirInfo = new DirectoryInfo(DocumentationProjectDirectory.FullName);
 			if (DirInfo.Exists && DirectoryReference.Exists(DocumentationSourceDirectory))
 			{
-				Log.TraceVerbose("Adding documentation files...");
+				Logger.LogDebug("Adding documentation files...");
 
 				List<string> SubdirectoryNamesToExclude = new List<string>();
 				{
@@ -2711,7 +2723,7 @@ namespace UnrealBuildTool
 			}
 			else
 			{
-				Log.TraceVerbose("Skipping documentation project... directory not found");
+				Logger.LogDebug("Skipping documentation project... directory not found");
 			}
 		}
 
@@ -2784,16 +2796,16 @@ namespace UnrealBuildTool
 		/// Writes the project files to disk
 		/// </summary>
 		/// <returns>True if successful</returns>
-		protected virtual bool WriteProjectFiles(PlatformProjectGeneratorCollection PlatformProjectGenerators)
+		protected virtual bool WriteProjectFiles(PlatformProjectGeneratorCollection PlatformProjectGenerators, ILogger Logger)
 		{
-			using (ProgressWriter Progress = new ProgressWriter("Writing project files...", true))
+			using (ProgressWriter Progress = new ProgressWriter("Writing project files...", true, Logger))
 			{
 				int TotalProjectFileCount = GeneratedProjectFiles.Count + 1;    // +1 for the primary project file, which we'll save next
 
 				for (int ProjectFileIndex = 0; ProjectFileIndex < GeneratedProjectFiles.Count; ++ProjectFileIndex)
 				{
 					ProjectFile CurProject = GeneratedProjectFiles[ProjectFileIndex];
-					if (!CurProject.WriteProjectFile(SupportedPlatforms, SupportedConfigurations, PlatformProjectGenerators))
+					if (!CurProject.WriteProjectFile(SupportedPlatforms, SupportedConfigurations, PlatformProjectGenerators, Logger))
 					{
 						return false;
 					}
@@ -2801,7 +2813,7 @@ namespace UnrealBuildTool
 					Progress.Write(ProjectFileIndex + 1, TotalProjectFileCount);
 				}
 
-				WritePrimaryProjectFile(UBTProject, PlatformProjectGenerators);
+				WritePrimaryProjectFile(UBTProject, PlatformProjectGenerators, Logger);
 				Progress.Write(TotalProjectFileCount, TotalProjectFileCount);
 			}
 
@@ -2813,8 +2825,9 @@ namespace UnrealBuildTool
 		/// </summary>
 		/// <param name="UBTProject">The UnrealBuildTool project</param>
 		/// <param name="PlatformProjectGenerators">The platform project file generators</param>
+		/// <param name="Logger"></param>
 		/// <returns>True if successful</returns>
-		protected abstract bool WritePrimaryProjectFile(ProjectFile? UBTProject, PlatformProjectGeneratorCollection PlatformProjectGenerators);
+		protected abstract bool WritePrimaryProjectFile(ProjectFile? UBTProject, PlatformProjectGeneratorCollection PlatformProjectGenerators, ILogger Logger);
 
 
 		/// <summary>
@@ -2822,7 +2835,8 @@ namespace UnrealBuildTool
 		/// </summary>
 		/// <param name="PlatformProjectGenerators">The platform project generators</param>
 		/// <param name="IntermediateProjectFilesPath">Intermediate project files folder</param>
-		protected virtual void WriteDebugSolutionFiles(PlatformProjectGeneratorCollection PlatformProjectGenerators, DirectoryReference IntermediateProjectFilesPath)
+		/// <param name="Logger">Logger for output</param>
+		protected virtual void WriteDebugSolutionFiles(PlatformProjectGeneratorCollection PlatformProjectGenerators, DirectoryReference IntermediateProjectFilesPath, ILogger Logger)
 		{
 		}
 
@@ -2832,9 +2846,10 @@ namespace UnrealBuildTool
 		/// </summary>
 		/// <param name="FileName">File to write</param>
 		/// <param name="NewFileContents">File content</param>
+		/// <param name="Logger"></param>
 		/// <param name="InEncoding"></param>
 		/// <returns>True if the file was saved, or if it didn't need to be overwritten because the content was unchanged</returns>
-		public static bool WriteFileIfChanged(string FileName, string NewFileContents, Encoding? InEncoding = null)
+		public static bool WriteFileIfChanged(string FileName, string NewFileContents, ILogger Logger, Encoding? InEncoding = null)
 		{
 			// Check to see if the file already exists, and if so, load it up
 			string? LoadedFileContent = null;
@@ -2847,7 +2862,7 @@ namespace UnrealBuildTool
 				}
 				catch (Exception)
 				{
-					Log.TraceInformation("Error while trying to load existing file {0}.  Ignored.", FileName);
+					Logger.LogInformation("Error while trying to load existing file {FileName}.  Ignored.", FileName);
 				}
 			}
 
@@ -2865,7 +2880,7 @@ namespace UnrealBuildTool
 
 				if (!FileNeedsSave)
 				{
-					Log.TraceVerbose("Skipped saving {0} because contents haven't changed.", Path.GetFileName(FileName));
+					Logger.LogDebug("Skipped saving {Path} because contents haven't changed.", Path.GetFileName(FileName));
 				}
 			}
 
@@ -2886,14 +2901,14 @@ namespace UnrealBuildTool
 						File.WriteAllText(FileName, NewFileContents, new UTF8Encoding());
 					else
 						File.WriteAllText(FileName, NewFileContents, InEncoding != null ? InEncoding : Encoding.UTF8);
-					Log.TraceVerbose("Saved {0}", Path.GetFileName(FileName));
+					Logger.LogDebug("Saved {Path}", Path.GetFileName(FileName));
 				}
 				catch (Exception ex)
 				{
 					// Unable to write to the project file.
 					string Message = string.Format("Error while trying to write file {0}.  The file is probably read-only.", FileName);
-					Log.TraceInformation("");
-					Log.TraceError(Message);
+					Logger.LogInformation("");
+					Logger.LogError("{Message}", Message);
 					throw new BuildException(ex, Message);
 				}
 			}
diff --git a/Engine/Source/Programs/UnrealBuildTool/ProjectFiles/QMake/QMakefileGenerator.cs b/Engine/Source/Programs/UnrealBuildTool/ProjectFiles/QMake/QMakefileGenerator.cs
--- a/Engine/Source/Programs/UnrealBuildTool/ProjectFiles/QMake/QMakefileGenerator.cs
+++ b/Engine/Source/Programs/UnrealBuildTool/ProjectFiles/QMake/QMakefileGenerator.cs
@@ -7,6 +7,7 @@
 using System.Linq;
 using EpicGames.Core;
 using UnrealBuildBase;
+using Microsoft.Extensions.Logging;
 
 namespace UnrealBuildTool
 {
@@ -38,7 +39,7 @@
 			}
 		}
 
-		protected override bool WritePrimaryProjectFile(ProjectFile? UBTProject, PlatformProjectGeneratorCollection PlatformProjectGenerators)
+		protected override bool WritePrimaryProjectFile(ProjectFile? UBTProject, PlatformProjectGeneratorCollection PlatformProjectGenerators, ILogger Logger)
 		{
 			bool bSuccess = true;
 			return bSuccess;
@@ -91,7 +92,7 @@
 		{
 		}
 
-		private bool WriteQMakePro()
+		private bool WriteQMakePro(ILogger Logger)
 		{
 			// Some more stuff borrowed from Mac side of things.
 			List<string> IncludeDirectories = new List<string>();
@@ -120,7 +121,7 @@
 				QMakefileProjectFile? QMakeProject = CurProject as QMakefileProjectFile;
 				if (QMakeProject == null)
 				{
-					Log.TraceInformation("QMakeProject == null");
+					Logger.LogInformation("QMakeProject == null");
 					continue;
 				}
 
@@ -149,7 +150,7 @@
 				QMakefileProjectFile? QMakeProject = CurProject as QMakefileProjectFile;
 				if (QMakeProject == null)
 				{
-					Log.TraceInformation("QMakeProject == null");
+					Logger.LogInformation("QMakeProject == null");
 					continue;
 				}
 
@@ -162,7 +163,7 @@
 
 					if (!DefinesAndValues.Contains(define))
 					{
-						// Log.TraceInformation (CurDefine);
+						// Logger.LogInformation (CurDefine);
 						if (string.IsNullOrEmpty(value))
 						{
 							DefinesAndValues.Add("\t");
@@ -265,7 +266,7 @@
 			);
 
 			// Create SourceFiles, HeaderFiles, and ConfigFiles sections.
-			List<FileReference> AllModuleFiles = DiscoverModules(FindGameProjects(), null);
+			List<FileReference> AllModuleFiles = DiscoverModules(FindGameProjects(Logger), null);
 			foreach (FileReference CurModuleFile in AllModuleFiles)
 			{
 				List<FileReference> FoundFiles = SourceFileSearch.FindModuleSourceFiles(CurModuleFile);
@@ -400,21 +401,21 @@
 			string FullQMakeHeaderPriFileName = Path.Combine(PrimaryProjectPath.FullName, QMakeHeaderPriFileName);
 			string FullQMakeConfigPriFileName = Path.Combine(PrimaryProjectPath.FullName, QMakeConfigPriFileName);
 
-			WriteFileIfChanged(FullQMakeDefinesFileName, QMakeDefinesPriFileContent.ToString());
-			WriteFileIfChanged(FullQMakeIncludesFileName, QMakeIncludesPriFileContent.ToString());
-			WriteFileIfChanged(FullQMakeSourcePriFileName, QMakeSourcePriFileContent.ToString());
+			WriteFileIfChanged(FullQMakeDefinesFileName, QMakeDefinesPriFileContent.ToString(), Logger);
+			WriteFileIfChanged(FullQMakeIncludesFileName, QMakeIncludesPriFileContent.ToString(), Logger);
+			WriteFileIfChanged(FullQMakeSourcePriFileName, QMakeSourcePriFileContent.ToString(), Logger);
 
-			WriteFileIfChanged(FullQMakeHeaderPriFileName, QMakeHeaderPriFileContent.ToString());
-			WriteFileIfChanged(FullQMakeConfigPriFileName, QMakeConfigPriFileContent.ToString());
+			WriteFileIfChanged(FullQMakeHeaderPriFileName, QMakeHeaderPriFileContent.ToString(), Logger);
+			WriteFileIfChanged(FullQMakeConfigPriFileName, QMakeConfigPriFileContent.ToString(), Logger);
 
-			return WriteFileIfChanged(FullFileName, QMakeFileContent.ToString());
+			return WriteFileIfChanged(FullFileName, QMakeFileContent.ToString(), Logger);
 		}
 
 		/// ProjectFileGenerator interface
 		//protected override bool WritePrimaryProjectFile( ProjectFile UBTProject )
-		protected override bool WriteProjectFiles(PlatformProjectGeneratorCollection PlatformProjectGenerators)
+		protected override bool WriteProjectFiles(PlatformProjectGeneratorCollection PlatformProjectGenerators, ILogger Logger)
 		{
-			return WriteQMakePro();
+			return WriteQMakePro(Logger);
 		}
 
 		/// ProjectFileGenerator interface
@@ -430,7 +431,7 @@
 		}
 
 		/// ProjectFileGenerator interface
-		public override void CleanProjectFiles(DirectoryReference InPrimaryProjectDirectory, string InPrimaryProjectName, DirectoryReference InIntermediateProjectFilesDirectory)
+		public override void CleanProjectFiles(DirectoryReference InPrimaryProjectDirectory, string InPrimaryProjectName, DirectoryReference InIntermediateProjectFilesDirectory, ILogger Logger)
 		{
 		}
 	}
diff --git a/Engine/Source/Programs/UnrealBuildTool/ProjectFiles/Rider/RiderProjectFile.cs b/Engine/Source/Programs/UnrealBuildTool/ProjectFiles/Rider/RiderProjectFile.cs
--- a/Engine/Source/Programs/UnrealBuildTool/ProjectFiles/Rider/RiderProjectFile.cs
+++ b/Engine/Source/Programs/UnrealBuildTool/ProjectFiles/Rider/RiderProjectFile.cs
@@ -8,6 +8,7 @@
 using System.Text.RegularExpressions;
 using EpicGames.Core;
 using UnrealBuildBase;
+using Microsoft.Extensions.Logging;
 
 namespace UnrealBuildTool
 {
@@ -47,10 +48,11 @@
 		/// <param name="InConfigurations"></param>
 		/// <param name="PlatformProjectGenerators"></param>
 		/// <param name="Minimize"></param>
+		/// <param name="Logger"></param>
 		/// <returns></returns>
 		public bool WriteProjectFile(List<UnrealTargetPlatform> InPlatforms,
 			List<UnrealTargetConfiguration> InConfigurations,
-			PlatformProjectGeneratorCollection PlatformProjectGenerators, JsonWriterStyle Minimize)
+			PlatformProjectGeneratorCollection PlatformProjectGenerators, JsonWriterStyle Minimize, ILogger Logger)
 		{
 			string ProjectName = ProjectFilePath.GetFileNameWithoutAnyExtensions();
 			DirectoryReference ProjectRootFolder = RootPath;
@@ -94,15 +96,15 @@
 							Platform, Configuration, DefaultArchitecture, Arguments);
 						try
 						{
-							UEBuildTarget BuildTarget = UEBuildTarget.Create(TargetDesc, false, false, false);
+							UEBuildTarget BuildTarget = UEBuildTarget.Create(TargetDesc, false, false, false, Logger);
 						
 							FileReference OutputFile = FileReference.Combine(TargetFolder, $"{ProjectName}.json");
 							FileToTarget.Add(Tuple.Create(OutputFile, BuildTarget));
 						}
 						catch(Exception Ex)
 						{
-							Log.TraceWarning("Exception while generating include data for Target:{0}, Platform: {1}, Configuration: {2}", TargetDesc.Name, Platform.ToString(), Configuration.ToString());
-							Log.TraceWarning(Ex.ToString());
+							Logger.LogWarning("Exception while generating include data for Target:{Target}, Platform: {Platform}, Configuration: {Configuration}", TargetDesc.Name, Platform.ToString(), Configuration.ToString());
+							Logger.LogWarning("{Ex}", Ex.ToString());
 						}
 					}
 				}
@@ -112,13 +114,13 @@
 				try
 				{
 					CurrentTarget = tuple.Item2;
-					CurrentTarget.PreBuildSetup();
-					SerializeTarget(tuple.Item1, CurrentTarget, Minimize);
+					CurrentTarget.PreBuildSetup(Logger);
+					SerializeTarget(tuple.Item1, CurrentTarget, Minimize, Logger);
 				}
 				catch (Exception Ex)
 				{
-					Log.TraceWarning("Exception while generating include data for Target:{0}, Platform: {1}, Configuration: {2}", tuple.Item2.AppName, tuple.Item2.Platform.ToString(), tuple.Item2.Configuration.ToString());
-					Log.TraceWarning(Ex.ToString());
+					Logger.LogWarning("Exception while generating include data for Target:{Target}, Platform: {Platform}, Configuration: {Configuration}", tuple.Item2.AppName, tuple.Item2.Platform.ToString(), tuple.Item2.Configuration.ToString());
+					Logger.LogWarning("{Ex}", Ex.ToString());
 				}
 			}
 			
@@ -140,12 +142,12 @@
 			return false;
 		}
 
-		private void SerializeTarget(FileReference OutputFile, UEBuildTarget BuildTarget, JsonWriterStyle Minimize)
+		private void SerializeTarget(FileReference OutputFile, UEBuildTarget BuildTarget, JsonWriterStyle Minimize, ILogger Logger)
 		{
 			DirectoryReference.CreateDirectory(OutputFile.Directory);
 			using (JsonWriter Writer = new JsonWriter(OutputFile, Minimize))
 			{
-				ExportTarget(BuildTarget, Writer);
+				ExportTarget(BuildTarget, Writer, Logger);
 			}
 		}
 
@@ -154,7 +156,8 @@
 		/// </summary>
 		/// <param name="Target"></param>
 		/// <param name="Writer">Writer for the array data</param>
-		private void ExportTarget(UEBuildTarget Target, JsonWriter Writer)
+		/// <param name="Logger">Logger for output</param>
+		private void ExportTarget(UEBuildTarget Target, JsonWriter Writer, ILogger Logger)
 		{
 			Writer.WriteObjectStart();
 
@@ -167,7 +170,7 @@
 				Writer.WriteValue("ProjectFile", Target.ProjectFile.FullName);
 			}
 			
-			ExportEnvironmentToJson(Target, Writer);
+			ExportEnvironmentToJson(Target, Writer, Logger);
 			
 			if(Target.Binaries.Any())
 			{
@@ -181,7 +184,7 @@
 				Writer.WriteArrayEnd();
 			}
 			
-			CppCompileEnvironment GlobalCompileEnvironment = Target.CreateCompileEnvironmentForProjectFiles();
+			CppCompileEnvironment GlobalCompileEnvironment = Target.CreateCompileEnvironmentForProjectFiles(Logger);
 			HashSet<string> ModuleNames = new HashSet<string>();
 			Writer.WriteObjectStart("Modules");
 			foreach (UEBuildBinary Binary in Target.Binaries)
@@ -195,23 +198,23 @@
 						UEBuildModuleCPP? ModuleCpp = Module as UEBuildModuleCPP;
 						if (ModuleCpp != null)
 						{
-							CppCompileEnvironment ModuleCompileEnvironment = ModuleCpp.CreateCompileEnvironmentForIntellisense(Target.Rules, BinaryCompileEnvironment);
-							ExportModuleCpp(ModuleCpp, ModuleCompileEnvironment, Writer);
+							CppCompileEnvironment ModuleCompileEnvironment = ModuleCpp.CreateCompileEnvironmentForIntellisense(Target.Rules, BinaryCompileEnvironment, Logger);
+							ExportModuleCpp(ModuleCpp, ModuleCompileEnvironment, Writer, Logger);
 							Module.PrivateIncludePaths.UnionWith(ModuleCompileEnvironment.UserIncludePaths);
 						}
-						ExportModule(Module, Binary.OutputDir, Target.GetExecutableDir(), Writer);
+						ExportModule(Module, Binary.OutputDir, Target.GetExecutableDir(), Writer, Logger);
 						Writer.WriteObjectEnd();
 					}
 				}
 			}
 			Writer.WriteObjectEnd();
 			
-			ExportPluginsFromTarget(Target, Writer);
+			ExportPluginsFromTarget(Target, Writer, Logger);
 			
 			Writer.WriteObjectEnd();
 		}
 
-		private void ExportModuleCpp(UEBuildModuleCPP ModuleCPP, CppCompileEnvironment ModuleCompileEnvironment, JsonWriter Writer)
+		private void ExportModuleCpp(UEBuildModuleCPP ModuleCPP, CppCompileEnvironment ModuleCompileEnvironment, JsonWriter Writer, ILogger Logger)
 		{
 			Writer.WriteValue("GeneratedCodeDirectory", ModuleCPP.GeneratedCodeDirectory != null ? ModuleCPP.GeneratedCodeDirectory.FullName  : string.Empty);
 			
@@ -221,7 +224,7 @@
 				Writer.WriteObjectStart("ToolchainInfo");
 				foreach (Tuple<string,object?> Field in ModuleToolchainInfo.GetDiff(RootToolchainInfo))
 				{
-					WriteField(ModuleCPP.Name, Writer, Field);
+					WriteField(ModuleCPP.Name, Writer, Field, Logger);
 				}
 				Writer.WriteObjectEnd();
 			}
@@ -229,12 +232,12 @@
 			if (ModuleCompileEnvironment.PrecompiledHeaderIncludeFilename != null)
 			{
 				string CorrectFilePathPch;
-				if(ExtractWrappedIncludeFile(ModuleCompileEnvironment.PrecompiledHeaderIncludeFilename, out CorrectFilePathPch))
+				if(ExtractWrappedIncludeFile(ModuleCompileEnvironment.PrecompiledHeaderIncludeFilename, Logger, out CorrectFilePathPch))
 					Writer.WriteValue("SharedPCHFilePath", CorrectFilePathPch);
 			}
 		}
 
-		private static bool ExtractWrappedIncludeFile(FileSystemReference FileRef, out string CorrectFilePathPch)
+		private static bool ExtractWrappedIncludeFile(FileSystemReference FileRef, ILogger Logger, out string CorrectFilePathPch)
 		{
 			CorrectFilePathPch = "";
 			try
@@ -251,7 +254,7 @@
 			}
 			finally
 			{
-				Log.TraceVerbose("Couldn't extract path to PCH from {0}", FileRef);
+				Logger.LogDebug("Couldn't extract path to PCH from {FileRef}", FileRef);
 			}
 			return false;
 		}
@@ -263,7 +266,8 @@
 		/// <param name="TargetOutputDir"></param>
 		/// <param name="Writer">Writer for the array data</param>
 		/// <param name="Module"></param>
-		private static void ExportModule(UEBuildModule Module, DirectoryReference BinaryOutputDir, DirectoryReference TargetOutputDir, JsonWriter Writer)
+		/// <param name="Logger"></param>
+		private static void ExportModule(UEBuildModule Module, DirectoryReference BinaryOutputDir, DirectoryReference TargetOutputDir, JsonWriter Writer, ILogger Logger)
 		{
 			Writer.WriteValue("Name", Module.Name);
 			Writer.WriteValue("Directory", Module.ModuleDirectory.FullName );
@@ -330,7 +334,7 @@
 					}
 					catch(BuildException buildException)
 					{
-						Log.TraceVerbose("Value {0} for module {1} will not be stored. Reason: {2}", "Path", Module.Name, buildException);	
+						Logger.LogDebug("Value {Value} for module {ModuleName} will not be stored. Reason: {Ex}", "Path", Module.Name, buildException);	
 					}
 					
 					if (RuntimeDependency.SourcePath != null)
@@ -343,7 +347,7 @@
 						}
 						catch(BuildException buildException)
 						{
-							Log.TraceVerbose("Value {0} for module {1} will not be stored. Reason: {2}", "SourcePath", Module.Name, buildException);	
+							Logger.LogDebug("Value {Value} for module {ModuleName} will not be stored. Reason: {Ex}", "SourcePath", Module.Name, buildException);	
 						}
 					}
 
@@ -419,9 +423,10 @@
 		/// </summary>
 		/// <param name="Target"></param>
 		/// <param name="Writer"></param>
-		private static void ExportPluginsFromTarget(UEBuildTarget Target, JsonWriter Writer)
+		/// <param name="Logger"></param>
+		private static void ExportPluginsFromTarget(UEBuildTarget Target, JsonWriter Writer, ILogger Logger)
 		{
-			Target.SetupPlugins();
+			Target.SetupPlugins(Logger);
 			if (Target.BuildPlugins == null || !Target.BuildPlugins.Any()) return;
 			
 			Writer.WriteObjectStart("Plugins");
@@ -455,16 +460,17 @@
 		/// </summary>
 		/// <param name="Target"></param>
 		/// <param name="Writer"></param>
-		private void ExportEnvironmentToJson(UEBuildTarget Target, JsonWriter Writer)
+		/// <param name="Logger"></param>
+		private void ExportEnvironmentToJson(UEBuildTarget Target, JsonWriter Writer, ILogger Logger)
 		{
-			CppCompileEnvironment GlobalCompileEnvironment = Target.CreateCompileEnvironmentForProjectFiles();
+			CppCompileEnvironment GlobalCompileEnvironment = Target.CreateCompileEnvironmentForProjectFiles(Logger);
 			
 			RootToolchainInfo = GenerateToolchainInfo(GlobalCompileEnvironment);
 			
 			Writer.WriteObjectStart("ToolchainInfo");
 			foreach (Tuple<string, object?> Field in RootToolchainInfo.GetFields())
 			{
-				WriteField(Target.TargetName, Writer, Field);
+				WriteField(Target.TargetName, Writer, Field, Logger);
 			}
 			Writer.WriteObjectEnd();
 			
@@ -490,7 +496,7 @@
 			{
 				// Only generate Apple system include paths when host platform is Apple OS
 				// TODO: Fix case when working with MacOS on Windows host platform  
-				foreach (string Path in AppleHelper.GetAppleSystemIncludePaths(GlobalCompileEnvironment.Architecture, Target.Platform))
+				foreach (string Path in AppleHelper.GetAppleSystemIncludePaths(GlobalCompileEnvironment.Architecture, Target.Platform, Logger))
 				{
 					Writer.WriteValue(Path);
 				}
@@ -549,7 +555,7 @@
 			Writer.WriteArrayEnd();
 		}
 
-		private static void WriteField(string ModuleOrTargetName, JsonWriter Writer, Tuple<string, object?> Field)
+		private static void WriteField(string ModuleOrTargetName, JsonWriter Writer, Tuple<string, object?> Field, ILogger Logger)
 		{
 			if (Field.Item2 == null) return;
 			string Name = Field.Item1;
@@ -589,7 +595,7 @@
 						Writer.WriteValue(Name, "Latest");
 						break;
 					default:
-						Log.TraceError("Unsupported C++ standard type: {0}", version);
+						Logger.LogError("Unsupported C++ standard type: {Type}", version);
 						break;
 				}
 			}
@@ -605,7 +611,7 @@
 			}
 			else
 			{
-				Log.TraceWarning("Dumping incompatible ToolchainInfo field: {0} with type: {1} for: {2}",
+				Logger.LogWarning("Dumping incompatible ToolchainInfo field: {Name} with type: {Field} for: {ModuleOrTarget}",
 					Name, Field.Item2, ModuleOrTargetName);
 			}
 		}
@@ -687,14 +693,14 @@
 			private Process? XcrunProcess;
 			private bool IsReadingIncludesSection;
 
-			public IList<string> GetAppleSystemIncludePaths(string Architecture, UnrealTargetPlatform Platform)
+			public IList<string> GetAppleSystemIncludePaths(string Architecture, UnrealTargetPlatform Platform, ILogger Logger)
 			{
 				if (!UEBuildPlatform.IsPlatformInGroup(Platform, UnrealPlatformGroup.Apple))
 				{
 					throw new InvalidOperationException("xcrun can be run only for Apple's platforms");
 				}
 
-				string SDKPath = GetSDKPath(Architecture, Platform);
+				string SDKPath = GetSDKPath(Architecture, Platform, Logger);
 				if (!CachedIncludePaths.ContainsKey(SDKPath))
 				{
 					CalculateSystemIncludePaths(SDKPath);
@@ -765,7 +771,7 @@
 				}
 			}
 
-			private string GetSDKPath(string Architecture, UnrealTargetPlatform Platform)
+			private string GetSDKPath(string Architecture, UnrealTargetPlatform Platform, ILogger Logger)
 			{
 				if (Platform == UnrealTargetPlatform.Mac)
 				{
@@ -774,12 +780,12 @@
 
 				if (Platform == UnrealTargetPlatform.IOS)
 				{
-					return new IOSToolChainSettings().GetSDKPath(Architecture);
+					return new IOSToolChainSettings(Logger).GetSDKPath(Architecture);
 				}
 
 				if (Platform == UnrealTargetPlatform.TVOS)
 				{
-					return new TVOSToolChainSettings().GetSDKPath(Architecture);
+					return new TVOSToolChainSettings(Logger).GetSDKPath(Architecture);
 				}
 
 				throw new NotImplementedException("Path to SDK has to be specified for each Apple's platform");
diff --git a/Engine/Source/Programs/UnrealBuildTool/ProjectFiles/Rider/RiderProjectFileGenerator.cs b/Engine/Source/Programs/UnrealBuildTool/ProjectFiles/Rider/RiderProjectFileGenerator.cs
--- a/Engine/Source/Programs/UnrealBuildTool/ProjectFiles/Rider/RiderProjectFileGenerator.cs
+++ b/Engine/Source/Programs/UnrealBuildTool/ProjectFiles/Rider/RiderProjectFileGenerator.cs
@@ -6,6 +6,7 @@ using System.IO;
 using System.Linq;
 using EpicGames.Core;
 using UnrealBuildBase;
+using Microsoft.Extensions.Logging;
 
 namespace UnrealBuildTool
 {
@@ -67,7 +68,7 @@ namespace UnrealBuildTool
 		public override bool ShouldGenerateIntelliSenseData() => true;
 
 		public override void CleanProjectFiles(DirectoryReference InPrimaryProjectDirectory, string InPrimaryProjectName,
-			DirectoryReference InIntermediateProjectFilesDirectory)
+			DirectoryReference InIntermediateProjectFilesDirectory, ILogger Logger)
 		{
 			DirectoryReference.Delete(InPrimaryProjectDirectory);
 		}
@@ -91,9 +92,9 @@ namespace UnrealBuildTool
 			return projectFile;
 		}
 
-		protected override bool WriteProjectFiles(PlatformProjectGeneratorCollection PlatformProjectGenerators)
+		protected override bool WriteProjectFiles(PlatformProjectGeneratorCollection PlatformProjectGenerators, ILogger Logger)
 		{
-			using (ProgressWriter Progress = new ProgressWriter("Writing project files...", true))
+			using (ProgressWriter Progress = new ProgressWriter("Writing project files...", true, Logger))
 			{
 				List<ProjectFile> ProjectsToGenerate = new List<ProjectFile>(GeneratedProjectFiles);
 				if (ProjectNames.Any())
@@ -117,10 +118,10 @@ namespace UnrealBuildTool
 					
 					if (UEBuildPlatform.IsPlatformAvailable(it))
 						return true;
-					Log.TraceWarning(
-						"Platform {0} is not a valid platform to build. Check that the SDK is installed properly",
+					Logger.LogWarning(
+						"Platform {Platform} is not a valid platform to build. Check that the SDK is installed properly",
 						it);
-					Log.TraceWarning("Platform will be ignored in project file generation");
+					Logger.LogWarning("Platform will be ignored in project file generation");
 					return false;
 				}).ToList();
 
@@ -136,7 +137,7 @@ namespace UnrealBuildTool
 					if(CurProject != null)
 					{
 						if (!CurProject.WriteProjectFile(FilteredPlatforms, ConfigurationsToGenerate.ToList(),
-							PlatformProjectGenerators, Minimize))
+							PlatformProjectGenerators, Minimize, Logger))
 						{
 							return false;
 						}
@@ -166,6 +167,7 @@ namespace UnrealBuildTool
 		private void AddProjectsForAllTargets(
 			PlatformProjectGeneratorCollection PlatformProjectGenerators,
 			List<FileReference> AllGames,
+			ILogger Logger,
 			out ProjectFile? EngineProject,
 			out List<ProjectFile> GameProjects,
 			out Dictionary<FileReference, ProjectFile> ProgramProjects)
@@ -177,7 +179,7 @@ namespace UnrealBuildTool
 
 			// Find all of the target files.  This will filter out any modules or targets that don't
 			// belong to platforms we're generating project files for.
-			List<FileReference> AllTargetFiles = DiscoverTargets(AllGames);
+			List<FileReference> AllTargetFiles = DiscoverTargets(AllGames, Logger);
 
 			// Sort the targets by name. When we have multiple targets of a given type for a project, we'll use the order to determine which goes in the primary project file (so that client names with a suffix will go into their own project).
 			AllTargetFiles = AllTargetFiles.OrderBy(x => x.FullName, StringComparer.OrdinalIgnoreCase).ToList();
@@ -217,18 +219,18 @@ namespace UnrealBuildTool
 						AllGames.FirstOrDefault(x => TargetFilePath.IsUnderDirectory(x.Directory));
 					if (CheckProjectFile == null)
 					{
-						RulesAssembly = RulesCompiler.CreateEngineRulesAssembly(false, false, false);
+						RulesAssembly = RulesCompiler.CreateEngineRulesAssembly(false, false, false, Logger);
 					}
 					else
 					{
-						RulesAssembly = RulesCompiler.CreateProjectRulesAssembly(CheckProjectFile, false, false, false);
+						RulesAssembly = RulesCompiler.CreateProjectRulesAssembly(CheckProjectFile, false, false, false, Logger);
 					}
 
 					// Create target rules for all of the platforms and configuration combinations that we want to enable support for.
 					// Just use the current platform as we only need to recover the target type and both should be supported for all targets...
 					TargetRules TargetRulesObject = RulesAssembly.CreateTargetRules(TargetName,
 						BuildHostPlatform.Current.Platform, UnrealTargetConfiguration.Development, "", CheckProjectFile,
-						null);
+						null, Logger);
 
 					bool IsProgramTarget = false;
 
@@ -305,7 +307,7 @@ namespace UnrealBuildTool
 					ProjectFile ProjectFile = FindOrAddProject(ProjectFilePath, BaseFolder,
 						true, out bProjectAlreadyExisted);
 					ProjectFile.IsForeignProject =
-						CheckProjectFile != null && !NativeProjects.IsNativeProject(CheckProjectFile);
+						CheckProjectFile != null && !NativeProjects.IsNativeProject(CheckProjectFile, Logger);
 					ProjectFile.IsGeneratedProject = true;
 					ProjectFile.IsStubProject = UnrealBuildTool.IsProjectInstalled();
 					if (TargetRulesObject.bBuildInSolutionByDefault.HasValue)
@@ -385,7 +387,7 @@ namespace UnrealBuildTool
 							.Where(x => UEBuildPlatform.TryGetBuildPlatform(x, out _)).ToArray(),
 						CreateRulesDelegate: (Platform, Configuration) =>
 							RulesAssembly.CreateTargetRules(TargetName, Platform, Configuration, "", CheckProjectFile,
-								null)
+								null, Logger)
 					);
 
 					ProjectFile.ProjectTargets.Add(ProjectTarget);
@@ -399,20 +401,20 @@ namespace UnrealBuildTool
 						ProjectFile.AddFilesToProject(AllSubTargetFilesPerTarget[TargetName], BaseFolder);
 					}
 
-					Log.TraceVerbose("Generating target {0} for {1}", TargetRulesObject.Type.ToString(),
+					Logger.LogDebug("Generating target {Target} for {Project}", TargetRulesObject.Type.ToString(),
 						ProjectFilePath);
 				}
 			}
 		}
 		
-		private void SetupSupportedPlatformsAndConfigurations()
+		private void SetupSupportedPlatformsAndConfigurations(ILogger Logger)
 		{
 			string SupportedPlatformNames;
-			SetupSupportedPlatformsAndConfigurations(true, out SupportedPlatformNames);
+			SetupSupportedPlatformsAndConfigurations(true, Logger, out SupportedPlatformNames);
 		}
 		
 		public override bool GenerateProjectFiles(PlatformProjectGeneratorCollection PlatformProjectGenerators,
-			String[] arguments)
+			String[] arguments, ILogger Logger)
 		{
 			ConfigureProjectFileGeneration();
 
@@ -443,23 +445,23 @@ namespace UnrealBuildTool
 					DirectoryReference.Combine(PrimaryProjectPath, "Intermediate", "ProjectFiles");
 			}
 
-			SetupSupportedPlatformsAndConfigurations();
+			SetupSupportedPlatformsAndConfigurations(Logger);
 
-			Log.TraceVerbose("Detected supported platforms: " + SupportedPlatforms);
+			Logger.LogDebug("Detected supported platforms: {Platforms}", SupportedPlatforms);
 
-			List<FileReference> AllGameProjects = FindGameProjects();
+			List<FileReference> AllGameProjects = FindGameProjects(Logger);
 
-			GatherProjects(PlatformProjectGenerators, AllGameProjects);
+			GatherProjects(PlatformProjectGenerators, AllGameProjects, Logger);
 
-			WriteProjectFiles(PlatformProjectGenerators);
-			Log.TraceVerbose("Project generation complete ({0} generated, {1} imported)", GeneratedProjectFiles.Count,
+			WriteProjectFiles(PlatformProjectGenerators, Logger);
+			Logger.LogDebug("Project generation complete ({NumGenerated} generated, {NumImported} imported)", GeneratedProjectFiles.Count,
 				OtherProjectFiles.Count);
 
 			return true;
 		}
 
 		private void GatherProjects(PlatformProjectGeneratorCollection PlatformProjectGenerators,
-			List<FileReference> AllGameProjects)
+			List<FileReference> AllGameProjects, ILogger Logger)
 		{ 
 			ProjectFile? EngineProject = null;
 			List<ProjectFile>? GameProjects = null;
@@ -467,8 +469,8 @@ namespace UnrealBuildTool
 			Dictionary<FileReference, ProjectFile>? ProgramProjects = null;
 
 			// Setup buildable projects for all targets
-			AddProjectsForAllTargets(PlatformProjectGenerators, AllGameProjects, out EngineProject,
-				out GameProjects, out ProgramProjects);
+			AddProjectsForAllTargets(PlatformProjectGenerators, AllGameProjects, Logger,
+				out EngineProject, out GameProjects, out ProgramProjects);
 
 			AddProjectsForMods(GameProjects, ModProjects);
 			AddAllGameProjects(GameProjects);
@@ -489,7 +491,7 @@ namespace UnrealBuildTool
 		}
 
 		protected override bool WritePrimaryProjectFile(ProjectFile? ProjectFile,
-			PlatformProjectGeneratorCollection PlatformProjectGenerators)
+			PlatformProjectGeneratorCollection PlatformProjectGenerators, ILogger Logger)
 		{
 			return true;
 		}
