UBT: Use ILogger throughout UBT, rather than legacy EpicGames.Core.Log methods. All output going forwards should use structured logging rather than String.Format style logging.

diff --git a/Engine/Source/Programs/UnrealBuildTool/ProjectFiles/VisualStudio/VCProject.cs b/Engine/Source/Programs/UnrealBuildTool/ProjectFiles/VisualStudio/VCProject.cs
--- a/Engine/Source/Programs/UnrealBuildTool/ProjectFiles/VisualStudio/VCProject.cs
+++ b/Engine/Source/Programs/UnrealBuildTool/ProjectFiles/VisualStudio/VCProject.cs
@@ -14,6 +14,7 @@ using System.Security;
 using EpicGames.Core;
 using UnrealBuildBase;
 using System.Diagnostics.CodeAnalysis;
+using Microsoft.Extensions.Logging;
 
 namespace UnrealBuildTool
 {
@@ -167,8 +168,9 @@ namespace UnrealBuildTool
 		/// <param name="SolutionConfiguration">The solution configuration</param>
 		/// <param name="SolutionPlatform">The solution platform</param>
 		/// <param name="PlatformProjectGenerators">Set of platform project generators</param>
+		/// <param name="Logger">Logger for output</param>
 		/// <returns>Project context matching the given solution context</returns>
-		public abstract MSBuildProjectContext GetMatchingProjectContext(TargetType SolutionTarget, UnrealTargetConfiguration SolutionConfiguration, UnrealTargetPlatform SolutionPlatform, PlatformProjectGeneratorCollection PlatformProjectGenerators);
+		public abstract MSBuildProjectContext GetMatchingProjectContext(TargetType SolutionTarget, UnrealTargetConfiguration SolutionConfiguration, UnrealTargetPlatform SolutionPlatform, PlatformProjectGeneratorCollection PlatformProjectGenerators, ILogger Logger);
 
 		/// <summary>
 		/// Checks to see if the specified solution platform and configuration is able to map to this project
@@ -176,8 +178,9 @@ namespace UnrealBuildTool
 		/// <param name="ProjectTarget">The target that we're checking for a valid platform/config combination</param>
 		/// <param name="Platform">Platform</param>
 		/// <param name="Configuration">Configuration</param>
+		/// <param name="Logger">Logger for output</param>
 		/// <returns>True if this is a valid combination for this project, otherwise false</returns>
-		public static bool IsValidProjectPlatformAndConfiguration(Project ProjectTarget, UnrealTargetPlatform Platform, UnrealTargetConfiguration Configuration)
+		public static bool IsValidProjectPlatformAndConfiguration(Project ProjectTarget, UnrealTargetPlatform Platform, UnrealTargetConfiguration Configuration, ILogger Logger)
 		{
 			if (!ProjectFileGenerator.bIncludeTestAndShippingConfigs)
 			{
@@ -324,9 +327,10 @@ namespace UnrealBuildTool
 		/// <param name="Configuration">Actual configuration</param>
 		/// <param name="TargetConfigurationName">The configuration name from the target rules, or null if we don't have one</param>
 		/// <param name="PlatformProjectGenerators">Set of platform project generators</param>
+		/// <param name="Logger"></param>
 		/// <param name="ProjectPlatformName">Name of platform string to use for Visual Studio project</param>
 		/// <param name="ProjectConfigurationName">Name of configuration string to use for Visual Studio project</param>
-		private void MakeProjectPlatformAndConfigurationNames(UnrealTargetPlatform Platform, UnrealTargetConfiguration Configuration, TargetType TargetConfigurationName, PlatformProjectGeneratorCollection PlatformProjectGenerators, out string ProjectPlatformName, out string ProjectConfigurationName)
+		private void MakeProjectPlatformAndConfigurationNames(UnrealTargetPlatform Platform, UnrealTargetConfiguration Configuration, TargetType TargetConfigurationName, PlatformProjectGeneratorCollection PlatformProjectGenerators, ILogger Logger, out string ProjectPlatformName, out string ProjectConfigurationName)
 		{
 			PlatformProjectGenerator? PlatformProjectGenerator = PlatformProjectGenerators.GetPlatformProjectGenerator(Platform, bInAllowFailure: true);
 
@@ -372,8 +376,9 @@ namespace UnrealBuildTool
 		/// <param name="SolutionConfiguration">The solution configuration</param>
 		/// <param name="SolutionPlatform">The solution platform</param>
 		/// <param name="PlatformProjectGenerators">Set of platform project generations</param>
+		/// <param name="Logger"></param>
 		/// <returns>Project context matching the given solution context</returns>
-		public override MSBuildProjectContext GetMatchingProjectContext(TargetType SolutionTarget, UnrealTargetConfiguration SolutionConfiguration, UnrealTargetPlatform SolutionPlatform, PlatformProjectGeneratorCollection PlatformProjectGenerators)
+		public override MSBuildProjectContext GetMatchingProjectContext(TargetType SolutionTarget, UnrealTargetConfiguration SolutionConfiguration, UnrealTargetPlatform SolutionPlatform, PlatformProjectGeneratorCollection PlatformProjectGenerators, ILogger Logger)
 		{
 			// Stub projects always build in the same configuration
 			if (IsStubProject)
@@ -409,7 +414,7 @@ namespace UnrealBuildTool
 				List<Project> MatchingProjectTargets = new List<Project>();
 				foreach (Project ProjectTarget in ProjectTargets)
 				{
-					if (VCProjectFile.IsValidProjectPlatformAndConfiguration(ProjectTarget, SolutionPlatform, SolutionConfiguration))
+					if (VCProjectFile.IsValidProjectPlatformAndConfiguration(ProjectTarget, SolutionPlatform, SolutionConfiguration, Logger))
 					{
 						if (ProjectTarget.TargetRules != null)
 						{
@@ -474,7 +479,7 @@ namespace UnrealBuildTool
 					}
 					else
 					{
-						MakeProjectPlatformAndConfigurationNames(ProjectPlatform, ProjectConfiguration, TargetConfigurationName, PlatformProjectGenerators, out ProjectPlatformName, out ProjectConfigurationName);
+						MakeProjectPlatformAndConfigurationNames(ProjectPlatform, ProjectConfiguration, TargetConfigurationName, PlatformProjectGenerators, Logger, out ProjectPlatformName, out ProjectConfigurationName);
 					}
 
 					// Set whether this project configuration should be built when the user initiates "build solution"
@@ -631,7 +636,7 @@ namespace UnrealBuildTool
 		HashSet<string>? InvalidConfigPlatformNames;
 		List<ProjectConfigAndTargetCombination>? ProjectConfigAndTargetCombinations;
 
-		private void BuildProjectConfigAndTargetCombinations(List<UnrealTargetPlatform> InPlatforms, List<UnrealTargetConfiguration> InConfigurations, PlatformProjectGeneratorCollection PlatformProjectGenerators)
+		private void BuildProjectConfigAndTargetCombinations(List<UnrealTargetPlatform> InPlatforms, List<UnrealTargetConfiguration> InConfigurations, PlatformProjectGeneratorCollection PlatformProjectGenerators, ILogger Logger)
 		{
 			//no need to do this more than once
 			if (ProjectConfigAndTargetCombinations == null)
@@ -676,10 +681,10 @@ namespace UnrealBuildTool
 
 								foreach (ProjectTarget ProjectTarget in ProjectTargets.OfType<ProjectTarget>())
 								{
-									if (IsValidProjectPlatformAndConfiguration(ProjectTarget, Platform, Configuration))
+									if (IsValidProjectPlatformAndConfiguration(ProjectTarget, Platform, Configuration, Logger))
 									{
 										string ProjectPlatformName, ProjectConfigurationName;
-										MakeProjectPlatformAndConfigurationNames(Platform, Configuration, ProjectTarget.TargetRules!.Type, PlatformProjectGenerators, out ProjectPlatformName, out ProjectConfigurationName);
+										MakeProjectPlatformAndConfigurationNames(Platform, Configuration, ProjectTarget.TargetRules!.Type, PlatformProjectGenerators, Logger, out ProjectPlatformName, out ProjectConfigurationName);
 
 										ProjectConfigAndTargetCombination Combination = new ProjectConfigAndTargetCombination(Platform, Configuration, ProjectPlatformName, ProjectConfigurationName, ProjectTarget);
 										ProjectConfigAndTargetCombinations.Add(Combination);
@@ -711,14 +716,14 @@ namespace UnrealBuildTool
 		/// If found writes a debug project file to disk
 		/// </summary>
 		/// <returns>True on success</returns>
-		public override List<Tuple<ProjectFile, string>> WriteDebugProjectFiles(List<UnrealTargetPlatform> InPlatforms, List<UnrealTargetConfiguration> InConfigurations, PlatformProjectGeneratorCollection PlatformProjectGenerators)
+		public override List<Tuple<ProjectFile, string>> WriteDebugProjectFiles(List<UnrealTargetPlatform> InPlatforms, List<UnrealTargetConfiguration> InConfigurations, PlatformProjectGeneratorCollection PlatformProjectGenerators, ILogger Logger)
 		{
 			//string ProjectName = ProjectFilePath.GetFileNameWithoutExtension();
 
 			List<UnrealTargetPlatform> ProjectPlatforms = new List<UnrealTargetPlatform>();
 			List<Tuple<ProjectFile, string>> ProjectFiles = new List<Tuple<ProjectFile, string>>();
 
-			BuildProjectConfigAndTargetCombinations(InPlatforms, InConfigurations, PlatformProjectGenerators);
+			BuildProjectConfigAndTargetCombinations(InPlatforms, InConfigurations, PlatformProjectGenerators, Logger);
 
 
 			foreach (ProjectConfigAndTargetCombination Combination in ProjectConfigAndTargetCombinations!)
@@ -839,7 +844,7 @@ namespace UnrealBuildTool
 		}
 
 		/// Implements Project interface
-		public override bool WriteProjectFile(List<UnrealTargetPlatform> InPlatforms, List<UnrealTargetConfiguration> InConfigurations, PlatformProjectGeneratorCollection PlatformProjectGenerators)
+		public override bool WriteProjectFile(List<UnrealTargetPlatform> InPlatforms, List<UnrealTargetConfiguration> InConfigurations, PlatformProjectGeneratorCollection PlatformProjectGenerators, ILogger Logger)
 		{
 			string ProjectName = ProjectFilePath.GetFileNameWithoutExtension();
 
@@ -909,16 +914,16 @@ namespace UnrealBuildTool
 				{
 					if (InPlatforms.Contains(UnrealTargetPlatform.Win64))
 					{
-						SharedIncludeSearchPaths.Append(VCToolChain.GetVCIncludePaths(UnrealTargetPlatform.Win64, GetCompilerForIntellisense(), null) + ";");
+						SharedIncludeSearchPaths.Append(VCToolChain.GetVCIncludePaths(UnrealTargetPlatform.Win64, GetCompilerForIntellisense(), null, Logger) + ";");
 					}
 					else if (InPlatforms.Contains(UnrealTargetPlatform.HoloLens))
 					{
-						SharedIncludeSearchPaths.Append(VCToolChain.GetVCIncludePaths(UnrealTargetPlatform.HoloLens, GetCompilerForIntellisense(), null) + ";");
+						SharedIncludeSearchPaths.Append(VCToolChain.GetVCIncludePaths(UnrealTargetPlatform.HoloLens, GetCompilerForIntellisense(), null, Logger) + ";");
 					}
 				}
 				else
 				{
-					Log.TraceInformation("Unable to compute VC include paths on non-Windows host");
+					Logger.LogInformation("Unable to compute VC include paths on non-Windows host");
 				}
 			}
 
@@ -952,7 +957,7 @@ namespace UnrealBuildTool
 				VCUserFileContent.AppendLine("<Project ToolsVersion=\"{0}\" xmlns=\"http://schemas.microsoft.com/developer/msbuild/2003\">", VCProjectFileGenerator.GetProjectFileToolVersionString(ProjectFileFormat));
 			}
 
-			BuildProjectConfigAndTargetCombinations(InPlatforms, InConfigurations, PlatformProjectGenerators);
+			BuildProjectConfigAndTargetCombinations(InPlatforms, InConfigurations, PlatformProjectGenerators, Logger);
 
 			VCProjectFileContent.AppendLine("  <ItemGroup Label=\"ProjectConfigurations\">");
 
@@ -966,7 +971,7 @@ namespace UnrealBuildTool
 				{
 					continue;
 				}
-				if (Combination.ProjectTarget != null && !IsValidProjectPlatformAndConfiguration(Combination.ProjectTarget, (UnrealTargetPlatform)Combination.Platform, Combination.Configuration))
+				if (Combination.ProjectTarget != null && !IsValidProjectPlatformAndConfiguration(Combination.ProjectTarget, (UnrealTargetPlatform)Combination.Platform, Combination.Configuration, Logger))
 				{
 					continue;
 				}
@@ -1084,7 +1089,7 @@ namespace UnrealBuildTool
 				{
 					string ProjectPlatformName = PlatformTuple.Item1;
 					UnrealTargetPlatform TargetPlatform = PlatformTuple.Item2;
-					WritePostDefaultPropsConfiguration(TargetPlatform, TargetConfiguration, ProjectPlatformName, ProjectConfigurationName, PlatformProjectGenerators, VCProjectFileContent);
+					WritePostDefaultPropsConfiguration(TargetPlatform, TargetConfiguration, ProjectPlatformName, ProjectConfigurationName, PlatformProjectGenerators, VCProjectFileContent, Logger);
 				}
 			}
 
@@ -1112,7 +1117,7 @@ namespace UnrealBuildTool
 			// Write each project configuration
 			foreach (ProjectConfigAndTargetCombination Combination in ProjectConfigAndTargetCombinations)
 			{
-				WriteConfiguration(ProjectName, Combination, VCProjectFileContent, PlatformProjectGenerators, bGenerateUserFileContent ? VCUserFileContent : null);
+				WriteConfiguration(ProjectName, Combination, VCProjectFileContent, PlatformProjectGenerators, bGenerateUserFileContent ? VCUserFileContent : null, Logger);
 			}
 
 			// Write IntelliSense info
@@ -1368,7 +1373,7 @@ namespace UnrealBuildTool
 			else if (VSTestRunSettingsFiles.Length > 1)
 			{
 				string Files = String.Join(", ", VSTestRunSettingsFiles.Select(File => File!.FullName));
-				Log.TraceWarning($"Inconsistent VSTest run settings files for project '{ProjectFilePath}': {Files}");
+				Logger.LogWarning("Inconsistent VSTest run settings files for project '{ProjectFilePath}': {Files}", ProjectFilePath, Files);
 			}
 
 			if (!IsStubProject)
@@ -1396,7 +1401,7 @@ namespace UnrealBuildTool
 			// Save the project file
 			if (bSuccess)
 			{
-				bSuccess = ProjectFileGenerator.WriteFileIfChanged(ProjectFilePath.FullName, VCProjectFileContent.ToString());
+				bSuccess = ProjectFileGenerator.WriteFileIfChanged(ProjectFilePath.FullName, VCProjectFileContent.ToString(), Logger);
 			}
 
 
@@ -1407,11 +1412,11 @@ namespace UnrealBuildTool
 				string VCFiltersFilePath = ProjectFilePath.FullName + ".filters";
 				if (FiltersFileIsNeeded)
 				{
-					bSuccess = ProjectFileGenerator.WriteFileIfChanged(VCFiltersFilePath, VCFiltersFileContent.ToString());
+					bSuccess = ProjectFileGenerator.WriteFileIfChanged(VCFiltersFilePath, VCFiltersFileContent.ToString(), Logger);
 				}
 				else
 				{
-					Log.TraceVerbose("Deleting Visual C++ filters file which is no longer needed: " + VCFiltersFilePath);
+					Logger.LogDebug("Deleting Visual C++ filters file which is no longer needed: {File}", VCFiltersFilePath);
 
 					// Delete the filters file, if one exists.  We no longer need it
 					try
@@ -1420,7 +1425,7 @@ namespace UnrealBuildTool
 					}
 					catch (Exception)
 					{
-						Log.TraceInformation("Error deleting filters file (file may not be writable): " + VCFiltersFilePath);
+						Logger.LogInformation("Error deleting filters file (file may not be writable): {File}", VCFiltersFilePath);
 					}
 				}
 			}
@@ -1433,7 +1438,7 @@ namespace UnrealBuildTool
 				// Never overwrite the existing user path as it will cause them to lose their settings
 				if (File.Exists(VCUserFilePath) == false)
 				{
-					bSuccess = ProjectFileGenerator.WriteFileIfChanged(VCUserFilePath, VCUserFileContent.ToString());
+					bSuccess = ProjectFileGenerator.WriteFileIfChanged(VCUserFilePath, VCUserFileContent.ToString(), Logger);
 				}
 			}
 
@@ -1561,7 +1566,7 @@ namespace UnrealBuildTool
 		}
 
 		// Anonymous function that writes post-Default.props configuration data
-		private void WritePostDefaultPropsConfiguration(UnrealTargetPlatform TargetPlatform, UnrealTargetConfiguration TargetConfiguration, string ProjectPlatformName, string ProjectConfigurationName, PlatformProjectGeneratorCollection PlatformProjectGenerators, StringBuilder VCProjectFileContent)
+		private void WritePostDefaultPropsConfiguration(UnrealTargetPlatform TargetPlatform, UnrealTargetConfiguration TargetConfiguration, string ProjectPlatformName, string ProjectConfigurationName, PlatformProjectGeneratorCollection PlatformProjectGenerators, StringBuilder VCProjectFileContent, ILogger Logger)
 		{
 			PlatformProjectGenerator? ProjGenerator = PlatformProjectGenerators.GetPlatformProjectGenerator(TargetPlatform, true);
 
@@ -1591,7 +1596,7 @@ namespace UnrealBuildTool
 		}
 
 		// Anonymous function that writes project configuration data
-		private void WriteConfiguration(string ProjectName, ProjectConfigAndTargetCombination Combination, StringBuilder VCProjectFileContent, PlatformProjectGeneratorCollection PlatformProjectGenerators, StringBuilder? VCUserFileContent)
+		private void WriteConfiguration(string ProjectName, ProjectConfigAndTargetCombination Combination, StringBuilder VCProjectFileContent, PlatformProjectGeneratorCollection PlatformProjectGenerators, StringBuilder? VCUserFileContent, ILogger Logger)
 		{
 			UnrealTargetConfiguration Configuration = Combination.Configuration;
 
@@ -1881,8 +1886,9 @@ namespace UnrealBuildTool
 		/// Constructs a new project file object
 		/// </summary>
 		/// <param name="InitFilePath">The path to the project file on disk</param>
+		/// <param name="Logger">Logger for output</param>
 		/// <param name="BaseDir">The base directory for files within this project - if not specified, InitFilePath.Directory will be used</param>
-		public VCSharpProjectFile(FileReference InitFilePath, DirectoryReference? BaseDir = null)
+		public VCSharpProjectFile(FileReference InitFilePath, ILogger Logger, DirectoryReference? BaseDir = null)
 			: base(InitFilePath, BaseDir ?? InitFilePath.Directory)
 		{
 			try
@@ -1915,7 +1921,7 @@ namespace UnrealBuildTool
 								}
 								else
 								{
-									Log.TraceWarning("Unable to parse configuration/platform from condition '{0}': {1}", InitFilePath, Condition);
+									Logger.LogWarning("Unable to parse configuration/platform from condition '{InitFilePath}': {Condition}", InitFilePath, Condition);
 								}
 							}
 						}
@@ -1935,7 +1941,7 @@ namespace UnrealBuildTool
 						if (PropertyGroup.HasAttribute("Condition"))
 						{
 							string Condition = PropertyGroup.GetAttribute("Condition");
-							Log.TraceWarning("Unable to parse configuration from property group with condition '{0}': {1}. UBT Requires you to set the configuration without conditionals.", InitFilePath, Condition);
+							Logger.LogWarning("Unable to parse configuration from property group with condition '{InitFilePath}': {Condition}. UBT Requires you to set the configuration without conditionals.", InitFilePath, Condition);
 							continue;
 						}
 						string[]? ParsedConfigurations = ConfigNodeList[0]?.FirstChild?.Value?.Split(';');
@@ -1966,7 +1972,7 @@ namespace UnrealBuildTool
 			}
 			catch (Exception Ex)
 			{
-				Log.TraceWarning("Unable to parse {0}: {1}", InitFilePath, Ex.ToString());
+				Logger.LogWarning("Unable to parse {Path}: {Ex}", InitFilePath, Ex.ToString());
 			}
 		}
 
@@ -2009,8 +2015,9 @@ namespace UnrealBuildTool
 		/// <param name="SolutionConfiguration">The solution configuration</param>
 		/// <param name="SolutionPlatform">The solution platform</param>
 		/// <param name="PlatformProjectGenerators">Set of platform project generators</param>
+		/// <param name="Logger"></param>
 		/// <returns>Project context matching the given solution context</returns>
-		public override MSBuildProjectContext GetMatchingProjectContext(TargetType SolutionTarget, UnrealTargetConfiguration SolutionConfiguration, UnrealTargetPlatform SolutionPlatform, PlatformProjectGeneratorCollection PlatformProjectGenerators)
+		public override MSBuildProjectContext GetMatchingProjectContext(TargetType SolutionTarget, UnrealTargetConfiguration SolutionConfiguration, UnrealTargetPlatform SolutionPlatform, PlatformProjectGeneratorCollection PlatformProjectGenerators, ILogger Logger)
 		{
 			// Find the matching platform name
 			string ProjectPlatformName;
@@ -2055,8 +2062,9 @@ namespace UnrealBuildTool
 		/// <param name="InPlatforms">Not used.</param>
 		/// <param name="InConfigurations">Not Used.</param>
 		/// <param name="PlatformProjectGenerators">Set of platform project generators</param>
+		/// <param name="Logger"></param>
 		/// <returns>true if the opration was successful, false otherwise</returns>
-		public override bool WriteProjectFile(List<UnrealTargetPlatform> InPlatforms, List<UnrealTargetConfiguration> InConfigurations, PlatformProjectGeneratorCollection PlatformProjectGenerators)
+		public override bool WriteProjectFile(List<UnrealTargetPlatform> InPlatforms, List<UnrealTargetConfiguration> InConfigurations, PlatformProjectGeneratorCollection PlatformProjectGenerators, ILogger Logger)
 		{
 			throw new BuildException("Support for writing C# projects from UnrealBuildTool has been removed.");
 		}
diff --git a/Engine/Source/Programs/UnrealBuildTool/ProjectFiles/VisualStudio/VCProjectFileGenerator.cs b/Engine/Source/Programs/UnrealBuildTool/ProjectFiles/VisualStudio/VCProjectFileGenerator.cs
--- a/Engine/Source/Programs/UnrealBuildTool/ProjectFiles/VisualStudio/VCProjectFileGenerator.cs
+++ b/Engine/Source/Programs/UnrealBuildTool/ProjectFiles/VisualStudio/VCProjectFileGenerator.cs
@@ -9,6 +9,7 @@ using System.Text;
 using System.Xml.Linq;
 using EpicGames.Core;
 using UnrealBuildBase;
+using Microsoft.Extensions.Logging;
 
 namespace UnrealBuildTool
 {
@@ -144,7 +145,7 @@ namespace UnrealBuildTool
 
 		/// <summary>
 		/// </summary>
-		public override void CleanProjectFiles(DirectoryReference InPrimaryProjectDirectory, string InPrimaryProjectName, DirectoryReference InIntermediateProjectFilesDirectory)
+		public override void CleanProjectFiles(DirectoryReference InPrimaryProjectDirectory, string InPrimaryProjectName, DirectoryReference InIntermediateProjectFilesDirectory, ILogger Logger)
 		{
 			FileReference PrimaryProjectFile = FileReference.Combine(InPrimaryProjectDirectory, InPrimaryProjectName);
 			FileReference PrimaryProjDeleteFilename = PrimaryProjectFile + ".sln";
@@ -182,8 +183,8 @@ namespace UnrealBuildTool
 				}
 				catch (Exception Ex)
 				{
-					Log.TraceInformation("Error while trying to clean project files path {0}. Ignored.", InIntermediateProjectFilesDirectory);
-					Log.TraceInformation("\t" + Ex.Message);
+					Logger.LogInformation("Error while trying to clean project files path {InIntermediateProjectFilesDirectory}. Ignored.", InIntermediateProjectFilesDirectory);
+					Logger.LogInformation("\t{Message}", Ex.Message);
 				}
 			}
 		}
@@ -234,41 +235,38 @@ namespace UnrealBuildTool
 			VCProjectFileContent.AppendLine("    <PlatformToolset>{0}</PlatformToolset>", PlatformToolsetVersionString);
 		}
 
-		/// <summary>
-		/// Configures project generator based on command-line options
-		/// </summary>
-		/// <param name="Arguments">Arguments passed into the program</param>
-		/// <param name="IncludeAllPlatforms">True if all platforms should be included</param>
-		protected override void ConfigureProjectFileGeneration(String[] Arguments, ref bool IncludeAllPlatforms)
+		/// <inheritdoc/>
+		protected override void ConfigureProjectFileGeneration(String[] Arguments, ref bool IncludeAllPlatforms, ILogger Logger)
 		{
 			// Call parent implementation first
-			base.ConfigureProjectFileGeneration(Arguments, ref IncludeAllPlatforms);
+			base.ConfigureProjectFileGeneration(Arguments, ref IncludeAllPlatforms, Logger);
 		}
 
 		/// <summary>
 		/// Selects which platforms and build configurations we want in the project file
 		/// </summary>
 		/// <param name="IncludeAllPlatforms">True if we should include ALL platforms that are supported on this machine.  Otherwise, only desktop platforms will be included.</param>
+		/// <param name="Logger"></param>
 		/// <param name="SupportedPlatformNames">Output string for supported platforms, returned as comma-separated values.</param>
-		protected override void SetupSupportedPlatformsAndConfigurations(bool IncludeAllPlatforms, out string SupportedPlatformNames)
+		protected override void SetupSupportedPlatformsAndConfigurations(bool IncludeAllPlatforms, ILogger Logger, out string SupportedPlatformNames)
 		{
 			// Call parent implementation to figure out the actual platforms
-			base.SetupSupportedPlatformsAndConfigurations(IncludeAllPlatforms, out SupportedPlatformNames);
+			base.SetupSupportedPlatformsAndConfigurations(IncludeAllPlatforms, Logger, out SupportedPlatformNames);
 
 			// If we have a non-default setting for visual studio, check the compiler exists. If not, revert to the default.
 			if (Settings.ProjectFileFormat == VCProjectFileFormat.VisualStudio2019)
 			{
-				if (!WindowsPlatform.HasCompiler(WindowsCompiler.VisualStudio2019, WindowsArchitecture.x64))
+				if (!WindowsPlatform.HasCompiler(WindowsCompiler.VisualStudio2019, WindowsArchitecture.x64, Logger))
 				{
-					Log.TraceWarning("Visual Studio C++ 2019 installation not found - ignoring preferred project file format.");
+					Logger.LogWarning("Visual Studio C++ 2019 installation not found - ignoring preferred project file format.");
 					Settings.ProjectFileFormat = VCProjectFileFormat.Default;
 				}
 			}
 			else if (Settings.ProjectFileFormat == VCProjectFileFormat.VisualStudio2022)
 			{
-				if (!WindowsPlatform.HasCompiler(WindowsCompiler.VisualStudio2022, WindowsArchitecture.x64))
+				if (!WindowsPlatform.HasCompiler(WindowsCompiler.VisualStudio2022, WindowsArchitecture.x64, Logger))
 				{
-					Log.TraceWarning("Visual Studio C++ 2022 installation not found - ignoring preferred project file format.");
+					Logger.LogWarning("Visual Studio C++ 2022 installation not found - ignoring preferred project file format.");
 					Settings.ProjectFileFormat = VCProjectFileFormat.Default;
 				}
 			}
@@ -279,7 +277,7 @@ namespace UnrealBuildTool
 			if (Settings.ProjectFileFormat == VCProjectFileFormat.Default)
 			{
 				// Enumerate all the valid installations. This list is already sorted by preference.
-				List<VisualStudioInstallation> Installations = MicrosoftPlatformSDK.VisualStudioInstallations.Where(x => WindowsPlatform.HasCompiler(x.Compiler, WindowsArchitecture.x64)).ToList();
+				List<VisualStudioInstallation> Installations = MicrosoftPlatformSDK.FindVisualStudioInstallations(Logger).Where(x => WindowsPlatform.HasCompiler(x.Compiler, WindowsArchitecture.x64, Logger)).ToList();
 
 				// Get the corresponding project file format
 				VCProjectFileFormat Format = VCProjectFileFormat.VisualStudio2019;
@@ -437,9 +435,9 @@ namespace UnrealBuildTool
 		/// Writes the project files to disk
 		/// </summary>
 		/// <returns>True if successful</returns>
-		protected override bool WriteProjectFiles(PlatformProjectGeneratorCollection PlatformProjectGenerators)
+		protected override bool WriteProjectFiles(PlatformProjectGeneratorCollection PlatformProjectGenerators, ILogger Logger)
 		{
-			if (!base.WriteProjectFiles(PlatformProjectGenerators))
+			if (!base.WriteProjectFiles(PlatformProjectGenerators, Logger))
 			{
 				return false;
 			}
@@ -469,7 +467,7 @@ namespace UnrealBuildTool
 		}
 
 
-		protected override bool WritePrimaryProjectFile(ProjectFile? UBTProject, PlatformProjectGeneratorCollection PlatformProjectGenerators)
+		protected override bool WritePrimaryProjectFile(ProjectFile? UBTProject, PlatformProjectGeneratorCollection PlatformProjectGenerators, ILogger Logger)
 		{
 			bool bSuccess = true;
 
@@ -651,7 +649,7 @@ namespace UnrealBuildTool
 												// Figure out the set of valid target configuration names
 												foreach (Project ProjectTarget in CurProject.ProjectTargets)
 												{
-													if (VCProjectFile.IsValidProjectPlatformAndConfiguration(ProjectTarget, CurPlatform, CurConfiguration))
+													if (VCProjectFile.IsValidProjectPlatformAndConfiguration(ProjectTarget, CurPlatform, CurConfiguration, Logger))
 													{
 														PlatformsValidForProjects.Add(CurPlatform);
 
@@ -736,7 +734,7 @@ namespace UnrealBuildTool
 							foreach (VCSolutionConfigCombination SolutionConfigCombination in SolutionConfigCombinations)
 							{
 								// Get the context for the current solution context
-								MSBuildProjectContext ProjectContext = CurProject.GetMatchingProjectContext(SolutionConfigCombination.TargetConfigurationName, SolutionConfigCombination.Configuration, SolutionConfigCombination.Platform, PlatformProjectGenerators);
+								MSBuildProjectContext ProjectContext = CurProject.GetMatchingProjectContext(SolutionConfigCombination.TargetConfigurationName, SolutionConfigCombination.Configuration, SolutionConfigCombination.Platform, PlatformProjectGenerators, Logger);
 
 								// Override the configuration to build for UBT
 								if (Settings.bBuildUBTInDebug && CurProject == UBTProject)
@@ -825,7 +823,7 @@ namespace UnrealBuildTool
 			if (bSuccess)
 			{
 				string SolutionFilePath = FileReference.Combine(PrimaryProjectPath, SolutionFileName).FullName;
-				bSuccess = WriteFileIfChanged(SolutionFilePath, VCSolutionFileContent.ToString());
+				bSuccess = WriteFileIfChanged(SolutionFilePath, VCSolutionFileContent.ToString(), Logger);
 			}
 
 
@@ -886,7 +884,7 @@ namespace UnrealBuildTool
 			return bSuccess;
 		}
 
-		protected override void WriteDebugSolutionFiles(PlatformProjectGeneratorCollection PlatformProjectGenerators, DirectoryReference IntermediateProjectFilesPath)
+		protected override void WriteDebugSolutionFiles(PlatformProjectGeneratorCollection PlatformProjectGenerators, DirectoryReference IntermediateProjectFilesPath, ILogger Logger)
 		{
 			//build and collect UnrealVS configuration
 			StringBuilder UnrealVSContent = new StringBuilder();
@@ -905,7 +903,7 @@ namespace UnrealBuildTool
 
 				string ConfigFilePath = FileReference.Combine(IntermediateProjectFilesPath, "UnrealVS.xml").FullName;
 				/* bool bSuccess = */
-				ProjectFileGenerator.WriteFileIfChanged(ConfigFilePath, UnrealVSContent.ToString());
+				ProjectFileGenerator.WriteFileIfChanged(ConfigFilePath, UnrealVSContent.ToString(), Logger);
 			}
 		}
 
diff --git a/Engine/Source/Programs/UnrealBuildTool/ProjectFiles/VisualStudioCode/VSCodeProjectFileGenerator.cs b/Engine/Source/Programs/UnrealBuildTool/ProjectFiles/VisualStudioCode/VSCodeProjectFileGenerator.cs
--- a/Engine/Source/Programs/UnrealBuildTool/ProjectFiles/VisualStudioCode/VSCodeProjectFileGenerator.cs
+++ b/Engine/Source/Programs/UnrealBuildTool/ProjectFiles/VisualStudioCode/VSCodeProjectFileGenerator.cs
@@ -8,6 +8,7 @@ using EpicGames.Core;
 using System.Diagnostics;
 using System.Linq;
 using UnrealBuildBase;
+using Microsoft.Extensions.Logging;
 
 namespace UnrealBuildTool
 {
@@ -18,7 +19,7 @@ namespace UnrealBuildTool
 		{
 		}
 
-		public override bool WriteProjectFile(List<UnrealTargetPlatform> InPlatforms, List<UnrealTargetConfiguration> InConfigurations, PlatformProjectGeneratorCollection PlatformProjectGenerators)
+		public override bool WriteProjectFile(List<UnrealTargetPlatform> InPlatforms, List<UnrealTargetConfiguration> InConfigurations, PlatformProjectGeneratorCollection PlatformProjectGenerators, ILogger Logger)
 		{
 			return true;
 		}
@@ -231,7 +232,7 @@ namespace UnrealBuildTool
 			}
 		}
 
-		public override void CleanProjectFiles(DirectoryReference InPrimaryProjectDirectory, string InPrimaryProjectName, DirectoryReference InIntermediateProjectFilesPath)
+		public override void CleanProjectFiles(DirectoryReference InPrimaryProjectDirectory, string InPrimaryProjectName, DirectoryReference InIntermediateProjectFilesPath, ILogger Logger)
 		{
 		}
 
@@ -245,7 +246,7 @@ namespace UnrealBuildTool
 			return new VSCodeProject(InitFilePath, BaseDir);
 		}
 
-		protected override bool WritePrimaryProjectFile(ProjectFile? UBTProject, PlatformProjectGeneratorCollection PlatformProjectGenerators)
+		protected override bool WritePrimaryProjectFile(ProjectFile? UBTProject, PlatformProjectGeneratorCollection PlatformProjectGenerators, ILogger Logger)
 		{
 			DirectoryReference VSCodeDir = DirectoryReference.Combine(PrimaryProjectPath, ".vscode");
 			DirectoryReference.CreateDirectory(VSCodeDir);
@@ -272,10 +273,10 @@ namespace UnrealBuildTool
 			}
 			Projects.Sort((A, B) => { return A.ProjectFilePath.GetFileName().CompareTo(B.ProjectFilePath.GetFileName()); });
 
-			ProjectData ProjectData = GatherProjectData(Projects);
+			ProjectData ProjectData = GatherProjectData(Projects, Logger);
 
 			WriteTasksFile(ProjectData, VSCodeDir);
-			WriteLaunchFile(ProjectData, VSCodeDir);
+			WriteLaunchFile(ProjectData, VSCodeDir, Logger);
 			WriteWorkspaceIgnoreFile(Projects);
 			WriteCppPropertiesFile(VSCodeDir, ProjectData);
 			WriteWorkspaceFile();
@@ -318,22 +319,22 @@ namespace UnrealBuildTool
 			}
 		}
 
-		protected override void AddTargetForIntellisense(UEBuildTarget Target)
+		protected override void AddTargetForIntellisense(UEBuildTarget Target, ILogger Logger)
 		{
-			base.AddTargetForIntellisense(Target);
+			base.AddTargetForIntellisense(Target, Logger);
 
 			bool UsingClang = true;
 			FileReference? CompilerPath = null;
 			DirectoryReference? SysRootPath = null;
 			if (OperatingSystem.IsWindows())
 			{
-				VCEnvironment Environment = VCEnvironment.Create(WindowsPlatform.GetDefaultCompiler(null, Target.Rules.WindowsPlatform.Architecture), Target.Platform, Target.Rules.WindowsPlatform.Architecture, null, Target.Rules.WindowsPlatform.WindowsSdkVersion, null);
+				VCEnvironment Environment = VCEnvironment.Create(WindowsPlatform.GetDefaultCompiler(null, Target.Rules.WindowsPlatform.Architecture, Logger), Target.Platform, Target.Rules.WindowsPlatform.Architecture, null, Target.Rules.WindowsPlatform.WindowsSdkVersion, null, Logger);
 				CompilerPath = FileReference.FromString(Environment.CompilerPath.FullName);
 				UsingClang = false;
 			}
 			else if (OperatingSystem.IsLinux())
 			{
-				CompilerPath = FileReference.FromString(LinuxCommon.WhichClang());
+				CompilerPath = FileReference.FromString(LinuxCommon.WhichClang(Logger));
 				string? InternalSDKPath = UEBuildPlatform.GetSDK(UnrealTargetPlatform.Linux)?.GetInternalSDKPath();
 				if (!string.IsNullOrEmpty(InternalSDKPath))
 				{
@@ -342,7 +343,7 @@ namespace UnrealBuildTool
 			}
 			else if (OperatingSystem.IsMacOS())
 			{
-				MacToolChainSettings Settings = new MacToolChainSettings(false);
+				MacToolChainSettings Settings = new MacToolChainSettings(false, Logger);
 				CompilerPath = FileReference.FromString(Settings.ToolchainDir + "clang++");
 				SysRootPath = DirectoryReference.FromString(Settings.BaseSDKDir + "/MacOSX" + Settings.MacOSSDKVersion + ".sdk");
 			}
@@ -355,13 +356,13 @@ namespace UnrealBuildTool
 			Dictionary<DirectoryReference, string> ModuleDirectoryToCompileCommand = new Dictionary<DirectoryReference, string>();
 
 			// Generate a compile environment for each module in the binary
-			CppCompileEnvironment GlobalCompileEnvironment = Target.CreateCompileEnvironmentForProjectFiles();
+			CppCompileEnvironment GlobalCompileEnvironment = Target.CreateCompileEnvironmentForProjectFiles(Logger);
 			foreach (UEBuildBinary Binary in Target.Binaries)
 			{
 				CppCompileEnvironment BinaryCompileEnvironment = Binary.CreateBinaryCompileEnvironment(GlobalCompileEnvironment);
 				foreach (UEBuildModuleCPP Module in Binary.Modules.OfType<UEBuildModuleCPP>())
 				{
-					CppCompileEnvironment ModuleCompileEnvironment = Module.CreateCompileEnvironmentForIntellisense(Target.Rules, BinaryCompileEnvironment);
+					CppCompileEnvironment ModuleCompileEnvironment = Module.CreateCompileEnvironmentForIntellisense(Target.Rules, BinaryCompileEnvironment, Logger);
 
 					List<FileReference> ForceIncludePaths = new List<FileReference>(ModuleCompileEnvironment.ForceIncludeFiles.Select(x => x.Location));
 					if (ModuleCompileEnvironment.PrecompiledHeaderIncludeFilename != null)
@@ -469,7 +470,7 @@ namespace UnrealBuildTool
 		}
 
 
-		private ProjectData GatherProjectData(List<ProjectFile> InProjects)
+		private ProjectData GatherProjectData(List<ProjectFile> InProjects, ILogger Logger)
 		{
 			ProjectData ProjectData = new ProjectData();
 
@@ -502,7 +503,7 @@ namespace UnrealBuildTool
 							{
 								foreach (UnrealTargetConfiguration Config in Configs)
 								{
-									if (MSBuildProjectFile.IsValidProjectPlatformAndConfiguration(Target, Platform, Config))
+									if (MSBuildProjectFile.IsValidProjectPlatformAndConfiguration(Target, Platform, Config, Logger))
 									{
 										NewTarget.BuildProducts.Add(new ProjectData.BuildProduct(GetExecutableFilename(Project, Target, Platform, Config))
 										{
@@ -1127,7 +1128,7 @@ namespace UnrealBuildTool
 			return new FileReference(ExecutableFilename);
 		}
 
-		private void WriteNativeLaunchConfigAndroidOculus(ProjectData.Project InProject, JsonFile OutFile, ProjectData.Target Target, ProjectData.BuildProduct BuildProduct)
+		private void WriteNativeLaunchConfigAndroidOculus(ProjectData.Project InProject, JsonFile OutFile, ProjectData.Target Target, ProjectData.BuildProduct BuildProduct, ILogger Logger)
 		{
 			ConfigHierarchy Ini = ConfigCache.ReadHierarchy(ConfigHierarchyType.Engine, DirectoryReference.FromFile(BuildProduct.UProjectFile), BuildProduct.Platform);
 
@@ -1358,7 +1359,7 @@ namespace UnrealBuildTool
 			}
 		}
 
-		private void WriteNativeLaunchConfig(ProjectData.Project InProject, JsonFile OutFile)
+		private void WriteNativeLaunchConfig(ProjectData.Project InProject, JsonFile OutFile, ILogger Logger)
 		{
 			foreach (ProjectData.Target Target in InProject.Targets)
 			{
@@ -1370,7 +1371,7 @@ namespace UnrealBuildTool
 					}
 					else if (BuildProduct.Platform == UnrealTargetPlatform.Android)
 					{
-						WriteNativeLaunchConfigAndroidOculus(InProject, OutFile, Target, BuildProduct);
+						WriteNativeLaunchConfigAndroidOculus(InProject, OutFile, Target, BuildProduct, Logger);
 					}
 				}
 			}
@@ -1437,7 +1438,7 @@ namespace UnrealBuildTool
 			}
 		}
 
-		private void WriteLaunchFile(ProjectData ProjectData, DirectoryReference VSCodeDir)
+		private void WriteLaunchFile(ProjectData ProjectData, DirectoryReference VSCodeDir, ILogger Logger)
 		{
 			JsonFile OutFile = new JsonFile();
 
@@ -1467,7 +1468,7 @@ namespace UnrealBuildTool
 				{
 					foreach (ProjectData.Project Project in ProjectData.NativeProjects)
 					{
-						WriteNativeLaunchConfig(Project, OutFile);
+						WriteNativeLaunchConfig(Project, OutFile, Logger);
 					}
 
 					foreach (ProjectData.Project Project in ProjectData.CSharpProjects)
diff --git a/Engine/Source/Programs/UnrealBuildTool/ProjectFiles/VisualStudioMac/VisualStudioMacProjectFileGenerator.cs b/Engine/Source/Programs/UnrealBuildTool/ProjectFiles/VisualStudioMac/VisualStudioMacProjectFileGenerator.cs
--- a/Engine/Source/Programs/UnrealBuildTool/ProjectFiles/VisualStudioMac/VisualStudioMacProjectFileGenerator.cs
+++ b/Engine/Source/Programs/UnrealBuildTool/ProjectFiles/VisualStudioMac/VisualStudioMacProjectFileGenerator.cs
@@ -9,6 +9,7 @@ using System.Xml;
 using System.Xml.Linq;
 using EpicGames.Core;
 using UnrealBuildBase;
+using Microsoft.Extensions.Logging;
 
 namespace UnrealBuildTool
 {
@@ -66,7 +67,7 @@ namespace UnrealBuildTool
         /// Writes the project files to disk
         /// </summary>
         /// <returns>True if successful</returns>
-        protected override bool WriteProjectFiles(PlatformProjectGeneratorCollection PlatformProjectGenerators)
+        protected override bool WriteProjectFiles(PlatformProjectGeneratorCollection PlatformProjectGenerators, ILogger Logger)
         {
 			// This can be reset by higher level code when it detects that we don't have
 			// VS2019 installed (TODO - add custom format for Mac?)
@@ -80,7 +81,7 @@ namespace UnrealBuildTool
             OtherProjectFiles.RemoveAll(x => !IsValidProject(x));
             AutomationProjectFiles.RemoveAll(x => !IsValidProject(x));
 
-            if (!base.WriteProjectFiles(PlatformProjectGenerators))
+            if (!base.WriteProjectFiles(PlatformProjectGenerators, Logger))
             {
                 return false;
             }
diff --git a/Engine/Source/Programs/UnrealBuildTool/ProjectFiles/Xcode/XcodeFrameworkWrapperProject.cs b/Engine/Source/Programs/UnrealBuildTool/ProjectFiles/Xcode/XcodeFrameworkWrapperProject.cs
--- a/Engine/Source/Programs/UnrealBuildTool/ProjectFiles/Xcode/XcodeFrameworkWrapperProject.cs
+++ b/Engine/Source/Programs/UnrealBuildTool/ProjectFiles/Xcode/XcodeFrameworkWrapperProject.cs
@@ -7,6 +7,7 @@
 using System.Collections.Generic;
 using EpicGames.Core;
 using UnrealBuildBase;
+using Microsoft.Extensions.Logging;
 
 namespace UnrealBuildTool.ProjectFiles.Xcode
 {
diff --git a/Engine/Source/Programs/UnrealBuildTool/ProjectFiles/Xcode/XcodeProject.cs b/Engine/Source/Programs/UnrealBuildTool/ProjectFiles/Xcode/XcodeProject.cs
--- a/Engine/Source/Programs/UnrealBuildTool/ProjectFiles/Xcode/XcodeProject.cs
+++ b/Engine/Source/Programs/UnrealBuildTool/ProjectFiles/Xcode/XcodeProject.cs
@@ -10,6 +10,7 @@
 using System.Text;
 using EpicGames.Core;
 using UnrealBuildBase;
+using Microsoft.Extensions.Logging;
 
 namespace UnrealBuildTool
 {
@@ -359,7 +360,7 @@
 		}
 
 		private void GenerateSectionsWithExtensions(StringBuilder PBXBuildFileSection, StringBuilder PBXFileReferenceSection, StringBuilder PBXCopyFilesBuildPhaseSection, StringBuilder PBXResourcesBuildPhaseSection,
-													List<XcodeExtensionInfo> AllExtensions, FileReference? UProjectPath, List<XcodeBuildConfig> BuildConfigs)
+													List<XcodeExtensionInfo> AllExtensions, FileReference? UProjectPath, List<XcodeBuildConfig> BuildConfigs, ILogger Logger)
 		{
 			if (UProjectPath != null)
 			{
@@ -463,7 +464,7 @@
 							}
 
 							string? IOSRuntimeVersion, TVOSRuntimeVersion;
-							AppendPlatformConfiguration(ConfigSection, Configuration, ExtensionInfo.Name, UProjectPath, false, bSupportIOS, bSupportTVOS, out IOSRuntimeVersion, out TVOSRuntimeVersion);
+							AppendPlatformConfiguration(ConfigSection, Configuration, ExtensionInfo.Name, UProjectPath, false, bSupportIOS, bSupportTVOS, Logger, out IOSRuntimeVersion, out TVOSRuntimeVersion);
 
 							ConfigSection.Append("\t\t\t};" + ProjectFileGenerator.NewLine);
 							ConfigSection.Append("\t\t\tname = \"" + ConfigName + "\";" + ProjectFileGenerator.NewLine);
@@ -1069,7 +1070,7 @@
 			return CachedMacProjectArcitectures[TargetName];
 		}
 
-		private void AppendPlatformConfiguration(StringBuilder Content, XcodeBuildConfig Config, string TargetName, FileReference? ProjectFile, bool bSupportMac, bool bSupportIOS, bool bSupportTVOS, out string? IOSRunTimeVersion, out string? TVOSRunTimeVersion, string BinariesSubDir = "/Payload")
+		private void AppendPlatformConfiguration(StringBuilder Content, XcodeBuildConfig Config, string TargetName, FileReference? ProjectFile, bool bSupportMac, bool bSupportIOS, bool bSupportTVOS, ILogger Logger, out string? IOSRunTimeVersion, out string? TVOSRunTimeVersion, string BinariesSubDir = "/Payload")
 		{
 			FileReference MacExecutablePath = Config.MacExecutablePath!;
 
@@ -1238,7 +1239,7 @@
 			}
 		}
 
-		private void AppendNativeTargetBuildConfiguration(StringBuilder Content, XcodeBuildConfig Config, string ConfigGuid, FileReference? ProjectFile)
+		private void AppendNativeTargetBuildConfiguration(StringBuilder Content, XcodeBuildConfig Config, string ConfigGuid, FileReference? ProjectFile, ILogger Logger)
 		{
 			bool bMacOnly = true;
 			if (Config.ProjectTarget!.TargetRules != null && XcodeProjectFileGenerator.ProjectFilePlatform.HasFlag(XcodeProjectFileGenerator.XcodeProjectFilePlatform.iOS))
@@ -1258,7 +1259,7 @@
 			string MacExecutableFileName = Config.MacExecutablePath!.GetFileName();
 
 			string? IOSRunTimeVersion, TVOSRunTimeVersion;
-			AppendPlatformConfiguration(Content, Config, Config.BuildTarget, ProjectFile, true, !bMacOnly, !bMacOnly, out IOSRunTimeVersion, out TVOSRunTimeVersion);
+			AppendPlatformConfiguration(Content, Config, Config.BuildTarget, ProjectFile, true, !bMacOnly, !bMacOnly, Logger, out IOSRunTimeVersion, out TVOSRunTimeVersion);
 
 			if (!bMacOnly)
 			{
@@ -1384,12 +1385,12 @@
 							TargetReceipt? Receipt;
 							TargetReceipt.TryRead(ReceiptFilename, out Receipt);
 							bool bBuildAsFramework = UEDeployIOS.GetCompileAsDll(Receipt);
-							UEDeployIOS.GenerateIOSPList(ProjectFile, Config.BuildConfig, ProjectPath.FullName, bIsUnrealGame, GameName, bIsClient, Config.BuildTarget, EngineDir.FullName, ProjectPath + "/Binaries/IOS/Payload", null, BundleIdentifier, bBuildAsFramework, out bSupportPortrait, out bSupportLandscape);
+							UEDeployIOS.GenerateIOSPList(ProjectFile, Config.BuildConfig, ProjectPath.FullName, bIsUnrealGame, GameName, bIsClient, Config.BuildTarget, EngineDir.FullName, ProjectPath + "/Binaries/IOS/Payload", null, BundleIdentifier, bBuildAsFramework, Logger, out bSupportPortrait, out bSupportLandscape);
 						}
 						if (bCreateTVOSInfoPlist)
 						{
 							Directory.CreateDirectory(Path.GetDirectoryName(TVOSInfoPlistPath)!);
-							UEDeployTVOS.GenerateTVOSPList(ProjectPath.FullName, bIsUnrealGame, GameName, bIsClient, Config.BuildTarget, EngineDir.FullName, ProjectPath + "/Binaries/TVOS/Payload", null, BundleIdentifier);
+							UEDeployTVOS.GenerateTVOSPList(ProjectPath.FullName, bIsUnrealGame, GameName, bIsClient, Config.BuildTarget, EngineDir.FullName, ProjectPath + "/Binaries/TVOS/Payload", null, BundleIdentifier, Logger);
 						}
 					}
 				}
@@ -1475,7 +1476,7 @@
 
 		private void AppendXCBuildConfigurationSection(StringBuilder Content, Dictionary<string, XcodeBuildConfig> ProjectBuildConfigs, Dictionary<string, XcodeBuildConfig> TargetBuildConfigs,
 				Dictionary<string, XcodeBuildConfig> BuildTargetBuildConfigs, Dictionary<string, XcodeBuildConfig> IndexTargetBuildConfigs, FileReference? GameProjectPath,
-				List<XcodeExtensionInfo> AllExtensions)
+				List<XcodeExtensionInfo> AllExtensions, ILogger Logger)
 		{
 			Content.Append("/* Begin XCBuildConfiguration section */" + ProjectFileGenerator.NewLine);
 
@@ -1486,7 +1487,7 @@
 
 			foreach (KeyValuePair<string, XcodeBuildConfig> Config in TargetBuildConfigs)
 			{
-				AppendNativeTargetBuildConfiguration(Content, Config.Value, Config.Key, GameProjectPath);
+				AppendNativeTargetBuildConfiguration(Content, Config.Value, Config.Key, GameProjectPath, Logger);
 			}
 
 			foreach (KeyValuePair<string, XcodeBuildConfig> Config in BuildTargetBuildConfigs)
@@ -1496,7 +1497,7 @@
 
 			foreach (KeyValuePair<string, XcodeBuildConfig> Config in IndexTargetBuildConfigs)
 			{
-				AppendNativeTargetBuildConfiguration(Content, Config.Value, Config.Key, GameProjectPath);
+				AppendNativeTargetBuildConfiguration(Content, Config.Value, Config.Key, GameProjectPath, Logger);
 			}
 
 			foreach (XcodeExtensionInfo EI in AllExtensions)
@@ -1543,7 +1544,7 @@
 			Content.Append("/* End XCConfigurationList section */" + ProjectFileGenerator.NewLine);
 			}
 
-		private List<XcodeBuildConfig> GetSupportedBuildConfigs(List<UnrealTargetPlatform> Platforms, List<UnrealTargetConfiguration> Configurations)
+		private List<XcodeBuildConfig> GetSupportedBuildConfigs(List<UnrealTargetPlatform> Platforms, List<UnrealTargetConfiguration> Configurations, ILogger Logger)
 		{
 			List<XcodeBuildConfig> BuildConfigs = new List<XcodeBuildConfig>();
 
@@ -1569,7 +1570,7 @@
 								// Now go through all of the target types for this project
 								foreach (ProjectTarget ProjectTarget in ProjectTargets.OfType<ProjectTarget>())
 								{
-									if (MSBuildProjectFile.IsValidProjectPlatformAndConfiguration(ProjectTarget, Platform, Configuration))
+									if (MSBuildProjectFile.IsValidProjectPlatformAndConfiguration(ProjectTarget, Platform, Configuration, Logger))
 									{
 										// Figure out if this is a monolithic build
 										bool bShouldCompileMonolithic = BuildPlatform.ShouldCompileMonolithicBinary(Platform);
@@ -1904,12 +1905,12 @@
 			};
 		}
 
-		public bool ShouldIncludeProjectInWorkspace()
+		public bool ShouldIncludeProjectInWorkspace(ILogger Logger)
 		{
-			return CanBuildProjectLocally();
+			return CanBuildProjectLocally(Logger);
 		}
 
-		public bool CanBuildProjectLocally()
+		public bool CanBuildProjectLocally(ILogger Logger)
 		{
 			foreach (Project ProjectTarget in ProjectTargets)
 			{
@@ -1917,7 +1918,7 @@
 				{
 					foreach (UnrealTargetConfiguration Config in GetSupportedConfigurations())
 					{
-						if (MSBuildProjectFile.IsValidProjectPlatformAndConfiguration(ProjectTarget, Platform, Config))
+						if (MSBuildProjectFile.IsValidProjectPlatformAndConfiguration(ProjectTarget, Platform, Config, Logger))
 						{
 							return true;
 						}
@@ -1929,7 +1930,7 @@
 		}
 
 		/// Implements Project interface
-		public override bool WriteProjectFile(List<UnrealTargetPlatform> InPlatforms, List<UnrealTargetConfiguration> InConfigurations, PlatformProjectGeneratorCollection PlatformProjectGenerators)
+		public override bool WriteProjectFile(List<UnrealTargetPlatform> InPlatforms, List<UnrealTargetConfiguration> InConfigurations, PlatformProjectGeneratorCollection PlatformProjectGenerators, ILogger Logger)
 		{
 			bool bSuccess = true;
 
@@ -1955,7 +1956,7 @@
 
 
 			// Figure out all the desired configurations
-			List<XcodeBuildConfig> BuildConfigs = GetSupportedBuildConfigs(InPlatforms, InConfigurations);
+			List<XcodeBuildConfig> BuildConfigs = GetSupportedBuildConfigs(InPlatforms, InConfigurations, Logger);
 			if (BuildConfigs.Count == 0)
 			{
 				return true;
@@ -2002,7 +2003,7 @@
 			StringBuilder PBXResourcesBuildPhaseSection = new StringBuilder();
 			List<XcodeExtensionInfo> AllExtensions = new List<XcodeExtensionInfo>();
 			GenerateSectionsWithSourceFiles(PBXBuildFileSection, PBXFileReferenceSection, PBXSourcesBuildPhaseSection, TargetAppGuid, TargetName, bIsAppBundle);
-			GenerateSectionsWithExtensions(PBXBuildFileSection, PBXFileReferenceSection, PBXCopyExtensionsBuildPhaseSection, PBXResourcesBuildPhaseSection, AllExtensions, GameProjectPath, BuildConfigs);
+			GenerateSectionsWithExtensions(PBXBuildFileSection, PBXFileReferenceSection, PBXCopyExtensionsBuildPhaseSection, PBXResourcesBuildPhaseSection, AllExtensions, GameProjectPath, BuildConfigs, Logger);
 
 			StringBuilder ProjectFileContent = new StringBuilder();
 
@@ -2035,7 +2036,7 @@
 			AppendIndexTargetSection(ProjectFileContent, IndexTargetName, IndexTargetGuid, IndexTargetConfigListGuid, SourcesBuildPhaseGuid);
 			AppendExtensionTargetSections(ProjectFileContent, AllExtensions);
 			AppendProjectSection(ProjectFileContent, TargetName, TargetGuid, BuildTargetName, BuildTargetGuid, IndexTargetName, IndexTargetGuid, MainGroupGuid, ProductRefGroupGuid, ProjectGuid, ProjectConfigListGuid, GameProjectPath, AllExtensions);
-			AppendXCBuildConfigurationSection(ProjectFileContent, ProjectBuildConfigs, TargetBuildConfigs, BuildTargetBuildConfigs, IndexTargetBuildConfigs, GameProjectPath, AllExtensions);
+			AppendXCBuildConfigurationSection(ProjectFileContent, ProjectBuildConfigs, TargetBuildConfigs, BuildTargetBuildConfigs, IndexTargetBuildConfigs, GameProjectPath, AllExtensions, Logger);
 			AppendXCConfigurationListSection(ProjectFileContent, TargetName, BuildTargetName, IndexTargetName, ProjectConfigListGuid, ProjectBuildConfigs,
 				TargetConfigListGuid, TargetBuildConfigs, BuildTargetConfigListGuid, BuildTargetBuildConfigs, IndexTargetConfigListGuid, IndexTargetBuildConfigs, AllExtensions);
 			AppendShellScriptSection(ProjectFileContent, ShellScriptSectionGuid, GameProjectPath);
@@ -2047,10 +2048,10 @@
 			if (bSuccess)
 			{
 				FileReference PBXProjFilePath = ProjectFilePath + "/project.pbxproj";
-				bSuccess = ProjectFileGenerator.WriteFileIfChanged(PBXProjFilePath.FullName, ProjectFileContent.ToString(), new UTF8Encoding());
+				bSuccess = ProjectFileGenerator.WriteFileIfChanged(PBXProjFilePath.FullName, ProjectFileContent.ToString(), Logger, new UTF8Encoding());
 			}
 
-			bool bNeedScheme = CanBuildProjectLocally();
+			bool bNeedScheme = CanBuildProjectLocally(Logger);
 
 			if (bNeedScheme)
 			{
diff --git a/Engine/Source/Programs/UnrealBuildTool/ProjectFiles/Xcode/XcodeProjectFileGenerator.cs b/Engine/Source/Programs/UnrealBuildTool/ProjectFiles/Xcode/XcodeProjectFileGenerator.cs
--- a/Engine/Source/Programs/UnrealBuildTool/ProjectFiles/Xcode/XcodeProjectFileGenerator.cs
+++ b/Engine/Source/Programs/UnrealBuildTool/ProjectFiles/Xcode/XcodeProjectFileGenerator.cs
@@ -6,6 +6,7 @@
 using System.IO;
 using EpicGames.Core;
 using System.Linq;
+using Microsoft.Extensions.Logging;
 
 namespace UnrealBuildTool
 {
@@ -78,7 +79,7 @@
 
 		/// <summary>
 		/// </summary>
-		public override void CleanProjectFiles(DirectoryReference InPrimaryProjectDirectory, string InPrimaryProjectName, DirectoryReference InIntermediateProjectFilesPath)
+		public override void CleanProjectFiles(DirectoryReference InPrimaryProjectDirectory, string InPrimaryProjectName, DirectoryReference InIntermediateProjectFilesPath, ILogger Logger)
 		{
 			DirectoryReference PrimaryProjDeleteFilename = DirectoryReference.Combine(InPrimaryProjectDirectory, InPrimaryProjectName + ".xcworkspace");
 			if (DirectoryReference.Exists(PrimaryProjDeleteFilename))
@@ -95,8 +96,8 @@
 				}
 				catch (Exception Ex)
 				{
-					Log.TraceInformation("Error while trying to clean project files path {0}. Ignored.", InIntermediateProjectFilesPath);
-					Log.TraceInformation("\t" + Ex.Message);
+					Logger.LogInformation("Error while trying to clean project files path {InIntermediateProjectFilesPath}. Ignored.", InIntermediateProjectFilesPath);
+					Logger.LogInformation("\t{Ex}", Ex.Message);
 				}
 			}
 		}
@@ -112,7 +113,7 @@
 			return new XcodeProjectFile(InitFilePath, BaseDir, bForDistribution, BundleIdentifier, AppName);
 		}
 
-		private bool WriteWorkspaceSettingsFile(string Path)
+		private bool WriteWorkspaceSettingsFile(string Path, ILogger Logger)
 		{
 			StringBuilder WorkspaceSettingsContent = new StringBuilder();
 			WorkspaceSettingsContent.Append("<?xml version=\"1.0\" encoding=\"UTF-8\"?>" + ProjectFileGenerator.NewLine);
@@ -137,10 +138,10 @@
 			WorkspaceSettingsContent.Append("\t<string>Default</string>" + ProjectFileGenerator.NewLine);
 			WorkspaceSettingsContent.Append("</dict>" + ProjectFileGenerator.NewLine);
 			WorkspaceSettingsContent.Append("</plist>" + ProjectFileGenerator.NewLine);
-			return WriteFileIfChanged(Path, WorkspaceSettingsContent.ToString(), new UTF8Encoding());
+			return WriteFileIfChanged(Path, WorkspaceSettingsContent.ToString(), Logger, new UTF8Encoding());
 		}
 
-		private bool WriteWorkspaceSharedSettingsFile(string Path)
+		private bool WriteWorkspaceSharedSettingsFile(string Path, ILogger Logger)
 		{
 			StringBuilder WorkspaceSettingsContent = new StringBuilder();
 			WorkspaceSettingsContent.Append("<?xml version=\"1.0\" encoding=\"UTF-8\"?>" + ProjectFileGenerator.NewLine);
@@ -153,10 +154,10 @@
             WorkspaceSettingsContent.Append("\t<true/>" + ProjectFileGenerator.NewLine);
 			WorkspaceSettingsContent.Append("</dict>" + ProjectFileGenerator.NewLine);
 			WorkspaceSettingsContent.Append("</plist>" + ProjectFileGenerator.NewLine);
-			return WriteFileIfChanged(Path, WorkspaceSettingsContent.ToString(), new UTF8Encoding());
+			return WriteFileIfChanged(Path, WorkspaceSettingsContent.ToString(), Logger, new UTF8Encoding());
 		}
 
-		private bool WriteXcodeWorkspace()
+		private bool WriteXcodeWorkspace(ILogger Logger)
 		{
 			bool bSuccess = true;
 
@@ -181,7 +182,7 @@
 						// Filter out anything that isn't an XC project, and that shouldn't be in the workspace
 						IEnumerable<XcodeProjectFile> SupportedProjects =
 								CurFolder.ChildProjects.OfType<XcodeProjectFile>()
-									.Where(P => P.ShouldIncludeProjectInWorkspace())
+									.Where(P => P.ShouldIncludeProjectInWorkspace(Logger))
 									.OrderBy(P => P.ProjectFilePath.GetFileName());
 
 						foreach (XcodeProjectFile XcodeProject in SupportedProjects)
@@ -237,22 +238,22 @@
 				ProjectName += ProjectFilePlatform == XcodeProjectFilePlatform.Mac ? "_Mac" : (ProjectFilePlatform == XcodeProjectFilePlatform.iOS ? "_IOS" : "_TVOS");
 			}
 			string WorkspaceDataFilePath = PrimaryProjectPath + "/" + ProjectName + ".xcworkspace/contents.xcworkspacedata";
-			bSuccess = WriteFileIfChanged(WorkspaceDataFilePath, WorkspaceDataContent.ToString(), new UTF8Encoding());
+			bSuccess = WriteFileIfChanged(WorkspaceDataFilePath, WorkspaceDataContent.ToString(), Logger, new UTF8Encoding());
 			if (bSuccess)
 			{
 				string WorkspaceSettingsFilePath = PrimaryProjectPath + "/" + ProjectName + ".xcworkspace/xcuserdata/" + Environment.UserName + ".xcuserdatad/WorkspaceSettings.xcsettings";
-				bSuccess = WriteWorkspaceSettingsFile(WorkspaceSettingsFilePath);
+				bSuccess = WriteWorkspaceSettingsFile(WorkspaceSettingsFilePath, Logger);
 				string WorkspaceSharedSettingsFilePath = PrimaryProjectPath + "/" + ProjectName + ".xcworkspace/xcshareddata/WorkspaceSettings.xcsettings";
-				bSuccess = WriteWorkspaceSharedSettingsFile(WorkspaceSharedSettingsFilePath);
+				bSuccess = WriteWorkspaceSharedSettingsFile(WorkspaceSharedSettingsFilePath, Logger);
 			}
 
 
 			return bSuccess;
 		}
 
-		protected override bool WritePrimaryProjectFile(ProjectFile? UBTProject, PlatformProjectGeneratorCollection PlatformProjectGenerators)
+		protected override bool WritePrimaryProjectFile(ProjectFile? UBTProject, PlatformProjectGeneratorCollection PlatformProjectGenerators, ILogger Logger)
 		{
-			return WriteXcodeWorkspace();
+			return WriteXcodeWorkspace(Logger);
 		}
 
 		[Flags]
@@ -273,15 +274,11 @@
 		/// Should we generate a special project to use for tvOS signing instead of a normal one
 		static public bool bGeneratingRunTVOSProject = false;
 
-		/// <summary>
-		/// Configures project generator based on command-line options
-		/// </summary>
-		/// <param name="Arguments">Arguments passed into the program</param>
-		/// <param name="IncludeAllPlatforms">True if all platforms should be included</param>
-		protected override void ConfigureProjectFileGeneration(string[] Arguments, ref bool IncludeAllPlatforms)
+		/// <inheritdoc/>
+		protected override void ConfigureProjectFileGeneration(string[] Arguments, ref bool IncludeAllPlatforms, ILogger Logger)
 		{
 			// Call parent implementation first
-			base.ConfigureProjectFileGeneration(Arguments, ref IncludeAllPlatforms);
+			base.ConfigureProjectFileGeneration(Arguments, ref IncludeAllPlatforms, Logger);
 			ProjectFilePlatform = IncludeAllPlatforms ? XcodeProjectFilePlatform.All : XcodeProjectFilePlatform.Mac;
 
 			foreach (string CurArgument in Arguments)
diff --git a/Engine/Source/Programs/UnrealBuildTool/System/ActionExecutor.cs b/Engine/Source/Programs/UnrealBuildTool/System/ActionExecutor.cs
--- a/Engine/Source/Programs/UnrealBuildTool/System/ActionExecutor.cs
+++ b/Engine/Source/Programs/UnrealBuildTool/System/ActionExecutor.cs
@@ -1,34 +1,35 @@
-// Copyright Epic Games, Inc. All Rights Reserved.
-
-using System;
-using System.Collections.Generic;
-
-namespace UnrealBuildTool
-{
-	abstract class ActionExecutor
-	{
-		public abstract string Name
-		{
-			get;
-		}
-
-		static protected double MemoryPerActionBytesOverride
-		{
-			get;
-			private set;
-		} = 0.0;
-
-		/// <summary>
-		/// Allow targets to override the expected amount of memory required for compiles, used to control the number
-		/// of parallel action processes.
-		/// </summary>
-		/// <param name="MemoryPerActionOverrideGB"></param>
-		public static void SetMemoryPerActionOverride(double MemoryPerActionOverrideGB)
-		{
-			MemoryPerActionBytesOverride = Math.Max(MemoryPerActionBytesOverride, MemoryPerActionOverrideGB * 1024 * 1024 * 1024);
-		}
-
-		public abstract bool ExecuteActions(List<LinkedAction> ActionsToExecute);
-	}
-
-}
+// Copyright Epic Games, Inc. All Rights Reserved.
+
+using Microsoft.Extensions.Logging;
+using System;
+using System.Collections.Generic;
+
+namespace UnrealBuildTool
+{
+	abstract class ActionExecutor
+	{
+		public abstract string Name
+		{
+			get;
+		}
+
+		static protected double MemoryPerActionBytesOverride
+		{
+			get;
+			private set;
+		} = 0.0;
+
+		/// <summary>
+		/// Allow targets to override the expected amount of memory required for compiles, used to control the number
+		/// of parallel action processes.
+		/// </summary>
+		/// <param name="MemoryPerActionOverrideGB"></param>
+		public static void SetMemoryPerActionOverride(double MemoryPerActionOverrideGB)
+		{
+			MemoryPerActionBytesOverride = Math.Max(MemoryPerActionBytesOverride, MemoryPerActionOverrideGB * 1024 * 1024 * 1024);
+		}
+
+		public abstract bool ExecuteActions(List<LinkedAction> ActionsToExecute, ILogger Logger);
+	}
+
+}
diff --git a/Engine/Source/Programs/UnrealBuildTool/System/ActionGraph.cs b/Engine/Source/Programs/UnrealBuildTool/System/ActionGraph.cs
--- a/Engine/Source/Programs/UnrealBuildTool/System/ActionGraph.cs
+++ b/Engine/Source/Programs/UnrealBuildTool/System/ActionGraph.cs
@@ -14,6 +14,7 @@
 using System.Threading.Tasks;
 using OpenTracing.Util;
 using UnrealBuildBase;
+using Microsoft.Extensions.Logging;
 
 namespace UnrealBuildTool
 {
@@ -75,7 +76,8 @@
 		/// Checks a set of actions for conflicts (ie. different actions producing the same output items)
 		/// </summary>
 		/// <param name="Actions">The set of actions to check</param>
-		public static void CheckForConflicts(IEnumerable<IExternalAction> Actions)
+		/// <param name="Logger">Logger for output</param>
+		public static void CheckForConflicts(IEnumerable<IExternalAction> Actions, ILogger Logger)
 		{
 			bool bResult = true;
 
@@ -86,7 +88,7 @@
 				{
 					if (ItemToProducingAction.TryGetValue(ProducedItem, out IExternalAction? ExistingAction))
 					{
-						bResult &= CheckForConflicts(ExistingAction, Action);
+						bResult &= CheckForConflicts(ExistingAction, Action, Logger);
 					}
 					else
 					{
@@ -106,8 +108,9 @@
 		/// </summary>
 		/// <param name="A">The first action</param>
 		/// <param name="B">The second action</param>
+		/// <param name="Logger">Logger for output</param>
 		/// <returns>True if no conflicts were found, false otherwise.</returns>
-		private static bool CheckForConflicts(IExternalAction A, IExternalAction B)
+		private static bool CheckForConflicts(IExternalAction A, IExternalAction B, ILogger Logger)
 		{
 			ActionConflictReasonFlags Reason = ActionConflictReasonFlags.None;
 			if (A.ActionType != B.ActionType)
@@ -147,7 +150,7 @@
 
 			if (Reason != ActionConflictReasonFlags.None)
 			{
-				LogConflict(A, B, Reason);
+				LogConflict(A, B, Reason, Logger);
 				return false;
 			}
 
@@ -212,7 +215,8 @@
 		/// <param name="A">The first action with the conflict</param>
 		/// <param name="B">The second action with the conflict</param>
 		/// <param name="Reason">Enum flags for which properties are in conflict</param>
-		static void LogConflict(IExternalAction A, IExternalAction B, ActionConflictReasonFlags Reason)
+		/// <param name="Logger">Logger for output</param>
+		static void LogConflict(IExternalAction A, IExternalAction B, ActionConflictReasonFlags Reason, ILogger Logger)
 		{
 			// Convert some complex types in IExternalAction to strings when printing json
 			JsonSerializerOptions Options = new JsonSerializerOptions
@@ -237,24 +241,24 @@
 			File.WriteAllText(AJsonPath, AJson);
 			File.WriteAllText(BJsonPath, BJson);
 
-			Log.TraceError($"Unable to merge actions '{A.StatusDescription}' and '{B.StatusDescription}': {Reason} are different");
-			Log.TraceInformation($"  First Action: {AJson}");
-			Log.TraceInformation($"  Second Action: {BJson}");
-			Log.TraceInformation($"  First Action json written to '{AJsonPath}'");
-			Log.TraceInformation($"  Second Action json written to '{BJsonPath}'");
+			Logger.LogError("Unable to merge actions '{StatusA}' and '{StatusB}': {Reason} are different", A.StatusDescription, B.StatusDescription, Reason);
+			Logger.LogInformation("  First Action: {AJson}", AJson);
+			Logger.LogInformation("  Second Action: {BJson}", BJson);
+			Logger.LogInformation("  First Action json written to '{AJsonPath}'", AJsonPath);
+			Logger.LogInformation("  Second Action json written to '{BJsonPath}'", BJsonPath);
 		}
 
 		/// <summary>
 		/// Builds a list of actions that need to be executed to produce the specified output items.
 		/// </summary>
 		public static List<LinkedAction> GetActionsToExecute(List<LinkedAction> Actions,
-			CppDependencyCache CppDependencies, ActionHistory History, bool bIgnoreOutdatedImportLibraries)
+			CppDependencyCache CppDependencies, ActionHistory History, bool bIgnoreOutdatedImportLibraries, ILogger Logger)
 		{
 			using (GlobalTracer.Instance.BuildSpan("ActionGraph.GetActionsToExecute()").StartActive())
 			{
 				// For all targets, build a set of all actions that are outdated.
 				Dictionary<LinkedAction, bool> OutdatedActionDictionary = new Dictionary<LinkedAction, bool>();
-				GatherAllOutdatedActions(Actions, History, OutdatedActionDictionary, CppDependencies, bIgnoreOutdatedImportLibraries);
+				GatherAllOutdatedActions(Actions, History, OutdatedActionDictionary, CppDependencies, bIgnoreOutdatedImportLibraries, Logger);
 
 				// Build a list of actions that are both needed for this target and outdated.
 				return Actions.Where(Action => OutdatedActionDictionary[Action]).ToList();
@@ -266,7 +270,8 @@
 		/// </summary>
 		/// <param name="BuildConfiguration">The build configuration</param>
 		/// <param name="Actions">List of actions in the graph</param>
-		public static void CheckPathLengths(BuildConfiguration BuildConfiguration, IEnumerable<IExternalAction> Actions)
+		/// <param name="Logger"></param>
+		public static void CheckPathLengths(BuildConfiguration BuildConfiguration, IEnumerable<IExternalAction> Actions, ILogger Logger)
 		{
 			if (BuildHostPlatform.Current.Platform == UnrealTargetPlatform.Win64)
 			{
@@ -322,7 +327,7 @@
 					}
 
 					Message.Append($"\n\nConsider setting {nameof(ModuleRules.ShortName)} = ... in module *.Build.cs files to use alternative names for intermediate paths.");
-					Log.TraceWarning(Message.ToString());
+					Logger.LogWarning("{Message}", Message.ToString());
 				}
 			}
 		}
@@ -330,57 +335,57 @@
 		/// <summary>
 		/// Selects an ActionExecutor
 		/// </summary>
-		private static ActionExecutor SelectExecutor(BuildConfiguration BuildConfiguration, int ActionCount, List<TargetDescriptor> TargetDescriptors)
+		private static ActionExecutor SelectExecutor(BuildConfiguration BuildConfiguration, int ActionCount, List<TargetDescriptor> TargetDescriptors, ILogger Logger)
 		{
-			if (ActionCount > ParallelExecutor.GetDefaultNumParallelProcesses())
+			if (ActionCount > ParallelExecutor.GetDefaultNumParallelProcesses(Logger))
 			{
-				if (BuildConfiguration.bAllowHybridExecutor && HybridExecutor.IsAvailable())
+				if (BuildConfiguration.bAllowHybridExecutor && HybridExecutor.IsAvailable(Logger))
 				{
-					return new HybridExecutor(TargetDescriptors, BuildConfiguration.MaxParallelActions, BuildConfiguration.bCompactOutput);
+					return new HybridExecutor(TargetDescriptors, BuildConfiguration.MaxParallelActions, BuildConfiguration.bCompactOutput, Logger);
 				}
-				else if (BuildConfiguration.bAllowXGE && XGE.IsAvailable() && ActionCount >= XGE.MinActions)
+				else if (BuildConfiguration.bAllowXGE && XGE.IsAvailable(Logger) && ActionCount >= XGE.MinActions)
 				{
 					return new XGE();
 				}
-				else if (BuildConfiguration.bAllowFASTBuild && FASTBuild.IsAvailable())
+				else if (BuildConfiguration.bAllowFASTBuild && FASTBuild.IsAvailable(Logger))
 				{
-					return new FASTBuild(BuildConfiguration.MaxParallelActions, BuildConfiguration.bCompactOutput);
+					return new FASTBuild(BuildConfiguration.MaxParallelActions, BuildConfiguration.bCompactOutput, Logger);
 				}
-				else if (BuildConfiguration.bAllowSNDBS && SNDBS.IsAvailable())
+				else if (BuildConfiguration.bAllowSNDBS && SNDBS.IsAvailable(Logger))
 				{
 					return new SNDBS(TargetDescriptors);
 				}
 				else if (BuildConfiguration.bAllowHordeCompute && HordeExecutor.IsAvailable())
 				{
-					return new HordeExecutor(BuildConfiguration.MaxParallelActions, BuildConfiguration.bCompactOutput);
+					return new HordeExecutor(BuildConfiguration.MaxParallelActions, BuildConfiguration.bCompactOutput, Logger);
 				}
 			}
 
-			return new ParallelExecutor(BuildConfiguration.MaxParallelActions, BuildConfiguration.bCompactOutput);
+			return new ParallelExecutor(BuildConfiguration.MaxParallelActions, BuildConfiguration.bCompactOutput, Logger);
 		}
 
 		/// <summary>
 		/// Executes a list of actions.
 		/// </summary>
-		public static void ExecuteActions(BuildConfiguration BuildConfiguration, List<LinkedAction> ActionsToExecute, List<TargetDescriptor> TargetDescriptors)
+		public static void ExecuteActions(BuildConfiguration BuildConfiguration, List<LinkedAction> ActionsToExecute, List<TargetDescriptor> TargetDescriptors, ILogger Logger)
 		{
 			if (ActionsToExecute.Count == 0)
 			{
-				Log.TraceInformation("Target is up to date");
+				Logger.LogInformation("Target is up to date");
 			}
 			else
 			{
 				// Figure out which executor to use
-				ActionExecutor Executor = SelectExecutor(BuildConfiguration, ActionsToExecute.Count, TargetDescriptors);
+				ActionExecutor Executor = SelectExecutor(BuildConfiguration, ActionsToExecute.Count, TargetDescriptors, Logger);
 
 				// Execute the build
 				Stopwatch Timer = Stopwatch.StartNew();
-				if (!Executor.ExecuteActions(ActionsToExecute))
+				if (!Executor.ExecuteActions(ActionsToExecute, Logger))
 				{
 					throw new CompilationResultException(CompilationResult.OtherCompilationError);
 				}
 
-				Log.TraceInformation($"Total time in {Executor.Name} executor: {Timer.Elapsed.TotalSeconds:0.00} seconds");
+				Logger.LogInformation("Total time in {ExecutorName} executor: {TotalSeconds:0.00} seconds", Executor.Name, Timer.Elapsed.TotalSeconds);
 
 				// Reset the file info for all the produced items
 				foreach (LinkedAction BuildAction in ActionsToExecute)
@@ -605,10 +610,11 @@
 		/// <param name="ActionHistory"></param>
 		/// <param name="CppDependencies"></param>
 		/// <param name="bIgnoreOutdatedImportLibraries"></param>
+		/// <param name="Logger"></param>
 		/// <returns>true if outdated</returns>
 		private static void IsIndividualActionOutdated(LinkedAction RootAction,
 			Dictionary<LinkedAction, bool> OutdatedActionDictionary, ReaderWriterLockSlim OutdatedActionLock,
-			ActionHistory ActionHistory, CppDependencyCache CppDependencies, bool bIgnoreOutdatedImportLibraries)
+			ActionHistory ActionHistory, CppDependencyCache CppDependencies, bool bIgnoreOutdatedImportLibraries, ILogger Logger)
 		{
 			// Only compute the outdated-ness for actions that don't aren't cached in the outdated action dictionary.
 			bool bIsOutdated = false;
@@ -629,12 +635,12 @@
 			{
 				// Check if the command-line of the action previously used to produce the item is outdated.
 				string NewProducingAttributes = $"{RootAction.CommandPath.FullName} {RootAction.CommandArguments} (ver {RootAction.CommandVersion})";
-				if (ActionHistory.UpdateProducingAttributes(ProducedItem, NewProducingAttributes) && RootAction.bUseActionHistory)
+				if (ActionHistory.UpdateProducingAttributes(ProducedItem, NewProducingAttributes, Logger) && RootAction.bUseActionHistory)
 				{
 					if (ProducedItem.Exists)
 					{
-						Log.TraceLog($"{RootAction.StatusDescription}: Produced item \"{ProducedItem.AbsolutePath}\" was produced by outdated attributes.");
-						Log.TraceLog($"  New attributes: {NewProducingAttributes}");
+						Logger.LogDebug("{StatusDescription}: Produced item \"{ProducedItem}\" was produced by outdated attributes.", RootAction.StatusDescription, ProducedItem.Location);
+						Logger.LogDebug("  New attributes: {NewProducingAttributes}", NewProducingAttributes);
 					}
 
 					bIsOutdated = true;
@@ -655,7 +661,7 @@
 				else
 				{
 					// If any of the produced items doesn't exist, the action is outdated.
-					Log.TraceLog($"{RootAction.StatusDescription}: Produced item \"{ProducedItem.AbsolutePath}\" doesn't exist.");
+					Logger.LogDebug("{StatusDescription}: Produced item \"{ProducedItem}\" doesn't exist.", RootAction.StatusDescription, ProducedItem.Location);
 					bIsOutdated = true;
 				}
 			}
@@ -675,7 +681,7 @@
 							// Need to check for import libraries here too
 							if (!bIgnoreOutdatedImportLibraries || !IsImportLibraryDependency(RootAction, PrerequisiteItem))
 							{
-								Log.TraceLog($"{RootAction.StatusDescription}: Prerequisite {PrerequisiteItem.AbsolutePath} is newer than the last execution of the action: " +
+								Logger.LogDebug("{StatusDescription}: Prerequisite {PrerequisiteItem} is newer than the last execution of the action: {Message}", RootAction.StatusDescription, PrerequisiteItem.Location, 
 									$"{PrerequisiteItem.LastWriteTimeUtc.ToLocalTime().ToString(CultureInfo.CurrentCulture)} vs {LastExecutionTimeUtc.LocalDateTime.ToString(CultureInfo.CurrentCulture)}");
 								bIsOutdated = true;
 								break;
@@ -688,9 +694,9 @@
 			// Check the dependency list
 			if (!bIsOutdated && RootAction.DependencyListFile != null)
 			{
-				if (!CppDependencies.TryGetDependencies(RootAction.DependencyListFile, out List<FileItem>? DependencyFiles))
+				if (!CppDependencies.TryGetDependencies(RootAction.DependencyListFile, Logger, out List<FileItem>? DependencyFiles))
 				{
-					Log.TraceLog($"{RootAction.StatusDescription}: Missing dependency list file \"{RootAction.DependencyListFile}\"");
+					Logger.LogDebug("{StatusDescription}: Missing dependency list file \"{DependencyListFile}\"", RootAction.StatusDescription, RootAction.DependencyListFile);
 					bIsOutdated = true;
 				}
 				else
@@ -699,7 +705,7 @@
 					{
 						if (!DependencyFile.Exists || DependencyFile.LastWriteTimeUtc > LastExecutionTimeUtc)
 						{
-							Log.TraceLog($"{RootAction.StatusDescription}: Dependency {DependencyFile.AbsolutePath} is newer than the last execution of the action:" +
+							Logger.LogDebug("{RootAction.StatusDescription}: Dependency {DependencyFile} is newer than the last execution of the action: {Message}", RootAction.StatusDescription, DependencyFile.AbsolutePath,
 								$"{DependencyFile.LastWriteTimeUtc.ToLocalTime().ToString(CultureInfo.CurrentCulture)} vs {LastExecutionTimeUtc.LocalDateTime.ToString(CultureInfo.CurrentCulture)}");
 							bIsOutdated = true;
 							break;
@@ -725,8 +731,9 @@
 		/// <param name="RootAction">- The action being considered.</param>
 		/// <param name="OutdatedActionDictionary">-</param>
 		/// <param name="bIgnoreOutdatedImportLibraries"></param>
+		/// <param name="Logger">Logger instance</param>
 		/// <returns>true if outdated</returns>
-		private static bool IsActionOutdatedDueToPrerequisites(LinkedAction RootAction,	Dictionary<LinkedAction, bool> OutdatedActionDictionary, bool bIgnoreOutdatedImportLibraries)
+		private static bool IsActionOutdatedDueToPrerequisites(LinkedAction RootAction,	Dictionary<LinkedAction, bool> OutdatedActionDictionary, bool bIgnoreOutdatedImportLibraries, ILogger Logger)
 		{
 			// Only compute the outdated-ness for actions that aren't already cached in the outdated action dictionary.
 			if (OutdatedActionDictionary.TryGetValue(RootAction, out bool bIsOutdated))
@@ -737,7 +744,7 @@
 			// Check if any of the prerequisite actions are out of date
 			foreach (LinkedAction PrerequisiteAction in RootAction.PrerequisiteActions)
 			{
-				if (IsActionOutdatedDueToPrerequisites(PrerequisiteAction, OutdatedActionDictionary, bIgnoreOutdatedImportLibraries))
+				if (IsActionOutdatedDueToPrerequisites(PrerequisiteAction, OutdatedActionDictionary, bIgnoreOutdatedImportLibraries, Logger))
 				{
 					// Only check for outdated import libraries if we were configured to do so.  Often, a changed import library
 					// won't affect a dependency unless a public header file was also changed, in which case we would be forced
@@ -745,7 +752,7 @@
 					// won't have to wait for dependent targets to be relinked after each change.
 					if (!bIgnoreOutdatedImportLibraries || !IsImportLibraryDependency(RootAction, PrerequisiteAction))
 					{
-						Log.TraceLog($"{RootAction.StatusDescription}: Prerequisite {PrerequisiteAction.StatusDescription} is produced by outdated action.");
+						Logger.LogDebug("{StatusDescription}: Prerequisite {PrereqStatusDescription} is produced by outdated action.", RootAction.StatusDescription, PrerequisiteAction.StatusDescription);
 						bIsOutdated = true;
 						break;
 					}
@@ -804,7 +811,7 @@
 		/// </summary>
 		public static void GatherAllOutdatedActions(IReadOnlyList<LinkedAction> Actions, ActionHistory ActionHistory,
 			Dictionary<LinkedAction, bool> OutdatedActions, CppDependencyCache CppDependencies,
-			bool bIgnoreOutdatedImportLibraries)
+			bool bIgnoreOutdatedImportLibraries, ILogger Logger)
 		{
 			using (GlobalTracer.Instance.BuildSpan("Prefetching include dependencies").StartActive())
 			{
@@ -817,7 +824,7 @@
 					}
 				}
 
-				Parallel.ForEach(Dependencies, File => { CppDependencies.TryGetDependencies(File, out _); });
+				Parallel.ForEach(Dependencies, File => { CppDependencies.TryGetDependencies(File, Logger, out _); });
 			}
 
 			using (GlobalTracer.Instance.BuildSpan("Cache individual outdated actions").StartActive())
@@ -825,14 +832,14 @@
 				ReaderWriterLockSlim OutdatedActionsLock = new ReaderWriterLockSlim();
 				Parallel.ForEach(Actions,
 					Action => IsIndividualActionOutdated(Action, OutdatedActions, OutdatedActionsLock, 
-						ActionHistory, CppDependencies, bIgnoreOutdatedImportLibraries));
+						ActionHistory, CppDependencies, bIgnoreOutdatedImportLibraries, Logger));
 			}
 
 			using (GlobalTracer.Instance.BuildSpan("Cache outdated actions based on recursive prerequisites").StartActive())
 			{
 				foreach (var Action in Actions)
 				{
-					IsActionOutdatedDueToPrerequisites(Action, OutdatedActions, bIgnoreOutdatedImportLibraries);
+					IsActionOutdatedDueToPrerequisites(Action, OutdatedActions, bIgnoreOutdatedImportLibraries, Logger);
 				}
 			}
 		}
@@ -841,7 +848,8 @@
 		/// Deletes all the items produced by actions in the provided outdated action dictionary.
 		/// </summary>
 		/// <param name="OutdatedActions">List of outdated actions</param>
-		public static void DeleteOutdatedProducedItems(List<LinkedAction> OutdatedActions)
+		/// <param name="Logger">Logger for output</param>
+		public static void DeleteOutdatedProducedItems(List<LinkedAction> OutdatedActions, ILogger Logger)
 		{
 			foreach (LinkedAction OutdatedAction in OutdatedActions)
 			{
@@ -849,8 +857,8 @@
 				{
 					if (DeleteItem.Exists)
 					{
-						Log.TraceLog($"Deleting outdated item: {DeleteItem.AbsolutePath}");
-						DeleteItem.Delete();
+						Logger.LogDebug("Deleting outdated item: {DeleteItem}", DeleteItem.Location);
+						DeleteItem.Delete(Logger);
 					}
 				}
 			}
diff --git a/Engine/Source/Programs/UnrealBuildTool/System/ActionGraphBuilder.cs b/Engine/Source/Programs/UnrealBuildTool/System/ActionGraphBuilder.cs
--- a/Engine/Source/Programs/UnrealBuildTool/System/ActionGraphBuilder.cs
+++ b/Engine/Source/Programs/UnrealBuildTool/System/ActionGraphBuilder.cs
@@ -7,6 +7,7 @@ using System.Reflection;
 using System.Text;
 using System.Threading.Tasks;
 using EpicGames.Core;
+using Microsoft.Extensions.Logging;
 using UnrealBuildBase;
 
 namespace UnrealBuildTool
@@ -80,52 +81,63 @@ namespace UnrealBuildTool
 	/// <summary>
 	/// Implementation of IActionGraphBuilder which discards all unnecessary operations
 	/// </summary>
-	class NullActionGraphBuilder : IActionGraphBuilder
+	sealed class NullActionGraphBuilder : IActionGraphBuilder
 	{
+		private readonly ILogger Logger;
+
+		/// <summary>
+		/// Constructor
+		/// </summary>
+		/// <param name="InLogger"></param>
+		public NullActionGraphBuilder(ILogger InLogger)
+		{
+			Logger = InLogger;
+		}
+
 		/// <inheritdoc/>
 		public void AddAction(IExternalAction Action)
 		{
 		}
 
 		/// <inheritdoc/>
-		public virtual void CreateIntermediateTextFile(FileItem FileItem, string Contents)
+		public void CreateIntermediateTextFile(FileItem FileItem, string Contents)
 		{
-			Utils.WriteFileIfChanged(FileItem, Contents);
+			Utils.WriteFileIfChanged(FileItem, Contents, Logger);
 		}
 
 		/// <inheritdoc/>
-		public virtual void CreateIntermediateTextFile(FileItem FileItem, IEnumerable<string> ContentLines)
+		public void CreateIntermediateTextFile(FileItem FileItem, IEnumerable<string> ContentLines)
 		{
-			Utils.WriteFileIfChanged(FileItem, ContentLines);
+			Utils.WriteFileIfChanged(FileItem, ContentLines, Logger);
 		}
 
 		/// <inheritdoc/>
-		public virtual void AddSourceDir(DirectoryItem SourceDir)
+		public void AddSourceDir(DirectoryItem SourceDir)
 		{
 		}
 
 		/// <inheritdoc/>
-		public virtual void AddSourceFiles(DirectoryItem SourceDir, FileItem[] SourceFiles)
+		public void AddSourceFiles(DirectoryItem SourceDir, FileItem[] SourceFiles)
 		{
 		}
 
 		/// <inheritdoc/>
-		public virtual void AddFileToWorkingSet(FileItem File)
+		public void AddFileToWorkingSet(FileItem File)
 		{
 		}
 
 		/// <inheritdoc/>
-		public virtual void AddCandidateForWorkingSet(FileItem File)
+		public void AddCandidateForWorkingSet(FileItem File)
 		{
 		}
 
 		/// <inheritdoc/>
-		public virtual void AddDiagnostic(string Message)
+		public void AddDiagnostic(string Message)
 		{
 		}
 
 		/// <inheritdoc/>
-		public virtual void SetOutputItemsForModule(string ModuleName, FileItem[] OutputItems)
+		public void SetOutputItemsForModule(string ModuleName, FileItem[] OutputItems)
 		{
 		}
 	}
diff --git a/Engine/Source/Programs/UnrealBuildTool/System/ActionHistory.cs b/Engine/Source/Programs/UnrealBuildTool/System/ActionHistory.cs
--- a/Engine/Source/Programs/UnrealBuildTool/System/ActionHistory.cs
+++ b/Engine/Source/Programs/UnrealBuildTool/System/ActionHistory.cs
@@ -11,6 +11,7 @@
 using System.Text;
 using EpicGames.Core;
 using UnrealBuildBase;
+using Microsoft.Extensions.Logging;
 
 namespace UnrealBuildTool
 {
@@ -52,20 +53,21 @@
 		/// Constructor
 		/// </summary>
 		/// <param name="Location">File to store this history in</param>
-		public ActionHistoryLayer(FileReference Location)
+		/// <param name="Logger">Logger for output</param>
+		public ActionHistoryLayer(FileReference Location, ILogger Logger)
 		{
 			this.Location = Location;
 
 			if(FileReference.Exists(Location))
 			{
-				Load();
+				Load(Logger);
 			}
 		}
 
 		/// <summary>
 		/// Attempts to load this action history from disk
 		/// </summary>
-		void Load()
+		void Load(ILogger Logger)
 		{
 			try
 			{
@@ -74,7 +76,7 @@
 					int Version = Reader.ReadInt();
 					if(Version != CurrentVersion)
 					{
-						Log.TraceLog("Unable to read action history from {0}; version {1} vs current {2}", Location, Version, CurrentVersion);
+						Logger.LogDebug("Unable to read action history from {Location}; version {Version} vs current {CurrentVersion}", Location, Version, CurrentVersion);
 						return;
 					}
 
@@ -83,8 +85,8 @@
 			}
 			catch(Exception Ex)
 			{
-				Log.TraceWarning("Unable to read {0}. See log for additional information.", Location);
-				Log.TraceLog("{0}", ExceptionUtils.FormatExceptionDetails(Ex));
+				Logger.LogWarning("Unable to read {Location}. See log for additional information.", Location);
+				Logger.LogDebug("{Ex}", ExceptionUtils.FormatExceptionDetails(Ex));
 			}
 		}
 
@@ -270,8 +272,9 @@
 		/// </summary>
 		/// <param name="File">The file to update</param>
 		/// <param name="Attributes">The new attributes</param>
+		/// <param name="Logger">Logger for output</param>
 		/// <returns>True if the attributes were updated, false otherwise</returns>
-		public bool UpdateProducingAttributes(FileItem File, string Attributes)
+		public bool UpdateProducingAttributes(FileItem File, string Attributes, ILogger Logger)
 		{
 			FileReference LayerLocation = GetLayerLocationForFile(File.Location);
 
@@ -283,7 +286,7 @@
 					Layer = Layers.FirstOrDefault(x => x.Location == LayerLocation);
 					if(Layer == null)
 					{
-						Layer = new ActionHistoryLayer(LayerLocation);
+						Layer = new ActionHistoryLayer(LayerLocation, Logger);
 
 						List<ActionHistoryLayer> NewLayers = new List<ActionHistoryLayer>(Layers);
 						NewLayers.Add(Layer);
@@ -426,18 +429,19 @@
 		/// </summary>
 		/// <param name="File">The output file to look for</param>
 		/// <param name="Attributes">Receives the Attributes used to produce this file</param>
+		/// <param name="Logger">Logger for output</param>
 		/// <returns>True if the output item exists</returns>
-		public bool UpdateProducingAttributes(FileItem File, string Attributes)
+		public bool UpdateProducingAttributes(FileItem File, string Attributes, ILogger Logger)
 		{
 			foreach (ActionHistoryPartition Partition in Partitions)
 			{
 				if (File.Location.IsUnderDirectory(Partition.BaseDir))
 				{
-					return Partition.UpdateProducingAttributes(File, Attributes);
+					return Partition.UpdateProducingAttributes(File, Attributes, Logger);
 				}
 			}
 
-			Log.TraceWarning("File {0} is not under any action history root directory", File.Location);
+			Logger.LogWarning("File {FileLocation} is not under any action history root directory", File.Location);
 			return false;
 		}
 
diff --git a/Engine/Source/Programs/UnrealBuildTool/System/BuildVersion.cs b/Engine/Source/Programs/UnrealBuildTool/System/BuildVersion.cs
--- a/Engine/Source/Programs/UnrealBuildTool/System/BuildVersion.cs
+++ b/Engine/Source/Programs/UnrealBuildTool/System/BuildVersion.cs
@@ -8,6 +8,7 @@
 using System.Text;
 using System.Threading.Tasks;
 using EpicGames.Core;
+using Microsoft.Extensions.Logging;
 using UnrealBuildBase;
 
 namespace UnrealBuildTool
@@ -182,11 +183,12 @@
 		/// Exports this object as Json
 		/// </summary>
 		/// <param name="FileName">The filename to write to</param>
-		public void WriteIfModified(FileReference FileName)
+		/// <param name="Logger">Logger for output</param>
+		public void WriteIfModified(FileReference FileName, ILogger Logger)
 		{
 			using StringWriter Writer = new StringWriter();
 			Write(Writer);
-			Utils.WriteFileIfChanged(FileName, Writer.ToString());
+			Utils.WriteFileIfChanged(FileName, Writer.ToString(), Logger);
 		}
 
 		/// <summary>
diff --git a/Engine/Source/Programs/UnrealBuildTool/System/CommandLine.cs b/Engine/Source/Programs/UnrealBuildTool/System/CommandLine.cs
--- a/Engine/Source/Programs/UnrealBuildTool/System/CommandLine.cs
+++ b/Engine/Source/Programs/UnrealBuildTool/System/CommandLine.cs
@@ -9,6 +9,7 @@
 using System.Threading.Tasks;
 using EpicGames.Core;
 using UnrealBuildBase;
+using Microsoft.Extensions.Logging;
 
 namespace UnrealBuildTool
 {
@@ -56,9 +57,17 @@
 		/// </summary>
 		/// <param name="Arguments">List of arguments. Parsed arguments will be removed from this list when the function returns.</param>
 		/// <param name="TargetObject">Object to receive the parsed arguments. Fields in this object should be marked up with CommandLineArgumentAttribute's to indicate how they should be parsed.</param>
-		public static void ParseArguments(IEnumerable<string> Arguments, object TargetObject)
+		public static void ParseArguments(IEnumerable<string> Arguments, object TargetObject) => ParseArguments(Arguments, TargetObject, Log.Logger);
+
+		/// <summary>
+		/// Parse the given list of arguments and apply them to the given object
+		/// </summary>
+		/// <param name="Arguments">List of arguments. Parsed arguments will be removed from this list when the function returns.</param>
+		/// <param name="TargetObject">Object to receive the parsed arguments. Fields in this object should be marked up with CommandLineArgumentAttribute's to indicate how they should be parsed.</param>
+		/// <param name="Logger">Logger for output</param>
+		public static void ParseArguments(IEnumerable<string> Arguments, object TargetObject, ILogger Logger)
 		{
-			ParseAndRemoveArguments(Arguments.ToList(), TargetObject);
+			ParseAndRemoveArguments(Arguments.ToList(), TargetObject, Logger);
 		}
 
 		private static Type NonNullableType(Type type)
@@ -76,7 +85,8 @@
 		/// </summary>
 		/// <param name="Arguments">List of arguments. Parsed arguments will be removed from this list when the function returns.</param>
 		/// <param name="TargetObject">Object to receive the parsed arguments. Fields in this object should be marked up with CommandLineArgumentAttribute's to indicate how they should be parsed.</param>
-		public static void ParseAndRemoveArguments(List<string> Arguments, object TargetObject)
+		/// <param name="Logger">Logger for output</param>
+		public static void ParseAndRemoveArguments(List<string> Arguments, object TargetObject, ILogger Logger)
 		{
 			// Build a mapping from name to field and attribute for this object
 			Dictionary<string, Parameter> PrefixToParameter = new Dictionary<string, Parameter>(StringComparer.InvariantCultureIgnoreCase);
@@ -133,24 +143,24 @@
 						{
 							if(EqualsIdx != -1)
 							{
-								Log.WriteLine(LogEventType.Warning, "Cannot specify a value for {0}", Parameter.Prefix);
+								Logger.LogWarning("Cannot specify a value for {ParameterPrefix}", Parameter.Prefix);
 							}
 							else
 							{
-								AssignValue(Parameter, Parameter.Attribute.Value, TargetObject, AssignedFieldToParameter);
+								AssignValue(Parameter, Parameter.Attribute.Value, TargetObject, AssignedFieldToParameter, Logger);
 							}
 						}
 						else if(EqualsIdx != -1)
 						{
-							AssignValue(Parameter, Argument.Substring(EqualsIdx + 1), TargetObject, AssignedFieldToParameter);
+							AssignValue(Parameter, Argument.Substring(EqualsIdx + 1), TargetObject, AssignedFieldToParameter, Logger);
 						}
 						else if(NonNullableType(Parameter.FieldInfo.FieldType) == typeof(bool))
 						{
-							AssignValue(Parameter, "true", TargetObject, AssignedFieldToParameter);
+							AssignValue(Parameter, "true", TargetObject, AssignedFieldToParameter, Logger);
 						}
 						else
 						{
-							Log.WriteLine(LogEventType.Warning, "Missing value for {0}", Parameter.Prefix);
+							Logger.LogWarning("Missing value for {ParameterPrefix}", Parameter.Prefix);
 						}
 
 						// Remove the argument from the list
@@ -186,17 +196,18 @@
 		/// Checks that the list of arguments is empty. If not, throws an exception listing them.
 		/// </summary>
 		/// <param name="RemainingArguments">List of remaining arguments</param>
-		public static void CheckNoRemainingArguments(List<string> RemainingArguments)
+		/// <param name="Logger">Logger for output</param>
+		public static void CheckNoRemainingArguments(List<string> RemainingArguments, ILogger Logger)
 		{
 			if(RemainingArguments.Count > 0)
 			{
 				if(RemainingArguments.Count == 1)
 				{
-					Log.TraceWarning("Invalid argument: {0}", RemainingArguments[0]);
+					Logger.LogWarning("Invalid argument: {Arg}", RemainingArguments[0]);
 				}
 				else
 				{
-					Log.TraceWarning("Invalid arguments:\n{0}", String.Join("\n", RemainingArguments));
+					Logger.LogWarning("Invalid arguments:\n{Args}", String.Join("\n", RemainingArguments));
 				}
 			}
 		}
@@ -208,7 +219,8 @@
 		/// <param name="Text">Argument text</param>
 		/// <param name="TargetObject">The target object to assign values to</param>
 		/// <param name="AssignedFieldToParameter">Maps assigned fields to the parameter that wrote to it. Used to detect duplicate and conflicting arguments.</param>
-		static void AssignValue(Parameter Parameter, string Text, object TargetObject, Dictionary<FieldInfo, Parameter> AssignedFieldToParameter)
+		/// <param name="Logger">Logger for output</param>
+		static void AssignValue(Parameter Parameter, string Text, object TargetObject, Dictionary<FieldInfo, Parameter> AssignedFieldToParameter, ILogger Logger)
 		{
 			// Check if the field type implements ICollection<>. If so, we can take multiple values.
 			Type? CollectionType = null;
@@ -228,7 +240,7 @@
 				object? Value;
 				if(!TryParseValue(Parameter.FieldInfo.FieldType, Text, out Value))
 				{
-					Log.WriteLine(LogEventType.Warning, "Invalid value for {0}... - ignoring {1}", Parameter.Prefix, Text);
+					Logger.LogWarning("Invalid value for {ParameterPrefix}... - ignoring {Text}", Parameter.Prefix, Text);
 					return;
 				}
 
@@ -241,11 +253,11 @@
 					{
 						if(PreviousParameter.Prefix == Parameter.Prefix)
 						{
-							Log.WriteLine(LogEventType.Warning, "Conflicting {0} arguments - ignoring", Parameter.Prefix);
+							Logger.LogWarning("Conflicting {ParameterPrefix} arguments - ignoring", Parameter.Prefix);
 						}
 						else
 						{
-							Log.WriteLine(LogEventType.Warning, "{0} conflicts with {1} - ignoring", Parameter.Prefix, PreviousParameter.Prefix);
+							Logger.LogWarning("{ParameterPrefix} conflicts with {PreviousParameterPrefix} - ignoring", Parameter.Prefix, PreviousParameter.Prefix);
 						}
 					}
 					return;
@@ -278,7 +290,7 @@
 					}
 					else
 					{
-						Log.WriteLine(LogEventType.Warning, "'{0}' is not a valid value for -{1}=... - ignoring", Item, Parameter.Prefix);
+						Logger.LogWarning("'{Item}' is not a valid value for -{ParameterPrefix}=... - ignoring", Item, Parameter.Prefix);
 					}
 				}
 			}
diff --git a/Engine/Source/Programs/UnrealBuildTool/System/ConfigCache.cs b/Engine/Source/Programs/UnrealBuildTool/System/ConfigCache.cs
--- a/Engine/Source/Programs/UnrealBuildTool/System/ConfigCache.cs
+++ b/Engine/Source/Programs/UnrealBuildTool/System/ConfigCache.cs
@@ -10,6 +10,7 @@
 using System.Threading.Tasks;
 using EpicGames.Core;
 using UnrealBuildTool;
+using Microsoft.Extensions.Logging;
 
 namespace UnrealBuildTool
 {
@@ -573,20 +574,21 @@
 		/// <param name="Section">Name of the section with the Key in it</param>
 		/// <param name="Key">Key to update</param>
 		/// <param name="Value">Value to write for te Key</param>
+		/// <param name="Logger">Logger for output</param>
 		/// <returns></returns>		
-		public static bool WriteSettingToDefaultConfig(ConfigHierarchyType ConfigType, DirectoryReference ProjectDir, ConfigDefaultUpdateType UpdateType, string Section, string Key, string Value)
+		public static bool WriteSettingToDefaultConfig(ConfigHierarchyType ConfigType, DirectoryReference ProjectDir, ConfigDefaultUpdateType UpdateType, string Section, string Key, string Value, ILogger Logger)
 		{
 			FileReference DefaultConfigFile = GetDefaultConfigFileReference(ConfigType, ProjectDir);
 
 			if (!FileReference.Exists(DefaultConfigFile))
 			{
-				Log.TraceWarning($"Failed to find config file '{DefaultConfigFile.FullName}' to update");
+				Logger.LogWarning("Failed to find config file '{DefaultConfigFile}' to update", DefaultConfigFile);
 				return false;
 			}
 
 			if (File.GetAttributes(DefaultConfigFile.FullName).HasFlag(FileAttributes.ReadOnly))
 			{
-				Log.TraceWarning($"Config file '{DefaultConfigFile.FullName}' is read-only, unable to write setting {Key}");
+				Logger.LogWarning("Config file '{ConfigFile}' is read-only, unable to write setting {Key}", DefaultConfigFile.FullName, Key);
 				return false;
 			}
 
diff --git a/Engine/Source/Programs/UnrealBuildTool/System/ConfigFile.cs b/Engine/Source/Programs/UnrealBuildTool/System/ConfigFile.cs
--- a/Engine/Source/Programs/UnrealBuildTool/System/ConfigFile.cs
+++ b/Engine/Source/Programs/UnrealBuildTool/System/ConfigFile.cs
@@ -9,6 +9,7 @@
 using System.Threading.Tasks;
 using EpicGames.Core;
 using UnrealBuildTool;
+using Microsoft.Extensions.Logging;
 
 namespace UnrealBuildTool
 {
@@ -221,10 +222,7 @@
 							// Otherwise add it to the current section or add a new one
 							if(CurrentSection != null)
 							{
-								if(!TryAddConfigLine(CurrentSection, Line, StartIdx, EndIdx, DefaultAction))
-								{
-									Log.TraceWarning("Couldn't parse '{0}' in {1} of {2}", Line, CurrentSection, Location.FullName);
-								}
+								TryAddConfigLine(CurrentSection, Line, StartIdx, EndIdx, DefaultAction);
 								break;
 							}
 
@@ -262,10 +260,7 @@
 					if (CurrentSection != null)
 					{
 						string IniKeyValue = Setting.Substring(Setting.IndexOf(':') + 1);
-						if (!TryAddConfigLine(CurrentSection, IniKeyValue, 0, IniKeyValue.Length, DefaultAction))
-						{
-							Log.TraceWarning("Couldn't parse '{0}'", IniKeyValue);
-						}
+						TryAddConfigLine(CurrentSection, IniKeyValue, 0, IniKeyValue.Length, DefaultAction);
 					}
 				}
 			}
diff --git a/Engine/Source/Programs/UnrealBuildTool/System/ConfigValueTracker.cs b/Engine/Source/Programs/UnrealBuildTool/System/ConfigValueTracker.cs
--- a/Engine/Source/Programs/UnrealBuildTool/System/ConfigValueTracker.cs
+++ b/Engine/Source/Programs/UnrealBuildTool/System/ConfigValueTracker.cs
@@ -7,6 +7,7 @@
 using System.Text;
 using System.Threading.Tasks;
 using EpicGames.Core;
+using Microsoft.Extensions.Logging;
 
 namespace UnrealBuildTool
 {
