UBT: Use ILogger throughout UBT, rather than legacy EpicGames.Core.Log methods. All output going forwards should use structured logging rather than String.Format style logging.

diff --git a/Engine/Source/Programs/UnrealBuildTool/System/CppDependencyCache.cs b/Engine/Source/Programs/UnrealBuildTool/System/CppDependencyCache.cs
--- a/Engine/Source/Programs/UnrealBuildTool/System/CppDependencyCache.cs
+++ b/Engine/Source/Programs/UnrealBuildTool/System/CppDependencyCache.cs
@@ -11,6 +11,7 @@
 using System.Threading.Tasks;
 using EpicGames.Core;
 using UnrealBuildBase;
+using Microsoft.Extensions.Logging;
 
 namespace UnrealBuildTool
 {
@@ -94,21 +95,22 @@
 			/// </summary>
 			/// <param name="Location">File to store the cache</param>
 			/// <param name="BaseDir">Base directory for files that this cache should store data for</param>
-			public CachePartition(FileReference Location, DirectoryReference BaseDir)
+			/// <param name="Logger">Logger for output</param>
+			public CachePartition(FileReference Location, DirectoryReference BaseDir, ILogger Logger)
 			{
 				this.Location = Location;
 				this.BaseDir = BaseDir;
 
 				if (FileReference.Exists(Location))
 				{
-					Read();
+					Read(Logger);
 				}
 			}
 
 			/// <summary>
 			/// Reads data for this dependency cache from disk
 			/// </summary>
-			public void Read()
+			public void Read(ILogger Logger)
 			{
 				try
 				{
@@ -117,7 +119,7 @@
 						int Version = Reader.ReadInt();
 						if (Version != CurrentVersion)
 						{
-							Log.TraceLog("Unable to read dependency cache from {0}; version {1} vs current {2}", Location, Version, CurrentVersion);
+							Logger.LogDebug("Unable to read dependency cache from {File}; version {Version} vs current {CurrentVersion}", Location, Version, CurrentVersion);
 							return;
 						}
 
@@ -131,8 +133,8 @@
 				}
 				catch (Exception Ex)
 				{
-					Log.TraceWarning("Unable to read {0}. See log for additional information.", Location);
-					Log.TraceLog("{0}", ExceptionUtils.FormatExceptionDetails(Ex));
+					Logger.LogWarning("Unable to read {Location}. See log for additional information.", Location);
+					Logger.LogDebug("{Ex}", ExceptionUtils.FormatExceptionDetails(Ex));
 				}
 			}
 
@@ -191,12 +193,13 @@
 		/// Gets the produced module from a dependencies file
 		/// </summary>
 		/// <param name="InputFile">The dependencies file</param>
+		/// <param name="Logger">Logger for output</param>
 		/// <param name="OutModule">The produced module name</param>
 		/// <returns>True if a produced module was found</returns>
-		public bool TryGetProducedModule(FileItem InputFile, [NotNullWhen(true)] out string? OutModule)
+		public bool TryGetProducedModule(FileItem InputFile, ILogger Logger, [NotNullWhen(true)] out string? OutModule)
 		{
 			DependencyInfo? Info;
-			if (TryGetDependencyInfo(InputFile, out Info) && Info.ProducedModule != null)
+			if (TryGetDependencyInfo(InputFile, Logger, out Info) && Info.ProducedModule != null)
 			{
 				OutModule = Info.ProducedModule;
 				return true;
@@ -212,12 +215,13 @@
 		/// Attempts to get a list of imported modules for the given file
 		/// </summary>
 		/// <param name="InputFile">The dependency file to query</param>
+		/// <param name="Logger">Logger for output</param>
 		/// <param name="OutImportedModules">List of imported modules</param>
 		/// <returns>True if a list of imported modules was obtained</returns>
-		public bool TryGetImportedModules(FileItem InputFile, [NotNullWhen(true)] out List<(string Name, string BMI)>? OutImportedModules)
+		public bool TryGetImportedModules(FileItem InputFile, ILogger Logger, [NotNullWhen(true)] out List<(string Name, string BMI)>? OutImportedModules)
 		{
 			DependencyInfo? Info;
-			if (TryGetDependencyInfo(InputFile, out Info))
+			if (TryGetDependencyInfo(InputFile, Logger, out Info))
 			{
 				OutImportedModules = Info.ImportedModules;
 				return OutImportedModules != null;
@@ -233,12 +237,13 @@
 		/// Attempts to read the dependencies from the given input file
 		/// </summary>
 		/// <param name="InputFile">File to be read</param>
+		/// <param name="Logger">Logger for output</param>
 		/// <param name="OutDependencyItems">Receives a list of output items</param>
 		/// <returns>True if the input file exists and the dependencies were read</returns>
-		public bool TryGetDependencies(FileItem InputFile, [NotNullWhen(true)] out List<FileItem>? OutDependencyItems)
+		public bool TryGetDependencies(FileItem InputFile, ILogger Logger, [NotNullWhen(true)] out List<FileItem>? OutDependencyItems)
 		{
 			DependencyInfo? Info;
-			if (TryGetDependencyInfo(InputFile, out Info))
+			if (TryGetDependencyInfo(InputFile, Logger, out Info))
 			{
 				OutDependencyItems = Info.Files;
 				return true;
@@ -254,9 +259,10 @@
 		/// Attempts to read the dependencies from the given input file
 		/// </summary>
 		/// <param name="InputFile">File to be read</param>
+		/// <param name="Logger">Logger for output</param>
 		/// <param name="OutInfo">The dependency info</param>
 		/// <returns>True if the input file exists and the dependencies were read</returns>
-		private bool TryGetDependencyInfo(FileItem InputFile, [NotNullWhen(true)] out DependencyInfo? OutInfo)
+		private bool TryGetDependencyInfo(FileItem InputFile, ILogger Logger, [NotNullWhen(true)] out DependencyInfo? OutInfo)
 		{
 			if (!InputFile.Exists)
 			{
@@ -274,7 +280,7 @@
 			}
 			catch (Exception Ex)
 			{
-				Log.TraceLog("Unable to read {0}:\n{1}", InputFile, ExceptionUtils.FormatExceptionDetails(Ex));
+				Logger.LogDebug("Unable to read {File}:\n{Ex}", InputFile, ExceptionUtils.FormatExceptionDetails(Ex));
 				OutInfo = null;
 				return false;
 			}
@@ -336,8 +342,9 @@
 		/// <param name="Configuration">Configuration being built</param>
 		/// <param name="TargetType">The target type</param>
 		/// <param name="Architecture">The target architecture</param>
+		/// <param name="Logger">Logger for output</param>
 		/// <returns>Dependency cache hierarchy for the given project</returns>
-		public void Mount(FileReference? ProjectFile, string TargetName, UnrealTargetPlatform Platform, UnrealTargetConfiguration Configuration, TargetType TargetType, string Architecture)
+		public void Mount(FileReference? ProjectFile, string TargetName, UnrealTargetPlatform Platform, UnrealTargetConfiguration Configuration, TargetType TargetType, string Architecture, ILogger Logger)
 		{
 			if (ProjectFile == null || !Unreal.IsEngineInstalled())
 			{
@@ -352,13 +359,13 @@
 				}
 
 				FileReference EngineCacheLocation = FileReference.Combine(Unreal.EngineDirectory, UEBuildTarget.GetPlatformIntermediateFolder(Platform, Architecture, false), AppName, Configuration.ToString(), "DependencyCache.bin");
-				FindOrAddPartition(EngineCacheLocation, Unreal.EngineDirectory);
+				FindOrAddPartition(EngineCacheLocation, Unreal.EngineDirectory, Logger);
 			}
 
 			if (ProjectFile != null)
 			{
 				FileReference ProjectCacheLocation = FileReference.Combine(ProjectFile.Directory, UEBuildTarget.GetPlatformIntermediateFolder(Platform, Architecture, false), TargetName, Configuration.ToString(), "DependencyCache.bin");
-				FindOrAddPartition(ProjectCacheLocation, ProjectFile.Directory);
+				FindOrAddPartition(ProjectCacheLocation, ProjectFile.Directory, Logger);
 			}
 		}
 
@@ -367,8 +374,9 @@
 		/// </summary>
 		/// <param name="Location">File to store the cache</param>
 		/// <param name="BaseDir">Base directory for files that this cache should store data for</param>
+		/// <param name="Logger">Logger for output</param>
 		/// <returns>Reference to a dependency cache with the given settings</returns>
-		void FindOrAddPartition(FileReference Location, DirectoryReference BaseDir)
+		void FindOrAddPartition(FileReference Location, DirectoryReference BaseDir, ILogger Logger)
 		{
 			lock (GlobalPartitions)
 			{
@@ -377,7 +385,7 @@
 					CachePartition? Partition;
 					if (!GlobalPartitions.TryGetValue(Location, out Partition))
 					{
-						Partition = new CachePartition(Location, BaseDir);
+						Partition = new CachePartition(Location, BaseDir, Logger);
 						GlobalPartitions.Add(Location, Partition);
 					}
 					Partitions.Add(Partition);
diff --git a/Engine/Source/Programs/UnrealBuildTool/System/DynamicCompilation.cs b/Engine/Source/Programs/UnrealBuildTool/System/DynamicCompilation.cs
--- a/Engine/Source/Programs/UnrealBuildTool/System/DynamicCompilation.cs
+++ b/Engine/Source/Programs/UnrealBuildTool/System/DynamicCompilation.cs
@@ -14,6 +14,7 @@
 using Microsoft.CodeAnalysis.CSharp;
 using Microsoft.CodeAnalysis.CSharp.Syntax;
 using Microsoft.CodeAnalysis.Emit;
+using Microsoft.Extensions.Logging;
 using System.Reflection.Metadata;
 using Microsoft.CodeAnalysis.Text;
 using System.Runtime.InteropServices;
@@ -33,13 +34,14 @@
 		/// <param name="SourceFiles">Set of source files</param>
 		/// <param name="AssemblyManifestFilePath">File containing information about this assembly, like which source files it was built with and engine version</param>
 		/// <param name="OutputAssemblyPath">Output path for the assembly</param>
+		/// <param name="Logger">Logger for output</param>
 		/// <returns>True if the assembly needs to be built</returns>
-		private static bool RequiresCompilation(HashSet<FileReference> SourceFiles, FileReference AssemblyManifestFilePath, FileReference OutputAssemblyPath)
+		private static bool RequiresCompilation(HashSet<FileReference> SourceFiles, FileReference AssemblyManifestFilePath, FileReference OutputAssemblyPath, ILogger Logger)
 		{
 			// Do not compile the file if it's installed
 			if (UnrealBuildTool.IsFileInstalled(OutputAssemblyPath))
 			{
-				Log.TraceLog("Skipping {0}: File is installed", OutputAssemblyPath);
+				Logger.LogDebug("Skipping {OutputAssemblyPath}: File is installed", OutputAssemblyPath);
 				return false;
 			}
 
@@ -47,7 +49,7 @@
 			FileItem OutputAssemblyInfo = FileItem.GetItemByFileReference(OutputAssemblyPath);
 			if (!OutputAssemblyInfo.Exists)
 			{
-				Log.TraceLog("Compiling {0}: Assembly does not exist", OutputAssemblyPath);
+				Logger.LogDebug("Compiling {OutputAssemblyPath}: Assembly does not exist", OutputAssemblyPath);
 				return true;
 			}
 
@@ -58,7 +60,7 @@
 			FileItem UnrealBuildToolDllItem = FileItem.GetItemByFileReference(Unreal.UnrealBuildToolDllPath);
 			if (UnrealBuildToolDllItem.LastWriteTimeUtc > OutputAssemblyInfo.LastWriteTimeUtc)
 			{
-				Log.TraceLog("Compiling {0}: {1} is newer", OutputAssemblyPath, UnrealBuildToolDllItem.Name);
+				Logger.LogDebug("Compiling {OutputAssemblyPath}: {UnrealBuildToolDllItemName} is newer", OutputAssemblyPath, UnrealBuildToolDllItem.Name);
 				return true;
 			}
 
@@ -68,7 +70,7 @@
 			FileItem AssemblySourceListFile = FileItem.GetItemByFileReference(AssemblyManifestFilePath);
 			if (!AssemblySourceListFile.Exists)
 			{
-				Log.TraceLog("Compiling {0}: Missing source file list ({1})", OutputAssemblyPath, AssemblyManifestFilePath);
+				Logger.LogDebug("Compiling {OutputAssemblyPath}: Missing source file list ({AssemblyManifestFilePath})", OutputAssemblyPath, AssemblyManifestFilePath);
 				return true;
 			}
 
@@ -79,7 +81,7 @@
 			string EngineVersionCurrent = FormatVersionNumber(ReadOnlyBuildVersion.Current);
 			if (EngineVersionManifest != EngineVersionCurrent)
 			{
-				Log.TraceLog("Compiling {0}: Engine Version changed from {1} to {2}", OutputAssemblyPath, EngineVersionManifest, EngineVersionCurrent);
+				Logger.LogDebug("Compiling {OutputAssemblyPath}: Engine Version changed from {EngineVersionManifest} to {EngineVersionCurrent}", OutputAssemblyPath, EngineVersionManifest, EngineVersionCurrent);
 				return true;
 			}
 
@@ -104,7 +106,7 @@
 			{
 				if(!SourceFileItems.Contains(CurrentSourceFileItem))
 				{
-					Log.TraceLog("Compiling {0}: Removed source file ({1})", OutputAssemblyPath, CurrentSourceFileItem);
+					Logger.LogDebug("Compiling {OutputAssemblyPath}: Removed source file ({CurrentSourceFileItem})", OutputAssemblyPath, CurrentSourceFileItem);
 					return true;
 				}
 			}
@@ -112,7 +114,7 @@
 			{
 				if(!CurrentSourceFileItems.Contains(SourceFileItem))
 				{
-					Log.TraceLog("Compiling {0}: Added source file ({1})", OutputAssemblyPath, SourceFileItem);
+					Logger.LogDebug("Compiling {OutputAssemblyPath}: Added source file ({SourceFileItem})", OutputAssemblyPath, SourceFileItem);
 					return true;
 				}
 			}
@@ -122,7 +124,7 @@
 			{
 				if(SourceFileItem.LastWriteTimeUtc > OutputAssemblyInfo.LastWriteTimeUtc)
 				{
-					Log.TraceLog("Compiling {0}: {1} is newer", OutputAssemblyPath, SourceFileItem);
+					Logger.LogDebug("Compiling {OutputAssemblyPath}: {SourceFileItem} is newer", OutputAssemblyPath, SourceFileItem);
 					return true;
 				}
 			}
@@ -160,7 +162,7 @@
 			}
 		}
 
-		private static Assembly? CompileAssembly(FileReference OutputAssemblyPath, HashSet<FileReference> SourceFileNames, List<string>? ReferencedAssembies, List<string>? PreprocessorDefines = null, bool TreatWarningsAsErrors = false)
+		private static Assembly? CompileAssembly(FileReference OutputAssemblyPath, HashSet<FileReference> SourceFileNames, ILogger Logger, List<string>? ReferencedAssembies, List<string>? PreprocessorDefines = null, bool TreatWarningsAsErrors = false)
 		{
 			CSharpParseOptions ParseOptions = new CSharpParseOptions(
 				languageVersion:LanguageVersion.Latest, 
@@ -178,7 +180,7 @@
 				IEnumerable<Diagnostic> Diagnostics = Tree.GetDiagnostics();
 				if (Diagnostics.Any())
 				{
-					Log.TraceWarning($"Errors generated while parsing '{SourceFileName.FullName}'");
+					Logger.LogWarning("Errors generated while parsing '{SourceFileName}'", SourceFileName);
 					LogDiagnostics(Tree.GetDiagnostics());
 					return null;
 				}
@@ -220,6 +222,7 @@
 			MetadataReferences.Add(MetadataReference.CreateFromFile(Assembly.Load("System.Text.RegularExpressions").Location));
 			MetadataReferences.Add(MetadataReference.CreateFromFile(Assembly.Load("System.Console").Location));
 			MetadataReferences.Add(MetadataReference.CreateFromFile(Assembly.Load("System.Runtime.Extensions").Location));
+			MetadataReferences.Add(MetadataReference.CreateFromFile(Assembly.Load("Microsoft.Extensions.Logging.Abstractions").Location));
 			MetadataReferences.Add(MetadataReference.CreateFromFile(Assembly.Load("netstandard").Location));
 			
 			// process start dependencies
@@ -288,13 +291,14 @@
 		/// </summary>
 		/// <param name="OutputAssemblyPath">Full path to the assembly to be created</param>
 		/// <param name="SourceFileNames">List of source file name</param>
+		/// <param name="Logger">Logger for output</param>
 		/// <param name="ReferencedAssembies"></param>
 		/// <param name="PreprocessorDefines"></param>
 		/// <param name="DoNotCompile"></param>
 		/// <param name="ForceCompile"></param>
 		/// <param name="TreatWarningsAsErrors"></param>
 		/// <returns>The assembly that was loaded</returns>
-		public static Assembly? CompileAndLoadAssembly(FileReference OutputAssemblyPath, HashSet<FileReference> SourceFileNames, List<string>? ReferencedAssembies = null, List<string>? PreprocessorDefines = null, bool DoNotCompile = false, bool ForceCompile = false, bool TreatWarningsAsErrors = false)
+		public static Assembly? CompileAndLoadAssembly(FileReference OutputAssemblyPath, HashSet<FileReference> SourceFileNames, ILogger Logger, List<string>? ReferencedAssembies = null, List<string>? PreprocessorDefines = null, bool DoNotCompile = false, bool ForceCompile = false, bool TreatWarningsAsErrors = false)
 		{
 			// Check to see if the resulting assembly is compiled and up to date
 			FileReference AssemblyManifestFilePath = FileReference.Combine(OutputAssemblyPath.Directory, Path.GetFileNameWithoutExtension(OutputAssemblyPath.FullName) + "Manifest.json");
@@ -302,7 +306,7 @@
 			bool bNeedsCompilation = ForceCompile;
 			if (!DoNotCompile)
 			{
-				bNeedsCompilation = RequiresCompilation(SourceFileNames, AssemblyManifestFilePath, OutputAssemblyPath);
+				bNeedsCompilation = RequiresCompilation(SourceFileNames, AssemblyManifestFilePath, OutputAssemblyPath, Logger);
 			}
 
 			// Load the assembly to ensure it is correct
@@ -316,12 +320,12 @@
 				}
 				catch (FileLoadException Ex)
 				{
-					Log.TraceInformation(String.Format("Unable to load the previously-compiled assembly file '{0}'.  Unreal Build Tool will try to recompile this assembly now.  (Exception: {1})", OutputAssemblyPath, Ex.Message));
+					Logger.LogInformation("Unable to load the previously-compiled assembly file '{File}'.  Unreal Build Tool will try to recompile this assembly now.  (Exception: {Ex})", OutputAssemblyPath, Ex.Message);
 					bNeedsCompilation = true;
 				}
 				catch (BadImageFormatException Ex)
 				{
-					Log.TraceInformation(String.Format("Compiled assembly file '{0}' appears to be for a newer CLR version or is otherwise invalid.  Unreal Build Tool will try to recompile this assembly now.  (Exception: {1})", OutputAssemblyPath, Ex.Message));
+					Logger.LogInformation("Compiled assembly file '{File}' appears to be for a newer CLR version or is otherwise invalid.  Unreal Build Tool will try to recompile this assembly now.  (Exception: {Ex})", OutputAssemblyPath, Ex.Message);
 					bNeedsCompilation = true;
 				}
 				catch (FileNotFoundException)
@@ -339,7 +343,7 @@
 			{
 				using (GlobalTracer.Instance.BuildSpan(String.Format("Compiling rules assembly ({0})", OutputAssemblyPath.GetFileName())).StartActive())
 				{
-					CompiledAssembly = CompileAssembly(OutputAssemblyPath, SourceFileNames, ReferencedAssembies, PreprocessorDefines, TreatWarningsAsErrors);
+					CompiledAssembly = CompileAssembly(OutputAssemblyPath, SourceFileNames, Logger, ReferencedAssembies, PreprocessorDefines, TreatWarningsAsErrors);
 				}
 
 				using (JsonWriter Writer = new JsonWriter(AssemblyManifestFilePath))
diff --git a/Engine/Source/Programs/UnrealBuildTool/System/EncryptionAndSigning.cs b/Engine/Source/Programs/UnrealBuildTool/System/EncryptionAndSigning.cs
--- a/Engine/Source/Programs/UnrealBuildTool/System/EncryptionAndSigning.cs
+++ b/Engine/Source/Programs/UnrealBuildTool/System/EncryptionAndSigning.cs
@@ -8,6 +8,7 @@
 using System.Text.Json;
 using System.Text.RegularExpressions;
 using EpicGames.Core;
+using Microsoft.Extensions.Logging;
 
 namespace UnrealBuildTool
 {
@@ -225,7 +226,7 @@
 		/// <summary>
 		/// Parse crypto settings from INI file
 		/// </summary>
-		public static CryptoSettings ParseCryptoSettings(DirectoryReference? InProjectDirectory, UnrealTargetPlatform InTargetPlatform)
+		public static CryptoSettings ParseCryptoSettings(DirectoryReference? InProjectDirectory, UnrealTargetPlatform InTargetPlatform, ILogger Logger)
 		{
 			CryptoSettings Settings = new CryptoSettings();
 			
@@ -282,11 +283,11 @@
 				{
 					if (EncryptionKeyString.Length < 32)
 					{
-						Log.WriteLine(LogEventType.Warning, "AES key parsed from encryption.ini is too short. It must be 32 bytes, so will be padded with 0s, giving sub-optimal security!");
+						Logger.LogWarning("AES key parsed from encryption.ini is too short. It must be 32 bytes, so will be padded with 0s, giving sub-optimal security!");
 					}
 					else if (EncryptionKeyString.Length > 32)
 					{
-						Log.WriteLine(LogEventType.Warning, "AES key parsed from encryption.ini is too long. It must be 32 bytes, so will be truncated!");
+						Logger.LogWarning("AES key parsed from encryption.ini is too long. It must be 32 bytes, so will be truncated!");
 					}
 
 					Settings.EncryptionKey.Key = ParseAnsiStringToByteArray(EncryptionKeyString, 32);
diff --git a/Engine/Source/Programs/UnrealBuildTool/System/ExternalExecution.cs b/Engine/Source/Programs/UnrealBuildTool/System/ExternalExecution.cs
--- a/Engine/Source/Programs/UnrealBuildTool/System/ExternalExecution.cs
+++ b/Engine/Source/Programs/UnrealBuildTool/System/ExternalExecution.cs
@@ -20,6 +20,7 @@
 using OpenTracing.Util;
 using UnrealBuildTool.Modes;
 using EpicGames.UHT.Utils;
+using Microsoft.Extensions.Logging;
 
 namespace UnrealBuildTool
 {
@@ -448,7 +449,7 @@
 			}
 		}
 
-		public static void SetupUObjectModules(IEnumerable<UEBuildModuleCPP> ModulesToGenerateHeadersFor, UnrealTargetPlatform Platform, ProjectDescriptor? ProjectDescriptor, List<UHTModuleInfo> UObjectModules, List<UHTModuleHeaderInfo> UObjectModuleHeaders, EGeneratedCodeVersion GeneratedCodeVersion, SourceFileMetadataCache MetadataCache)
+		public static void SetupUObjectModules(IEnumerable<UEBuildModuleCPP> ModulesToGenerateHeadersFor, UnrealTargetPlatform Platform, ProjectDescriptor? ProjectDescriptor, List<UHTModuleInfo> UObjectModules, List<UHTModuleHeaderInfo> UObjectModuleHeaders, EGeneratedCodeVersion GeneratedCodeVersion, SourceFileMetadataCache MetadataCache, ILogger Logger)
 		{
 			// Find the type of each module
 			Dictionary<UEBuildModuleCPP, UHTModuleType> ModuleToType = new Dictionary<UEBuildModuleCPP, UHTModuleType>();
@@ -491,7 +492,7 @@
 					DirectoryItem GeneratedCodeDirectory = DirectoryItem.GetItemByDirectoryReference(Module.GeneratedCodeDirectory);
 					foreach (FileItem File in GeneratedCodeDirectory.EnumerateFiles())
 					{
-						File.Delete();
+						File.Delete(Logger);
 					}
 				}
 
@@ -626,9 +627,9 @@
 		/// Gets the latest write time of any of the UnrealHeaderTool binaries (including DLLs and Plugins) or DateTime.MaxValue if UnrealHeaderTool does not exist
 		/// </summary>
 		/// <returns>Latest timestamp of UHT binaries or DateTime.MaxValue if UnrealHeaderTool is out of date and needs to be rebuilt.</returns>
-		static bool GetHeaderToolTimestampUtc(FileReference ReceiptPath, out DateTime Timestamp)
+		static bool GetHeaderToolTimestampUtc(FileReference ReceiptPath, ILogger Logger, out DateTime Timestamp)
 		{
-			using (ScopedTimer TimestampTimer = new ScopedTimer("GetHeaderToolTimestamp"))
+			using (ScopedTimer TimestampTimer = new ScopedTimer("GetHeaderToolTimestamp", Logger))
 			{
 				// Try to read the receipt for UHT.
 				FileItem ReceiptFile = FileItem.GetItemByFileReference(ReceiptPath);
@@ -695,8 +696,9 @@
 		/// <param name="BuildConfiguration">Build configuration</param>
 		/// <param name="UObjectModules">Modules that we generate headers for</param>
 		/// <param name="HeaderToolTimestampUtc">Timestamp for UHT</param>
+		/// <param name="Logger">Logger for output</param>
 		/// <returns>True if the code files are out of date</returns>
-		private static bool AreGeneratedCodeFilesOutOfDate(BuildConfiguration BuildConfiguration, List<UHTModuleInfo> UObjectModules, DateTime HeaderToolTimestampUtc)
+		private static bool AreGeneratedCodeFilesOutOfDate(BuildConfiguration BuildConfiguration, List<UHTModuleInfo> UObjectModules, DateTime HeaderToolTimestampUtc, ILogger Logger)
 		{
 			// Get CoreUObject.init.gen.cpp timestamp.  If the source files are older than the CoreUObject generated code, we'll
 			// need to regenerate code for the module
@@ -732,7 +734,7 @@
 				if (!TestDirectory.Exists)
 				{
 					// Generated code directory is missing entirely!
-					Log.TraceLog("UnrealHeaderTool needs to run because no generated code directory was found for module {0}", Module.ModuleName);
+					Logger.LogDebug("UnrealHeaderTool needs to run because no generated code directory was found for module {ModuleName}", Module.ModuleName);
 					return true;
 				}
 
@@ -745,7 +747,7 @@
 				if (!SavedTimestampFileInfo.Exists)
 				{
 					// Timestamp file was missing (possibly deleted/cleaned), so headers are out of date
-					Log.TraceLog("UnrealHeaderTool needs to run because UHT Timestamp file did not exist for module {0}", Module.ModuleName);
+					Logger.LogDebug("UnrealHeaderTool needs to run because UHT Timestamp file did not exist for module {ModuleName}", Module.ModuleName);
 					return true;
 				}
 
@@ -754,13 +756,13 @@
 				if (HeaderToolTimestampUtc > SavedTimestampUtc)
 				{
 					// Generated code is older than UnrealHeaderTool.exe.  Out of date!
-					Log.TraceLog("UnrealHeaderTool needs to run because UnrealHeaderTool timestamp ({0}) is later than timestamp for module {1} ({2})", HeaderToolTimestampUtc.ToLocalTime(), Module.ModuleName, SavedTimestampUtc.ToLocalTime());
+					Logger.LogDebug("UnrealHeaderTool needs to run because UnrealHeaderTool timestamp ({Time}) is later than timestamp for module {ModuleName} ({ModuleTime})", HeaderToolTimestampUtc.ToLocalTime(), Module.ModuleName, SavedTimestampUtc.ToLocalTime());
 					return true;
 				}
 				if (CoreGeneratedTimestampUtc > SavedTimestampUtc)
 				{
 					// Generated code is older than CoreUObject headers.  Out of date!
-					Log.TraceLog("UnrealHeaderTool needs to run because CoreUObject timestamp ({0}) is newer than timestamp for module {1} ({2})", CoreGeneratedTimestampUtc.Value.ToLocalTime(), Module.ModuleName, SavedTimestampUtc.ToLocalTime());
+					Logger.LogDebug("UnrealHeaderTool needs to run because CoreUObject timestamp ({Time}) is newer than timestamp for module {ModuleName} ({ModuleTime})", CoreGeneratedTimestampUtc.Value.ToLocalTime(), Module.ModuleName, SavedTimestampUtc.ToLocalTime());
 					return true;
 				}
 
@@ -768,7 +770,7 @@
 				FileInfo ModuleRulesFile = new FileInfo(Module.ModuleRulesFile.FullName);
 				if (!ModuleRulesFile.Exists || ModuleRulesFile.LastWriteTimeUtc > SavedTimestampUtc)
 				{
-					Log.TraceLog("UnrealHeaderTool needs to run because SavedTimestamp is older than the rules file ({0}) for module {1}", Module.ModuleRulesFile, Module.ModuleName);
+					Logger.LogDebug("UnrealHeaderTool needs to run because SavedTimestamp is older than the rules file ({ModuleModuleRulesFile}) for module {ModuleModuleName}", Module.ModuleRulesFile, Module.ModuleName);
 					return true;
 				}
 
@@ -784,7 +786,7 @@
 					string[] UObjectFilesFromPreviousRun = File.ReadAllLines(TimestampFile);
 					if (AllUObjectHeaders.Count != UObjectFilesFromPreviousRun.Length)
 					{
-						Log.TraceLog("UnrealHeaderTool needs to run because there are a different number of UObject source files in module {0}", Module.ModuleName);
+						Logger.LogDebug("UnrealHeaderTool needs to run because there are a different number of UObject source files in module {ModuleModuleName}", Module.ModuleName);
 						return true;
 					}
 
@@ -794,7 +796,7 @@
 					{
 						if(!ObjectHeadersSet.Contains(FileName))
 						{
-							Log.TraceLog("UnrealHeaderTool needs to run because the set of UObject source files in module {0} has changed ({1})", Module.ModuleName, FileName);
+							Logger.LogDebug("UnrealHeaderTool needs to run because the set of UObject source files in module {ModuleModuleName} has changed ({FileName})", Module.ModuleName, FileName);
 							return true;
 						}
 					}
@@ -807,7 +809,7 @@
 					// Has the source header changed since we last generated headers successfully?
 					if (HeaderFileTimestampUtc > SavedTimestampUtc)
 					{
-						Log.TraceLog("UnrealHeaderTool needs to run because SavedTimestamp is older than HeaderFileTimestamp ({0}) for module {1}", HeaderFile.AbsolutePath, Module.ModuleName);
+						Logger.LogDebug("UnrealHeaderTool needs to run because SavedTimestamp is older than HeaderFileTimestamp ({HeaderFileAbsolutePath}) for module {ModuleModuleName}", HeaderFile.AbsolutePath, Module.ModuleName);
 						return true;
 					}
 				}
@@ -891,16 +893,16 @@
 		/// <summary>
 		/// Run an external native executable (and capture the output), given the executable path and the commandline.
 		/// </summary>
-		public static int RunExternalNativeExecutable(FileReference ExePath, string Commandline)
+		public static int RunExternalNativeExecutable(FileReference ExePath, string Commandline, ILogger Logger)
 		{
-			Log.TraceVerbose("RunExternalExecutable {0} {1}", ExePath.FullName, Commandline);
+			Logger.LogDebug("RunExternalExecutable {ExePathFullName} {Commandline}", ExePath.FullName, Commandline);
 			using (Process GameProcess = new Process())
 			{
 				GameProcess.StartInfo.FileName = ExePath.FullName;
 				GameProcess.StartInfo.Arguments = Commandline;
 				GameProcess.StartInfo.UseShellExecute = false;
 				GameProcess.StartInfo.RedirectStandardOutput = true;
-				GameProcess.OutputDataReceived += PrintProcessOutputAsync;
+				GameProcess.OutputDataReceived += (s, e) => PrintProcessOutputAsync(s, e, Logger);
 				GameProcess.Start();
 				GameProcess.BeginOutputReadLine();
 				GameProcess.WaitForExit();
@@ -912,13 +914,13 @@
 		/// <summary>
 		/// Simple function to pipe output asynchronously
 		/// </summary>
-		private static void PrintProcessOutputAsync(object Sender, DataReceivedEventArgs Event)
+		private static void PrintProcessOutputAsync(object Sender, DataReceivedEventArgs Event, ILogger Logger)
 		{
 			// DataReceivedEventHandler is fired with a null string when the output stream is closed.  We don't want to
 			// print anything for that event.
 			if (!String.IsNullOrEmpty(Event.Data))
 			{
-				Log.TraceInformation(Event.Data);
+				Logger.LogInformation("{Output}", Event.Data);
 			}
 		}
 
@@ -926,15 +928,15 @@
 		/// Builds and runs the header tool and touches the header directories.
 		/// Performs any early outs if headers need no changes, given the UObject modules, tool path, game name, and configuration
 		/// </summary>
-		public static void ExecuteHeaderToolIfNecessary(BuildConfiguration BuildConfiguration, FileReference? ProjectFile, TargetMakefile Makefile, string TargetName, ISourceFileWorkingSet WorkingSet)
+		public static void ExecuteHeaderToolIfNecessary(BuildConfiguration BuildConfiguration, FileReference? ProjectFile, TargetMakefile Makefile, string TargetName, ISourceFileWorkingSet WorkingSet, ILogger Logger)
 		{
 			if (!BuildConfiguration.bUseBuiltInUnrealHeaderTool || Makefile.bHasRequiredProjectScriptPlugin || TargetName.Equals("UnrealHeaderTool", StringComparison.InvariantCultureIgnoreCase))
 			{
-				ExecuteExternalHeaderToolIfNecessary(BuildConfiguration, ProjectFile, Makefile, TargetName, WorkingSet);
+				ExecuteExternalHeaderToolIfNecessary(BuildConfiguration, ProjectFile, Makefile, TargetName, WorkingSet, Logger);
 			}
 			else
 			{
-				ExecuteInternalHeaderToolIfNecessary(BuildConfiguration, ProjectFile, Makefile, TargetName, WorkingSet);
+				ExecuteInternalHeaderToolIfNecessary(BuildConfiguration, ProjectFile, Makefile, TargetName, WorkingSet, Logger);
 			}
 		}
 
@@ -1032,13 +1034,13 @@
 		/// Builds and runs the header tool and touches the header directories.
 		/// Performs any early outs if headers need no changes, given the UObject modules, tool path, game name, and configuration
 		/// </summary>
-		private static void ExecuteExternalHeaderToolIfNecessary(BuildConfiguration BuildConfiguration, FileReference? ProjectFile, TargetMakefile Makefile, string TargetName, ISourceFileWorkingSet WorkingSet)
+		private static void ExecuteExternalHeaderToolIfNecessary(BuildConfiguration BuildConfiguration, FileReference? ProjectFile, TargetMakefile Makefile, string TargetName, ISourceFileWorkingSet WorkingSet, ILogger Logger)
 		{
 			if (ProgressWriter.bWriteMarkup)
 			{
-				Log.WriteLine(LogEventType.Console, "@progress push 5%");
+				Logger.LogInformation("@progress push 5%");
 			}
-			using (ProgressWriter Progress = new ProgressWriter("Generating code...", false))
+			using (ProgressWriter Progress = new ProgressWriter("Generating code...", false, Logger))
 			{
 				// We never want to try to execute the header tool when we're already trying to build it!
 				bool bIsBuildingUHT = TargetName.Equals("UnrealHeaderTool", StringComparison.InvariantCultureIgnoreCase);
@@ -1052,7 +1054,7 @@
 
 				// check if UHT is out of date
 				DateTime HeaderToolTimestampUtc = DateTime.MaxValue;
-				bool bHaveHeaderTool = !bIsBuildingUHT && GetHeaderToolTimestampUtc(HeaderToolReceipt, out HeaderToolTimestampUtc);
+				bool bHaveHeaderTool = !bIsBuildingUHT && GetHeaderToolTimestampUtc(HeaderToolReceipt, Logger, out HeaderToolTimestampUtc);
 
 				// ensure the headers are up to date
 				bool bUHTNeedsToRun = false;
@@ -1064,7 +1066,7 @@
 				{
 					bUHTNeedsToRun = true;
 				}
-				else if(AreGeneratedCodeFilesOutOfDate(BuildConfiguration, Makefile.UObjectModules, HeaderToolTimestampUtc))
+				else if(AreGeneratedCodeFilesOutOfDate(BuildConfiguration, Makefile.UObjectModules, HeaderToolTimestampUtc, Logger))
 				{
 					bUHTNeedsToRun = true;
 				}
@@ -1131,14 +1133,14 @@
 
 						using (GlobalTracer.Instance.BuildSpan("Building UnrealHeaderTool").StartActive())
 						{
-							BuildMode.Build(new List<TargetDescriptor>{ TargetDescriptor }, BuildConfiguration, WorkingSet, BuildOptions.None, null);
+							BuildMode.Build(new List<TargetDescriptor>{ TargetDescriptor }, BuildConfiguration, WorkingSet, BuildOptions.None, null, Logger);
 						}
 					}
 
 					Progress.Write(1, 3);
 
 					string ActualTargetName = String.IsNullOrEmpty(TargetName) ? "UE5" : TargetName;
-					Log.TraceInformation("Parsing headers for {0}", ActualTargetName);
+					Logger.LogInformation("Parsing headers for {ActualTargetName}", ActualTargetName);
 
 					FileReference HeaderToolPath = GetHeaderToolPath(HeaderToolReceipt);
 					if (!FileReference.Exists(HeaderToolPath))
@@ -1178,12 +1180,12 @@
 						CmdLine += " -FailIfGeneratedCodeChanges";
 					}
 
-					Log.TraceInformation("  Running UnrealHeaderTool {0}", CmdLine);
+					Logger.LogInformation("  Running UnrealHeaderTool {CmdLine}", CmdLine);
 
 					Stopwatch s = new Stopwatch();
 					s.Start();
 					IScope Timer = GlobalTracer.Instance.BuildSpan("Executing UnrealHeaderTool").StartActive();
-					CompilationResult UHTResult = (CompilationResult)RunExternalNativeExecutable(ExternalExecution.GetHeaderToolPath(HeaderToolReceipt), CmdLine);
+					CompilationResult UHTResult = (CompilationResult)RunExternalNativeExecutable(ExternalExecution.GetHeaderToolPath(HeaderToolReceipt), CmdLine, Logger);
 					Timer.Span.Finish();
 					s.Stop();
 
@@ -1202,13 +1204,13 @@
 						if (BuildHostPlatform.Current.Platform == UnrealTargetPlatform.Win64 &&
 							(int)(UHTResult) < 0)
 						{
-							Log.TraceError(String.Format("UnrealHeaderTool failed with exit code 0x{0:X} - check that Unreal Engine prerequisites are installed.", (int)UHTResult));
+							Logger.LogError("UnrealHeaderTool failed with exit code 0x{Result:X} - check that Unreal Engine prerequisites are installed.", (int)UHTResult);
 						}
 
 						throw new CompilationResultException(UHTResult);
 					}
 
-					Log.TraceInformation("Reflection code generated for {0} in {1} seconds", ActualTargetName, s.Elapsed.TotalSeconds);
+					Logger.LogInformation("Reflection code generated for {Target} in {Time} seconds", ActualTargetName, s.Elapsed.TotalSeconds);
 
 					// Update the tool info file
 					DirectoryReference.CreateDirectory(ToolInfoFile.Directory);
@@ -1223,7 +1225,7 @@
 				}
 				else
 				{
-					Log.TraceVerbose("Generated code is up to date.");
+					Logger.LogDebug("Generated code is up to date.");
 				}
 
 				Progress.Write(2, 3);
@@ -1237,7 +1239,7 @@
 			}
 			if (ProgressWriter.bWriteMarkup)
 			{
-				Log.WriteLine(LogEventType.Console, "@progress pop");
+				Logger.LogInformation("@progress pop");
 			}
 		}
 
@@ -1246,13 +1248,13 @@
 		/// Performs any early outs if headers need no changes, given the UObject modules, tool path, game name, and configuration
 		/// </summary>
 		private static void ExecuteInternalHeaderToolIfNecessary(BuildConfiguration BuildConfiguration, FileReference? ProjectFile, 
-			TargetMakefile Makefile, string TargetName, ISourceFileWorkingSet WorkingSet)
+			TargetMakefile Makefile, string TargetName, ISourceFileWorkingSet WorkingSet, ILogger Logger)
 		{
 			if (ProgressWriter.bWriteMarkup)
 			{
-				Log.WriteLine(LogEventType.Console, "@progress push 5%");
+				Logger.LogInformation("@progress push 5%");
 			}
-			using (ProgressWriter Progress = new ProgressWriter("Generating code...", false))
+			using (ProgressWriter Progress = new ProgressWriter("Generating code...", false, Logger))
 			{
 				string RootLocalPath = Unreal.RootDirectory.FullName;
 
@@ -1286,7 +1288,7 @@
 				{
 					bUHTNeedsToRun = true;
 				}
-				else if (AreGeneratedCodeFilesOutOfDate(BuildConfiguration, Makefile.UObjectModules, CompositeTimestamp))
+				else if (AreGeneratedCodeFilesOutOfDate(BuildConfiguration, Makefile.UObjectModules, CompositeTimestamp, Logger))
 				{
 					bUHTNeedsToRun = true;
 				}
@@ -1320,7 +1322,7 @@
 					WriteUHTManifest(Makefile, TargetName, ModuleInfoFileName, ExternalDependenciesFile);
 
 					string ActualTargetName = String.IsNullOrEmpty(TargetName) ? "UE5" : TargetName;
-					Log.TraceInformation("Parsing headers for {0}", ActualTargetName);
+					Logger.LogInformation("Parsing headers for {ActualTargetName}", ActualTargetName);
 
 					// Generate the command line
 					List<string> CmdArgs = new List<string>();
@@ -1360,14 +1362,14 @@
 						CmdLine.AppendArgument(Arg);
 					}
 
-					Log.TraceInformation("  Running Internal UnrealHeaderTool {0}", CmdLine);
+					Logger.LogInformation("  Running Internal UnrealHeaderTool {CmdLine}", CmdLine);
 
 					// Run UHT
 					Stopwatch s = new Stopwatch();
 					s.Start();
 					IScope Timer = GlobalTracer.Instance.BuildSpan("Executing UnrealHeaderTool").StartActive();
 					UnrealHeaderToolMode UHTTool = new UnrealHeaderToolMode();
-					CompilationResult UHTResult = (CompilationResult)UHTTool.Execute(Arguments);
+					CompilationResult UHTResult = (CompilationResult)UHTTool.Execute(Arguments, Logger);
 					Timer.Span.Finish();
 					s.Stop();
 
@@ -1386,13 +1388,13 @@
 						if (BuildHostPlatform.Current.Platform == UnrealTargetPlatform.Win64 &&
 							(int)(UHTResult) < 0)
 						{
-							Log.TraceError(String.Format("UnrealHeaderTool failed with exit code 0x{0:X} - check that Unreal Engine prerequisites are installed.", (int)UHTResult));
+							Logger.LogError("UnrealHeaderTool failed with exit code 0x{Result:X} - check that Unreal Engine prerequisites are installed.", (int)UHTResult);
 						}
 
 						throw new CompilationResultException(UHTResult);
 					}
 
-					Log.TraceInformation("Reflection code generated for {0} in {1} seconds", ActualTargetName, s.Elapsed.TotalSeconds);
+					Logger.LogInformation("Reflection code generated for {TargetName} in {Time} seconds", ActualTargetName, s.Elapsed.TotalSeconds);
 
 					// Update the tool info file
 					DirectoryReference.CreateDirectory(ToolInfoFile.Directory);
@@ -1407,7 +1409,7 @@
 				}
 				else
 				{
-					Log.TraceVerbose("Generated code is up to date.");
+					Logger.LogDebug("Generated code is up to date.");
 				}
 
 				Progress.Write(2, 3);
@@ -1421,7 +1423,7 @@
 			}
 			if (ProgressWriter.bWriteMarkup)
 			{
-				Log.WriteLine(LogEventType.Console, "@progress pop");
+				Logger.LogInformation("@progress pop");
 			}
 		}
 
diff --git a/Engine/Source/Programs/UnrealBuildTool/System/HotReload.cs b/Engine/Source/Programs/UnrealBuildTool/System/HotReload.cs
--- a/Engine/Source/Programs/UnrealBuildTool/System/HotReload.cs
+++ b/Engine/Source/Programs/UnrealBuildTool/System/HotReload.cs
@@ -9,6 +9,7 @@
 using System.Threading;
 using System.Threading.Tasks;
 using EpicGames.Core;
+using Microsoft.Extensions.Logging;
 using OpenTracing.Util;
 using UnrealBuildBase;
 
@@ -138,14 +139,15 @@
 		/// </summary>
 		/// <param name="TargetDescriptor">The target being built</param>
 		/// <param name="BuildConfiguration">Global build configuration</param>
+		/// <param name="Logger">Logger for output</param>
 		/// <returns>Default hotreload mode</returns>
-		public static HotReloadMode GetDefaultMode(TargetDescriptor TargetDescriptor, BuildConfiguration BuildConfiguration)
+		public static HotReloadMode GetDefaultMode(TargetDescriptor TargetDescriptor, BuildConfiguration BuildConfiguration, ILogger Logger)
 		{
 			if (TargetDescriptor.HotReloadModuleNameToSuffix.Count > 0 && TargetDescriptor.ForeignPlugin == null)
 			{
 				return HotReloadMode.FromEditor;
 			}
-			else if (BuildConfiguration.bAllowHotReloadFromIDE && HotReload.ShouldDoHotReloadFromIDE(BuildConfiguration, TargetDescriptor))
+			else if (BuildConfiguration.bAllowHotReloadFromIDE && HotReload.ShouldDoHotReloadFromIDE(BuildConfiguration, TargetDescriptor, Logger))
 			{
 				return HotReloadMode.FromIDE;
 			}
@@ -162,7 +164,8 @@
 		/// <param name="Makefile">Makefile for the targe</param>
 		/// <param name="Actions">Actions for this target</param>
 		/// <param name="BuildConfiguration">Global build configuration</param>
-		public static Dictionary<FileReference, FileReference>? Setup(TargetDescriptor TargetDescriptor, TargetMakefile Makefile, List<LinkedAction> Actions, BuildConfiguration BuildConfiguration)
+		/// <param name="Logger">Logger for output</param>
+		public static Dictionary<FileReference, FileReference>? Setup(TargetDescriptor TargetDescriptor, TargetMakefile Makefile, List<LinkedAction> Actions, BuildConfiguration BuildConfiguration, ILogger Logger)
 		{
 			Dictionary<FileReference, FileReference>? PatchedOldLocationToNewLocation = null;
 
@@ -178,7 +181,7 @@
 			}
 			else if (TargetDescriptor.HotReloadMode == HotReloadMode.Default)
 			{
-				TargetDescriptor.HotReloadMode = GetDefaultMode(TargetDescriptor, BuildConfiguration);
+				TargetDescriptor.HotReloadMode = GetDefaultMode(TargetDescriptor, BuildConfiguration, Logger);
 			}
 
 			// Apply the previous hot reload state
@@ -189,7 +192,7 @@
 				{
 					// Delete the previous state file
 					FileReference StateFile = HotReloadState.GetLocation(TargetDescriptor);
-					HotReload.DeleteTemporaryFiles(StateFile);
+					HotReload.DeleteTemporaryFiles(StateFile, Logger);
 				}
 			}
 			else
@@ -211,10 +214,10 @@
 			return PatchedOldLocationToNewLocation;
 		}
 
-		public static void CheckForLiveCodingSessionActive(TargetDescriptor TargetDescriptor, TargetMakefile Makefile, BuildConfiguration BuildConfiguration)
+		public static void CheckForLiveCodingSessionActive(TargetDescriptor TargetDescriptor, TargetMakefile Makefile, BuildConfiguration BuildConfiguration, ILogger Logger)
 		{
 			// Guard against a live coding session for this target being active
-			if (BuildConfiguration.bAllowHotReloadFromIDE && TargetDescriptor.HotReloadMode != HotReloadMode.LiveCoding && TargetDescriptor.ForeignPlugin == null && HotReload.IsLiveCodingSessionActive(Makefile))
+			if (BuildConfiguration.bAllowHotReloadFromIDE && TargetDescriptor.HotReloadMode != HotReloadMode.LiveCoding && TargetDescriptor.ForeignPlugin == null && HotReload.IsLiveCodingSessionActive(Makefile, Logger))
 			{
 				throw new BuildException("Unable to build while Live Coding is active. Exit the editor and game, or press Ctrl+Alt+F11 if iterating on code in the editor or game");
 			}
@@ -224,8 +227,9 @@
 		/// Checks whether a live coding session is currently active for a target. If so, we don't want to allow modifying any object files before they're loaded.
 		/// </summary>
 		/// <param name="Makefile">Makefile for the target being built</param>
+		/// <param name="Logger">Logger for output</param>
 		/// <returns>True if a live coding session is active, false otherwise</returns>
-		static bool IsLiveCodingSessionActive(TargetMakefile Makefile)
+		static bool IsLiveCodingSessionActive(TargetMakefile Makefile, ILogger Logger)
 		{
 			// Find the first output executable
 			FileReference Executable = Makefile.ExecutableFile;
@@ -245,7 +249,7 @@
 						MutexName.Append(Character);
 					}
 				}
-				Log.TraceLog("Checking for live coding mutex: {0}", MutexName);
+				Logger.LogDebug("Checking for live coding mutex: {MutexName}", MutexName);
 
 				// Try to open the mutex
 				Mutex? Mutex;
@@ -261,7 +265,7 @@
 		/// <summary>
 		/// Checks if the editor is currently running and this is a hot-reload
 		/// </summary>
-		static bool ShouldDoHotReloadFromIDE(BuildConfiguration BuildConfiguration, TargetDescriptor TargetDesc)
+		static bool ShouldDoHotReloadFromIDE(BuildConfiguration BuildConfiguration, TargetDescriptor TargetDesc, ILogger Logger)
 		{
 			// Check if Hot-reload is disabled globally for this project
 			ConfigHierarchy Hierarchy = ConfigCache.ReadHierarchy(ConfigHierarchyType.Engine, DirectoryReference.FromFile(TargetDesc.ProjectFile), TargetDesc.Platform);
@@ -401,7 +405,7 @@
 								}
 								catch
 								{
-									Log.TraceLog("Failed to delete EditorRun file for exited process: {0}", EditorInstanceFile.GetFileName());
+									Logger.LogDebug("Failed to delete EditorRun file for exited process: {Process}", EditorInstanceFile.GetFileName());
 								}
 								continue;
 							}
@@ -445,7 +449,8 @@
 		/// Delete all temporary files created by previous hot reload invocations
 		/// </summary>
 		/// <param name="HotReloadStateFile">Location of the state file</param>
-		public static void DeleteTemporaryFiles(FileReference HotReloadStateFile)
+		/// <param name="Logger">Logger for output</param>
+		public static void DeleteTemporaryFiles(FileReference HotReloadStateFile, ILogger Logger)
 		{
 			if(FileReference.Exists(HotReloadStateFile))
 			{
@@ -457,7 +462,7 @@
 				}
 				catch(Exception Ex)
 				{
-					Log.TraceWarning("Unable to read hot reload state file: {0}", HotReloadStateFile);
+					Logger.LogWarning("Unable to read hot reload state file: {HotReloadStateFile}", HotReloadStateFile);
 					Log.WriteException(Ex, null);
 					return;
 				}
@@ -475,7 +480,7 @@
 						{
 							throw new BuildException(Ex, "Unable to delete hot-reload file: {0}", Location);
 						}
-						Log.TraceInformation("Deleted hot-reload file: {0}", Location);
+						Logger.LogInformation("Deleted hot-reload file: {Location}", Location);
 					}
 				}
 
@@ -526,12 +531,13 @@
 		/// <param name="PrerequisiteActions">The actions to execute</param>
 		/// <param name="TargetActionsToExecute">Actions to execute for this target</param>
 		/// <param name="InitialPatchedOldLocationToNewLocation">Collection of all the renamed as part of module reload requests.  Can be null</param>
+		/// <param name="Logger">Logger for output</param>
 		/// <returns>Set of actions to execute</returns>
-		public static List<LinkedAction> PatchActionsForTarget(BuildConfiguration BuildConfiguration, TargetDescriptor TargetDescriptor, TargetMakefile Makefile, List<LinkedAction> PrerequisiteActions, List<LinkedAction> TargetActionsToExecute, Dictionary<FileReference, FileReference>? InitialPatchedOldLocationToNewLocation)
+		public static List<LinkedAction> PatchActionsForTarget(BuildConfiguration BuildConfiguration, TargetDescriptor TargetDescriptor, TargetMakefile Makefile, List<LinkedAction> PrerequisiteActions, List<LinkedAction> TargetActionsToExecute, Dictionary<FileReference, FileReference>? InitialPatchedOldLocationToNewLocation, ILogger Logger)
 		{
 			// Get the dependency history
 			CppDependencyCache CppDependencies = new CppDependencyCache();
-			CppDependencies.Mount(TargetDescriptor.ProjectFile, TargetDescriptor.Name, TargetDescriptor.Platform, TargetDescriptor.Configuration, Makefile.TargetType, TargetDescriptor.Architecture);
+			CppDependencies.Mount(TargetDescriptor.ProjectFile, TargetDescriptor.Name, TargetDescriptor.Platform, TargetDescriptor.Configuration, Makefile.TargetType, TargetDescriptor.Architecture, Logger);
 
 			ActionHistory History = new ActionHistory();
 			if(TargetDescriptor.ProjectFile != null)
@@ -577,7 +583,7 @@
 						FileReference.WriteAllLines(new FileReference(TargetDescriptor.LiveCodingModules.FullName + ".out"), ProtectedOutputFiles.Select(x => x.ToString()));
 						foreach (FileReference ProtectedOutputFile in ProtectedOutputFiles)
 						{
-							Log.TraceInformation("Module {0} is not currently enabled for Live Coding", ProtectedOutputFile);
+							Logger.LogInformation("Module {ProtectedOutputFile} is not currently enabled for Live Coding", ProtectedOutputFile);
 						}
 
 						// Note the issue but continue processing to allow the limit to generate an error if hit.
@@ -587,21 +593,21 @@
 
 				// Filter the prerequisite actions down to just the compile actions, then recompute all the actions to execute
 				PrerequisiteActions = new List<LinkedAction>(TargetActionsToExecute.Where(x => x.ActionType == ActionType.Compile));
-				TargetActionsToExecute = ActionGraph.GetActionsToExecute(PrerequisiteActions, CppDependencies, History, BuildConfiguration.bIgnoreOutdatedImportLibraries);
+				TargetActionsToExecute = ActionGraph.GetActionsToExecute(PrerequisiteActions, CppDependencies, History, BuildConfiguration.bIgnoreOutdatedImportLibraries, Logger);
 
 				// Update the action graph with these new paths
 				Dictionary<FileReference, FileReference> OriginalFileToPatchedFile = new Dictionary<FileReference, FileReference>();
-				HotReload.PatchActionGraphForLiveCoding(PrerequisiteActions, OriginalFileToPatchedFile);
+				HotReload.PatchActionGraphForLiveCoding(PrerequisiteActions, OriginalFileToPatchedFile, Logger);
 
 				// Get a new list of actions to execute now that the graph has been modified
-				TargetActionsToExecute = ActionGraph.GetActionsToExecute(PrerequisiteActions, CppDependencies, History, BuildConfiguration.bIgnoreOutdatedImportLibraries);
+				TargetActionsToExecute = ActionGraph.GetActionsToExecute(PrerequisiteActions, CppDependencies, History, BuildConfiguration.bIgnoreOutdatedImportLibraries, Logger);
 
 				// Check to see if we exceed the limit for live coding actions
 				if (TargetDescriptor.LiveCodingLimit > 0 && TargetDescriptor.LiveCodingLimit < TargetActionsToExecute.Count)
 				{
-					Log.TraceInformation("The live coding request of {0} actions exceeds the number of allowed actions of {1}", TargetActionsToExecute.Count, TargetDescriptor.LiveCodingLimit);
-					Log.TraceInformation("This limit helps to prevent the situation where seemingly simple changes result in large scale rebuilds.");
-					Log.TraceInformation("It can also help to detect when the engine needs to be rebuilt outside of Live Coding due to compiler changes.");
+					Logger.LogInformation("The live coding request of {TargetActionsToExecuteCount} actions exceeds the number of allowed actions of {TargetDescriptorLiveCodingLimit}", TargetActionsToExecute.Count, TargetDescriptor.LiveCodingLimit);
+					Logger.LogInformation("This limit helps to prevent the situation where seemingly simple changes result in large scale rebuilds.");
+					Logger.LogInformation("It can also help to detect when the engine needs to be rebuilt outside of Live Coding due to compiler changes.");
 					Result = CompilationResult.LiveCodingLimitError;
 				}
 
@@ -691,7 +697,7 @@
 				Dictionary<FileReference, FileReference> PatchedOldLocationToNewLocation = HotReload.PatchActionGraph(PrerequisiteActions, OldLocationToNewLocation);
 
 				// Get a new list of actions to execute now that the graph has been modified
-				TargetActionsToExecute = ActionGraph.GetActionsToExecute(PrerequisiteActions, CppDependencies, History, BuildConfiguration.bIgnoreOutdatedImportLibraries);
+				TargetActionsToExecute = ActionGraph.GetActionsToExecute(PrerequisiteActions, CppDependencies, History, BuildConfiguration.bIgnoreOutdatedImportLibraries, Logger);
 
 				// Record all of the updated locations directly associated with actions.
 				if (InitialPatchedOldLocationToNewLocation != null)
@@ -801,7 +807,8 @@
 		/// </summary>
 		/// <param name="Actions">Set of actions</param>
 		/// <param name="OriginalFileToPatchedFile">Dictionary that receives a map of original object file to patched object file</param>
-		public static void PatchActionGraphForLiveCoding(IEnumerable<LinkedAction> Actions, Dictionary<FileReference, FileReference> OriginalFileToPatchedFile)
+		/// <param name="Logger"></param>
+		public static void PatchActionGraphForLiveCoding(IEnumerable<LinkedAction> Actions, Dictionary<FileReference, FileReference> OriginalFileToPatchedFile, ILogger Logger)
 		{
 			foreach (LinkedAction Action in Actions)
 			{
@@ -925,7 +932,7 @@
 						}
 					}
 
-					Utils.WriteFileIfChanged(NewResponseFile, ResponseLines);
+					Utils.WriteFileIfChanged(NewResponseFile, ResponseLines, Logger);
 
 					Arguments[ResponseFileIdx] = "@" + NewResponseFile.FullName;
 
diff --git a/Engine/Source/Programs/UnrealBuildTool/System/InstalledPlatformInfo.cs b/Engine/Source/Programs/UnrealBuildTool/System/InstalledPlatformInfo.cs
--- a/Engine/Source/Programs/UnrealBuildTool/System/InstalledPlatformInfo.cs
+++ b/Engine/Source/Programs/UnrealBuildTool/System/InstalledPlatformInfo.cs
@@ -8,6 +8,7 @@
 using System.IO;
 using EpicGames.Core;
 using UnrealBuildBase;
+using Microsoft.Extensions.Logging;
 
 namespace UnrealBuildTool
 {
@@ -150,7 +151,7 @@
 			// Unused, but allows timing call of the static constructor
 		}
 
-		private static void ParsePlatformConfiguration(string PlatformConfiguration)
+		private static void ParsePlatformConfiguration(string PlatformConfiguration)//, ILogger Logger)
 		{
 			// Trim whitespace at the beginning.
 			PlatformConfiguration = PlatformConfiguration.Trim();
@@ -168,7 +169,7 @@
 			}
 			if (Configuration == UnrealTargetConfiguration.Unknown)
 			{
-				Log.TraceWarning("Unable to read configuration from {0}", PlatformConfiguration);
+//				Logger.LogWarning("Unable to read configuration from {PlatformConfiguration}", PlatformConfiguration);
 				bCanCreateEntry = false;
 			}
 
@@ -177,7 +178,7 @@
 			{
 				if (!UnrealTargetPlatform.IsValidName(PlatformName))
 				{
-					Log.TraceWarning("Unable to read platform from {0}", PlatformConfiguration);
+//					Logger.LogWarning("Unable to read platform from {PlatformConfiguration}", PlatformConfiguration);
 					bCanCreateEntry = false;
 				}
 			}
@@ -188,13 +189,13 @@
 			{
 				if (!Enum.TryParse(PlatformTypeName, out PlatformType))
 				{
-					Log.TraceWarning("Unable to read Platform Type from {0}, defaulting to Game", PlatformConfiguration);
+//					Logger.LogWarning("Unable to read Platform Type from {PlatformConfiguration}, defaulting to Game", PlatformConfiguration);
 					PlatformType = TargetType.Game;
 				}
 			}
 			if (PlatformType == TargetType.Program)
 			{
-				Log.TraceWarning("Program is not a valid PlatformType for an Installed Platform, defaulting to Game");
+//				Logger.LogWarning("Program is not a valid PlatformType for an Installed Platform, defaulting to Game");
 				PlatformType = TargetType.Game;
 			}
 
@@ -215,7 +216,7 @@
 			}
 			if (ProjectType == EProjectType.Unknown)
 			{
-				Log.TraceWarning("Unable to read project type from {0}", PlatformConfiguration);
+//				Logger.LogWarning("Unable to read project type from {PlatformConfiguration}", PlatformConfiguration);
 				bCanCreateEntry = false;
 			}
 
diff --git a/Engine/Source/Programs/UnrealBuildTool/System/JunkDeleter.cs b/Engine/Source/Programs/UnrealBuildTool/System/JunkDeleter.cs
--- a/Engine/Source/Programs/UnrealBuildTool/System/JunkDeleter.cs
+++ b/Engine/Source/Programs/UnrealBuildTool/System/JunkDeleter.cs
@@ -6,6 +6,7 @@
 using System.Text;
 using System.IO;
 using EpicGames.Core;
+using Microsoft.Extensions.Logging;
 
 namespace UnrealBuildTool
 {
@@ -14,10 +15,10 @@
 		/// <summary>
 		/// Loads JunkManifest.txt file and removes all junk files/folders defined in it.
 		/// </summary>
-		static public void DeleteJunk()
+		static public void DeleteJunk(ILogger Logger)
 		{
 			List<string> JunkManifest = LoadJunkManifest();
-			DeleteAllJunk(JunkManifest);
+			DeleteAllJunk(JunkManifest, Logger);
 		}
 
 		/// <summary>
@@ -94,7 +95,8 @@
 		/// Goes through each entry from the junk manifest and deletes it.
 		/// </summary>
 		/// <param name="JunkManifest">JunkManifest.txt entries.</param>
-		static private void DeleteAllJunk(List<string> JunkManifest)
+		/// <param name="Logger">Logger for output</param>
+		static private void DeleteAllJunk(List<string> JunkManifest, ILogger Logger)
 		{
 			foreach (string Junk in JunkManifest)
 			{
@@ -111,27 +113,27 @@
 							string[] FilesToDelete = Directory.GetFiles(DirectoryToLookIn, FileName, SearchOption.TopDirectoryOnly);
 							foreach (string JunkFile in FilesToDelete)
 							{
-								DeleteFile(JunkFile);
+								DeleteFile(JunkFile, Logger);
 							}
 
 							// Delete all subdirectories with the specified folder
 							string[] DirectoriesToDelete = Directory.GetDirectories(DirectoryToLookIn, FileName, SearchOption.TopDirectoryOnly);
 							foreach (string JunkFolder in DirectoriesToDelete)
 							{
-								DeleteDirectory(JunkFolder);
+								DeleteDirectory(JunkFolder, Logger);
 							}
 						}
 					}
 					else
 					{
 						// Delete single file
-						DeleteFile(Junk);
+						DeleteFile(Junk, Logger);
 					}
 				}
 				else if (Directory.Exists(Junk))
 				{
 					// Delete the selected folder and all its contents
-					DeleteDirectory(Junk);
+					DeleteDirectory(Junk, Logger);
 				}
 			}
 		}
@@ -161,17 +163,18 @@
 		/// Deletes a directory recursively gracefully handling all exceptions.
 		/// </summary>
 		/// <param name="DirectoryPath">Path.</param>
-		static private void DeleteDirectory(string DirectoryPath)
+		/// <param name="Logger">Logger for output</param>
+		static private void DeleteDirectory(string DirectoryPath, ILogger Logger)
 		{
 			try
 			{
-				Log.TraceInformation("Deleting junk directory: \"{0}\".", DirectoryPath);
+				Logger.LogInformation("Deleting junk directory: \"{Dir}\".", DirectoryPath);
 				Directory.Delete(DirectoryPath, true);
 			}
 			catch (Exception Ex)
 			{
 				// Ignore all exceptions
-				Log.TraceInformation("Unable to delete junk directory: \"{0}\". Error: {1}", DirectoryPath, Ex.Message.TrimEnd());
+				Logger.LogInformation("Unable to delete junk directory: \"{Dir}\". Error: {Ex}", DirectoryPath, Ex.Message.TrimEnd());
 			}
 		}
 
@@ -179,17 +182,18 @@
 		/// Deletes a file gracefully handling all exceptions.
 		/// </summary>
 		/// <param name="Filename">Filename.</param>
-		static private void DeleteFile(string Filename)
+		/// <param name="Logger">Logger for output</param>
+		static private void DeleteFile(string Filename, ILogger Logger)
 		{
 			try
 			{
-				Log.TraceInformation("Deleting junk file: \"{0}\".", Filename);
+				Logger.LogInformation("Deleting junk file: \"{File}\".", Filename);
 				File.Delete(Filename);
 			}
 			catch (Exception Ex)
 			{
 				// Ingore all exceptions
-				Log.TraceInformation("Unable to delete junk file: \"{0}\". Error: {1}", Filename, Ex.Message.TrimEnd());
+				Logger.LogInformation("Unable to delete junk file: \"{File}\". Error: {Ex}", Filename, Ex.Message.TrimEnd());
 			}
 		}
 	}
diff --git a/Engine/Source/Programs/UnrealBuildTool/System/NativeProjects.cs b/Engine/Source/Programs/UnrealBuildTool/System/NativeProjects.cs
--- a/Engine/Source/Programs/UnrealBuildTool/System/NativeProjects.cs
+++ b/Engine/Source/Programs/UnrealBuildTool/System/NativeProjects.cs
@@ -9,6 +9,7 @@
 using System.Threading.Tasks;
 using EpicGames.Core;
 using UnrealBuildBase;
+using Microsoft.Extensions.Logging;
 
 namespace UnrealBuildTool
 {
@@ -35,9 +36,10 @@
 		/// Get the project folder for the given target name
 		/// </summary>
 		/// <param name="InTargetName">Name of the target of interest</param>
+		/// <param name="Logger">Logger for output</param>
 		/// <param name="OutProjectFileName">The project filename</param>
 		/// <returns>True if the target was found</returns>
-		public static bool TryGetProjectForTarget(string InTargetName, [NotNullWhen(true)] out FileReference? OutProjectFileName)
+		public static bool TryGetProjectForTarget(string InTargetName, ILogger Logger, [NotNullWhen(true)] out FileReference? OutProjectFileName)
 		{
 			if(CachedTargetNameToProjectFile == null)
 			{
@@ -46,7 +48,7 @@
 					if(CachedTargetNameToProjectFile == null)
 					{
 						Dictionary<string, FileReference> TargetNameToProjectFile = new Dictionary<string, FileReference>();
-						foreach(FileReference ProjectFile in EnumerateProjectFiles())
+						foreach(FileReference ProjectFile in EnumerateProjectFiles(Logger))
 						{
 							foreach (DirectoryReference ExtensionDir in Unreal.GetExtensionDirs(ProjectFile.Directory))
 							{
diff --git a/Engine/Source/Programs/UnrealBuildTool/System/PlatformExports.cs b/Engine/Source/Programs/UnrealBuildTool/System/PlatformExports.cs
--- a/Engine/Source/Programs/UnrealBuildTool/System/PlatformExports.cs
+++ b/Engine/Source/Programs/UnrealBuildTool/System/PlatformExports.cs
@@ -7,6 +7,7 @@
 using System.Text;
 using System.Threading.Tasks;
 using EpicGames.Core;
+using Microsoft.Extensions.Logging;
 
 namespace UnrealBuildTool
 {
@@ -104,18 +105,20 @@
 		/// Check whether the given platform supports XGE
 		/// </summary>
 		/// <param name="Platform">Platform to check</param>
+		/// <param name="Logger">Logger for output</param>
 		/// <returns>True if the platform supports XGE</returns>
-		public static bool CanUseXGE(UnrealTargetPlatform Platform)
+		public static bool CanUseXGE(UnrealTargetPlatform Platform, ILogger Logger)
 		{
-			return UEBuildPlatform.IsPlatformAvailable(Platform) && UEBuildPlatform.GetBuildPlatform(Platform).CanUseXGE() && XGE.IsAvailable();
+			return UEBuildPlatform.IsPlatformAvailable(Platform) && UEBuildPlatform.GetBuildPlatform(Platform).CanUseXGE() && XGE.IsAvailable(Logger);
 		}
 
 		/// <summary>
 		/// Check whether the given platform supports the parallel executor in UAT
 		/// </summary>
 		/// <param name="Platform">Platform to check</param>
+		/// <param name="Logger">Logger for output</param>
 		/// <returns>True if the platform supports the parallel executor in UAT</returns>
-		public static bool CanUseParallelExecutor(UnrealTargetPlatform Platform)
+		public static bool CanUseParallelExecutor(UnrealTargetPlatform Platform, ILogger Logger)
 		{
 			return UEBuildPlatform.IsPlatformAvailable(Platform) && UEBuildPlatform.GetBuildPlatform(Platform).CanUseParallelExecutor();
 		}
@@ -156,14 +159,15 @@
 		/// <summary>
 		/// Initialize UBT in the context of another host process (presumably UAT)
 		/// </summary>
+		/// <param name="Logger">Logger for output</param>
 		/// <returns>True if initialization was successful</returns>
-		public static bool Initialize()
+		public static bool Initialize(ILogger Logger)
 		{
 			// Read the XML configuration files
-			XmlConfig.ReadConfigFiles(null);
+			XmlConfig.ReadConfigFiles(null, Logger);
 
 			// Register all the platform classes
-			UEBuildPlatform.RegisterPlatforms(false, false);
+			UEBuildPlatform.RegisterPlatforms(false, false, Logger);
 			return true;
 		}
 	}
diff --git a/Engine/Source/Programs/UnrealBuildTool/System/Plugins.cs b/Engine/Source/Programs/UnrealBuildTool/System/Plugins.cs
--- a/Engine/Source/Programs/UnrealBuildTool/System/Plugins.cs
+++ b/Engine/Source/Programs/UnrealBuildTool/System/Plugins.cs
@@ -7,6 +7,7 @@
 using System.Linq;
 using EpicGames.Core;
 using UnrealBuildBase;
+using Microsoft.Extensions.Logging;
 
 namespace UnrealBuildTool
 {
@@ -268,13 +269,14 @@
 		/// Read all of the plugins found in the project specified additional plugin directories
 		/// </summary>
 		/// <param name="AdditionalDirectory">The additional directory to scan</param>
+		/// <param name="Logger">Logger for output</param>
 		/// <returns>List of the found PluginInfo objects</returns>
-		public static IReadOnlyList<PluginInfo> ReadAdditionalPlugins(DirectoryReference AdditionalDirectory)
+		public static IReadOnlyList<PluginInfo> ReadAdditionalPlugins(DirectoryReference AdditionalDirectory, ILogger Logger)
 		{
 			DirectoryReference FullPath = DirectoryReference.Combine(AdditionalDirectory, "");
 			if (!DirectoryReference.Exists(FullPath))
 			{
-				Log.TraceWarning("AdditionalPluginDirectory {0} not found. Path should be relative to the project", FullPath);
+				Logger.LogWarning("AdditionalPluginDirectory {FullPath} not found. Path should be relative to the project", FullPath);
 			}
 			return ReadPluginsFromDirectory(AdditionalDirectory, "", PluginType.External);
 		}
diff --git a/Engine/Source/Programs/UnrealBuildTool/System/ProgressWriter.cs b/Engine/Source/Programs/UnrealBuildTool/System/ProgressWriter.cs
--- a/Engine/Source/Programs/UnrealBuildTool/System/ProgressWriter.cs
+++ b/Engine/Source/Programs/UnrealBuildTool/System/ProgressWriter.cs
@@ -6,6 +6,7 @@
 using System.Text;
 using System.Threading.Tasks;
 using EpicGames.Core;
+using Microsoft.Extensions.Logging;
 
 namespace UnrealBuildTool
 {
@@ -20,6 +21,11 @@
 		public static bool bWriteMarkup = false;
 
 		/// <summary>
+		/// Logger for output
+		/// </summary>
+		ILogger Logger;
+
+		/// <summary>
 		/// The name to include with the status message
 		/// </summary>
 		string Message;
@@ -39,9 +45,11 @@
 		/// </summary>
 		/// <param name="InMessage">The message to display before the progress percentage</param>
 		/// <param name="bInUpdateStatus">Whether to write messages to the console</param>
-		public ProgressWriter(string InMessage, bool bInUpdateStatus)
+		/// <param name="InLogger">Logger for output</param>
+		public ProgressWriter(string InMessage, bool bInUpdateStatus, ILogger InLogger)
 		{
 			Message = InMessage;
+			Logger = InLogger;
 			if(bInUpdateStatus)
 			{
 				Status = new LogStatusScope(InMessage);
@@ -75,7 +83,7 @@
 				CurrentProgressString = ProgressString;
 				if (bWriteMarkup)
 				{
-					Log.WriteLine(LogEventType.Console, "@progress '{0}' {1}", Message, ProgressString);
+					Logger.LogInformation("@progress '{Message}' {ProgressString}", Message, ProgressString);
 				}
 				if(Status != null)
 				{
diff --git a/Engine/Source/Programs/UnrealBuildTool/System/RulesAssembly.cs b/Engine/Source/Programs/UnrealBuildTool/System/RulesAssembly.cs
--- a/Engine/Source/Programs/UnrealBuildTool/System/RulesAssembly.cs
+++ b/Engine/Source/Programs/UnrealBuildTool/System/RulesAssembly.cs
@@ -10,6 +10,7 @@
 using System.Text;
 using System.Threading.Tasks;
 using EpicGames.Core;
+using Microsoft.Extensions.Logging;
 using UnrealBuildBase;
 
 namespace UnrealBuildTool
@@ -113,7 +114,8 @@
 		/// <param name="bSkipCompile">Whether to skip compiling this assembly</param>
 		/// <param name="bForceCompile">Whether to always compile this assembly</param>
 		/// <param name="Parent">The parent rules assembly</param>
-		internal RulesAssembly(RulesScope Scope, List<DirectoryReference> BaseDirs, IReadOnlyList<PluginInfo> Plugins, Dictionary<FileReference, ModuleRulesContext> ModuleFileToContext, List<FileReference> TargetFiles, FileReference AssemblyFileName, bool bContainsEngineModules, BuildSettingsVersion? DefaultBuildSettings, bool bReadOnly, bool bSkipCompile, bool bForceCompile, RulesAssembly? Parent)
+		/// <param name="Logger"></param>
+		internal RulesAssembly(RulesScope Scope, List<DirectoryReference> BaseDirs, IReadOnlyList<PluginInfo> Plugins, Dictionary<FileReference, ModuleRulesContext> ModuleFileToContext, List<FileReference> TargetFiles, FileReference AssemblyFileName, bool bContainsEngineModules, BuildSettingsVersion? DefaultBuildSettings, bool bReadOnly, bool bSkipCompile, bool bForceCompile, RulesAssembly? Parent, ILogger Logger)
 		{
 			this.Scope = Scope;
 			this.BaseDirs = BaseDirs;
@@ -133,7 +135,7 @@
 			if (AssemblySourceFiles.Count > 0)
 			{
 				PreprocessorDefines = GetPreprocessorDefinitions();
-				CompiledAssembly = DynamicCompilation.CompileAndLoadAssembly(AssemblyFileName, AssemblySourceFiles, PreprocessorDefines: PreprocessorDefines, DoNotCompile: bSkipCompile, ForceCompile: bForceCompile);
+				CompiledAssembly = DynamicCompilation.CompileAndLoadAssembly(AssemblyFileName, AssemblySourceFiles, Logger, PreprocessorDefines: PreprocessorDefines, DoNotCompile: bSkipCompile, ForceCompile: bForceCompile);
 			}
 
 			// Setup the module map
@@ -172,7 +174,7 @@
 							{
 								Location = new FileReference(CompiledAssembly.Location);
 							}
-							Log.TraceWarning("{0}: warning: {1}", Location, Attribute.Message);
+							Logger.LogWarning("{Location}: warning: {AttributeMessage}", Location, Attribute.Message);
 						}
 					}
 					if(CompiledType.BaseType == typeof(ModuleRules))
@@ -185,7 +187,7 @@
 							{
 								Location = new FileReference(CompiledAssembly.Location);
 							}
-							Log.TraceWarning("{0}: warning: Module constructors should take a ReadOnlyTargetRules argument (rather than a TargetInfo argument) and pass it to the base class constructor from 4.15 onwards. Please update the method signature.", Location);
+							Logger.LogWarning("{Location}: warning: Module constructors should take a ReadOnlyTargetRules argument (rather than a TargetInfo argument) and pass it to the base class constructor from 4.15 onwards. Please update the method signature.", Location);
 						}
 					}
 				}
@@ -385,8 +387,9 @@
 		/// <param name="ModuleName">Name of the module</param>
 		/// <param name="Target">Information about the target associated with this module</param>
 		/// <param name="ReferenceChain">Chain of references leading to this module</param>
+		/// <param name="Logger">Logger for output</param>
 		/// <returns>Compiled module rule info</returns>
-		public ModuleRules CreateModuleRules(string ModuleName, ReadOnlyTargetRules Target, string ReferenceChain)
+		public ModuleRules CreateModuleRules(string ModuleName, ReadOnlyTargetRules Target, string ReferenceChain, ILogger Logger)
 		{
 			if (Target.IsTestTarget && !Target.ExplicitTestsTarget)
 			{
@@ -406,7 +409,7 @@
 				}
 				else
 				{
-					return Parent.CreateModuleRules(ModuleName, Target, ReferenceChain);
+					return Parent.CreateModuleRules(ModuleName, Target, ReferenceChain, Logger);
 				}
 			}
 
@@ -440,7 +443,7 @@
 			if (Target.OptedInModulePlatforms != null)
 			{
 				// figure out what platforms/groups aren't allowed with this opted in list
-				List<string> DisallowedPlatformsAndGroups = Utils.MakeListOfUnsupportedPlatforms(Target.OptedInModulePlatforms.ToList(), false);
+				List<string> DisallowedPlatformsAndGroups = Utils.MakeListOfUnsupportedPlatforms(Target.OptedInModulePlatforms.ToList(), false, Logger);
 
 				// check if the module file is disallowed
 				if (ModuleFileName.ContainsAnyNames(DisallowedPlatformsAndGroups, Unreal.EngineDirectory) ||
@@ -536,9 +539,10 @@
 		/// </summary>
 		/// <param name="TypeName">Type name of the target rules</param>
 		/// <param name="TargetInfo">Target configuration information to pass to the constructor</param>
+		/// <param name="Logger">Logger for output</param>
 		/// <param name="IsTestTarget">If building a low level tests target</param>
 		/// <returns>Instance of the corresponding TargetRules</returns>
-		protected TargetRules CreateTargetRulesInstance(string TypeName, TargetInfo TargetInfo, bool IsTestTarget = false)
+		protected TargetRules CreateTargetRulesInstance(string TypeName, TargetInfo TargetInfo, ILogger Logger, bool IsTestTarget = false)
 		{
 			// The build module must define a type named '<TargetName>Target' that derives from our 'TargetRules' type.  
 			Type? BaseRulesType = CompiledAssembly?.GetType(TypeName);
@@ -591,6 +595,9 @@
 			// The platform/group-specific target file name
 			Rules.TargetSourceFile = TargetNameToTargetFile.TryGetValue(PlatformRulesName, out FileReference? PlatformTargetFile) ? PlatformTargetFile : Rules.File;
 
+			// Initialize the logger
+			Rules.Logger = Logger;
+
 			// Find the constructor
 			ConstructorInfo? Constructor = RulesType.GetConstructor(new Type[] { typeof(TargetInfo) });
 			if(Constructor == null)
@@ -670,14 +677,14 @@
 			if (Rules.OptedInModulePlatforms != null)
 			{
 				// figure out what platforms/groups aren't allowed with this opted in list
-				List<string> DisallowedPlatformsAndGroups = Utils.MakeListOfUnsupportedPlatforms(Rules.OptedInModulePlatforms.ToList(), false);
+				List<string> DisallowedPlatformsAndGroups = Utils.MakeListOfUnsupportedPlatforms(Rules.OptedInModulePlatforms.ToList(), false, Logger);
 
 				// look in all plugins' paths to see if any disallowed
 				IEnumerable<PluginInfo> DisallowedPlugins = EnumeratePlugins().Where(Plugin =>
 					Plugin.File.ContainsAnyNames(DisallowedPlatformsAndGroups, Unreal.EngineDirectory) ||
 					(Rules.ProjectFile != null && Plugin.File.ContainsAnyNames(DisallowedPlatformsAndGroups, Rules.ProjectFile.Directory)));
 				// log out the plugins we are disabling
-				DisallowedPlugins.ToList().ForEach(x => Log.TraceLog($"Disallowing non-opted-in platform plugin {x.File}"));
+				DisallowedPlugins.ToList().ForEach(x => Logger.LogDebug("Disallowing non-opted-in platform plugin {PluginFile}", x.File));
 				// and, disable these plugins
 				Rules.DisablePlugins.AddRange(DisallowedPlugins.Select(x => x.Name));
 			}
@@ -709,9 +716,10 @@
 		/// <param name="Architecture">Architecture being built</param>
 		/// <param name="ProjectFile">Path to the project file for this target</param>
 		/// <param name="Arguments">Command line arguments for this target</param>
+		/// <param name="Logger"></param>
 		/// <param name="IsTestTarget">If building a low level test target</param>
 		/// <returns>The build target rules for the specified target</returns>
-		public TargetRules CreateTargetRules(string TargetName, UnrealTargetPlatform Platform, UnrealTargetConfiguration Configuration, string Architecture, FileReference? ProjectFile, CommandLineArguments? Arguments, bool IsTestTarget = false)
+		public TargetRules CreateTargetRules(string TargetName, UnrealTargetPlatform Platform, UnrealTargetConfiguration Configuration, string Architecture, FileReference? ProjectFile, CommandLineArguments? Arguments, ILogger Logger, bool IsTestTarget = false)
 		{
 			if (IsTestTarget)
 			{
@@ -741,7 +749,7 @@
 				}
 				else
 				{
-					return Parent.CreateTargetRules(TargetName, Platform, Configuration, Architecture, ProjectFile, Arguments, IsTestTarget);
+					return Parent.CreateTargetRules(TargetName, Platform, Configuration, Architecture, ProjectFile, Arguments, Logger, IsTestTarget);
 				}
 			}
 
@@ -749,7 +757,7 @@
 			string TargetTypeName = TargetName + "Target";
 
 			// The build module must define a type named '<TargetName>Target' that derives from our 'TargetRules' type.  
-			return CreateTargetRulesInstance(TargetTypeName, new TargetInfo(TargetName, Platform, Configuration, Architecture, ProjectFile, Arguments), IsTestTarget);
+			return CreateTargetRulesInstance(TargetTypeName, new TargetInfo(TargetName, Platform, Configuration, Architecture, ProjectFile, Arguments), Logger, IsTestTarget);
 		}
 
 		/// <summary>
@@ -760,14 +768,15 @@
 		/// <param name="Configuration">The configuration being built</param>
 		/// <param name="Architecture">The architecture being built</param>
 		/// <param name="ProjectFile">Project file for the target being built</param>
+		/// <param name="Logger">Logger for output</param>
 		/// <returns>Name of the target for the given type</returns>
-		public string GetTargetNameByType(TargetType Type, UnrealTargetPlatform Platform, UnrealTargetConfiguration Configuration, string Architecture, FileReference? ProjectFile)
+		public string GetTargetNameByType(TargetType Type, UnrealTargetPlatform Platform, UnrealTargetConfiguration Configuration, string Architecture, FileReference? ProjectFile, ILogger Logger)
 		{
 			// Create all the targets in this assembly 
 			List<string> Matches = new List<string>();
 			foreach(KeyValuePair<string, FileReference> TargetPair in TargetNameToTargetFile)
 			{
-				TargetRules Rules = CreateTargetRulesInstance(TargetPair.Key + "Target", new TargetInfo(TargetPair.Key, Platform, Configuration, Architecture, ProjectFile, null));
+				TargetRules Rules = CreateTargetRulesInstance(TargetPair.Key + "Target", new TargetInfo(TargetPair.Key, Platform, Configuration, Architecture, ProjectFile, null), Logger);
 				if(Rules.Type == Type)
 				{
 					Matches.Add(TargetPair.Key);
@@ -783,7 +792,7 @@
 				}
 				else
 				{
-					return Parent.GetTargetNameByType(Type, Platform, Configuration, Architecture, ProjectFile);
+					return Parent.GetTargetNameByType(Type, Platform, Configuration, Architecture, ProjectFile, Logger);
 				}
 			}
 			else
diff --git a/Engine/Source/Programs/UnrealBuildTool/System/RulesCompiler.cs b/Engine/Source/Programs/UnrealBuildTool/System/RulesCompiler.cs
--- a/Engine/Source/Programs/UnrealBuildTool/System/RulesCompiler.cs
+++ b/Engine/Source/Programs/UnrealBuildTool/System/RulesCompiler.cs
@@ -8,6 +8,7 @@
 using EpicGames.Core;
 using OpenTracing.Util;
 using UnrealBuildBase;
+using Microsoft.Extensions.Logging;
 
 namespace UnrealBuildTool
 {
@@ -71,8 +72,9 @@
 		/// <param name="bUsePrecompiled">Whether to use a precompiled engine</param>
 		/// <param name="bSkipCompile">Whether to skip compilation for this assembly</param>
 		/// <param name="bForceCompile">Whether to always compile this assembly</param>
+		/// <param name="Logger">Logger for output</param>
 		/// <returns>New rules assembly</returns>
-		public static RulesAssembly CreateEngineRulesAssembly(bool bUsePrecompiled, bool bSkipCompile, bool bForceCompile)
+		public static RulesAssembly CreateEngineRulesAssembly(bool bUsePrecompiled, bool bSkipCompile, bool bForceCompile, ILogger Logger)
 		{
 			if (EngineRulesAssembly == null)
 			{
@@ -92,11 +94,11 @@
 					}
 				}
 
-				EngineRulesAssembly = CreateEngineRulesAssemblyInternal(Unreal.GetExtensionDirs(Unreal.EngineDirectory), ProjectFileGenerator.EngineProjectFileNameBase, EnginePlugins, Unreal.IsEngineInstalled() || bUsePrecompiled, bSkipCompile, bForceCompile, null);
+				EngineRulesAssembly = CreateEngineRulesAssemblyInternal(Unreal.GetExtensionDirs(Unreal.EngineDirectory), ProjectFileGenerator.EngineProjectFileNameBase, EnginePlugins, Unreal.IsEngineInstalled() || bUsePrecompiled, bSkipCompile, bForceCompile, null, Logger);
 
 				if (MarketplacePlugins.Count > 0)
 				{
-					EngineRulesAssembly = CreateMarketplaceRulesAssembly(MarketplacePlugins, bSkipCompile, bForceCompile, EngineRulesAssembly);
+					EngineRulesAssembly = CreateMarketplaceRulesAssembly(MarketplacePlugins, bSkipCompile, bForceCompile, EngineRulesAssembly, Logger);
 				}
 			}
 			return EngineRulesAssembly;
@@ -112,8 +114,9 @@
 		/// <param name="bSkipCompile">Whether to skip compilation for this assembly</param>
 		/// <param name="bForceCompile">Whether to always compile this assembly</param>
 		/// <param name="Parent">The parent rules assembly</param>
+		/// <param name="Logger">Logger for output</param>
 		/// <returns>New rules assembly</returns>
-		private static RulesAssembly CreateEngineRulesAssemblyInternal(List<DirectoryReference> RootDirectories, string AssemblyPrefix, IReadOnlyList<PluginInfo> Plugins, bool bReadOnly, bool bSkipCompile, bool bForceCompile, RulesAssembly? Parent)
+		private static RulesAssembly CreateEngineRulesAssemblyInternal(List<DirectoryReference> RootDirectories, string AssemblyPrefix, IReadOnlyList<PluginInfo> Plugins, bool bReadOnly, bool bSkipCompile, bool bForceCompile, RulesAssembly? Parent, ILogger Logger)
 		{
 			// Scope hierarchy
 			RulesScope Scope= new RulesScope("Engine", null);
@@ -147,7 +150,7 @@
 			// Create the assembly
 			DirectoryReference AssemblyDir = RootDirectories[0];
 			FileReference EngineAssemblyFileName = FileReference.Combine(AssemblyDir, "Intermediate", "Build", "BuildRules", AssemblyPrefix + "Rules" + FrameworkAssemblyExtension);
-			RulesAssembly EngineAssembly = new RulesAssembly(Scope, RootDirectories, Plugins, ModuleFileToContext, new List<FileReference>(), EngineAssemblyFileName, bContainsEngineModules: true, DefaultBuildSettings: BuildSettingsVersion.Latest, bReadOnly: bReadOnly, bSkipCompile: bSkipCompile, bForceCompile: bForceCompile, Parent: Parent);
+			RulesAssembly EngineAssembly = new RulesAssembly(Scope, RootDirectories, Plugins, ModuleFileToContext, new List<FileReference>(), EngineAssemblyFileName, bContainsEngineModules: true, DefaultBuildSettings: BuildSettingsVersion.Latest, bReadOnly: bReadOnly, bSkipCompile: bSkipCompile, bForceCompile: bForceCompile, Parent: Parent, Logger: Logger);
 
 			List<FileReference> ProgramTargetFiles = new List<FileReference>();
 			Dictionary<FileReference, ModuleRulesContext> ProgramModuleFiles = new Dictionary<FileReference, ModuleRulesContext>();
@@ -174,7 +177,7 @@
 
 			// Create a path to the assembly that we'll either load or compile
 			FileReference ProgramAssemblyFileName = FileReference.Combine(AssemblyDir, "Intermediate", "Build", "BuildRules", AssemblyPrefix + "ProgramRules" + FrameworkAssemblyExtension);
-			RulesAssembly ProgramAssembly = new RulesAssembly(ProgramsScope, RootDirectories, new List<PluginInfo>().AsReadOnly(), ProgramModuleFiles, ProgramTargetFiles, ProgramAssemblyFileName, bContainsEngineModules: false, DefaultBuildSettings: BuildSettingsVersion.Latest, bReadOnly: bReadOnly, bSkipCompile: bSkipCompile, bForceCompile: bForceCompile, Parent: EngineAssembly);
+			RulesAssembly ProgramAssembly = new RulesAssembly(ProgramsScope, RootDirectories, new List<PluginInfo>().AsReadOnly(), ProgramModuleFiles, ProgramTargetFiles, ProgramAssemblyFileName, bContainsEngineModules: false, DefaultBuildSettings: BuildSettingsVersion.Latest, bReadOnly: bReadOnly, bSkipCompile: bSkipCompile, bForceCompile: bForceCompile, Parent: EngineAssembly, Logger: Logger);
 
 			// Return the combined assembly
 			return ProgramAssembly;
@@ -187,8 +190,9 @@
 		/// <param name="bSkipCompile">Whether to skip compilation for this assembly</param>
 		/// <param name="bForceCompile">Whether to always compile this assembly</param>
 		/// <param name="Parent">The parent rules assembly</param>
+		/// <param name="Logger"></param>
 		/// <returns>New rules assembly</returns>
-		private static RulesAssembly CreateMarketplaceRulesAssembly(IReadOnlyList<PluginInfo> Plugins, bool bSkipCompile, bool bForceCompile, RulesAssembly Parent)
+		private static RulesAssembly CreateMarketplaceRulesAssembly(IReadOnlyList<PluginInfo> Plugins, bool bSkipCompile, bool bForceCompile, RulesAssembly Parent, ILogger Logger)
 		{
 			RulesScope MarketplaceScope = new RulesScope("Marketplace", Parent.Scope);
 
@@ -205,7 +209,7 @@
 			if (ModuleFileToContext.Count > 0)
 			{
 				FileReference AssemblyFileName = FileReference.Combine(UnrealBuildTool.WritableEngineDirectory, "Intermediate", "Build", "BuildRules", "MarketplaceRules.dll");
-				Result = new RulesAssembly(MarketplaceScope, new List<DirectoryReference> { DirectoryReference.Combine(Unreal.EngineDirectory, "Plugins", "Marketplace") }, Plugins, ModuleFileToContext, new List<FileReference>(), AssemblyFileName, bContainsEngineModules: true, DefaultBuildSettings: BuildSettingsVersion.Latest, bReadOnly: true, bSkipCompile: bSkipCompile, bForceCompile: bForceCompile, Parent: Parent);
+				Result = new RulesAssembly(MarketplaceScope, new List<DirectoryReference> { DirectoryReference.Combine(Unreal.EngineDirectory, "Plugins", "Marketplace") }, Plugins, ModuleFileToContext, new List<FileReference>(), AssemblyFileName, bContainsEngineModules: true, DefaultBuildSettings: BuildSettingsVersion.Latest, bReadOnly: true, bSkipCompile: bSkipCompile, bForceCompile: bForceCompile, Parent: Parent, Logger: Logger);
 			}
 			return Result;
 		}
@@ -217,8 +221,9 @@
 		/// <param name="bUsePrecompiled">Whether to use a precompiled engine</param>
 		/// <param name="bSkipCompile">Whether to skip compilation for this assembly</param>
 		/// <param name="bForceCompile">Whether to always compile this assembly</param>
+		/// <param name="Logger">Logger for output</param>
 		/// <returns>New rules assembly</returns>
-		public static RulesAssembly CreateProjectRulesAssembly(FileReference ProjectFileName, bool bUsePrecompiled, bool bSkipCompile, bool bForceCompile)
+		public static RulesAssembly CreateProjectRulesAssembly(FileReference ProjectFileName, bool bUsePrecompiled, bool bSkipCompile, bool bForceCompile, ILogger Logger)
 		{
 			// Check if there's an existing assembly for this project
 			RulesAssembly? ProjectRulesAssembly;
@@ -227,7 +232,7 @@
 				ProjectDescriptor Project = ProjectDescriptor.FromFile(ProjectFileName);
 
 				// Create the parent assembly
-				RulesAssembly Parent = CreateEngineRulesAssembly(bUsePrecompiled, bSkipCompile, bForceCompile);
+				RulesAssembly Parent = CreateEngineRulesAssembly(bUsePrecompiled, bSkipCompile, bForceCompile, Logger);
 
 				DirectoryReference MainProjectDirectory = ProjectFileName.Directory;
 				//DirectoryReference MainProjectSourceDirectory = DirectoryReference.Combine(MainProjectDirectory, "Source");
@@ -271,7 +276,7 @@
 				{
 					foreach (DirectoryReference AdditionalPluginDirectory in Project.AdditionalPluginDirectories)
 					{
-						ProjectPlugins.AddRange(Plugins.ReadAdditionalPlugins(AdditionalPluginDirectory));
+						ProjectPlugins.AddRange(Plugins.ReadAdditionalPlugins(AdditionalPluginDirectory, Logger));
 					}
 				}
 
@@ -294,7 +299,7 @@
 				}
 				else
 				{
-					ProjectRulesAssembly = new RulesAssembly(Scope, new List<DirectoryReference> { MainProjectDirectory }, ProjectPlugins, ModuleFiles, TargetFiles, AssemblyFileName, bContainsEngineModules: false, DefaultBuildSettings: null, bReadOnly: UnrealBuildTool.IsProjectInstalled(), bSkipCompile: bSkipCompile, bForceCompile: bForceCompile, Parent: Parent);
+					ProjectRulesAssembly = new RulesAssembly(Scope, new List<DirectoryReference> { MainProjectDirectory }, ProjectPlugins, ModuleFiles, TargetFiles, AssemblyFileName, bContainsEngineModules: false, DefaultBuildSettings: null, bReadOnly: UnrealBuildTool.IsProjectInstalled(), bSkipCompile: bSkipCompile, bForceCompile: bForceCompile, Parent: Parent, Logger: Logger);
 				}
 				LoadedAssemblyMap.Add(ProjectFileName, ProjectRulesAssembly);
 			}
@@ -309,8 +314,9 @@
 		/// <param name="bForceCompile">Whether to always compile this assembly</param>
 		/// <param name="Parent">The parent rules assembly</param>
         /// <param name="bContainsEngineModules">Whether the plugin contains engine modules. Used to initialize the default value for ModuleRules.bTreatAsEngineModule.</param>
+		/// <param name="Logger">Logger for ouptut</param>
 		/// <returns>The new rules assembly</returns>
-        public static RulesAssembly CreatePluginRulesAssembly(FileReference PluginFileName, bool bSkipCompile, bool bForceCompile, RulesAssembly Parent, bool bContainsEngineModules)
+        public static RulesAssembly CreatePluginRulesAssembly(FileReference PluginFileName, bool bSkipCompile, bool bForceCompile, RulesAssembly Parent, bool bContainsEngineModules, ILogger Logger)
 		{
 			// Check if there's an existing assembly for this project
 			RulesAssembly? PluginRulesAssembly;
@@ -342,7 +348,7 @@
 
 				// Compile the assembly
 				FileReference AssemblyFileName = FileReference.Combine(PluginFileName.Directory, "Intermediate", "Build", "BuildRules", Path.GetFileNameWithoutExtension(PluginFileName.FullName) + "ModuleRules" + FrameworkAssemblyExtension);
-				PluginRulesAssembly = new RulesAssembly(Scope, new List<DirectoryReference> { PluginFileName.Directory }, ForeignPlugins, ModuleFiles, TargetFiles, AssemblyFileName, bContainsEngineModules, DefaultBuildSettings: null, bReadOnly: false, bSkipCompile: bSkipCompile, bForceCompile: bForceCompile, Parent: Parent);
+				PluginRulesAssembly = new RulesAssembly(Scope, new List<DirectoryReference> { PluginFileName.Directory }, ForeignPlugins, ModuleFiles, TargetFiles, AssemblyFileName, bContainsEngineModules, DefaultBuildSettings: null, bReadOnly: false, bSkipCompile: bSkipCompile, bForceCompile: bForceCompile, Parent: Parent, Logger: Logger);
 				LoadedAssemblyMap.Add(PluginFileName, PluginRulesAssembly);
 			}
 			return PluginRulesAssembly;
@@ -357,21 +363,22 @@
 		/// <param name="bForceRulesCompile">Whether to always compile all rules assemblies</param>
 		/// <param name="bUsePrecompiled">Whether to use a precompiled engine build</param>
 		/// <param name="ForeignPlugin">Foreign plugin to be compiled</param>
+		/// <param name="Logger">Logger for output</param>
 		/// <returns>The compiled rules assembly</returns>
-		public static RulesAssembly CreateTargetRulesAssembly(FileReference? ProjectFile, string TargetName, bool bSkipRulesCompile, bool bForceRulesCompile, bool bUsePrecompiled, FileReference? ForeignPlugin)
+		public static RulesAssembly CreateTargetRulesAssembly(FileReference? ProjectFile, string TargetName, bool bSkipRulesCompile, bool bForceRulesCompile, bool bUsePrecompiled, FileReference? ForeignPlugin, ILogger Logger)
 		{
 			RulesAssembly RulesAssembly;
 			if (ProjectFile != null)
 			{
-				RulesAssembly = CreateProjectRulesAssembly(ProjectFile, bUsePrecompiled, bSkipRulesCompile, bForceRulesCompile);
+				RulesAssembly = CreateProjectRulesAssembly(ProjectFile, bUsePrecompiled, bSkipRulesCompile, bForceRulesCompile, Logger);
 			}
 			else
 			{
-				RulesAssembly = CreateEngineRulesAssembly(bUsePrecompiled, bSkipRulesCompile, bForceRulesCompile);
+				RulesAssembly = CreateEngineRulesAssembly(bUsePrecompiled, bSkipRulesCompile, bForceRulesCompile, Logger);
 			}
 			if (ForeignPlugin != null)
 			{
-				RulesAssembly = CreatePluginRulesAssembly(ForeignPlugin, bSkipRulesCompile, bForceRulesCompile, RulesAssembly, true);
+				RulesAssembly = CreatePluginRulesAssembly(ForeignPlugin, bSkipRulesCompile, bForceRulesCompile, RulesAssembly, true, Logger);
 			}
 			return RulesAssembly;
 		}
diff --git a/Engine/Source/Programs/UnrealBuildTool/System/RulesDocumentation.cs b/Engine/Source/Programs/UnrealBuildTool/System/RulesDocumentation.cs
--- a/Engine/Source/Programs/UnrealBuildTool/System/RulesDocumentation.cs
+++ b/Engine/Source/Programs/UnrealBuildTool/System/RulesDocumentation.cs
@@ -12,6 +12,7 @@
 using System.Xml;
 using EpicGames.Core;
 using UnrealBuildBase;
+using Microsoft.Extensions.Logging;
 
 namespace UnrealBuildTool
 {
@@ -31,7 +32,7 @@
 			}
 		}
 
-		public static void WriteDocumentation(Type RulesType, FileReference OutputFile)
+		public static void WriteDocumentation(Type RulesType, FileReference OutputFile, ILogger Logger)
 		{
 			// Get the path to the XML documentation
 			FileReference InputDocumentationFile = new FileReference(Assembly.GetExecutingAssembly().Location).ChangeExtension(".xml");
@@ -54,7 +55,7 @@
 				if(!Field.FieldType.IsClass || !Field.FieldType.Name.EndsWith("TargetRules"))
 				{
 					List<string>? Lines;
-					if (TryGetXmlComment(InputDocumentation, Field, out Lines))
+					if (TryGetXmlComment(InputDocumentation, Field, Logger, out Lines))
 					{
 						SettingInfo Setting = new SettingInfo(Field.Name, Field.FieldType, Lines);
 						if (Field.IsInitOnly)
@@ -75,7 +76,7 @@
 				if (!Property.PropertyType.IsClass || !Property.PropertyType.Name.EndsWith("TargetRules"))
 				{
 					List<string>? Lines;
-					if (TryGetXmlComment(InputDocumentation, Property, out Lines))
+					if (TryGetXmlComment(InputDocumentation, Property, Logger, out Lines))
 					{
 						SettingInfo Setting = new SettingInfo(Property.Name, Property.PropertyType, Lines);
 						if (Property.SetMethod == null)
@@ -115,7 +116,7 @@
 			}
 
 			// Success!
-			Log.TraceInformation("Written documentation to {0}.", OutputFile);
+			Logger.LogInformation("Written documentation to {OutputFile}.", OutputFile);
 		}
 
 		static void WriteDocumentationUDN(FileReference OutputFile, List<SettingInfo> ReadOnlySettings, List<SettingInfo> ReadWriteSettings)
@@ -156,11 +157,12 @@
 		/// </summary>
 		/// <param name="Documentation">The XML documentation</param>
 		/// <param name="Field">The member to search for</param>
+		/// <param name="Logger">Logger for output</param>
 		/// <param name="Lines">Receives the description for the requested field</param>
 		/// <returns>True if a comment was found for the field</returns>
-		public static bool TryGetXmlComment(XmlDocument Documentation, FieldInfo Field, [NotNullWhen(true)] out List<string>? Lines)
+		public static bool TryGetXmlComment(XmlDocument Documentation, FieldInfo Field, ILogger Logger, [NotNullWhen(true)] out List<string>? Lines)
 		{
-			return TryGetXmlComment(Documentation, $"F:{Field.DeclaringType}.{Field.Name}", out Lines);
+			return TryGetXmlComment(Documentation, $"F:{Field.DeclaringType}.{Field.Name}", Logger, out Lines);
 		}
 
 		/// <summary>
@@ -168,11 +170,12 @@
 		/// </summary>
 		/// <param name="Documentation">The XML documentation</param>
 		/// <param name="Property">The member to search for</param>
+		/// <param name="Logger">Logger for output</param>
 		/// <param name="Lines">Receives the description for the requested field</param>
 		/// <returns>True if a comment was found for the field</returns>
-		public static bool TryGetXmlComment(XmlDocument Documentation, PropertyInfo Property, [NotNullWhen(true)] out List<string>? Lines)
+		public static bool TryGetXmlComment(XmlDocument Documentation, PropertyInfo Property, ILogger Logger, [NotNullWhen(true)] out List<string>? Lines)
 		{
-			return TryGetXmlComment(Documentation, $"P:{Property.DeclaringType}.{Property.Name}", out Lines);
+			return TryGetXmlComment(Documentation, $"P:{Property.DeclaringType}.{Property.Name}", Logger, out Lines);
 		}
 
 		/// <summary>
@@ -180,17 +183,18 @@
 		/// </summary>
 		/// <param name="Documentation">The XML documentation</param>
 		/// <param name="Member">The member to search for (field or property)</param>
+		/// <param name="Logger">Logger for output</param>
 		/// <param name="Lines">Receives the description for the requested member</param>
 		/// <returns>True if a comment was found for the member</returns>
-		public static bool TryGetXmlComment(XmlDocument Documentation, MemberInfo Member, [NotNullWhen(true)] out List<string>? Lines)
+		public static bool TryGetXmlComment(XmlDocument Documentation, MemberInfo Member, ILogger Logger, [NotNullWhen(true)] out List<string>? Lines)
 		{
 			if (Member is FieldInfo)
 			{
-				return TryGetXmlComment(Documentation, (Member as FieldInfo)!, out Lines);
+				return TryGetXmlComment(Documentation, (Member as FieldInfo)!, Logger, out Lines);
 			}
 			else if (Member is PropertyInfo)
 			{
-				return TryGetXmlComment(Documentation, (Member as PropertyInfo)!, out Lines);
+				return TryGetXmlComment(Documentation, (Member as PropertyInfo)!, Logger, out Lines);
 			}
 
 			Lines = null;
@@ -203,9 +207,10 @@
 		/// </summary>
 		/// <param name="Documentation">The XML documentation</param>
 		/// <param name="MemberName">The member to search for</param>
+		/// <param name="Logger">Logger for output</param>
 		/// <param name="Lines">Receives the description for the requested field</param>
 		/// <returns>True if a comment was found for the field</returns>
-		public static bool TryGetXmlComment(XmlDocument Documentation, string MemberName, [NotNullWhen(true)] out List<string>? Lines)
+		public static bool TryGetXmlComment(XmlDocument Documentation, string MemberName, ILogger Logger, [NotNullWhen(true)] out List<string>? Lines)
 		{
 			HashSet<string> VisitedProperties = new HashSet<string>(StringComparer.Ordinal);
 
@@ -261,7 +266,7 @@
 
 			if (!bExclude)
 			{
-				Log.TraceWarning("Missing XML comment for {0}", Regex.Replace(MemberName, @"^[A-Z]:", ""));
+				Logger.LogWarning("Missing XML comment for {Member}", Regex.Replace(MemberName, @"^[A-Z]:", ""));
 			}
 			Lines = null;
 			return false;
diff --git a/Engine/Source/Programs/UnrealBuildTool/System/ScopedTimer.cs b/Engine/Source/Programs/UnrealBuildTool/System/ScopedTimer.cs
--- a/Engine/Source/Programs/UnrealBuildTool/System/ScopedTimer.cs
+++ b/Engine/Source/Programs/UnrealBuildTool/System/ScopedTimer.cs
@@ -6,6 +6,7 @@
 using System.Text;
 using System.Threading.Tasks;
 using EpicGames.Core;
+using Microsoft.Extensions.Logging;
 
 namespace UnrealBuildTool
 {
@@ -15,8 +16,9 @@
 	public class ScopedTimer : IDisposable
 	{
 		DateTime StartTime;
-		string TimerName;
-		LogEventType Verbosity;
+		string Name;
+		ILogger Logger;
+		LogLevel Level;
 		bool bIncreaseIndent;
 		static int Indent = 0;
 		static object IndentLock = new object();
@@ -25,11 +27,13 @@
 		/// Constructor
 		/// </summary>
 		/// <param name="Name">Name of the block being measured</param>
-		/// <param name="InVerbosity">Verbosity for output messages</param>
+		/// <param name="Logger">Logger for output</param>
+		/// <param name="InLevel">Verbosity for output messages</param>
 		/// <param name="bIncreaseIndent">Whether gobal indent should be increased or not; set to false when running a scope in parallel. Message will still be printed indented relative to parent scope.</param>
-		public ScopedTimer(string Name, LogEventType InVerbosity = LogEventType.Verbose, bool bIncreaseIndent = true)
+		public ScopedTimer(string Name, ILogger Logger, LogLevel InLevel = LogLevel.Debug, bool bIncreaseIndent = true)
 		{
-			TimerName = Name;
+			this.Name = Name;
+			this.Logger = Logger;
 			if (bIncreaseIndent)
 			{
 				lock (IndentLock)
@@ -37,7 +41,7 @@
 					Indent++;
 				}
 			}
-			Verbosity = InVerbosity;
+			Level = InLevel;
 			StartTime = DateTime.UtcNow;
 			this.bIncreaseIndent = bIncreaseIndent;
 		}
@@ -59,7 +63,7 @@
 			StringBuilder IndentText = new StringBuilder(LogIndent * 2);
 			IndentText.Append(' ', LogIndent * 2);
 
-			Log.WriteLine(Verbosity, "{0}{1} took {2}s", IndentText.ToString(), TimerName, TotalSeconds);
+			Logger.Log(Level, "{Indent}{Name} took {TimeSeconds}s", IndentText.ToString(), Name, TotalSeconds);
 		}
 	}
 }
