UBT: Use ILogger throughout UBT, rather than legacy EpicGames.Core.Log methods. All output going forwards should use structured logging rather than String.Format style logging.

diff --git a/Engine/Source/Programs/UnrealBuildTool/System/SourceFileMetadataCache.cs b/Engine/Source/Programs/UnrealBuildTool/System/SourceFileMetadataCache.cs
--- a/Engine/Source/Programs/UnrealBuildTool/System/SourceFileMetadataCache.cs
+++ b/Engine/Source/Programs/UnrealBuildTool/System/SourceFileMetadataCache.cs
@@ -10,6 +10,7 @@
 using System.Text.RegularExpressions;
 using System.Threading.Tasks;
 using EpicGames.Core;
+using Microsoft.Extensions.Logging;
 using OpenTracing.Util;
 using UnrealBuildBase;
 
@@ -93,6 +94,11 @@
 		bool bModified;
 
 		/// <summary>
+		/// Logger for output
+		/// </summary>
+		ILogger Logger;
+
+		/// <summary>
 		/// Regex that matches C++ code with UObject declarations which we will need to generated code for.
 		/// </summary>
 		static readonly Regex ReflectionMarkupRegex = new Regex("^\\s*U(CLASS|STRUCT|ENUM|INTERFACE|DELEGATE)\\b", RegexOptions.Compiled | RegexOptions.Multiline);
@@ -118,11 +124,13 @@
 		/// <param name="Location">File to store the cache</param>
 		/// <param name="BaseDir">Base directory for files that this cache should store data for</param>
 		/// <param name="Parent">The parent cache to use</param>
-		private SourceFileMetadataCache(FileReference Location, DirectoryReference BaseDir, SourceFileMetadataCache? Parent)
+		/// <param name="Logger">Logger for output</param>
+		private SourceFileMetadataCache(FileReference Location, DirectoryReference BaseDir, SourceFileMetadataCache? Parent, ILogger Logger)
 		{
 			this.Location = Location;
 			this.BaseDirectory = BaseDir;
 			this.Parent = Parent;
+			this.Logger = Logger;
 
 			if(FileReference.Exists(Location))
 			{
@@ -269,21 +277,22 @@
 		/// Creates a cache hierarchy for a particular target
 		/// </summary>
 		/// <param name="ProjectFile">Project file for the target being built</param>
+		/// <param name="Logger">Logger for output</param>
 		/// <returns>Dependency cache hierarchy for the given project</returns>
-		public static SourceFileMetadataCache CreateHierarchy(FileReference? ProjectFile)
+		public static SourceFileMetadataCache CreateHierarchy(FileReference? ProjectFile, ILogger Logger)
 		{
 			SourceFileMetadataCache? Cache = null;
 
 			if(ProjectFile == null || !Unreal.IsEngineInstalled())
 			{
 				FileReference EngineCacheLocation = FileReference.Combine(Unreal.EngineDirectory, "Intermediate", "Build", "SourceFileCache.bin");
-				Cache = FindOrAddCache(EngineCacheLocation, Unreal.EngineDirectory, Cache);
+				Cache = FindOrAddCache(EngineCacheLocation, Unreal.EngineDirectory, Cache, Logger);
 			}
 
 			if(ProjectFile != null)
 			{
 				FileReference ProjectCacheLocation = FileReference.Combine(ProjectFile.Directory, "Intermediate", "Build", "SourceFileCache.bin");
-				Cache = FindOrAddCache(ProjectCacheLocation, ProjectFile.Directory, Cache);
+				Cache = FindOrAddCache(ProjectCacheLocation, ProjectFile.Directory, Cache, Logger);
 			}
 
 			return Cache!;
@@ -312,8 +321,9 @@
 		/// <param name="Location">File to store the cache</param>
 		/// <param name="BaseDirectory">Base directory for files that this cache should store data for</param>
 		/// <param name="Parent">The parent cache to use</param>
+		/// <param name="Logger"></param>
 		/// <returns>Reference to a dependency cache with the given settings</returns>
-		static SourceFileMetadataCache FindOrAddCache(FileReference Location, DirectoryReference BaseDirectory, SourceFileMetadataCache? Parent)
+		static SourceFileMetadataCache FindOrAddCache(FileReference Location, DirectoryReference BaseDirectory, SourceFileMetadataCache? Parent, ILogger Logger)
 		{
 			lock(Caches)
 			{
@@ -325,7 +335,7 @@
 				}
 				else
 				{
-					Cache = new SourceFileMetadataCache(Location, BaseDirectory, Parent);
+					Cache = new SourceFileMetadataCache(Location, BaseDirectory, Parent, Logger);
 					Caches.Add(Location, Cache);
 				}
 				return Cache;
@@ -352,7 +362,7 @@
 					int Version = Reader.ReadInt();
 					if(Version != CurrentVersion)
 					{
-						Log.TraceLog("Unable to read dependency cache from {0}; version {1} vs current {2}", Location, Version, CurrentVersion);
+						Logger.LogDebug("Unable to read dependency cache from {File}; version {Version} vs current {CurrentVersion}", Location, Version, CurrentVersion);
 						return;
 					}
 
@@ -383,8 +393,8 @@
 			}
 			catch(Exception Ex)
 			{
-				Log.TraceWarning("Unable to read {0}. See log for additional information.", Location);
-				Log.TraceLog("{0}", ExceptionUtils.FormatExceptionDetails(Ex));
+				Logger.LogWarning("Unable to read {Location}. See log for additional information.", Location);
+				Logger.LogDebug(Ex, "{Ex}", ExceptionUtils.FormatExceptionDetails(Ex));
 			}
 		}
 
diff --git a/Engine/Source/Programs/UnrealBuildTool/System/SourceFileWorkingSet.cs b/Engine/Source/Programs/UnrealBuildTool/System/SourceFileWorkingSet.cs
--- a/Engine/Source/Programs/UnrealBuildTool/System/SourceFileWorkingSet.cs
+++ b/Engine/Source/Programs/UnrealBuildTool/System/SourceFileWorkingSet.cs
@@ -10,6 +10,7 @@
 using System.Threading.Tasks;
 using EpicGames.Core;
 using UnrealBuildBase;
+using Microsoft.Extensions.Logging;
 
 namespace UnrealBuildTool
 {
@@ -97,6 +98,7 @@
 		List<DirectoryReference> Directories;
 		List<string> ErrorOutput;
 		GitSourceFileWorkingSet? Inner;
+		ILogger Logger;
 
 		/// <summary>
 		/// Constructor
@@ -104,15 +106,17 @@
 		/// <param name="GitPath">Path to the Git executable</param>
 		/// <param name="RootDir">Root directory to run queries from (typically the directory containing the .git folder, to ensure all subfolders can be searched)</param>
 		/// <param name="Inner">An inner working set. This allows supporting multiple Git repositories (one containing the engine, another containing the project, for example)</param>
-		public GitSourceFileWorkingSet(string GitPath, DirectoryReference RootDir, GitSourceFileWorkingSet? Inner)
+		/// <param name="Logger">Logger for output</param>
+		public GitSourceFileWorkingSet(string GitPath, DirectoryReference RootDir, GitSourceFileWorkingSet? Inner, ILogger Logger)
 		{
 			this.RootDir = RootDir;
 			this.Files = new HashSet<FileReference>();
 			this.Directories = new List<DirectoryReference>();
 			this.ErrorOutput = new List<string>();
 			this.Inner = Inner;
+			this.Logger = Logger;
 
-			Log.WriteLine(LogEventType.Console, "Using 'git status' to determine working set for adaptive non-unity build ({0}).", RootDir);
+			Logger.LogInformation("Using 'git status' to determine working set for adaptive non-unity build ({RootDir}).", RootDir);
 
 			BackgroundProcess = new Process();
 			BackgroundProcess.StartInfo.FileName = GitPath;
@@ -166,11 +170,11 @@
 			{
 				if(!BackgroundProcess.WaitForExit(500))
 				{
-					Log.WriteLine(LogEventType.Console, "Waiting for 'git status' command to complete");
+					Logger.LogInformation("Waiting for 'git status' command to complete");
 				}
 				if(!BackgroundProcess.WaitForExit(15000))
 				{
-					Log.WriteLine(LogEventType.Console, "Terminating git child process due to timeout");
+					Logger.LogInformation("Terminating git child process due to timeout");
 					try
 					{
 						BackgroundProcess.Kill();
@@ -315,8 +319,9 @@
 		/// </summary>
 		/// <param name="RootDir">The root directory</param>
 		/// <param name="ProjectDirs">The project directories</param>
+		/// <param name="Logger">Logger for output</param>
 		/// <returns>Working set instance for the given directory</returns>
-		public static ISourceFileWorkingSet Create(DirectoryReference RootDir, IEnumerable<DirectoryReference> ProjectDirs)
+		public static ISourceFileWorkingSet Create(DirectoryReference RootDir, IEnumerable<DirectoryReference> ProjectDirs, ILogger Logger)
 		{
 			if (Provider == ProviderType.None || ProjectFileGenerator.bGenerateProjectFiles)
 			{
@@ -327,11 +332,11 @@
 				GitSourceFileWorkingSet? WorkingSet;
 				if (!String.IsNullOrEmpty(RepositoryPath))
 				{
-					WorkingSet = new GitSourceFileWorkingSet(GitPath, DirectoryReference.Combine(RootDir, RepositoryPath), null);
+					WorkingSet = new GitSourceFileWorkingSet(GitPath, DirectoryReference.Combine(RootDir, RepositoryPath), null, Logger);
 				}
-				else if(!TryCreateGitWorkingSet(RootDir, ProjectDirs, out WorkingSet))
+				else if(!TryCreateGitWorkingSet(RootDir, ProjectDirs, Logger, out WorkingSet))
 				{
-					WorkingSet = new GitSourceFileWorkingSet(GitPath, RootDir, null);
+					WorkingSet = new GitSourceFileWorkingSet(GitPath, RootDir, null, Logger);
 				}
 				return WorkingSet;
 			}
@@ -342,7 +347,7 @@
 			else
 			{
 				GitSourceFileWorkingSet? WorkingSet;
-				if(TryCreateGitWorkingSet(RootDir, ProjectDirs, out WorkingSet))
+				if(TryCreateGitWorkingSet(RootDir, ProjectDirs, Logger, out WorkingSet))
 				{
 					return WorkingSet;
 				}
@@ -353,14 +358,14 @@
 			}
 		}
 
-		static bool TryCreateGitWorkingSet(DirectoryReference RootDir, IEnumerable<DirectoryReference> ProjectDirs, [NotNullWhen(true)] out GitSourceFileWorkingSet? OutWorkingSet)
+		static bool TryCreateGitWorkingSet(DirectoryReference RootDir, IEnumerable<DirectoryReference> ProjectDirs, ILogger Logger, [NotNullWhen(true)] out GitSourceFileWorkingSet? OutWorkingSet)
 		{
 			GitSourceFileWorkingSet? WorkingSet  = null;
 
 			// Create the working set for the engine directory
 			if (DirectoryReference.Exists(DirectoryReference.Combine(RootDir, ".git")))
 			{
-				WorkingSet = new GitSourceFileWorkingSet(GitPath, RootDir, WorkingSet);
+				WorkingSet = new GitSourceFileWorkingSet(GitPath, RootDir, WorkingSet, Logger);
 			}
 
 			// Try to create a working set for the project directory
@@ -370,11 +375,11 @@
 				{
 					if (DirectoryReference.Exists(DirectoryReference.Combine(ProjectDir, ".git")))
 					{
-						WorkingSet = new GitSourceFileWorkingSet(GitPath, ProjectDir, WorkingSet);
+						WorkingSet = new GitSourceFileWorkingSet(GitPath, ProjectDir, WorkingSet, Logger);
 					}
 					else if (DirectoryReference.Exists(DirectoryReference.Combine(ProjectDir.ParentDirectory!, ".git")))
 					{
-						WorkingSet = new GitSourceFileWorkingSet(GitPath, ProjectDir.ParentDirectory!, WorkingSet);
+						WorkingSet = new GitSourceFileWorkingSet(GitPath, ProjectDir.ParentDirectory!, WorkingSet, Logger);
 					}
 				}
 			}
diff --git a/Engine/Source/Programs/UnrealBuildTool/System/TargetMakefile.cs b/Engine/Source/Programs/UnrealBuildTool/System/TargetMakefile.cs
--- a/Engine/Source/Programs/UnrealBuildTool/System/TargetMakefile.cs
+++ b/Engine/Source/Programs/UnrealBuildTool/System/TargetMakefile.cs
@@ -7,15 +7,96 @@
 using System.Reflection;
 using EpicGames.Core;
 using EpicGames.UHT.Utils;
+using Microsoft.Extensions.Logging;
 using OpenTracing.Util;
 using UnrealBuildBase;
 
 namespace UnrealBuildTool
 {
+	class TargetMakefileBuilder : IActionGraphBuilder
+	{
+		private readonly ILogger Logger;
+		public TargetMakefile Makefile { get; }
+
+		public TargetMakefileBuilder(TargetMakefile InMakefile, ILogger InLogger)
+		{
+			Logger = InLogger;
+			Makefile = InMakefile;
+		}
+
+
+		/// <inheritdoc/>
+		public void AddAction(IExternalAction Action)
+		{
+			Makefile.Actions.Add(Action);
+		}
+
+		/// <inheritdoc/>
+		public void CreateIntermediateTextFile(FileItem FileItem, string Contents)
+		{
+			// Write the file
+			Utils.WriteFileIfChanged(FileItem.Location, Contents, Logger);
+
+			// Reset the file info, in case it already knows about the old file
+			Makefile.InternalDependencies.Add(FileItem);
+			FileItem.ResetCachedInfo();
+		}
+
+		/// <inheritdoc/>
+		public void CreateIntermediateTextFile(FileItem FileItem, IEnumerable<string> ContentLines)
+		{
+			// Write the file
+			Utils.WriteFileIfChanged(FileItem, ContentLines, Logger);
+
+			// Reset the file info, in case it already knows about the old file
+			Makefile.InternalDependencies.Add(FileItem);
+			FileItem.ResetCachedInfo();
+		}
+
+		/// <inheritdoc/>
+		public void AddSourceDir(DirectoryItem SourceDir)
+		{
+			Makefile.SourceDirectories.Add(SourceDir);
+		}
+
+		/// <inheritdoc/>
+		public void AddSourceFiles(DirectoryItem SourceDir, FileItem[] SourceFiles)
+		{
+			Makefile.DirectoryToSourceFiles[SourceDir] = SourceFiles;
+		}
+
+		/// <inheritdoc/>
+		public void AddDiagnostic(string Message)
+		{
+			if (!Makefile.Diagnostics.Contains(Message))
+			{
+				Makefile.Diagnostics.Add(Message);
+			}
+		}
+
+		/// <inheritdoc/>
+		public void AddFileToWorkingSet(FileItem File)
+		{
+			Makefile.WorkingSet.Add(File);
+		}
+
+		/// <inheritdoc/>
+		public void AddCandidateForWorkingSet(FileItem File)
+		{
+			Makefile.CandidatesForWorkingSet.Add(File);
+		}
+
+		/// <inheritdoc/>
+		public void SetOutputItemsForModule(string ModuleName, FileItem[] OutputItems)
+		{
+			Makefile.ModuleNameToOutputItems[ModuleName] = OutputItems;
+		}
+	}
+
 	/// <summary>
 	/// Cached list of actions that need to be executed to build a target, along with the information needed to determine whether they are valid.
 	/// </summary>
-	class TargetMakefile : IActionGraphBuilder
+	class TargetMakefile
 	{
 		/// <summary>
 		/// The version number to write
@@ -380,9 +461,10 @@
 		/// <param name="ProjectFile">Path to the project file</param>
 		/// <param name="Platform">Platform for this makefile</param>
 		/// <param name="Arguments">Command line arguments for this target</param>
+		/// <param name="Logger">Logger for output</param>
 		/// <param name="ReasonNotLoaded">If the function returns null, this string will contain the reason why</param>
 		/// <returns>The loaded makefile, or null if it failed for some reason.  On failure, the 'ReasonNotLoaded' variable will contain information about why</returns>
-		public static TargetMakefile? Load(FileReference MakefilePath, FileReference? ProjectFile, UnrealTargetPlatform Platform, string[] Arguments, out string? ReasonNotLoaded)
+		public static TargetMakefile? Load(FileReference MakefilePath, FileReference? ProjectFile, UnrealTargetPlatform Platform, string[] Arguments, ILogger Logger, out string? ReasonNotLoaded)
 		{
 			FileInfo MakefileInfo;
 			using (GlobalTracer.Instance.BuildSpan("Checking dependent timestamps").StartActive())
@@ -400,7 +482,7 @@
 				FileInfo BuildVersionFileInfo = new FileInfo(BuildVersion.GetDefaultFileName().FullName);
 				if (BuildVersionFileInfo.Exists && MakefileInfo.LastWriteTime.CompareTo(BuildVersionFileInfo.LastWriteTime) < 0)
 				{
-					Log.TraceLog("Existing makefile is older than Build.version, ignoring it");
+					Logger.LogDebug("Existing makefile is older than Build.version, ignoring it");
 					ReasonNotLoaded = "Build.version is newer";
 					return null;
 				}
@@ -413,7 +495,7 @@
 					if (!ProjectFileInfo.Exists || MakefileInfo.LastWriteTime.CompareTo(ProjectFileInfo.LastWriteTime) < 0)
 					{
 						// .uproject file is newer than makefile
-						Log.TraceLog("Makefile is older than .uproject file, ignoring it");
+						Logger.LogDebug("Makefile is older than .uproject file, ignoring it");
 						ReasonNotLoaded = ".uproject file is newer";
 						return null;
 					}
@@ -424,7 +506,7 @@
 				if (MakefileInfo.LastWriteTime.CompareTo(UnrealBuildToolTimestamp) < 0)
 				{
 					// UnrealBuildTool was compiled more recently than the makefile
-					Log.TraceLog("Makefile is older than UnrealBuildTool assembly, ignoring it");
+					Logger.LogDebug("Makefile is older than UnrealBuildTool assembly, ignoring it");
 					ReasonNotLoaded = "UnrealBuildTool assembly is newer";
 					return null;
 				}
@@ -434,7 +516,7 @@
 				if (MakefileInfo.LastWriteTime.CompareTo(UhtTimestamp) < 0)
 				{
 					// UnrealBuildTool was compiled more recently than the makefile
-					Log.TraceLog("Makefile is older than EpicGames.UHT assembly, ignoring it");
+					Logger.LogDebug("Makefile is older than EpicGames.UHT assembly, ignoring it");
 					ReasonNotLoaded = "EpicGames.UHT assembly is newer";
 					return null;
 				}
@@ -445,7 +527,7 @@
 					FileInfo InputFileInfo = new FileInfo(InputFile.Location.FullName);
 					if (InputFileInfo.LastWriteTime > MakefileInfo.LastWriteTime)
 					{
-						Log.TraceLog("Makefile is older than BuildConfiguration.xml, ignoring it");
+						Logger.LogDebug("Makefile is older than BuildConfiguration.xml, ignoring it");
 						ReasonNotLoaded = "BuildConfiguration.xml is newer";
 						return null;
 					}
@@ -480,8 +562,8 @@
 				}
 				catch (Exception Ex)
 				{
-					Log.TraceWarning("Failed to read makefile: {0}", Ex.Message);
-					Log.TraceLog("Exception: {0}", Ex.ToString());
+					Logger.LogWarning("Failed to read makefile: {ExMessage}", Ex.Message);
+					Logger.LogDebug(Ex, "Exception: {Ex}", Ex.ToString());
 					ReasonNotLoaded = "couldn't read existing makefile";
 					return null;
 				}
@@ -507,8 +589,8 @@
 				string CurrentExternalMetadata = UEBuildPlatform.GetBuildPlatform(Platform).GetExternalBuildMetadata(ProjectFile);
 				if(String.Compare(CurrentExternalMetadata, Makefile.ExternalMetadata, StringComparison.Ordinal) != 0)
 				{
-					Log.TraceLog("Old metadata:\n", Makefile.ExternalMetadata);
-					Log.TraceLog("New metadata:\n", CurrentExternalMetadata);
+					Logger.LogDebug("Old metadata:\n", Makefile.ExternalMetadata);
+					Logger.LogDebug("New metadata:\n", CurrentExternalMetadata);
 					ReasonNotLoaded = "build metadata has changed";
 					return null;
 				}
@@ -520,7 +602,7 @@
 					if ((Plugins != null) != (Makefile.UbtPlugins != null) ||
 						(Plugins != null && !Enumerable.SequenceEqual(Plugins, Makefile.UbtPlugins!)))
 					{
-						Log.TraceLog("Available UBT plugins changed");
+						Logger.LogDebug("Available UBT plugins changed");
 						ReasonNotLoaded = "Available UBT plugins changed";
 						return null;
 					}
@@ -536,9 +618,9 @@
 						BaseDirectories.Add(ProjectFile.Directory);
 					}
 					HashSet<FileReference> EnabledPlugins = new HashSet<FileReference>(Makefile.EnabledUbtPlugins);
-					if (!CompileScriptModule.AreScriptModulesUpToDate(EnabledPlugins, BaseDirectories))
+					if (!CompileScriptModule.AreScriptModulesUpToDate(EnabledPlugins, BaseDirectories, Logger))
 					{
-						Log.TraceLog("Enabled UBT plugins need to be recompiled");
+						Logger.LogDebug("Enabled UBT plugins need to be recompiled");
 						ReasonNotLoaded = "Enabled UBT plugins need to be recompiled";
 						return null;
 					}
@@ -557,9 +639,10 @@
 		/// <param name="ProjectFile">Path to the project file</param>
 		/// <param name="Platform">The platform being built</param>
 		/// <param name="WorkingSet">The current working set of source files</param>
+		/// <param name="Logger">Logger for output diagnostics</param>
 		/// <param name="ReasonNotLoaded">If the makefile is not valid, is set to a message describing why</param>
 		/// <returns>True if the makefile is valid, false otherwise</returns>
-		public static bool IsValidForSourceFiles(TargetMakefile Makefile, FileReference? ProjectFile, UnrealTargetPlatform Platform, ISourceFileWorkingSet WorkingSet, out string? ReasonNotLoaded)
+		public static bool IsValidForSourceFiles(TargetMakefile Makefile, FileReference? ProjectFile, UnrealTargetPlatform Platform, ISourceFileWorkingSet WorkingSet, ILogger Logger, out string? ReasonNotLoaded)
 		{
 			using (GlobalTracer.Instance.BuildSpan("TargetMakefile.IsValidForSourceFiles()").StartActive())
 			{
@@ -604,13 +687,13 @@
 				{
 					if (!ExternalDependency.Exists)
 					{
-						Log.TraceLog("{0} has been deleted since makefile was built.", ExternalDependency.Location);
+						Logger.LogDebug("{File} has been deleted since makefile was built.", ExternalDependency.Location);
 						ReasonNotLoaded = string.Format("{0} deleted", ExternalDependency.Location.GetFileName());
 						return false;
 					}
 					if(ExternalDependency.LastWriteTimeUtc > Makefile.CreateTimeUtc)
 					{
-						Log.TraceLog("{0} has been modified since makefile was built.", ExternalDependency.Location);
+						Logger.LogDebug("{File} has been modified since makefile was built.", ExternalDependency.Location);
 						ReasonNotLoaded = string.Format("{0} modified", ExternalDependency.Location.GetFileName());
 						return false;
 					}
@@ -621,13 +704,13 @@
 				{
 					if (!InternalDependency.Exists)
 					{
-						Log.TraceLog("{0} has been deleted since makefile was written.", InternalDependency.Location);
+						Logger.LogDebug("{File} has been deleted since makefile was written.", InternalDependency.Location);
 						ReasonNotLoaded = string.Format("{0} deleted", InternalDependency.Location.GetFileName());
 						return false;
 					}
 					if (InternalDependency.LastWriteTimeUtc > Makefile.ModifiedTimeUtc)
 					{
-						Log.TraceLog("{0} has been modified since makefile was written.", InternalDependency.Location);
+						Logger.LogDebug("{File} has been modified since makefile was written.", InternalDependency.Location);
 						ReasonNotLoaded = string.Format("{0} modified", InternalDependency.Location.GetFileName());
 						return false;
 					}
@@ -639,14 +722,14 @@
 					FileItem PluginFileItem = FileItem.GetItemByFileReference(PluginFile);
 					if(!Makefile.PluginFiles.Contains(PluginFileItem))
 					{
-						Log.TraceLog("{0} has been added", PluginFile.GetFileName());
+						Logger.LogDebug("{File} has been added", PluginFile.GetFileName());
 						ReasonNotLoaded = string.Format("{0} has been added", PluginFile.GetFileName());
 						return false;
 					}
 				}
 
 				// Load the metadata cache
-				SourceFileMetadataCache MetadataCache = SourceFileMetadataCache.CreateHierarchy(ProjectFile);
+				SourceFileMetadataCache MetadataCache = SourceFileMetadataCache.CreateHierarchy(ProjectFile, Logger);
 
 				// Find the set of files that contain reflection markup
 				ConcurrentBag<FileItem> NewFilesWithMarkupBag = new ConcurrentBag<FileItem>();
@@ -678,7 +761,7 @@
 				{
 					if (!WorkingSet.Contains(SourceFile) && SourceFile.LastWriteTimeUtc > Makefile.CreateTimeUtc)
 					{
-						Log.TraceLog("{0} was part of source working set and now is not; invalidating makefile", SourceFile.AbsolutePath);
+						Logger.LogDebug("{File} was part of source working set and now is not; invalidating makefile", SourceFile.Location);
 						ReasonNotLoaded = string.Format("working set of source files changed");
 						return false;
 					}
@@ -689,7 +772,7 @@
 				{
 					if (WorkingSet.Contains(SourceFile) && SourceFile.LastWriteTimeUtc > Makefile.CreateTimeUtc)
 					{
-						Log.TraceLog("{0} was part of source working set and now is not", SourceFile.AbsolutePath);
+						Logger.LogDebug("{File} was part of source working set and now is not", SourceFile.Location);
 						ReasonNotLoaded = string.Format("working set of source files changed");
 						return false;
 					}
@@ -773,72 +856,5 @@
 			DirectoryReference BaseDirectory = DirectoryReference.FromFile(ProjectFile) ?? Unreal.EngineDirectory;
 			return FileReference.Combine(BaseDirectory, UEBuildTarget.GetPlatformIntermediateFolder(Platform, Architecture, false), TargetName, Configuration.ToString(), "Makefile.bin");
 		}
-
-		/// <inheritdoc/>
-		public void AddAction(IExternalAction Action)
-		{
-			Actions.Add(Action);
-		}
-
-		/// <inheritdoc/>
-		public void CreateIntermediateTextFile(FileItem FileItem, string Contents)
-		{
-			// Write the file
-			Utils.WriteFileIfChanged(FileItem.Location, Contents);
-
-			// Reset the file info, in case it already knows about the old file
-			InternalDependencies.Add(FileItem);
-			FileItem.ResetCachedInfo();
-		}
-
-		/// <inheritdoc/>
-		public void CreateIntermediateTextFile(FileItem FileItem, IEnumerable<string> ContentLines)
-		{
-			// Write the file
-			Utils.WriteFileIfChanged(FileItem, ContentLines);
-
-			// Reset the file info, in case it already knows about the old file
-			InternalDependencies.Add(FileItem);
-			FileItem.ResetCachedInfo();
-		}
-
-		/// <inheritdoc/>
-		public void AddSourceDir(DirectoryItem SourceDir)
-		{
-			SourceDirectories.Add(SourceDir);
-		}
-
-		/// <inheritdoc/>
-		public void AddSourceFiles(DirectoryItem SourceDir, FileItem[] SourceFiles)
-		{
-			DirectoryToSourceFiles[SourceDir] = SourceFiles;
-		}
-
-		/// <inheritdoc/>
-		public void AddDiagnostic(string Message)
-		{
-			if(!Diagnostics.Contains(Message))
-			{
-				Diagnostics.Add(Message);
-			}
-		}
-
-		/// <inheritdoc/>
-		public void AddFileToWorkingSet(FileItem File)
-		{
-			WorkingSet.Add(File);
-		}
-
-		/// <inheritdoc/>
-		public void AddCandidateForWorkingSet(FileItem File)
-		{
-			CandidatesForWorkingSet.Add(File);
-		}
-
-		/// <inheritdoc/>
-		public void SetOutputItemsForModule(string ModuleName, FileItem[] OutputItems)
-		{
-			ModuleNameToOutputItems[ModuleName] = OutputItems;
-		}
 	}
 }
diff --git a/Engine/Source/Programs/UnrealBuildTool/System/Timeline.cs b/Engine/Source/Programs/UnrealBuildTool/System/Timeline.cs
--- a/Engine/Source/Programs/UnrealBuildTool/System/Timeline.cs
+++ b/Engine/Source/Programs/UnrealBuildTool/System/Timeline.cs
@@ -9,6 +9,7 @@
 using System.Text;
 using System.Threading.Tasks;
 using EpicGames.Core;
+using Microsoft.Extensions.Logging;
 
 namespace UnrealBuildTool
 {
@@ -137,12 +138,12 @@
 		/// <summary>
 		/// Prints this information to the log
 		/// </summary>
-		public static void Print(TimeSpan MaxUnknownTime, LogEventType Verbosity)
+		public static void Print(TimeSpan MaxUnknownTime, LogLevel Verbosity, ILogger Logger)
 		{
 			// Print the start time
-			Log.WriteLine(Verbosity, "Timeline:");
-			Log.WriteLine(Verbosity, "");
-			Log.WriteLine(Verbosity, "[{0,6}]", FormatTime(TimeSpan.Zero));
+			Logger.Log(Verbosity, "Timeline:");
+			Logger.Log(Verbosity, "");
+			Logger.Log(Verbosity, "[{Time,6}]", FormatTime(TimeSpan.Zero));
 
 			// Create the root event
 			TimeSpan FinishTime = Stopwatch.Elapsed;
@@ -164,14 +165,14 @@
 					{
 						break;
 					}
-					UpdateLastEventTime(ref LastTime, OuterEvent.FinishTime.Value, MaxUnknownTime, OuterEvents, Verbosity);
+					UpdateLastEventTime(ref LastTime, OuterEvent.FinishTime.Value, MaxUnknownTime, OuterEvents, Verbosity, Logger);
 				}
 
 				// If there's a gap since the last event, print an unknown marker
-				UpdateLastEventTime(ref LastTime, Event.StartTime, MaxUnknownTime, OuterEvents, Verbosity);
+				UpdateLastEventTime(ref LastTime, Event.StartTime, MaxUnknownTime, OuterEvents, Verbosity, Logger);
 
 				// Print this event
-				Print(Event.StartTime, Event.FinishTime, Event.Name, OuterEvents, Verbosity);
+				Print(Event.StartTime, Event.FinishTime, Event.Name, OuterEvents, Verbosity, Logger);
 
 				// Push it onto the stack
 				if(Event.FinishTime.HasValue)
@@ -190,11 +191,11 @@
 			// Remove everything from the stack
 			for(; OuterEvents.Count > 0; OuterEvents.RemoveAt(OuterEvents.Count - 1))
 			{
-				UpdateLastEventTime(ref LastTime, OuterEvents.Last().FinishTime!.Value, MaxUnknownTime, OuterEvents, Verbosity);
+				UpdateLastEventTime(ref LastTime, OuterEvents.Last().FinishTime!.Value, MaxUnknownTime, OuterEvents, Verbosity, Logger);
 			}
 
 			// Print the finish time
-			Log.WriteLine(Verbosity, "[{0,6}]", FormatTime(FinishTime));
+			Logger.Log(Verbosity, "[{Time,6}]", FormatTime(FinishTime));
 		}
 
 		/// <summary>
@@ -205,12 +206,13 @@
 		/// <param name="MaxUnknownTime"></param>
 		/// <param name="OuterEvents"></param>
 		/// <param name="Verbosity"></param>
-		static void UpdateLastEventTime(ref TimeSpan LastTime, TimeSpan NewTime, TimeSpan MaxUnknownTime, List<Event> OuterEvents, LogEventType Verbosity)
+		/// <param name="Logger"></param>
+		static void UpdateLastEventTime(ref TimeSpan LastTime, TimeSpan NewTime, TimeSpan MaxUnknownTime, List<Event> OuterEvents, LogLevel Verbosity, ILogger Logger)
 		{
 			const string UnknownEvent = "<unknown>";
 			if (NewTime - LastTime > MaxUnknownTime)
 			{
-				Print(LastTime, NewTime, UnknownEvent, OuterEvents, Verbosity);
+				Print(LastTime, NewTime, UnknownEvent, OuterEvents, Verbosity, Logger);
 			}
 			LastTime = NewTime;
 		}
@@ -223,7 +225,8 @@
 		/// <param name="Label">Event name</param>
 		/// <param name="OuterEvents">List of all the start times for parent events</param>
 		/// <param name="Verbosity">Verbosity for the output</param>
-		static void Print(TimeSpan StartTime, TimeSpan? FinishTime, string Label, List<Event> OuterEvents, LogEventType Verbosity)
+		/// <param name="Logger">Logger for output</param>
+		static void Print(TimeSpan StartTime, TimeSpan? FinishTime, string Label, List<Event> OuterEvents, LogLevel Verbosity, ILogger Logger)
 		{
 			StringBuilder Prefix = new StringBuilder();
 
@@ -247,7 +250,7 @@
 				Prefix.AppendFormat("({0,6})", "+" + FormatTime(FinishTime.Value - StartTime));
 			}
 
-			Log.WriteLine(Verbosity, "{0} {1}", Prefix.ToString(), Label);
+			Logger.Log(Verbosity, "{Prefix} {Label}", Prefix.ToString(), Label);
 		}
 
 		/// <summary>
diff --git a/Engine/Source/Programs/UnrealBuildTool/System/ToolMode.cs b/Engine/Source/Programs/UnrealBuildTool/System/ToolMode.cs
--- a/Engine/Source/Programs/UnrealBuildTool/System/ToolMode.cs
+++ b/Engine/Source/Programs/UnrealBuildTool/System/ToolMode.cs
@@ -6,6 +6,7 @@
 using System.Text;
 using System.Threading.Tasks;
 using EpicGames.Core;
+using Microsoft.Extensions.Logging;
 
 namespace UnrealBuildTool
 {
@@ -98,7 +99,8 @@
 		/// Entry point for this command.
 		/// </summary>
 		/// <param name="Arguments">List of command line arguments</param>
+		/// <param name="Logger"></param>
 		/// <returns>Exit code for the process</returns>
-		public abstract int Execute(CommandLineArguments Arguments);
+		public abstract int Execute(CommandLineArguments Arguments, ILogger Logger);
 	}
 }
diff --git a/Engine/Source/Programs/UnrealBuildTool/System/UnrealPluginLanguage.cs b/Engine/Source/Programs/UnrealBuildTool/System/UnrealPluginLanguage.cs
--- a/Engine/Source/Programs/UnrealBuildTool/System/UnrealPluginLanguage.cs
+++ b/Engine/Source/Programs/UnrealBuildTool/System/UnrealPluginLanguage.cs
@@ -11,6 +11,7 @@
 using EpicGames.Core;
 using UnrealBuildBase;
 using System.Diagnostics.CodeAnalysis;
+using Microsoft.Extensions.Logging;
 
 namespace UnrealBuildTool
 {
@@ -527,10 +528,12 @@
 		private Dictionary<string, UPLContext> Contexts;
 		private int ContextIndex;
 		private string? LastError;
+		private ILogger Logger;
 
-		public UnrealPluginLanguage(FileReference? InProjectFile, List<string> InXMLFiles, List<string> InArchitectures, string InXMLNameSpace, string InRootDefinition, UnrealTargetPlatform InTargetPlatform)
+		public UnrealPluginLanguage(FileReference? InProjectFile, List<string> InXMLFiles, List<string> InArchitectures, string InXMLNameSpace, string InRootDefinition, UnrealTargetPlatform InTargetPlatform, ILogger InLogger)
 		{
 			ProjectFile = InProjectFile;
+			Logger = InLogger;
 
 			LastError = null;
 
@@ -548,7 +551,7 @@
 			foreach (string Basename in InXMLFiles)
 			{
 				string Filename = Path.Combine(PathPrefix, Basename.Replace("\\", "/"));
-				Log.TraceInformation("UPL: {0}", Filename);
+				Logger.LogInformation("UPL: {FileName}", Filename);
 				if (File.Exists(Filename))
 				{
 					string PluginDir = Path.GetDirectoryName(Filename)!;
@@ -560,14 +563,14 @@
 					catch (Exception e)
 					{
 						LastError = String.Format("Unreal Plugin file {0} parsing failed! {1}", Filename, e);
-						Log.TraceError("\n{0}", LastError);
+						Logger.LogError("\n{LastError}", LastError);
 					}
 				}
 				else
 				{
 					LastError = String.Format("Unreal Plugin file {0} missing!", Filename);
-					Log.TraceError("\n{0}", LastError);
-					Log.TraceInformation("\nCWD: {0}", Directory.GetCurrentDirectory());
+					Logger.LogError("\n{LastError}", LastError);
+					Logger.LogInformation("\nCWD: {Cwd}", Directory.GetCurrentDirectory());
 				}
 			}
 		}
@@ -674,7 +677,7 @@
 			{
 				if (!GlobalContext.BoolVariables.TryGetValue(Condition, out Result))
 				{
-					Log.TraceWarning("\nMissing condition '{0}' in '{1}' (skipping instruction)", Condition, TraceNodeString(Node));
+					Logger.LogWarning("\nMissing condition '{Condition}' in '{Node}' (skipping instruction)", Condition, TraceNodeString(Node));
 					return false;
 				}
 			}
@@ -816,7 +819,7 @@
 			int Result = 0;
 			if (!int.TryParse(Input, out Result))
 			{
-				Log.TraceWarning("\nInvalid integer '{0}' in '{1}' (defaulting to 0)", Input, TraceNodeString(Node));
+				Logger.LogWarning("\nInvalid integer '{Input}' in '{Node}' (defaulting to 0)", Input, TraceNodeString(Node));
 			}
 			return Result;
 		}
@@ -829,7 +832,7 @@
 			{
 				if (bRequired)
 				{
-					Log.TraceWarning("\nMissing attribute '{0}' in '{1}' (skipping instruction)", AttributeName, TraceNodeString(Node));
+					Logger.LogWarning("\nMissing attribute '{Attribute}' in '{Node}' (skipping instruction)", AttributeName, TraceNodeString(Node));
 				}
 				return Fallback;
 			}
@@ -845,7 +848,7 @@
 			{
 				if (bRequired)
 				{
-					Log.TraceWarning("\nMissing attribute '{0}' in '{1}' (skipping instruction)", AttributeName, TraceNodeString(Node));
+					Logger.LogWarning("\nMissing attribute '{Attribute}' in '{Node}' (skipping instruction)", AttributeName, TraceNodeString(Node));
 				}
 				return Fallback;
 			}
@@ -865,19 +868,19 @@
 			if (!ConfigCache.TryGetValue(baseIniName, out config))
 			{
 				// note: use our own ConfigCacheIni since EngineConfiguration.cs only parses RequiredSections!
-				config = ConfigCacheIni_UPL.CreateConfigCacheIni_UPL(TargetPlatform, baseIniName, DirectoryReference.FromFile(ProjectFile));
+				config = ConfigCacheIni_UPL.CreateConfigCacheIni_UPL(TargetPlatform, baseIniName, DirectoryReference.FromFile(ProjectFile), Logger);
 				ConfigCache.Add(baseIniName, config);
 			}
 			return config;
 		}
 
-		private static bool FilesAreDifferent(string SourceFilename, string DestFilename)
+		private static bool FilesAreDifferent(string SourceFilename, string DestFilename, ILogger Logger)
 		{
 			// source must exist
 			FileInfo SourceInfo = new FileInfo(SourceFilename);
 			if (!SourceInfo.Exists)
 			{
-				Log.TraceInformation("File {0} does not exist", SourceFilename);
+				Logger.LogInformation("File {SourceFilename} does not exist", SourceFilename);
 				return false;
 			}
 
@@ -905,7 +908,7 @@
 			return false;
 		}
 
-		private static void CopyFileDirectory(string SourceDir, string DestDir, bool bForce = false)
+		private static void CopyFileDirectory(string SourceDir, string DestDir, ILogger Logger, bool bForce = false)
 		{
 			if (!Directory.Exists(SourceDir))
 			{
@@ -918,7 +921,7 @@
 				// make the dst filename with the same structure as it was in SourceDir
 				string DestFilename = Path.Combine(DestDir, Utils.MakePathRelativeTo(Filename, SourceDir));
 
-				if (bForce || FilesAreDifferent(Filename, DestFilename))
+				if (bForce || FilesAreDifferent(Filename, DestFilename, Logger))
 				{
 					if (File.Exists(DestFilename))
 					{
@@ -943,7 +946,7 @@
 			}
 		}
 
-		private static void DeleteFiles(string Filespec)
+		private static void DeleteFiles(string Filespec, ILogger Logger)
 		{
 			string BaseDir = Path.GetDirectoryName(Filespec)!;
 			string Mask = Path.GetFileName(Filespec);
@@ -958,7 +961,7 @@
 			{
 				File.SetAttributes(Filename, FileAttributes.Normal);
 				File.Delete(Filename);
-				Log.TraceInformation("\nDeleted file {0}", Filename);
+				Logger.LogInformation("\nDeleted file {Filename}", Filename);
 			}
 		}
 
@@ -1065,7 +1068,7 @@
 				XElement Node = ExecutionStack.Pop();
 				if (bGlobalTrace || CurrentContext.bTrace)
 				{
-					Log.TraceInformation("Execute: '{0}'", TraceNodeString(Node));
+					Logger.LogInformation("Execute: '{Node}'", TraceNodeString(Node));
 				}
 				switch (Node.Name.ToString())
 				{
@@ -1077,7 +1080,7 @@
 								CurrentContext.bTrace = StringToBool(Enable);
 								if (!bGlobalTrace && CurrentContext.bTrace)
 								{
-									Log.TraceInformation("Context: '{0}' using Architecture='{1}', NodeName='{2}', Input='{3}'", CurrentContext.StringVariables["PluginDir"], Architecture, NodeName, Input);
+									Logger.LogInformation("Context: '{Context}' using Architecture='{Arch}', NodeName='{NodeName}', Input='{Input}'", CurrentContext.StringVariables["PluginDir"], Architecture, NodeName, Input);
 								}
 							}
 						}
@@ -1087,9 +1090,9 @@
 						{
 							if (!bGlobalTrace && !CurrentContext.bTrace)
 							{
-								Log.TraceInformation("Context: '{0}' using Architecture='{1}', NodeName='{2}', Input='{3}'", CurrentContext.StringVariables["PluginDir"], Architecture, NodeName, Input);
+								Logger.LogInformation("Context: '{Context}' using Architecture='{Arch}', NodeName='{NodeName}', Input='{Input}'", CurrentContext.StringVariables["PluginDir"], Architecture, NodeName, Input);
 							}
-							Log.TraceInformation("Variables:\n{0}", DumpContext(CurrentContext));
+							Logger.LogInformation("Variables:\n{Variables}", DumpContext(CurrentContext));
 						}
 						break;
 
@@ -1106,7 +1109,7 @@
 
 							if (bGlobalTrace || CurrentContext.bTrace)
 							{
-								Log.TraceInformation("Context: '{0}' using Architecture='{1}', NodeName='{2}', Input='{3}'", CurrentContext.StringVariables["PluginDir"], Architecture, NodeName, Input);
+								Logger.LogInformation("Context: '{Context}' using Architecture='{Arch}', NodeName='{NodeName}', Input='{Input}'", CurrentContext.StringVariables["PluginDir"], Architecture, NodeName, Input);
 							}
 						}
 						break;
@@ -1233,7 +1236,7 @@
 							string? Text = GetAttribute(CurrentContext, Node, "text");
 							if (Text != null)
 							{
-								Log.TraceInformation("{0}", Text);
+								Logger.LogInformation("{Message}", Text);
 							}
 						}
 						break;
@@ -1272,7 +1275,7 @@
 										{
 											if (!GlobalContext.ElementVariables.TryGetValue(Tag.Substring(1), out Target))
 											{
-												Log.TraceWarning("\nMissing element variable '{0}' in '{1}' (skipping instruction)", Tag, TraceNodeString(Node));
+												Logger.LogWarning("Missing element variable '{Tag}' in '{Node}' (skipping instruction)", Tag, TraceNodeString(Node));
 												continue;
 											}
 										}
@@ -1309,7 +1312,7 @@
 										{
 											if (!GlobalContext.ElementVariables.TryGetValue(Tag.Substring(1), out Target))
 											{
-												Log.TraceWarning("\nMissing element variable '{0}' in '{1}' (skipping instruction)", Tag, TraceNodeString(Node));
+												Logger.LogInformation("\nMissing element variable '{Tag}' in '{Node}' (skipping instruction)", Tag, TraceNodeString(Node));
 												continue;
 											}
 										}
@@ -1545,7 +1548,7 @@
 								{
 									if (!GlobalContext.ElementVariables.TryGetValue(Name, out Element))
 									{
-										Log.TraceWarning("\nMissing element variable '{0}' in '{1}' (skipping instruction)", Name, TraceNodeString(Node));
+										Logger.LogWarning("Missing element variable '{Name}' in '{Node}' (skipping instruction)", Name, TraceNodeString(Node));
 										continue;
 									}
 								}
@@ -1558,7 +1561,7 @@
 										{
 											if (!GlobalContext.ElementVariables.TryGetValue(Tag.Substring(1), out Target))
 											{
-												Log.TraceWarning("\nMissing element variable '{0}' in '{1}' (skipping instruction)", Tag, TraceNodeString(Node));
+												Logger.LogWarning("Missing element variable '{Name}' in '{Node}' (skipping instruction)", Tag, TraceNodeString(Node));
 												continue;
 											}
 										}
@@ -1606,7 +1609,7 @@
 										{
 											if (!GlobalContext.ElementVariables.TryGetValue(Tag.Substring(1), out Target))
 											{
-												Log.TraceWarning("\nMissing element variable '{0}' in '{1}' (skipping instruction)", Tag, TraceNodeString(Node));
+												Logger.LogWarning("Missing element variable '{Tag}' in '{Node}' (skipping instruction)", Tag, TraceNodeString(Node));
 												continue;
 											}
 										}
@@ -1714,7 +1717,7 @@
 								if (File.Exists(Src))
 								{
 									// check to see if newer than last time we copied
-									if (bForce || FilesAreDifferent(Src, Dst))
+									if (bForce || FilesAreDifferent(Src, Dst, Logger))
 									{
 										if (File.Exists(Dst))
 										{
@@ -1723,7 +1726,7 @@
 										}
 										Directory.CreateDirectory(Path.GetDirectoryName(Dst)!);
 										File.Copy(Src, Dst, true);
-										Log.TraceInformation("\nFile {0} copied to {1}", Src, Dst);
+										Logger.LogInformation("File {Src} copied to {Dst}", Src, Dst);
 
 										// remove any read only flags and keep timestamp
 										FileInfo DestFileInfo = new FileInfo(Dst);
@@ -1733,7 +1736,7 @@
 								}
 								else
 								{
-									Log.TraceInformation("\nFile {0} does not exist, not copied!", Src);
+									Logger.LogInformation("File {Src} does not exist, not copied!", Src);
 								}
 							}
 						}
@@ -1746,8 +1749,8 @@
 							bool bForce = StringToBool(GetAttribute(CurrentContext, Node, "force", true, false, "true"));
 							if (Src != null && Dst != null)
 							{
-								CopyFileDirectory(Src, Dst, bForce);
-								Log.TraceInformation("\nDirectory {0} copied to {1}", Src, Dst, bForce);
+								CopyFileDirectory(Src, Dst, Logger, bForce);
+								Logger.LogInformation("\nDirectory {Src} copied to {Dst} ({Force})", Src, Dst, bForce);
 							}
 						}
 						break;
@@ -1759,12 +1762,12 @@
 							{
 								if (Filespec.Contains(":") || Filespec.Contains(".."))
 								{
-									Log.TraceInformation("\nFilespec {0} not allowed; ignored.", Filespec);
+									Logger.LogInformation("\nFilespec {FileSpec} not allowed; ignored.", Filespec);
 								}
 								else
 								{
 									// force relative to BuildDir (and only from global context so someone doesn't try to be clever)
-									DeleteFiles(Path.Combine(GlobalContext.StringVariables["BuildDir"], Filespec));
+									DeleteFiles(Path.Combine(GlobalContext.StringVariables["BuildDir"], Filespec), Logger);
 								}
 							}
 						}
@@ -2205,7 +2208,7 @@
 										{
 											if (!GlobalContext.ElementVariables.TryGetValue(Tag.Substring(1), out Element))
 											{
-												Log.TraceWarning("\nMissing element variable '{0}' in '{1}' (skipping instruction)", Tag, TraceNodeString(Node));
+												Logger.LogWarning("Missing element variable '{Tag}' in '{Node}' (skipping instruction)", Tag, TraceNodeString(Node));
 												continue;
 											}
 										}
@@ -2232,7 +2235,7 @@
 										{
 											if (!GlobalContext.ElementVariables.TryGetValue(Tag.Substring(1), out Element))
 											{
-												Log.TraceWarning("\nMissing element variable '{0}' in '{1}' (skipping instruction)", Tag, TraceNodeString(Node));
+												Logger.LogWarning("Missing element variable '{Tag}' in '{Node}' (skipping instruction)", Tag, TraceNodeString(Node));
 												continue;
 											}
 										}
@@ -2271,7 +2274,7 @@
 										{
 											if (!GlobalContext.ElementVariables.TryGetValue(Tag.Substring(1), out Element))
 											{
-												Log.TraceWarning("\nMissing element variable '{0}' in '{1}' (skipping instruction)", Tag, TraceNodeString(Node));
+												Logger.LogWarning("Missing element variable '{Tag}' in '{Node}' (skipping instruction)", Tag, TraceNodeString(Node));
 												continue;
 											}
 										}
@@ -2280,7 +2283,7 @@
 
 								if (Element.Value == null)
 								{
-									Log.TraceWarning("\nExpected text in element '{0}' in '{1}' but found none (skipping instruction)", Element.Name.ToString(), TraceNodeString(Node));
+									Logger.LogWarning("Expected text in element '{Element}' in '{Node}' but found none (skipping instruction)", Element.Name.ToString(), TraceNodeString(Node));
 									continue;
 								}
 
@@ -2413,7 +2416,7 @@
 									}
 									catch (Exception e)
 									{
-										Log.TraceError("\nXML parsing {0} failed! {1} (skipping instruction)", Parse, e);
+										Logger.LogError(e, "XML parsing {Parse} failed! {Ex} (skipping instruction)", Parse, e);
 									}
 								}
 							}
@@ -2421,7 +2424,7 @@
 						break;
 
 					default:
-						Log.TraceWarning("\nUnknown command: {0}", Node.Name);
+						Logger.LogWarning("Unknown command: {Name}", Node.Name);
 						break;
 				}
 			}
@@ -2487,13 +2490,13 @@
 					}
 				}
 
-				Log.TraceInformation("UPL Init: {0}", Arch);
+				Logger.LogInformation("UPL Init: {Arch}", Arch);
 				ProcessPluginNode(Arch, "init", "");
 			}
 
 			if (bGlobalTrace)
 			{
-				Log.TraceInformation("\nVariables:\n{0}", DumpVariables());
+				Logger.LogInformation("\nVariables:\n{Variables}", DumpVariables());
 			}
 		}
 
@@ -2568,7 +2571,7 @@
 		static Dictionary<string, ConfigCacheIni_UPL> BaseIniCache = new Dictionary<string, ConfigCacheIni_UPL>();
 
 		// static creation functions for ini files
-		public static ConfigCacheIni_UPL CreateConfigCacheIni_UPL(UnrealTargetPlatform Platform, string BaseIniName, DirectoryReference? ProjectDirectory, DirectoryReference? EngineDirectory = null)
+		public static ConfigCacheIni_UPL CreateConfigCacheIni_UPL(UnrealTargetPlatform Platform, string BaseIniName, DirectoryReference? ProjectDirectory, ILogger Logger, DirectoryReference? EngineDirectory = null)
 		{
 			if (EngineDirectory == null)
 			{
@@ -2578,7 +2581,7 @@
 			// cache base ini for use as the seed for the rest
 			if (!BaseIniCache.ContainsKey(BaseIniName))
 			{
-				BaseIniCache.Add(BaseIniName, new ConfigCacheIni_UPL(BuildHostPlatform.Current.Platform, BaseIniName, null, EngineDirectory, EngineOnly: true));
+				BaseIniCache.Add(BaseIniName, new ConfigCacheIni_UPL(BuildHostPlatform.Current.Platform, BaseIniName, null, Logger, EngineDirectory, EngineOnly: true));
 			}
 
 			// build the new ini and cache it for later re-use
@@ -2586,7 +2589,7 @@
 			string Key = GetIniPlatformName(Platform) + BaseIniName + EngineDirectory.FullName + (ProjectDirectory != null ? ProjectDirectory.FullName : "");
 			if (!IniCache.ContainsKey(Key))
 			{
-				IniCache.Add(Key, new ConfigCacheIni_UPL(Platform, BaseIniName, ProjectDirectory, EngineDirectory, BaseCache: BaseCache));
+				IniCache.Add(Key, new ConfigCacheIni_UPL(Platform, BaseIniName, ProjectDirectory, Logger, EngineDirectory, BaseCache: BaseCache));
 			}
 			return IniCache[Key];
 		}
@@ -2650,10 +2653,11 @@
 		/// Constructor. Parses a single ini file. No Platform settings, no engine hierarchy. Do not use this with ini files that have hierarchy!
 		/// </summary>
 		/// <param name="Filename">The ini file to load</param>
-		public ConfigCacheIni_UPL(FileReference Filename) : this()
+		/// <param name="Logger"></param>
+		public ConfigCacheIni_UPL(FileReference Filename, ILogger Logger) : this()
 		{
 			bIsMergingConfigs = false;
-			ParseIniFile(Filename);
+			ParseIniFile(Filename, Logger);
 		}
 
 		/// <summary>
@@ -2661,9 +2665,10 @@
 		/// </summary>
 		/// <param name="BaseIniName">Ini name (Engine, Editor, etc)</param>
 		/// <param name="ProjectDirectory">Project path</param>
+		/// <param name="Logger">Logger for </param>
 		/// <param name="EngineDirectory"></param>
-		public ConfigCacheIni_UPL(string BaseIniName, string? ProjectDirectory, string? EngineDirectory = null)
-			: this(BuildHostPlatform.Current.Platform, BaseIniName, ProjectDirectory, EngineDirectory)
+		public ConfigCacheIni_UPL(string BaseIniName, string? ProjectDirectory, ILogger Logger, string? EngineDirectory = null)
+			: this(BuildHostPlatform.Current.Platform, BaseIniName, ProjectDirectory, Logger, EngineDirectory)
 		{
 		}
 
@@ -2672,9 +2677,10 @@
 		/// </summary>
 		/// <param name="BaseIniName">Ini name (Engine, Editor, etc)</param>
 		/// <param name="ProjectDirectory">Project path</param>
+		/// <param name="Logger">Logger for output</param>
 		/// <param name="EngineDirectory"></param>
-		public ConfigCacheIni_UPL(string BaseIniName, DirectoryReference ProjectDirectory, DirectoryReference? EngineDirectory = null)
-			: this(BuildHostPlatform.Current.Platform, BaseIniName, ProjectDirectory, EngineDirectory)
+		public ConfigCacheIni_UPL(string BaseIniName, DirectoryReference ProjectDirectory, ILogger Logger, DirectoryReference? EngineDirectory = null)
+			: this(BuildHostPlatform.Current.Platform, BaseIniName, ProjectDirectory, Logger, EngineDirectory)
 		{
 		}
 
@@ -2684,10 +2690,12 @@
 		/// <param name="ProjectDirectory">Project path</param>
 		/// <param name="Platform">Target platform</param>
 		/// <param name="BaseIniName">Ini name (Engine, Editor, etc)</param>
+		/// <param name="Logger"></param>
 		/// <param name="EngineDirectory"></param>
-		public ConfigCacheIni_UPL(UnrealTargetPlatform Platform, string BaseIniName, string? ProjectDirectory, string? EngineDirectory = null)
+		public ConfigCacheIni_UPL(UnrealTargetPlatform Platform, string BaseIniName, string? ProjectDirectory, ILogger Logger, string? EngineDirectory = null)
 			: this(Platform, BaseIniName, 
 				  (ProjectDirectory == null) ? null : new DirectoryReference(ProjectDirectory), 
+				  Logger,
 				  (EngineDirectory == null) ? null : new DirectoryReference(EngineDirectory))
 		{
 		}
@@ -2698,10 +2706,11 @@
 		/// <param name="ProjectDirectory">Project path</param>
 		/// <param name="Platform">Target platform</param>
 		/// <param name="BaseIniName">Ini name (Engine, Editor, etc)</param>
+		/// <param name="Logger"></param>
 		/// <param name="EngineDirectory"></param>
 		/// <param name="EngineOnly"></param>
 		/// <param name="BaseCache"></param>
-		public ConfigCacheIni_UPL(UnrealTargetPlatform Platform, string BaseIniName, DirectoryReference? ProjectDirectory, DirectoryReference? EngineDirectory = null, bool EngineOnly = false, ConfigCacheIni_UPL? BaseCache = null) : this()
+		public ConfigCacheIni_UPL(UnrealTargetPlatform Platform, string BaseIniName, DirectoryReference? ProjectDirectory, ILogger Logger, DirectoryReference? EngineDirectory = null, bool EngineOnly = false, ConfigCacheIni_UPL? BaseCache = null) : this()
 		{
 			bIsMergingConfigs = true;
 			if (EngineDirectory == null)
@@ -2722,7 +2731,7 @@
 				{
 					if (FileReference.Exists(IniFileName))
 					{
-						ParseIniFile(IniFileName);
+						ParseIniFile(IniFileName, Logger);
 					}
 				}
 			}
@@ -2732,7 +2741,7 @@
 				{
 					if (FileReference.Exists(IniFileName))
 					{
-						ParseIniFile(IniFileName);
+						ParseIniFile(IniFileName, Logger);
 					}
 				}
 			}
@@ -3023,7 +3032,7 @@
 		/// <summary>
 		/// Loads and parses ini file.
 		/// </summary>
-		public void ParseIniFile(FileReference Filename)
+		public void ParseIniFile(FileReference Filename, ILogger Logger)
 		{
 			String[]? IniLines = null;
 			List<Command>? Commands = null;
@@ -3037,7 +3046,7 @@
 				}
 				catch (Exception ex)
 				{
-					Log.TraceInformation("Error reading ini file: " + Filename + " Exception: " + ex.Message);
+					Logger.LogInformation("Error reading ini file: {Filename} Exception: {Ex}", Filename, ex.Message);
 				}
 			}
 			else
diff --git a/Engine/Source/Programs/UnrealBuildTool/System/Utils.cs b/Engine/Source/Programs/UnrealBuildTool/System/Utils.cs
--- a/Engine/Source/Programs/UnrealBuildTool/System/Utils.cs
+++ b/Engine/Source/Programs/UnrealBuildTool/System/Utils.cs
@@ -14,6 +14,7 @@
 using EpicGames.Core;
 using UnrealBuildBase;
 using System.Collections.Concurrent;
+using Microsoft.Extensions.Logging;
 
 namespace UnrealBuildTool
 {
@@ -274,29 +275,29 @@
 		/// <summary>
 		/// Runs a local process and pipes the output to the log
 		/// </summary>
-		public static int RunLocalProcessAndLogOutput(ProcessStartInfo StartInfo)
+		public static int RunLocalProcessAndLogOutput(ProcessStartInfo StartInfo, ILogger Logger)
 		{
 			Process LocalProcess = new Process();
 			LocalProcess.StartInfo = StartInfo;
-			LocalProcess.OutputDataReceived += (Sender, Args) => { LocalProcessOutput(Args, false); };
-			LocalProcess.ErrorDataReceived += (Sender, Args) => { LocalProcessOutput(Args, true); };
+			LocalProcess.OutputDataReceived += (Sender, Args) => { LocalProcessOutput(Args, false, Logger); };
+			LocalProcess.ErrorDataReceived += (Sender, Args) => { LocalProcessOutput(Args, true, Logger); };
 			return RunLocalProcess(LocalProcess);
 		}
 
 		/// <summary>
 		/// Output a line of text from a local process. Implemented as a separate function to give a useful function name in the UAT log prefix.
 		/// </summary>
-		static void LocalProcessOutput(DataReceivedEventArgs Args, bool bIsError)
+		static void LocalProcessOutput(DataReceivedEventArgs Args, bool bIsError, ILogger Logger)
 		{
 			if(Args != null && Args.Data != null)
 			{
 				if(bIsError)
 				{
-					Log.TraceError(Args.Data.TrimEnd());
+					Logger.LogError("{Message}", Args.Data.TrimEnd());
 				}
 				else
 				{
-					Log.TraceInformation(Args.Data.TrimEnd());
+					Logger.LogInformation("{Message}", Args.Data.TrimEnd());
 				}
 			}
 		}
@@ -304,7 +305,7 @@
 		/// <summary>
 		/// Runs a local process and pipes the output to a file
 		/// </summary>
-		public static int RunLocalProcessAndPrintfOutput(ProcessStartInfo StartInfo)
+		public static int RunLocalProcessAndPrintfOutput(ProcessStartInfo StartInfo, ILogger Logger)
 		{
 			string AppName = Path.GetFileNameWithoutExtension(StartInfo.FileName);
 			string LogFilenameBase = string.Format("{0}_{1}", AppName, DateTime.Now.ToString("yyyy.MM.dd-HH.mm.ss"));
@@ -368,7 +369,7 @@
 					}
 					else
 					{
-						Log.TraceInformation(data);
+						Logger.LogInformation("{Output}", data);
 					}
 				}
 			};
@@ -379,7 +380,7 @@
 			var ExitCode = RunLocalProcess(LocalProcess);
 			if(ExitCode != 0 && !string.IsNullOrEmpty(LogFilename))
 			{
-				Log.TraceError("Process \'{0}\' failed. Details are in \'{1}\'", AppName, LogFilename);
+				Logger.LogError("Process \'{AppName}\' failed. Details are in \'{LogFilename}\'", AppName, LogFilename);
 			}
 
 			return ExitCode;
@@ -389,9 +390,9 @@
 		/// <summary>
 		/// Runs a local process and pipes the output to the log
 		/// </summary>
-		public static int RunLocalProcessAndLogOutput(string Command, string Args)
+		public static int RunLocalProcessAndLogOutput(string Command, string Args, ILogger Logger)
 		{
-			return RunLocalProcessAndLogOutput(new ProcessStartInfo(Command, Args));
+			return RunLocalProcessAndLogOutput(new ProcessStartInfo(Command, Args), Logger);
 		}
 
 		/// <summary>
@@ -400,10 +401,19 @@
 		/// <returns>The entire StdOut generated from the process as a single trimmed string</returns>
 		/// <param name="Command">Command to run</param>
 		/// <param name="Args">Arguments to Command</param>
-		public static string RunLocalProcessAndReturnStdOut(string Command, string Args)
+		public static string RunLocalProcessAndReturnStdOut(string Command, string Args) => RunLocalProcessAndReturnStdOut(Command, Args, null);
+
+		/// <summary>
+		/// Runs a command line process, and returns simple StdOut output. This doesn't handle errors or return codes
+		/// </summary>
+		/// <returns>The entire StdOut generated from the process as a single trimmed string</returns>
+		/// <param name="Command">Command to run</param>
+		/// <param name="Args">Arguments to Command</param>
+		/// <param name="Logger">Logger for output</param>
+		public static string RunLocalProcessAndReturnStdOut(string Command, string Args, ILogger? Logger)
 		{
 			int ExitCode;
-			return RunLocalProcessAndReturnStdOut(Command, Args, out ExitCode);	
+			return RunLocalProcessAndReturnStdOut(Command, Args, Logger, out ExitCode);	
 		}
 
 		/// <summary>
@@ -413,8 +423,27 @@
 		/// <param name="Command">Command to run</param>
 		/// <param name="Args">Arguments to Command</param>
 		/// <param name="ExitCode">The return code from the process after it exits</param>
+		public static string RunLocalProcessAndReturnStdOut(string Command, string? Args, out int ExitCode) => RunLocalProcessAndReturnStdOut(Command, Args, null, out ExitCode);
+
+		/// <summary>
+		/// Runs a command line process, and returns simple StdOut output.
+		/// </summary>
+		/// <returns>The entire StdOut generated from the process as a single trimmed string</returns>
+		/// <param name="Command">Command to run</param>
+		/// <param name="Args">Arguments to Command</param>
+		/// <param name="ExitCode">The return code from the process after it exits</param>
 		/// <param name="LogOutput">Whether to also log standard output and standard error</param>
-		public static string RunLocalProcessAndReturnStdOut(string Command, string? Args, out int ExitCode, bool LogOutput = false)
+		public static string RunLocalProcessAndReturnStdOut(string Command, string? Args, out int ExitCode, bool LogOutput) => RunLocalProcessAndReturnStdOut(Command, Args, LogOutput? Log.Logger : null, out ExitCode);
+
+		/// <summary>
+		/// Runs a command line process, and returns simple StdOut output.
+		/// </summary>
+		/// <returns>The entire StdOut generated from the process as a single trimmed string</returns>
+		/// <param name="Command">Command to run</param>
+		/// <param name="Args">Arguments to Command</param>
+		/// <param name="Logger">Logger for output. No output if null.</param>
+		/// <param name="ExitCode">The return code from the process after it exits</param>
+		public static string RunLocalProcessAndReturnStdOut(string Command, string? Args, ILogger? Logger, out int ExitCode)
 		{
 			// Process Arguments follow windows conventions in .NET Core
 			// Which means single quotes ' are not considered quotes.
@@ -442,16 +471,16 @@
 				StreamReader ErrorReader = LocalProcess.StandardError;
 				// trim off any extraneous new lines, helpful for those one-line outputs
 				ErrorOutput = ErrorReader.ReadToEnd().Trim();
-				if (LogOutput)
+				if (Logger != null)
 				{
 					if(FullOutput.Length > 0)
 					{
-						Log.TraceInformation(FullOutput);
+						Logger.LogInformation("{Output}", FullOutput);
 					}
 
 					if (ErrorOutput.Length > 0)
 					{
-						Log.TraceError(ErrorOutput);
+						Logger.LogError("{Output}", ErrorOutput);
 					}
 				}
 
@@ -497,9 +526,10 @@
 		/// </summary>
 		/// <param name="SupportedPlatforms">List of supported platforms</param>
 		/// <param name="bIncludeUnbuildablePlatforms">If true, add platforms that are present but not available for compiling</param>
+		/// <param name="Logger"></param>
 		/// 
 		/// <returns>List of unsupported platforms in string format</returns>
-		public static List<string> MakeListOfUnsupportedPlatforms(List<UnrealTargetPlatform> SupportedPlatforms, bool bIncludeUnbuildablePlatforms)
+		public static List<string> MakeListOfUnsupportedPlatforms(List<UnrealTargetPlatform> SupportedPlatforms, bool bIncludeUnbuildablePlatforms, ILogger Logger)
 		{
 			// Make a list of all platform name strings that we're *not* currently compiling, to speed
 			// up file path comparisons later on
@@ -801,8 +831,9 @@
 		/// </summary>
 		/// <typeparam name="T">The type to read</typeparam>
 		/// <param name="FileName">The XML file to read from</param>
+		/// <param name="Logger">Logger for output</param>
 		/// <returns>New deserialized instance of type T</returns>
-		static public T ReadClass<T>(string FileName) where T : new()
+		static public T ReadClass<T>(string FileName, ILogger Logger) where T : new()
 		{
 			T Instance = new T();
 			StreamReader? XmlStream = null;
@@ -822,7 +853,7 @@
 			}
 			catch (Exception E)
 			{
-				Log.TraceInformation(E.Message);
+				Logger.LogInformation("{Output}", E.Message);
 			}
 			finally
 			{
@@ -843,8 +874,9 @@
 		/// <param name="Data">Object to write</param>
 		/// <param name="FileName">File to write to</param>
 		/// <param name="DefaultNameSpace">Default namespace for the output elements</param>
+		/// <param name="Logger">Logger for output</param>
 		/// <returns>True if the file was written successfully</returns>
-		static public bool WriteClass<T>(T Data, string FileName, string DefaultNameSpace)
+		static public bool WriteClass<T>(T Data, string FileName, string DefaultNameSpace, ILogger Logger)
 		{
 			bool bSuccess = true;
 			StreamWriter? XmlStream = null;
@@ -873,7 +905,7 @@
 			}
 			catch (Exception E)
 			{
-				Log.TraceInformation(E.Message);
+				Logger.LogInformation("{Message}", E.Message);
 				bSuccess = false;
 			}
 			finally
@@ -1300,8 +1332,9 @@
 		/// Executes a list of custom build step scripts
 		/// </summary>
 		/// <param name="ScriptFiles">List of script files to execute</param>
+		/// <param name="Logger">Logger for output</param>
 		/// <returns>True if the steps succeeded, false otherwise</returns>
-		public static void ExecuteCustomBuildSteps(FileReference[] ScriptFiles)
+		public static void ExecuteCustomBuildSteps(FileReference[] ScriptFiles, ILogger Logger)
 		{
 			UnrealTargetPlatform HostPlatform = BuildHostPlatform.Current.Platform;
 			foreach(FileReference ScriptFile in ScriptFiles)
@@ -1318,7 +1351,7 @@
 					StartInfo.Arguments = String.Format("\"{0}\"", ScriptFile.FullName);
 				}
 
-				int ReturnCode = Utils.RunLocalProcessAndLogOutput(StartInfo);
+				int ReturnCode = Utils.RunLocalProcessAndLogOutput(StartInfo, Logger);
 				if(ReturnCode != 0)
 				{
 					throw new BuildException("Custom build step {0} {1} terminated with exit code {2}", StartInfo.FileName, StartInfo.Arguments, ReturnCode);
@@ -1412,11 +1445,34 @@
 		/// </summary>
 		/// <param name="Location">Location of the file</param>
 		/// <param name="Contents">New contents of the file</param>
+		/// <param name="Logger">Logger for output</param>
+		internal static void WriteFileIfChanged(FileReference Location, string Contents, ILogger Logger)
+		{
+			WriteFileIfChanged(Location, Contents, StringComparison.Ordinal, Logger);
+		}
+
+		/// <summary>
+		/// Writes a file if the contents have changed
+		/// </summary>
+		/// <param name="Location">Location of the file</param>
+		/// <param name="Contents">New contents of the file</param>
 		/// <param name="Comparison">The type of string comparison to use</param>
-		internal static void WriteFileIfChanged(FileReference Location, string Contents, StringComparison Comparison = StringComparison.Ordinal)
+		/// <param name="Logger">Logger for output</param>
+		internal static void WriteFileIfChanged(FileReference Location, string Contents, StringComparison Comparison, ILogger Logger)
 		{
 			FileItem FileItem = FileItem.GetItemByFileReference(Location);
-			WriteFileIfChanged(FileItem, Contents, Comparison);
+			WriteFileIfChanged(FileItem, Contents, Comparison, Logger);
+		}
+
+		/// <summary>
+		/// Writes a file if the contents have changed
+		/// </summary>
+		/// <param name="Location">Location of the file</param>
+		/// <param name="ContentLines">New contents of the file</param>
+		/// <param name="Logger">Logger for output</param>
+		internal static void WriteFileIfChanged(FileReference Location, IEnumerable<string> ContentLines, ILogger Logger)
+		{
+			WriteFileIfChanged(Location, ContentLines, StringComparison.Ordinal, Logger);
 		}
 
 		/// <summary>
@@ -1425,10 +1481,11 @@
 		/// <param name="Location">Location of the file</param>
 		/// <param name="ContentLines">New contents of the file</param>
 		/// <param name="Comparison">The type of string comparison to use</param>
-		internal static void WriteFileIfChanged(FileReference Location, IEnumerable<string> ContentLines, StringComparison Comparison = StringComparison.Ordinal)
+		/// <param name="Logger">Logger for output</param>
+		internal static void WriteFileIfChanged(FileReference Location, IEnumerable<string> ContentLines, StringComparison Comparison, ILogger Logger)
 		{
 			FileItem FileItem = FileItem.GetItemByFileReference(Location);
-			WriteFileIfChanged(FileItem, ContentLines, Comparison);
+			WriteFileIfChanged(FileItem, ContentLines, Comparison, Logger);
 		}
 
 		/// <summary>
@@ -1439,13 +1496,13 @@
 		static internal FileReference? WriteFileIfChangedTrace = null;
 		static internal string WriteFileIfChangedContext = "";
 
-		static void RecordWriteFileIfChanged(FileReference File, bool bNew, bool bChanged)
+		static void RecordWriteFileIfChanged(FileReference File, bool bNew, bool bChanged, ILogger Logger)
 		{
 			int NewWriteRequestCount = 1;
 			int NewActualWriteCount = bChanged ? 1 : 0;
 
 			bool bOverrideLogEventType = FileReference.Equals(WriteFileIfChangedTrace, File);
-			LogEventType OverrideType = LogEventType.Console;
+			LogLevel OverrideType = LogLevel.Information;
 
 			string Prefix = "";
 			if (bOverrideLogEventType)
@@ -1472,22 +1529,22 @@
 					{
 						if (bNew)
 						{
-							Log.WriteLine(bOverrideLogEventType ? OverrideType : LogEventType.Warning,
-								$"{Prefix}Writing a file that previously existed was not overwritten and then removed: \"{File}\"{Context}");
+							Logger.Log(bOverrideLogEventType ? OverrideType : LogLevel.Warning,
+								"{Prefix}Writing a file that previously existed was not overwritten and then removed: \"{File}\"{Context}", Prefix, File, Context);
 						}
 						else
 						{
 							if (bChanged)
 							{
-								Log.WriteLine(bOverrideLogEventType ? OverrideType : LogEventType.Warning,
-									$"{Prefix}Writing a file that previously was not written \"{File}\"{Context}");
+								Logger.Log(bOverrideLogEventType ? OverrideType : LogLevel.Warning,
+									"{Prefix}Writing a file that previously was not written \"{File}\"{Context}", Prefix, File, Context);
 							}
 							else
 							{
 								if (bOverrideLogEventType)
 								{
-									Log.WriteLine(OverrideType,
-										$"{Prefix}Not writing a file that was previously not written: \"{File}\"{Context}");
+									Logger.Log(OverrideType,
+										"{Prefix}Not writing a file that was previously not written: \"{File}\"{Context}", Prefix, File, Context);
 								}
 
 							}
@@ -1497,22 +1554,22 @@
 					{
 						if (bNew)
 						{
-							Log.WriteLine(bOverrideLogEventType ? OverrideType : LogEventType.Warning,
-								$"{Prefix}Re-writing a file that was previously written and then removed: \"{File}\"{Context}");
+							Logger.Log(bOverrideLogEventType ? OverrideType : LogLevel.Warning,
+								"{Prefix}Re-writing a file that was previously written and then removed: \"{File}\"{Context}", Prefix, File, Context);
 						}
 						else
 						{
 							if (bChanged)
 							{
-								Log.WriteLine(bOverrideLogEventType ? OverrideType : LogEventType.Warning,
-									$"{Prefix}Re-writing a file that was previously written: \"{File}\"{Context}");
+								Logger.Log(bOverrideLogEventType ? OverrideType : LogLevel.Warning,
+									"{Prefix}Re-writing a file that was previously written: \"{File}\"{Context}", Prefix, File, Context);
 							}
 							else
 							{
 								if (bOverrideLogEventType)
 								{
-									Log.WriteLine(OverrideType,
-										$"{Prefix}Not writing a file that was previously written: \"{File}\"{Context}");
+									Logger.Log(OverrideType,
+										"{Prefix}Not writing a file that was previously written: \"{File}\"{Context}", Prefix, File, Context);
 								}
 							}
 						}
@@ -1524,17 +1581,17 @@
 					{
 						if (bNew)
 						{
-							Log.TraceInformation($"{Prefix}Writing new file: \"{File}\"{Context}");
+							Logger.LogInformation("{Prefix}Writing new file: \"{File}\"{Context}", Prefix, File, Context);
 						}
 						else
 						{
 							if (bChanged)
 							{
-								Log.TraceInformation($"{Prefix}Writing changed file: \"{File}\"{Context}");
+								Logger.LogInformation("{Prefix}Writing changed file: \"{File}\"{Context}", Prefix, File, Context);
 							}
 							else
 							{
-								Log.TraceInformation($"{Prefix}Not writing unchanged file: \"{File}\"{Context}");
+								Logger.LogInformation("{Prefix}Not writing unchanged file: \"{File}\"{Context}", Prefix, File, Context);
 							}
 						}
 					}
@@ -1544,7 +1601,7 @@
 			}
 		}
 
-		internal static void LogWriteFileIfChangedActivity()
+		internal static void LogWriteFileIfChangedActivity(ILogger Logger)
 		{
 			int TotalRequests = 0;
 			int TotalWrites = 0;
@@ -1554,7 +1611,18 @@
 				TotalWrites += Actual;
 			}
 
-			Log.TraceLog($"WriteFileIfChanged() wrote {TotalWrites} changed files of {TotalRequests} requested writes.");
+			Logger.LogDebug("WriteFileIfChanged() wrote {TotalWrites} changed files of {TotalRequests} requested writes.", TotalWrites, TotalRequests);
+		}
+
+		/// <summary>
+		/// Writes a file if the contents have changed
+		/// </summary>
+		/// <param name="FileItem">Location of the file</param>
+		/// <param name="Contents">New contents of the file</param>
+		/// <param name="Logger">Logger for output</param>
+		internal static void WriteFileIfChanged(FileItem FileItem, string Contents, ILogger Logger)
+		{
+			WriteFileIfChanged(FileItem, Contents, StringComparison.Ordinal, Logger);
 		}
 
 		/// <summary>
@@ -1563,7 +1631,8 @@
 		/// <param name="FileItem">Location of the file</param>
 		/// <param name="Contents">New contents of the file</param>
 		/// <param name="Comparison">The type of string comparison to use</param>
-		internal static void WriteFileIfChanged(FileItem FileItem, string Contents, StringComparison Comparison = StringComparison.Ordinal)
+		/// <param name="Logger">Logger for output</param>
+		internal static void WriteFileIfChanged(FileItem FileItem, string Contents, StringComparison Comparison, ILogger Logger)
 		{
 			// Only write the file if its contents have changed.
 			FileReference Location = FileItem.Location;
@@ -1573,7 +1642,7 @@
 				FileReference.WriteAllText(Location, Contents, GetEncodingForString(Contents));
 				FileItem.ResetCachedInfo();
 
-				RecordWriteFileIfChanged(FileItem.Location, bNew: true, bChanged: true);
+				RecordWriteFileIfChanged(FileItem.Location, bNew: true, bChanged: true, Logger);
 			}
 			else
 			{
@@ -1583,23 +1652,23 @@
 					FileReference BackupFile = new FileReference(FileItem.FullName + ".old");
 					try
 					{
-						Log.TraceLog("Updating {0}: contents have changed. Saving previous version to {1}.", FileItem, BackupFile);
+						Logger.LogDebug("Updating {File}: contents have changed. Saving previous version to {BackupFile}.", FileItem.Location, BackupFile);
 						FileReference.Delete(BackupFile);
 						FileReference.Move(Location, BackupFile);
 					}
 					catch (Exception Ex)
 					{
-						Log.TraceWarning("Unable to rename {0} to {1}", FileItem, BackupFile);
-						Log.TraceLog("{0}", ExceptionUtils.FormatExceptionDetails(Ex));
+						Logger.LogWarning("Unable to rename {FileItem} to {BackupFile}", FileItem, BackupFile);
+						Logger.LogDebug(Ex, "{Ex}", ExceptionUtils.FormatExceptionDetails(Ex));
 					}
 					FileReference.WriteAllText(Location, Contents, GetEncodingForString(Contents));
 					FileItem.ResetCachedInfo();
 
-					RecordWriteFileIfChanged(FileItem.Location, bNew: false, bChanged: true);
+					RecordWriteFileIfChanged(FileItem.Location, bNew: false, bChanged: true, Logger);
 				}
 				else
 				{ 
-					RecordWriteFileIfChanged(FileItem.Location, bNew: false, bChanged: false);
+					RecordWriteFileIfChanged(FileItem.Location, bNew: false, bChanged: false, Logger);
 				}
 			}
 		}
@@ -1609,8 +1678,20 @@
 		/// </summary>
 		/// <param name="FileItem">Location of the file</param>
 		/// <param name="ContentLines">New contents of the file</param>
+		/// <param name="Logger">Logger for output</param>
+		internal static void WriteFileIfChanged(FileItem FileItem, IEnumerable<string> ContentLines, ILogger Logger)
+		{
+			WriteFileIfChanged(FileItem, ContentLines, StringComparison.Ordinal, Logger);
+		}
+
+		/// <summary>
+		/// Writes a file if the contents have changed
+		/// </summary>
+		/// <param name="FileItem">Location of the file</param>
+		/// <param name="ContentLines">New contents of the file</param>
 		/// <param name="Comparison">The type of string comparison to use</param>
-		internal static void WriteFileIfChanged(FileItem FileItem, IEnumerable<string> ContentLines, StringComparison Comparison = StringComparison.Ordinal)
+		/// <param name="Logger">Logger for output</param>
+		internal static void WriteFileIfChanged(FileItem FileItem, IEnumerable<string> ContentLines, StringComparison Comparison, ILogger Logger)
 		{
 			// Only write the file if its contents have changed.
 			FileReference Location = FileItem.Location;
@@ -1621,7 +1702,7 @@
 				FileReference.WriteAllLines(Location, ContentLines, GetEncodingForStrings(ContentLines));
 				FileItem.ResetCachedInfo();
 
-				RecordWriteFileIfChanged(FileItem.Location, bNew: true, bChanged: true);
+				RecordWriteFileIfChanged(FileItem.Location, bNew: true, bChanged: true, Logger);
 			}
 			else
 			{
@@ -1631,23 +1712,23 @@
 					FileReference BackupFile = new FileReference($"{FileItem.FullName}.old");
 					try
 					{
-						Log.TraceLog($"Updating {FileItem}: contents have changed. Saving previous version to {BackupFile}.");
+						Logger.LogDebug("Updating {FileItem}: contents have changed. Saving previous version to {BackupFile}.", FileItem.Location, BackupFile);
 						FileReference.Delete(BackupFile);
 						FileReference.Move(Location, BackupFile);
 					}
 					catch (Exception Ex)
 					{
-						Log.TraceWarning($"Unable to rename {FileItem} to {BackupFile}");
-						Log.TraceLog(ExceptionUtils.FormatExceptionDetails(Ex));
+						Logger.LogWarning("Unable to rename {FileItem} to {BackupFile}", FileItem, BackupFile);
+						Logger.LogDebug(Ex, "{Ex}", ExceptionUtils.FormatExceptionDetails(Ex));
 					}
 					FileReference.WriteAllLines(Location, ContentLines, GetEncodingForStrings(ContentLines));
 					FileItem.ResetCachedInfo();
 
-					RecordWriteFileIfChanged(FileItem.Location, bNew: false, bChanged: true);
+					RecordWriteFileIfChanged(FileItem.Location, bNew: false, bChanged: true, Logger);
 				}
 				else
 				{ 
-					RecordWriteFileIfChanged(FileItem.Location, bNew: false, bChanged: false);
+					RecordWriteFileIfChanged(FileItem.Location, bNew: false, bChanged: false, Logger);
 				}
 			}
 		}
diff --git a/Engine/Source/Programs/UnrealBuildTool/System/XmlConfig.cs b/Engine/Source/Programs/UnrealBuildTool/System/XmlConfig.cs
--- a/Engine/Source/Programs/UnrealBuildTool/System/XmlConfig.cs
+++ b/Engine/Source/Programs/UnrealBuildTool/System/XmlConfig.cs
@@ -12,6 +12,7 @@
 using System.Xml.Schema;
 using System.Xml.Serialization;
 using EpicGames.Core;
+using Microsoft.Extensions.Logging;
 using UnrealBuildBase;
 
 namespace UnrealBuildTool
@@ -67,7 +68,8 @@
 		/// Initialize the config system with the given types
 		/// </summary>
 		/// <param name="OverrideCacheFile">Force use of the cached XML config without checking if it's valid (useful for remote builds)</param>
-		public static void ReadConfigFiles(FileReference? OverrideCacheFile)
+		/// <param name="Logger">Logger for output</param>
+		public static void ReadConfigFiles(FileReference? OverrideCacheFile, ILogger Logger)
 		{
 			// Find all the configurable types
 			List<Type> ConfigTypes = FindConfigurableTypes();
@@ -132,7 +134,7 @@
 						new Dictionary<Type, Dictionary<XmlConfigData.TargetMember, XmlConfigData.ValueInfo>>();
 					foreach(FileReference InputFile in InputFileLocations)
 					{
-						if(!TryReadFile(InputFile, CategoryToFields, TypeToValues, Schema))
+						if(!TryReadFile(InputFile, CategoryToFields, TypeToValues, Schema, Logger))
 						{
 							throw new BuildException("Failed to properly read XML file : {0}", InputFile.FullName);
 						}
@@ -241,6 +243,8 @@
 					return CachedInputFiles;
 				}
 
+				ILogger Logger = Log.Logger;
+				
 				// Find all the input file locations
 				List<InputFile> InputFilesFound = new List<InputFile>(4);
 
@@ -255,14 +259,14 @@
 					}
 					else
 					{
-						Log.TraceLog($"No config file at {NotForLicenseesConfigLocation}");
+						Logger.LogDebug("No config file at {NotForLicenseesConfigLocation}", NotForLicenseesConfigLocation);
 					}
 
 					// Check for the user config file under /Engine/Saved/UnrealBuildTool
 					FileReference UserConfigLocation = FileReference.Combine(Unreal.EngineDirectory, "Saved", "UnrealBuildTool", "BuildConfiguration.xml");
 					if (!FileReference.Exists(UserConfigLocation))
 					{
-						Log.TraceLog($"Creating default config file at {UserConfigLocation}");
+						Logger.LogDebug("Creating default config file at {UserConfigLocation}", UserConfigLocation);
 						CreateDefaultConfigFile(UserConfigLocation);
 					}
 					InputFilesFound.Add(new InputFile(UserConfigLocation, "User"));
@@ -275,7 +279,7 @@
 					FileReference AppDataConfigLocation = FileReference.Combine(new DirectoryReference(AppDataFolder), "Unreal Engine", "UnrealBuildTool", "BuildConfiguration.xml");
 					if (!FileReference.Exists(AppDataConfigLocation))
 					{
-						Log.TraceLog($"Creating default config file at {AppDataConfigLocation}");
+						Logger.LogDebug("Creating default config file at {AppDataConfigLocation}", AppDataConfigLocation);
 						CreateDefaultConfigFile(AppDataConfigLocation);
 					}
 					InputFilesFound.Add(new InputFile(AppDataConfigLocation, "Global (AppData)"));
@@ -292,16 +296,16 @@
 					}
 					else
 					{
-						Log.TraceLog($"No config file at {PersonalConfigLocation}");
+						Logger.LogDebug("No config file at {PersonalConfigLocation}", PersonalConfigLocation);
 					}
 				}
 
 				CachedInputFiles = InputFilesFound.ToArray();
 
-				Log.TraceLog("Configuration will be read from:");
+				Logger.LogDebug("Configuration will be read from:");
 				foreach (InputFile InputFile in InputFiles)
 				{
-					Log.TraceLog($"  {InputFile.Location.FullName}");
+					Logger.LogDebug("  {File}", InputFile.Location.FullName);
 				}
 
 				return CachedInputFiles;
@@ -329,6 +333,7 @@
 		/// <param name="TargetObject">The object instance to be configured</param>
 		public static void ApplyTo(object TargetObject)
 		{
+			ILogger Logger = Log.Logger;
 			for (Type? TargetType = TargetObject.GetType(); TargetType != null; TargetType = TargetType.BaseType)
 			{ 
 				XmlConfigData.ValueInfo[]? FieldValues;
@@ -343,16 +348,13 @@
 							// Check if setting has been deprecated
 							if (FieldValue.XmlConfigAttribute.Deprecated)
 							{
-								Log.TraceWarning($"Deprecated setting found in \"{FieldValue.SourceFile}\":");
-								Log.TraceWarning(
-									$"The setting \"{FieldValue.Target.MemberInfo.Name}\" is deprecated. Support for this setting will be removed in a future version of Unreal Engine.");
+								Logger.LogWarning("Deprecated setting found in \"{SourceFile}\":", FieldValue.SourceFile);
+								Logger.LogWarning("The setting \"{Setting}\" is deprecated. Support for this setting will be removed in a future version of Unreal Engine.", FieldValue.Target.MemberInfo.Name);
 
 								if (FieldValue.XmlConfigAttribute.NewAttributeName != null)
 								{
-									Log.TraceWarning(
-										$"Use \"{FieldValue.XmlConfigAttribute.NewAttributeName}\" in place of \"{FieldValue.Target.MemberInfo.Name}\"");
-									Log.TraceInformation(
-										$"The value provided for deprecated setting \"{FieldValue.Target.MemberInfo.Name}\" will be applied to \"{FieldValue.XmlConfigAttribute.NewAttributeName}\"");
+									Logger.LogWarning("Use \"{NewAttributeName}\" in place of \"{OldAttributeName}\"", FieldValue.XmlConfigAttribute.NewAttributeName, FieldValue.Target.MemberInfo.Name);
+									Logger.LogInformation("The value provided for deprecated setting \"{OldName}\" will be applied to \"{NewName}\"", FieldValue.Target.MemberInfo.Name, FieldValue.XmlConfigAttribute.NewAttributeName);
 
 									TargetToWrite = GetTargetMember(TargetType, FieldValue.XmlConfigAttribute.NewAttributeName) ?? TargetToWrite;
 								}
@@ -669,13 +671,14 @@
 		/// <param name="CategoryToFields">Lookup for configurable fields by category</param>
 		/// <param name="TypeToValues">Map of types to fields and their associated values</param>
 		/// <param name="Schema">Schema to validate against</param>
+		/// <param name="Logger">Logger for output</param>
 		/// <returns>True if the file was read successfully</returns>
 		static bool TryReadFile(FileReference Location, Dictionary<string, Dictionary<string, XmlConfigData.TargetMember>> CategoryToFields,
-			Dictionary<Type, Dictionary<XmlConfigData.TargetMember, XmlConfigData.ValueInfo>> TypeToValues, XmlSchema Schema)
+			Dictionary<Type, Dictionary<XmlConfigData.TargetMember, XmlConfigData.ValueInfo>> TypeToValues, XmlSchema Schema, ILogger Logger)
 		{
 			// Read the XML file, and validate it against the schema
 			XmlConfigFile? ConfigFile;
-			if(!XmlConfigFile.TryRead(Location, Schema, out ConfigFile))
+			if(!XmlConfigFile.TryRead(Location, Schema, Logger, out ConfigFile))
 			{
 				return false;
 			}
@@ -821,7 +824,8 @@
 		/// Generates documentation files for the available settings, by merging the XML documentation from the compiler.
 		/// </summary>
 		/// <param name="OutputFile">The documentation file to write</param>
-		public static void WriteDocumentation(FileReference OutputFile)
+		/// <param name="Logger">Logger for output</param>
+		public static void WriteDocumentation(FileReference OutputFile, ILogger Logger)
 		{
 			// Find all the configurable types
 			List<Type> ConfigTypes = FindConfigurableTypes();
@@ -855,11 +859,11 @@
 			// Generate the documentation file
 			if(OutputFile.HasExtension(".udn"))
 			{
-				WriteDocumentationUDN(OutputFile, InputDocumentation, CategoryToFields);
+				WriteDocumentationUDN(OutputFile, InputDocumentation, CategoryToFields, Logger);
 			}
 			else if(OutputFile.HasExtension(".html"))
 			{
-				WriteDocumentationHTML(OutputFile, InputDocumentation, CategoryToFields);
+				WriteDocumentationHTML(OutputFile, InputDocumentation, CategoryToFields, Logger);
 			}
 			else
 			{
@@ -867,7 +871,7 @@
 			}
 
 			// Success!
-			Log.TraceInformation("Written documentation to {0}.", OutputFile);
+			Logger.LogInformation("Written documentation to {OutputFile}.", OutputFile);
 		}
 
 		/// <summary>
@@ -876,7 +880,8 @@
 		/// <param name="OutputFile">The output file</param>
 		/// <param name="InputDocumentation">The XML documentation for this assembly</param>
 		/// <param name="CategoryToFields">Map of string to types to fields</param>
-		private static void WriteDocumentationUDN(FileReference OutputFile, XmlDocument InputDocumentation, Dictionary<string, Dictionary<string, XmlConfigData.TargetMember>> CategoryToFields)
+		/// <param name="Logger">Logger for output</param>
+		private static void WriteDocumentationUDN(FileReference OutputFile, XmlDocument InputDocumentation, Dictionary<string, Dictionary<string, XmlConfigData.TargetMember>> CategoryToFields, ILogger Logger)
 		{
 			using (StreamWriter Writer = new StreamWriter(OutputFile.FullName))
 			{
@@ -898,7 +903,7 @@
 					{
 						// Get the XML comment for this field
 						List<string>? Lines;
-						if(!RulesDocumentation.TryGetXmlComment(InputDocumentation, FieldPair.Value.MemberInfo, out Lines) || Lines.Count == 0)
+						if(!RulesDocumentation.TryGetXmlComment(InputDocumentation, FieldPair.Value.MemberInfo, Logger, out Lines) || Lines.Count == 0)
 						{
 							continue;
 						}
@@ -928,7 +933,8 @@
 		/// <param name="OutputFile">The output file</param>
 		/// <param name="InputDocumentation">The XML documentation for this assembly</param>
 		/// <param name="CategoryToFields">Map of string to types to fields</param>
-		private static void WriteDocumentationHTML(FileReference OutputFile, XmlDocument InputDocumentation, Dictionary<string, Dictionary<string, XmlConfigData.TargetMember>> CategoryToFields)
+		/// <param name="Logger">Logger for output</param>
+		private static void WriteDocumentationHTML(FileReference OutputFile, XmlDocument InputDocumentation, Dictionary<string, Dictionary<string, XmlConfigData.TargetMember>> CategoryToFields, ILogger Logger)
 		{
 			using (StreamWriter Writer = new StreamWriter(OutputFile.FullName))
 			{
@@ -946,7 +952,7 @@
 					{
 						// Get the XML comment for this field
 						List<string>? Lines;
-						if (!RulesDocumentation.TryGetXmlComment(InputDocumentation, FieldPair.Value.MemberInfo, out Lines) || Lines.Count == 0)
+						if (!RulesDocumentation.TryGetXmlComment(InputDocumentation, FieldPair.Value.MemberInfo, Logger, out Lines) || Lines.Count == 0)
 						{
 							continue;
 						}
diff --git a/Engine/Source/Programs/UnrealBuildTool/System/XmlConfigFile.cs b/Engine/Source/Programs/UnrealBuildTool/System/XmlConfigFile.cs
--- a/Engine/Source/Programs/UnrealBuildTool/System/XmlConfigFile.cs
+++ b/Engine/Source/Programs/UnrealBuildTool/System/XmlConfigFile.cs
@@ -10,6 +10,7 @@
 using System.Xml.Schema;
 using EpicGames.Core;
 using UnrealBuildTool;
+using Microsoft.Extensions.Logging;
 
 namespace UnrealBuildTool
 {
@@ -64,9 +65,10 @@
 		/// </summary>
 		/// <param name="File">The file to load</param>
 		/// <param name="Schema">The schema to validate against</param>
+		/// <param name="Logger">Logger for output</param>
 		/// <param name="OutConfigFile">If successful, the document that was read</param>
 		/// <returns>True if the document could be read, false otherwise</returns>
-		public static bool TryRead(FileReference File, XmlSchema Schema, [NotNullWhen(true)] out XmlConfigFile? OutConfigFile)
+		public static bool TryRead(FileReference File, XmlSchema Schema, ILogger Logger, [NotNullWhen(true)] out XmlConfigFile? OutConfigFile)
 		{
 			XmlConfigFile ConfigFile = new XmlConfigFile(File);
 
@@ -102,13 +104,13 @@
 				// Check that the root element is valid. If not, we didn't actually validate against the schema.
 				if (ConfigFile.DocumentElement?.Name != RootElementName)
 				{
-					Log.TraceError("Script does not have a root element called '{0}'", RootElementName);
+					Logger.LogError("Script does not have a root element called '{RootElementName}'", RootElementName);
 					OutConfigFile = null;
 					return false;
 				}
 				if (ConfigFile.DocumentElement.NamespaceURI != SchemaNamespaceURI)
 				{
-					Log.TraceError("Script root element is not in the '{0}' namespace (add the xmlns=\"{0}\" attribute)", SchemaNamespaceURI);
+					Logger.LogError("Script root element is not in the '{NamespaceUri}' namespace (add the xmlns=\"{NamespaceUri2}\" attribute)", SchemaNamespaceURI, SchemaNamespaceURI);
 					OutConfigFile = null;
 					return false;
 				}
diff --git a/Engine/Source/Programs/UnrealBuildTool/ToolChain/AppleToolChain.cs b/Engine/Source/Programs/UnrealBuildTool/ToolChain/AppleToolChain.cs
--- a/Engine/Source/Programs/UnrealBuildTool/ToolChain/AppleToolChain.cs
+++ b/Engine/Source/Programs/UnrealBuildTool/ToolChain/AppleToolChain.cs
@@ -9,6 +9,7 @@ using System.Text;
 using EpicGames.Core;
 using System.Text.RegularExpressions;
 using UnrealBuildBase;
+using Microsoft.Extensions.Logging;
 
 namespace UnrealBuildTool
 {
@@ -19,12 +20,12 @@ namespace UnrealBuildTool
 		/// </summary>
 		public string XcodeDeveloperDir = "xcode-select";
 
-		public AppleToolChainSettings(bool bVerbose)
+		public AppleToolChainSettings(bool bVerbose, ILogger Logger)
 		{
-			SelectXcode(ref XcodeDeveloperDir, bVerbose);
+			SelectXcode(ref XcodeDeveloperDir, bVerbose, Logger);
 		}
 
-		private static void SelectXcode(ref string DeveloperDir, bool bVerbose)
+		private static void SelectXcode(ref string DeveloperDir, bool bVerbose, ILogger Logger)
 		{
 			string Reason = "hardcoded";
 
@@ -33,7 +34,7 @@ namespace UnrealBuildTool
 				Reason = "xcode-select";
 
 				// on the Mac, run xcode-select directly
-				DeveloperDir = Utils.RunLocalProcessAndReturnStdOut("xcode-select", "--print-path");
+				DeveloperDir = Utils.RunLocalProcessAndReturnStdOut("xcode-select", "--print-path", Logger);
 
 				// make sure we get a full path
 				if (Directory.Exists(DeveloperDir) == false)
@@ -68,7 +69,7 @@ namespace UnrealBuildTool
 			}
 		}
 
-		protected void SelectSDK(string BaseSDKDir, string OSPrefix, ref string PlatformSDKVersion, bool bVerbose)
+		protected void SelectSDK(string BaseSDKDir, string OSPrefix, ref string PlatformSDKVersion, bool bVerbose, ILogger Logger)
 		{
 			if (PlatformSDKVersion == "latest")
 			{
@@ -127,8 +128,8 @@ namespace UnrealBuildTool
 				catch (Exception Ex)
 				{
 					// on any exception, just use the backup version
-					Log.TraceInformation("Triggered an exception while looking for SDK directory in Xcode.app");
-					Log.TraceInformation("{0}", Ex.ToString());
+					Logger.LogInformation("Triggered an exception while looking for SDK directory in Xcode.app");
+					Logger.LogInformation("{Ex}", Ex.ToString());
 				}
 			}
 
@@ -140,7 +141,7 @@ namespace UnrealBuildTool
 
 			if (bVerbose && !ProjectFileGenerator.bGenerateProjectFiles)
 			{
-				Log.TraceInformation("Compiling with {0} SDK {1}", OSPrefix, PlatformSDKVersion);
+				Logger.LogInformation("Compiling with {Os} SDK {Sdk}", OSPrefix, PlatformSDKVersion);
 			}
 		}
 	}
@@ -149,7 +150,7 @@ namespace UnrealBuildTool
 	{
 		protected FileReference? ProjectFile;
 
-		public AppleToolChain(FileReference? InProjectFile, ClangToolChainOptions InOptions) : base(InOptions)
+		public AppleToolChain(FileReference? InProjectFile, ClangToolChainOptions InOptions, ILogger InLogger) : base(InOptions, InLogger)
 		{
 			ProjectFile = InProjectFile;
 		}
@@ -172,7 +173,7 @@ namespace UnrealBuildTool
 			StartInfo.Arguments = String.Format("\"{0}\" -S", TargetFile.FullName);
 			StartInfo.UseShellExecute = false;
 			StartInfo.CreateNoWindow = true;
-			Utils.RunLocalProcessAndLogOutput(StartInfo);
+			Utils.RunLocalProcessAndLogOutput(StartInfo, Logger);
 		}
 
 		static Version? ClangVersion = null;
@@ -282,13 +283,13 @@ namespace UnrealBuildTool
 			Arguments.Add("-fpascal-strings");
 		}
 
-		protected string GetDsymutilPath(out string ExtraOptions, bool bIsForLTOBuild=false)
+		protected string GetDsymutilPath(ILogger Logger, out string ExtraOptions, bool bIsForLTOBuild=false)
 		{
 			FileReference DsymutilLocation = new FileReference("/usr/bin/dsymutil");
 
 			// dsymutil before 10.0.1 has a bug that causes issues, it's fixed in autosdks but not everyone has those set up so for the timebeing we have
 			// a version in P4 - first determine if we need it
-			string DsymutilVersionString = Utils.RunLocalProcessAndReturnStdOut(DsymutilLocation.FullName, "-version");
+			string DsymutilVersionString = Utils.RunLocalProcessAndReturnStdOut(DsymutilLocation.FullName, "-version", Logger);
 
 			bool bUseInstalledDsymutil = true;
 			int Major = 0, Minor = 0, Patch = 0;
diff --git a/Engine/Source/Programs/UnrealBuildTool/ToolChain/ClangToolChain.cs b/Engine/Source/Programs/UnrealBuildTool/ToolChain/ClangToolChain.cs
--- a/Engine/Source/Programs/UnrealBuildTool/ToolChain/ClangToolChain.cs
+++ b/Engine/Source/Programs/UnrealBuildTool/ToolChain/ClangToolChain.cs
@@ -5,6 +5,7 @@ using System;
 using System.Collections.Generic;
 using System.Linq;
 using UnrealBuildBase;
+using Microsoft.Extensions.Logging;
 
 namespace UnrealBuildTool
 {
@@ -112,7 +113,8 @@ namespace UnrealBuildTool
 
 		protected ClangToolChainOptions Options;
 
-		public ClangToolChain(ClangToolChainOptions InOptions) : base()
+		public ClangToolChain(ClangToolChainOptions InOptions, ILogger InLogger) 
+			: base(InLogger)
 		{
 			Options = InOptions;
 		}
diff --git a/Engine/Source/Programs/UnrealBuildTool/ToolChain/ISPCToolChain.cs b/Engine/Source/Programs/UnrealBuildTool/ToolChain/ISPCToolChain.cs
--- a/Engine/Source/Programs/UnrealBuildTool/ToolChain/ISPCToolChain.cs
+++ b/Engine/Source/Programs/UnrealBuildTool/ToolChain/ISPCToolChain.cs
@@ -6,12 +6,17 @@ using System.IO;
 using System.Linq;
 using System.Text;
 using EpicGames.Core;
+using Microsoft.Extensions.Logging;
 using UnrealBuildBase;
 
 namespace UnrealBuildTool
 {
 	abstract class ISPCToolChain : UEToolChain
 	{
+		public ISPCToolChain(ILogger InLogger) : base(InLogger)
+		{
+		}
+
 		/// <summary>
 		/// Get CPU Instruction set targets for ISPC.
 		/// </summary>
@@ -40,12 +45,12 @@ namespace UnrealBuildTool
 					case "-arm64": ISPCTargets.Add("neon"); break;
 					case "-x86": ISPCTargets.AddRange(new string[] { "sse4", "sse2" }); break;
 					case "-x64": ISPCTargets.AddRange(new string[] { "sse4", "sse2" }); break;
-					default: Log.TraceWarning("Invalid Android architecture for ISPC. At least one architecture (armv7, x86, etc) needs to be selected in the project settings to build"); break;
+					default: Logger.LogWarning("Invalid Android architecture for ISPC. At least one architecture (armv7, x86, etc) needs to be selected in the project settings to build"); break;
 				}
 			}
 			else
 			{
-				Log.TraceWarning("Unsupported ISPC platform target!");
+				Logger.LogWarning("Unsupported ISPC platform target!");
 			}
 
 			return ISPCTargets;
@@ -78,7 +83,7 @@ namespace UnrealBuildTool
 			}
 			else
 			{
-				Log.TraceWarning("Unsupported ISPC platform target!");
+				Logger.LogWarning("Unsupported ISPC platform target!");
 			}
 
 			return ISPCOS;
@@ -112,12 +117,12 @@ namespace UnrealBuildTool
 					case "-arm64": ISPCArch += "aarch64"; break;
 					case "-x86": ISPCArch += "x86"; break;
 					case "-x64": ISPCArch += "x86-64"; break;
-					default: Log.TraceWarning("Invalid Android architecture for ISPC. At least one architecture (armv7, x86, etc) needs to be selected in the project settings to build"); break;
+					default: Logger.LogWarning("Invalid Android architecture for ISPC. At least one architecture (armv7, x86, etc) needs to be selected in the project settings to build"); break;
 				}
 			}
 			else
 			{
-				Log.TraceWarning("Unsupported ISPC platform target!");
+				Logger.LogWarning("Unsupported ISPC platform target!");
 			}
 
 			return ISPCArch;
@@ -160,7 +165,7 @@ namespace UnrealBuildTool
 			}
 			else
 			{
-				Log.TraceWarning("Unsupported ISPC host!");
+				Logger.LogWarning("Unsupported ISPC host!");
 			}
 
 			return Path.Combine(ISPCCompilerPathCommon, ISPCArchitecturePath, "ispc" + ExeExtension);
@@ -183,7 +188,7 @@ namespace UnrealBuildTool
 
 				if (!File.Exists(CompilerPath))
 				{
-					Log.TraceWarning("No ISPC compiler at {0}", CompilerPath);
+					Logger.LogWarning("No ISPC compiler at {CompilerPath}", CompilerPath);
 					CompilerVersion = new Version(-1, -1);
 				}
 
@@ -214,7 +219,7 @@ namespace UnrealBuildTool
 			}
 			else
 			{
-				Log.TraceWarning("Unsupported ISPC platform target!");
+				Logger.LogWarning("Unsupported ISPC platform target!");
 			}
 
 			return Format;
@@ -241,7 +246,7 @@ namespace UnrealBuildTool
 			}
 			else
 			{
-				Log.TraceWarning("Unsupported ISPC platform target!");
+				Logger.LogWarning("Unsupported ISPC platform target!");
 			}
 
 			return Suffix;
@@ -497,7 +502,7 @@ namespace UnrealBuildTool
 
 				Result.GeneratedHeaderFiles.Add(TargetFileItem);
 
-				Log.TraceVerbose("   ISPC Generating Header " + CompileAction.StatusDescription + ": \"" + CompileAction.CommandPath + "\"" + CompileAction.CommandArguments);
+				Logger.LogDebug("   ISPC Generating Header {StatusDescription}: \"{CommandPath}\" {CommandArguments}", CompileAction.StatusDescription, CompileAction.CommandPath, CompileAction.CommandArguments);
 			}
 
 			return Result;
@@ -682,7 +687,7 @@ namespace UnrealBuildTool
 				// Add the source file and its included files to the prerequisite item list.
 				CompileAction.PrerequisiteItems.Add(ISPCFile);
 
-				Log.TraceVerbose("   ISPC Compiling " + CompileAction.StatusDescription + ": \"" + CompileAction.CommandPath + "\"" + CompileAction.CommandArguments);
+				Logger.LogDebug("   ISPC Compiling {StatusDescription}: \"{CommandPath}\" {CommandArguments}", CompileAction.StatusDescription, CompileAction.CommandPath, CompileAction.CommandArguments);
 			}
 
 			return Result;
