UBT: Use ILogger throughout UBT, rather than legacy EpicGames.Core.Log methods. All output going forwards should use structured logging rather than String.Format style logging.

diff --git a/Engine/Source/Programs/UnrealBuildTool/ToolChain/RemoteMac.cs b/Engine/Source/Programs/UnrealBuildTool/ToolChain/RemoteMac.cs
--- a/Engine/Source/Programs/UnrealBuildTool/ToolChain/RemoteMac.cs
+++ b/Engine/Source/Programs/UnrealBuildTool/ToolChain/RemoteMac.cs
@@ -13,6 +13,7 @@ using System.Text.RegularExpressions;
 using EpicGames.Core;
 using System.Security.Cryptography.X509Certificates;
 using UnrealBuildBase;
+using Microsoft.Extensions.Logging;
 
 namespace UnrealBuildTool
 {
@@ -128,7 +129,8 @@ namespace UnrealBuildTool
 		/// Constructor
 		/// </summary>
 		/// <param name="ProjectFile">Project to read settings from</param>
-		public RemoteMac(FileReference? ProjectFile)
+		/// <param name="Logger">Logger for output</param>
+		public RemoteMac(FileReference? ProjectFile, ILogger Logger)
 		{
 			this.RsyncExe = FileReference.Combine(Unreal.EngineDirectory, "Extras", "ThirdPartyNotUE", "cwrsync", "bin", "rsync.exe");
 			this.SshExe = FileReference.Combine(Unreal.EngineDirectory, "Extras", "ThirdPartyNotUE", "cwrsync", "bin", "ssh.exe");
@@ -196,7 +198,7 @@ namespace UnrealBuildTool
 			}
 
 			// Print out the server info
-			Log.TraceInformation("[Remote] Using remote server '{0}' on port {1} (user '{2}')", ServerName, ServerPort, UserName);
+			Logger.LogInformation("[Remote] Using remote server '{ServerName}' on port {ServerPort} (user '{UserName}')", ServerName, ServerPort, UserName);
 
 			// Get the path to the SSH private key
 			string OverrideSshPrivateKeyPath;
@@ -214,7 +216,7 @@ namespace UnrealBuildTool
 			// If it's not set, look in the standard locations. If that fails, spawn the batch file to generate one.
 			if (SshPrivateKey == null && !TryGetSshPrivateKey(out SshPrivateKey))
 			{
-				Log.TraceWarning("No SSH private key found for {0}@{1}. Launching SSH to generate one.", UserName, ServerName);
+				Logger.LogWarning("No SSH private key found for {UserName}@{ServerName}. Launching SSH to generate one.", UserName, ServerName);
 
 				StringBuilder CommandLine = new StringBuilder();
 				CommandLine.AppendFormat("/C \"\"{0}\"", FileReference.Combine(Unreal.EngineDirectory, "Build", "BatchFiles", "MakeAndInstallSSHKey.bat"));
@@ -240,7 +242,7 @@ namespace UnrealBuildTool
 			}
 
 			// Print the path to the private key
-			Log.TraceInformation("[Remote] Using private key at {0}", SshPrivateKey);
+			Logger.LogInformation("[Remote] Using private key at {SshPrivateKey}", SshPrivateKey);
 
 			// resolve the rest of the strings
 			RsyncAuthentication = ExpandVariables(RsyncAuthentication);
@@ -279,14 +281,14 @@ namespace UnrealBuildTool
 			else
 			{
 				StringBuilder Output;
-				if (ExecuteAndCaptureOutput("'echo ~'", out Output) != 0)
+				if (ExecuteAndCaptureOutput("'echo ~'", Logger, out Output) != 0)
 				{
 					throw new BuildException("Unable to determine home directory for remote user. SSH output:\n{0}", StringUtils.Indent(Output.ToString(), "  "));
 				}
 				RemoteBaseDir = String.Format("{0}/UE5/Builds/{1}", Output.ToString().Trim().TrimEnd('/'), Environment.MachineName);
 			}
 
-			Log.TraceInformation("[Remote] Using base directory '{0}'", RemoteBaseDir);
+			Logger.LogInformation("[Remote] Using base directory '{RemoteBaseDir}'", RemoteBaseDir);
 
 			// Build the list of directory mappings between the local and remote machines
 			Mappings = new List<RemoteMapping>();
@@ -366,10 +368,10 @@ namespace UnrealBuildTool
 		/// <summary>
 		/// Flush the remote machine, removing all existing files
 		/// </summary>
-		public void FlushRemote()
+		public void FlushRemote(ILogger Logger)
 		{
-			Log.TraceInformation("[Remote] Deleting all files under {0}...", RemoteBaseDir);
-			Execute("/", String.Format("rm -rf \"{0}\"", RemoteBaseDir));
+			Logger.LogInformation("[Remote] Deleting all files under {RemoteBaseDir}...", RemoteBaseDir);
+			Execute("/", String.Format("rm -rf \"{0}\"", RemoteBaseDir), Logger);
 		}
 
 		/// <summary>
@@ -386,8 +388,9 @@ namespace UnrealBuildTool
 		/// Clean a target remotely
 		/// </summary>
 		/// <param name="TargetDesc">Descriptor for the target to build</param>
+		/// <param name="Logger">Logger for diagnostic output</param>
 		/// <returns>True if the build succeeded, false otherwise</returns>
-		public bool Clean(TargetDescriptor TargetDesc)
+		public bool Clean(TargetDescriptor TargetDesc, ILogger Logger)
 		{
 			// Translate all the arguments for the remote
 			List<string> RemoteArguments = GetRemoteArgumentsForTarget(TargetDesc, null);
@@ -395,10 +398,10 @@ namespace UnrealBuildTool
 		
 			// Upload the workspace
 			DirectoryReference TempDir = CreateTempDirectory(TargetDesc);
-			UploadWorkspace(TempDir);
+			UploadWorkspace(TempDir, Logger);
 
 			// Execute the compile
-			Log.TraceInformation("[Remote] Executing clean...");
+			Logger.LogInformation("[Remote] Executing clean...");
 
 			StringBuilder BuildCommandLine = new StringBuilder("Engine/Build/BatchFiles/Mac/Build.sh");
 			foreach(string RemoteArgument in RemoteArguments)
@@ -406,7 +409,7 @@ namespace UnrealBuildTool
 				BuildCommandLine.AppendFormat(" {0}", EscapeShellArgument(RemoteArgument));
 			}
 
-			int Result = Execute(GetRemotePath(Unreal.RootDirectory), BuildCommandLine.ToString());
+			int Result = Execute(GetRemotePath(Unreal.RootDirectory), BuildCommandLine.ToString(), Logger);
 			return Result == 0;
 		}
 
@@ -416,23 +419,24 @@ namespace UnrealBuildTool
 		/// <param name="TargetDesc">Descriptor for the target to build</param>
 		/// <param name="RemoteLogFile">Path to store the remote log file</param>
 		/// <param name="bSkipPreBuildTargets">If true then any PreBuildTargets will be skipped</param>
+		/// <param name="Logger">Logger for diagnostic output</param>
 		/// <returns>True if the build succeeded, false otherwise</returns>
-		public bool Build(TargetDescriptor TargetDesc, FileReference RemoteLogFile, bool bSkipPreBuildTargets)
+		public bool Build(TargetDescriptor TargetDesc, FileReference RemoteLogFile, bool bSkipPreBuildTargets, ILogger Logger)
 		{
 			// Compile the rules assembly
-			RulesAssembly RulesAssembly = RulesCompiler.CreateTargetRulesAssembly(TargetDesc.ProjectFile, TargetDesc.Name, false, false, false, TargetDesc.ForeignPlugin);
+			RulesAssembly RulesAssembly = RulesCompiler.CreateTargetRulesAssembly(TargetDesc.ProjectFile, TargetDesc.Name, false, false, false, TargetDesc.ForeignPlugin, Logger);
 
 			// Create the target rules
-			TargetRules Rules = RulesAssembly.CreateTargetRules(TargetDesc.Name, TargetDesc.Platform, TargetDesc.Configuration, TargetDesc.Architecture, TargetDesc.ProjectFile, TargetDesc.AdditionalArguments);
+			TargetRules Rules = RulesAssembly.CreateTargetRules(TargetDesc.Name, TargetDesc.Platform, TargetDesc.Configuration, TargetDesc.Architecture, TargetDesc.ProjectFile, TargetDesc.AdditionalArguments, Logger);
 			if (!bSkipPreBuildTargets)
 			{
 				foreach (TargetInfo PreBuildTargetInfo in Rules.PreBuildTargets)
 				{
-					RemoteMac PreBuildTargetRemoteMac = new RemoteMac(ProjectFile);
+					RemoteMac PreBuildTargetRemoteMac = new RemoteMac(ProjectFile, Logger);
 					TargetDescriptor PreBuildTargetDesc = new TargetDescriptor(PreBuildTargetInfo.ProjectFile, PreBuildTargetInfo.Name, PreBuildTargetInfo.Platform, PreBuildTargetInfo.Configuration, PreBuildTargetInfo.Architecture, PreBuildTargetInfo.Arguments);
 
-					Log.TraceInformation("[Remote] Building pre target [{0}] for [{1}] ", PreBuildTargetDesc.ToString(), TargetDesc.ToString());
-					if (!PreBuildTargetRemoteMac.Build(PreBuildTargetDesc, RemoteLogFile, false))
+					Logger.LogInformation("[Remote] Building pre target [{PreTarget}] for [{Target}] ", PreBuildTargetDesc.ToString(), TargetDesc.ToString());
+					if (!PreBuildTargetRemoteMac.Build(PreBuildTargetDesc, RemoteLogFile, false, Logger))
 					{
 						return false;
 					}
@@ -492,8 +496,8 @@ namespace UnrealBuildTool
 					FileReference.SetAttributes(MobileProvisionFile, FileAttributes.Normal);
 				}
 				FileReference.Copy(ProvisioningData.MobileProvisionFile, MobileProvisionFile, true);
-				Log.TraceInformation("[Remote] Uploading {0}", MobileProvisionFile);
-				UploadFile(MobileProvisionFile);
+				Logger.LogInformation("[Remote] Uploading {MobileProvisionFile}", MobileProvisionFile);
+				UploadFile(MobileProvisionFile, Logger);
 
 				// Extract the certificate for the project. Try to avoid calling IPP if we already have it.
 				FileReference CertificateFile = FileReference.Combine(TempDir, "Certificate.p12");
@@ -503,7 +507,7 @@ namespace UnrealBuildTool
 
 				if(!FileReference.Exists(CertificateFile) || !FileReference.Exists(CertificateInfoFile) || FileReference.ReadAllText(CertificateInfoFile) != CertificateInfoContents)
 				{
-					Log.TraceInformation("[Remote] Exporting certificate for {0}...", ProvisioningData.MobileProvisionFile);
+					Logger.LogInformation("[Remote] Exporting certificate for {ProvisioningDataMobileProvisionFile}...", ProvisioningData.MobileProvisionFile);
 
 					StringBuilder Arguments = new StringBuilder("ExportCertificate");
 					if(TargetDesc.ProjectFile == null)
@@ -524,7 +528,7 @@ namespace UnrealBuildTool
 					ProcessStartInfo StartInfo = new ProcessStartInfo();
 					StartInfo.FileName = FileReference.Combine(Unreal.EngineDirectory, "Binaries", "DotNET", "IOS", "IPhonePackager.exe").FullName;
 					StartInfo.Arguments = Arguments.ToString();
-					if(Utils.RunLocalProcessAndLogOutput(StartInfo) != 0)
+					if(Utils.RunLocalProcessAndLogOutput(StartInfo, Logger) != 0)
 					{
 						throw new BuildException("IphonePackager failed.");
 					}
@@ -533,8 +537,8 @@ namespace UnrealBuildTool
 				}
 
 				// Upload the certificate to the remote
-				Log.TraceInformation("[Remote] Uploading {0}", CertificateFile);
-				UploadFile(CertificateFile);
+				Logger.LogInformation("[Remote] Uploading {CertificateFile}", CertificateFile);
+				UploadFile(CertificateFile, Logger);
 
 				// Tell the remote UBT instance to use them
 				RemoteArguments.Add(String.Format("-ImportProvision={0}", GetRemotePath(MobileProvisionFile)));
@@ -543,10 +547,10 @@ namespace UnrealBuildTool
 			}
 
 			// Upload the workspace files
-			UploadWorkspace(TempDir);
+			UploadWorkspace(TempDir, Logger);
 
 			// Execute the compile
-			Log.TraceInformation("[Remote] Executing build");
+			Logger.LogInformation("[Remote] Executing build");
 
 			StringBuilder BuildCommandLine = new StringBuilder("Engine/Build/BatchFiles/Mac/Build.sh");
 			foreach(string RemoteArgument in RemoteArguments)
@@ -554,50 +558,50 @@ namespace UnrealBuildTool
 				BuildCommandLine.AppendFormat(" {0}", EscapeShellArgument(RemoteArgument));
 			}
 
-			int Result = Execute(GetRemotePath(Unreal.RootDirectory), BuildCommandLine.ToString());
+			int Result = Execute(GetRemotePath(Unreal.RootDirectory), BuildCommandLine.ToString(), Logger);
 			if(Result != 0)
 			{
 				if(RemoteLogFile != null)
 				{
-					Log.TraceInformation("[Remote] Downloading {0}", RemoteLogFile);
-					DownloadFile(RemoteLogFile);
+					Logger.LogInformation("[Remote] Downloading {RemoteLogFile}", RemoteLogFile);
+					DownloadFile(RemoteLogFile, Logger);
 				}
 				return false;
 			}
 
 			// Download the manifest
-			Log.TraceInformation("[Remote] Downloading {0}", RemoteManifestFile);
-			DownloadFile(RemoteManifestFile);
+			Logger.LogInformation("[Remote] Downloading {RemoteManifestFile}", RemoteManifestFile);
+			DownloadFile(RemoteManifestFile, Logger);
 
 			// Convert the manifest to local form
-			BuildManifest Manifest = Utils.ReadClass<BuildManifest>(RemoteManifestFile.FullName);
+			BuildManifest Manifest = Utils.ReadClass<BuildManifest>(RemoteManifestFile.FullName, Logger);
 			for(int Idx = 0; Idx < Manifest.BuildProducts.Count; Idx++)
 			{
 				Manifest.BuildProducts[Idx] = GetLocalPath(Manifest.BuildProducts[Idx]).FullName;
 			}
 
 			// Download the files from the remote
-			Log.TraceInformation("[Remote] Downloading build products");
+			Logger.LogInformation("[Remote] Downloading build products");
 
 			List<FileReference> FilesToDownload = new List<FileReference>();
 			FilesToDownload.Add(RemoteLogFile);
 			FilesToDownload.AddRange(Manifest.BuildProducts.Select(x => new FileReference(x)));
-			DownloadFiles(FilesToDownload);
+			DownloadFiles(FilesToDownload, Logger);
 
 			// Copy remote FrameworkAssets directory as it could contain resource bundles that must be packaged locally.
 			DirectoryReference BaseDir = DirectoryReference.FromFile(TargetDesc.ProjectFile) ?? Unreal.EngineDirectory;
 			DirectoryReference FrameworkAssetsDir = DirectoryReference.Combine(BaseDir, "Intermediate", TargetDesc.Platform == UnrealTargetPlatform.IOS ? "IOS" : "TVOS", "FrameworkAssets");
-			if(RemoteDirectoryExists(FrameworkAssetsDir))
+			if(RemoteDirectoryExists(FrameworkAssetsDir, Logger))
 			{
-				Log.TraceInformation("[Remote] Downloading {0}", FrameworkAssetsDir);
-				DownloadDirectory(FrameworkAssetsDir);
+				Logger.LogInformation("[Remote] Downloading {FrameworkAssetsDir}", FrameworkAssetsDir);
+				DownloadDirectory(FrameworkAssetsDir, Logger);
 			}
 
 			// Write out all the local manifests
 			foreach(FileReference LocalManifestFile in LocalManifestFiles)
 			{
-				Log.TraceInformation("[Remote] Writing {0}", LocalManifestFile);
-				Utils.WriteClass<BuildManifest>(Manifest, LocalManifestFile.FullName, "");
+				Logger.LogInformation("[Remote] Writing {LocalManifestFile}", LocalManifestFile);
+				Utils.WriteClass<BuildManifest>(Manifest, LocalManifestFile.FullName, "", Logger);
 			}
 
 			return true;
@@ -685,25 +689,26 @@ namespace UnrealBuildTool
 		/// <param name="Platform">The target platform</param>
 		/// <param name="InputDir">Input directory containing assets</param>
 		/// <param name="OutputFile">Path to the Assets.car file to produce</param>
-		public void RunAssetCatalogTool(UnrealTargetPlatform Platform, DirectoryReference InputDir, FileReference OutputFile)
+		/// <param name="Logger">Logger for output</param>
+		public void RunAssetCatalogTool(UnrealTargetPlatform Platform, DirectoryReference InputDir, FileReference OutputFile, ILogger Logger)
 		{
-			Log.TraceInformation("Running asset catalog tool for {0}: {1} -> {2}", Platform, InputDir, OutputFile);
+			Logger.LogInformation("Running asset catalog tool for {Platform}: {InputDir} -> {OutputFile}", Platform, InputDir, OutputFile);
 
 			string RemoteInputDir = GetRemotePath(InputDir);
-			UploadDirectory(InputDir);
+			UploadDirectory(InputDir, Logger);
 
 			string RemoteOutputFile = GetRemotePath(OutputFile);
-			Execute(RemoteBaseDir, String.Format("rm -f {0}", EscapeShellArgument(RemoteOutputFile)));
+			Execute(RemoteBaseDir, String.Format("rm -f {0}", EscapeShellArgument(RemoteOutputFile)), Logger);
 
 			string RemoteOutputDir = Path.GetDirectoryName(RemoteOutputFile)!.Replace(Path.DirectorySeparatorChar, '/');
-			Execute(RemoteBaseDir, String.Format("mkdir -p {0}", EscapeShellArgument(RemoteOutputDir)));
+			Execute(RemoteBaseDir, String.Format("mkdir -p {0}", EscapeShellArgument(RemoteOutputDir)), Logger);
 
 			string RemoteArguments = IOSToolChain.GetAssetCatalogArgs(Platform, RemoteInputDir, RemoteOutputDir); 
-			if(Execute(RemoteBaseDir, String.Format("/usr/bin/xcrun {0}", RemoteArguments)) != 0)
+			if(Execute(RemoteBaseDir, String.Format("/usr/bin/xcrun {0}", RemoteArguments), Logger) != 0)
 			{
 				throw new BuildException("Failed to run actool.");
 			}
-			DownloadFile(OutputFile);
+			DownloadFile(OutputFile, Logger);
 		}
 
 		/// <summary>
@@ -771,7 +776,8 @@ namespace UnrealBuildTool
 		/// Upload a single file to the remote
 		/// </summary>
 		/// <param name="LocalFile">The file to upload</param>
-		void UploadFile(FileReference LocalFile)
+		/// <param name="Logger">Logger for output</param>
+		void UploadFile(FileReference LocalFile, ILogger Logger)
 		{
 			string RemoteFile = GetRemotePath(LocalFile);
 			string RemoteDirectory = GetRemotePath(LocalFile.Directory);
@@ -782,7 +788,7 @@ namespace UnrealBuildTool
 			Arguments.Add(String.Format("\"{0}@{1}\":'{2}'", UserName, ServerName, RemoteFile));
 			Arguments.Add("-q");
 
-			int Result = Rsync(String.Join(" ", Arguments));
+			int Result = Rsync(String.Join(" ", Arguments), Logger);
 			if(Result != 0)
 			{
 				throw new BuildException("Error while running Rsync (exit code {0})", Result);
@@ -795,7 +801,8 @@ namespace UnrealBuildTool
 		/// <param name="LocalDirectory">The base directory to copy</param>
 		/// <param name="RemoteDirectory">The remote directory</param>
 		/// <param name="LocalFileList">The file to upload</param>
-		void UploadFiles(DirectoryReference LocalDirectory, string RemoteDirectory, FileReference LocalFileList)
+		/// <param name="Logger">Logger for output</param>
+		void UploadFiles(DirectoryReference LocalDirectory, string RemoteDirectory, FileReference LocalFileList, ILogger Logger)
 		{
 			List<string> Arguments = new List<string>(BasicRsyncArguments);
 			Arguments.Add(String.Format("--rsync-path=\"mkdir -p {0} && rsync\"", RemoteDirectory));
@@ -804,7 +811,7 @@ namespace UnrealBuildTool
 			Arguments.Add(String.Format("\"{0}@{1}\":'{2}/'", UserName, ServerName, RemoteDirectory));
 			Arguments.Add("-q");
 
-			int Result = Rsync(String.Join(" ", Arguments));
+			int Result = Rsync(String.Join(" ", Arguments), Logger);
 			if(Result != 0)
 			{
 				throw new BuildException("Error while running Rsync (exit code {0})", Result);
@@ -815,7 +822,8 @@ namespace UnrealBuildTool
 		/// Upload a single directory to the remote
 		/// </summary>
 		/// <param name="LocalDirectory">The local directory to upload</param>
-		void UploadDirectory(DirectoryReference LocalDirectory)
+		/// <param name="Logger">Logger for output</param>
+		void UploadDirectory(DirectoryReference LocalDirectory, ILogger Logger)
 		{
 			string RemoteDirectory = GetRemotePath(LocalDirectory);
 
@@ -825,7 +833,7 @@ namespace UnrealBuildTool
 			Arguments.Add(String.Format("\"{0}@{1}\":'{2}/'", UserName, ServerName, RemoteDirectory));
 			Arguments.Add("-q");
 
-			int Result = Rsync(String.Join(" ", Arguments));
+			int Result = Rsync(String.Join(" ", Arguments), Logger);
 			if(Result != 0)
 			{
 				throw new BuildException("Error while running Rsync (exit code {0})", Result);
@@ -838,7 +846,8 @@ namespace UnrealBuildTool
 		/// <param name="LocalDirectory">The local directory to copy from</param>
 		/// <param name="RemoteDirectory">The remote directory to copy to</param>
 		/// <param name="FilterLocations">List of paths to filter</param>
-		void UploadDirectory(DirectoryReference LocalDirectory, string RemoteDirectory, List<FileReference> FilterLocations)
+		/// <param name="Logger">Logger for output</param>
+		void UploadDirectory(DirectoryReference LocalDirectory, string RemoteDirectory, List<FileReference> FilterLocations, ILogger Logger)
 		{
 			List<string> Arguments = new List<string>(CommonRsyncArguments);
 			Arguments.Add(String.Format("--rsync-path=\"mkdir -p {0} && rsync\"", RemoteDirectory));
@@ -850,7 +859,7 @@ namespace UnrealBuildTool
 			Arguments.Add(String.Format("\"{0}/\"", GetLocalCygwinPath(LocalDirectory)));
 			Arguments.Add(String.Format("\"{0}@{1}\":'{2}/'", UserName, ServerName, RemoteDirectory));
 
-			int Result = Rsync(String.Join(" ", Arguments));
+			int Result = Rsync(String.Join(" ", Arguments), Logger);
 			if(Result != 0)
 			{
 				throw new BuildException("Error while running Rsync (exit code {0})", Result);
@@ -860,7 +869,7 @@ namespace UnrealBuildTool
 		/// <summary>
 		/// Upload all the files in the workspace for the current project
 		/// </summary>
-		void UploadWorkspace(DirectoryReference TempDir)
+		void UploadWorkspace(DirectoryReference TempDir, ILogger Logger)
 		{
 			// Path to the scripts to be uploaded
 			FileReference ScriptPathsFileName = FileReference.Combine(Unreal.EngineDirectory, "Build", "Rsync", "RsyncEngineScripts.txt");
@@ -907,12 +916,12 @@ namespace UnrealBuildTool
 			}
 
 			// Upload these files to the remote
-			Log.TraceInformation("[Remote] Uploading scripts...");
-			UploadFiles(TempDir, GetRemotePath(Unreal.EngineDirectory), ScriptPathsFileName);
+			Logger.LogInformation("[Remote] Uploading scripts...");
+			UploadFiles(TempDir, GetRemotePath(Unreal.EngineDirectory), ScriptPathsFileName, Logger);
 
 			// Upload the config files
-			Log.TraceInformation("[Remote] Uploading config files...");
-			UploadFile(XmlConfig.CacheFile!);
+			Logger.LogInformation("[Remote] Uploading config files...");
+			UploadFile(XmlConfig.CacheFile!, Logger);
 
 			// Upload the engine files
 			List<FileReference> EngineFilters = new List<FileReference>();
@@ -923,8 +932,8 @@ namespace UnrealBuildTool
 			}
 			EngineFilters.Add(FileReference.Combine(Unreal.EngineDirectory, "Build", "Rsync", "RsyncEngine.txt"));
 
-			Log.TraceInformation("[Remote] Uploading engine files...");
-			UploadDirectory(Unreal.EngineDirectory, GetRemotePath(Unreal.EngineDirectory), EngineFilters);
+			Logger.LogInformation("[Remote] Uploading engine files...");
+			UploadDirectory(Unreal.EngineDirectory, GetRemotePath(Unreal.EngineDirectory), EngineFilters, Logger);
 
 			// Upload the project files
 			DirectoryReference? ProjectDir = null;
@@ -951,8 +960,8 @@ namespace UnrealBuildTool
 				}
 				ProjectFilters.Add(FileReference.Combine(Unreal.EngineDirectory, "Build", "Rsync", "RsyncProject.txt"));
 
-				Log.TraceInformation("[Remote] Uploading project files...");
-				UploadDirectory(ProjectDir, GetRemotePath(ProjectDir), ProjectFilters);
+				Logger.LogInformation("[Remote] Uploading project files...");
+				UploadDirectory(ProjectDir, GetRemotePath(ProjectDir), ProjectFilters, Logger);
 			}
 
 			if (AdditionalPaths != null)
@@ -968,31 +977,32 @@ namespace UnrealBuildTool
 					}
 					CustomFilters.Add(FileReference.Combine(Unreal.EngineDirectory, "Build", "Rsync", "RsyncProject.txt"));
 
-					Log.TraceInformation(string.Format("[Remote] Uploading additional path files [{0}]...", AdditionalPath.FullName));
-					UploadDirectory(AdditionalPath, GetRemotePath(AdditionalPath), CustomFilters);
+					Logger.LogInformation("[Remote] Uploading additional path files [{Dir}]...", AdditionalPath);
+					UploadDirectory(AdditionalPath, GetRemotePath(AdditionalPath), CustomFilters, Logger);
 				}
 			}
 
-			Execute("/", String.Format("rm -rf {0}/Intermediate/IOS/*.plist", GetRemotePath(Unreal.EngineDirectory)), true);
-			Execute("/", String.Format("rm -rf {0}/Intermediate/TVOS/*.plist", GetRemotePath(Unreal.EngineDirectory)), true);
+			Execute("/", String.Format("rm -rf {0}/Intermediate/IOS/*.plist", GetRemotePath(Unreal.EngineDirectory)), Logger, true);
+			Execute("/", String.Format("rm -rf {0}/Intermediate/TVOS/*.plist", GetRemotePath(Unreal.EngineDirectory)), Logger, true);
 			if (ProjectFile != null)
 			{
-				Execute("/", String.Format("rm -rf {0}/Intermediate/IOS/*.plist", GetRemotePath(ProjectFile.Directory)), true);
-				Execute("/", String.Format("rm -rf {0}/Intermediate/TVOS/*.plist", GetRemotePath(ProjectFile.Directory)), true);
+				Execute("/", String.Format("rm -rf {0}/Intermediate/IOS/*.plist", GetRemotePath(ProjectFile.Directory)), Logger, true);
+				Execute("/", String.Format("rm -rf {0}/Intermediate/TVOS/*.plist", GetRemotePath(ProjectFile.Directory)), Logger, true);
 			}
 
 			// Convert CRLF to LF for all shell scripts
-			Execute(RemoteBaseDir, String.Format("for i in {0}/Build/BatchFiles/Mac/*.sh; do mv $i $i.crlf; tr -d '\r' < $i.crlf > $i; done", EscapeShellArgument(GetRemotePath(Unreal.EngineDirectory))));
+			Execute(RemoteBaseDir, String.Format("for i in {0}/Build/BatchFiles/Mac/*.sh; do mv $i $i.crlf; tr -d '\r' < $i.crlf > $i; done", EscapeShellArgument(GetRemotePath(Unreal.EngineDirectory))), Logger);
 
 			// Fixup permissions on any shell scripts
-			Execute(RemoteBaseDir, String.Format("chmod +x {0}/Build/BatchFiles/Mac/*.sh", EscapeShellArgument(GetRemotePath(Unreal.EngineDirectory))));
+			Execute(RemoteBaseDir, String.Format("chmod +x {0}/Build/BatchFiles/Mac/*.sh", EscapeShellArgument(GetRemotePath(Unreal.EngineDirectory))), Logger);
 		}
 
 		/// <summary>
 		/// Downloads a single file from the remote
 		/// </summary>
 		/// <param name="LocalFile">The file to download</param>
-		void DownloadFile(FileReference LocalFile)
+		/// <param name="Logger">Logger for output</param>
+		void DownloadFile(FileReference LocalFile, ILogger Logger)
 		{
 			RemoteMapping? Mapping = Mappings.FirstOrDefault(x => LocalFile.IsUnderDirectory(x.LocalDirectory));
 			if(Mapping == null)
@@ -1005,7 +1015,7 @@ namespace UnrealBuildTool
 			Arguments.Add(String.Format("\"{0}/\"", GetLocalCygwinPath(LocalFile.Directory)));
 			Arguments.Add("-q");
 
-			int Result = Rsync(String.Join(" ", Arguments));
+			int Result = Rsync(String.Join(" ", Arguments), Logger);
 			if(Result != 0)
 			{
 				throw new BuildException("Unable to download '{0}' from the remote Mac (exit code {1}).", LocalFile, Result);
@@ -1016,7 +1026,8 @@ namespace UnrealBuildTool
 		/// Download multiple files from the remote Mac
 		/// </summary>
 		/// <param name="Files">List of local files to download</param>
-		void DownloadFiles(IEnumerable<FileReference> Files)
+		/// <param name="Logger">Logger for output</param>
+		void DownloadFiles(IEnumerable<FileReference> Files, ILogger Logger)
 		{
 			List<FileReference>[] FileGroups = new List<FileReference>[Mappings.Count];
 			for(int Idx = 0; Idx < Mappings.Count; Idx++)
@@ -1045,7 +1056,7 @@ namespace UnrealBuildTool
 					Arguments.Add(String.Format("\"{0}@{1}\":'{2}/'", UserName, ServerName, Mappings[Idx].RemoteDirectory));
 					Arguments.Add(String.Format("\"{0}/\"", GetLocalCygwinPath(Mappings[Idx].LocalDirectory)));
 
-					int Result = Rsync(String.Join(" ", Arguments));
+					int Result = Rsync(String.Join(" ", Arguments), Logger);
 					if(Result != 0)
 					{
 						throw new BuildException("Unable to download files from remote Mac (exit code {0})", Result);
@@ -1058,18 +1069,20 @@ namespace UnrealBuildTool
 		/// Checks whether a directory exists on the remote machine
 		/// </summary>
 		/// <param name="LocalDirectory">Path to the directory on the local machine</param>
+		/// <param name="Logger">Logger for output</param>
 		/// <returns>True if the remote directory exists</returns>
-		private bool RemoteDirectoryExists(DirectoryReference LocalDirectory)
+		private bool RemoteDirectoryExists(DirectoryReference LocalDirectory, ILogger Logger)
 		{
 			string RemoteDirectory = GetRemotePath(LocalDirectory);
-			return Execute(Unreal.RootDirectory, String.Format("[ -d {0} ]", EscapeShellArgument(RemoteDirectory))) == 0;
+			return Execute(Unreal.RootDirectory, String.Format("[ -d {0} ]", EscapeShellArgument(RemoteDirectory)), Logger) == 0;
 		}
 
 		/// <summary>
 		/// Download a directory from the remote Mac
 		/// </summary>
 		/// <param name="LocalDirectory">Directory to download</param>
-		private void DownloadDirectory(DirectoryReference LocalDirectory)
+		/// <param name="Logger">Logger for output</param>
+		private void DownloadDirectory(DirectoryReference LocalDirectory, ILogger Logger)
 		{
 			DirectoryReference.CreateDirectory(LocalDirectory);
 
@@ -1079,7 +1092,7 @@ namespace UnrealBuildTool
 			Arguments.Add(String.Format("\"{0}@{1}\":'{2}/'", UserName, ServerName, RemoteDirectory));
 			Arguments.Add(String.Format("\"{0}/\"", GetLocalCygwinPath(LocalDirectory)));
 
-			int Result = Rsync(String.Join(" ", Arguments));
+			int Result = Rsync(String.Join(" ", Arguments), Logger);
 			if (Result != 0)
 			{
 				throw new BuildException("Unable to download '{0}' from the remote Mac (exit code {1}).", LocalDirectory, Result);
@@ -1090,13 +1103,14 @@ namespace UnrealBuildTool
 		/// Execute Rsync
 		/// </summary>
 		/// <param name="Arguments">Arguments for the Rsync command</param>
+		/// <param name="Logger">Logger for output</param>
 		/// <returns>Exit code from Rsync</returns>
-		private int Rsync(string Arguments)
+		private int Rsync(string Arguments, ILogger Logger)
 		{
 			using(Process RsyncProcess = new Process())
 			{
-				DataReceivedEventHandler OutputHandler = (E, Args) => { RsyncOutput(Args, false); };
-				DataReceivedEventHandler ErrorHandler = (E, Args) => { RsyncOutput(Args, true); };
+				DataReceivedEventHandler OutputHandler = (E, Args) => { RsyncOutput(Args, false, Logger); };
+				DataReceivedEventHandler ErrorHandler = (E, Args) => { RsyncOutput(Args, true, Logger); };
 
 				RsyncProcess.StartInfo.FileName = RsyncExe.FullName;
 				RsyncProcess.StartInfo.Arguments = Arguments;
@@ -1104,7 +1118,7 @@ namespace UnrealBuildTool
 				RsyncProcess.OutputDataReceived += OutputHandler;
 				RsyncProcess.ErrorDataReceived += ErrorHandler;
 
-				Log.TraceLog("[Rsync] {0} {1}", Utils.MakePathSafeToUseWithCommandLine(RsyncProcess.StartInfo.FileName), RsyncProcess.StartInfo.Arguments);
+				Logger.LogDebug("[Rsync] {File} {Args}", Utils.MakePathSafeToUseWithCommandLine(RsyncProcess.StartInfo.FileName), RsyncProcess.StartInfo.Arguments);
 				return Utils.RunLocalProcess(RsyncProcess);
 			}
 		}
@@ -1114,17 +1128,18 @@ namespace UnrealBuildTool
 		/// </summary>
 		/// <param name="Args">The received data</param>e
 		/// <param name="bStdErr">whether the data was received on stderr</param>
-		private void RsyncOutput(DataReceivedEventArgs Args, bool bStdErr)
+		/// <param name="Logger">Logger for output</param>
+		private void RsyncOutput(DataReceivedEventArgs Args, bool bStdErr, ILogger Logger)
 		{
 			if (Args.Data != null)
 			{
 				if (bStdErr)
 				{
-					Log.TraceError("  {0}", Args.Data);
+					Logger.LogError("  {Output}", Args.Data);
 				}
 				else
 				{
-					Log.TraceInformation("  {0}", Args.Data);
+					Logger.LogInformation("  {Output}", Args.Data);
 				}
 			}
 		}
@@ -1134,11 +1149,12 @@ namespace UnrealBuildTool
 		/// </summary>
 		/// <param name="WorkingDir"></param>
 		/// <param name="Command"></param>
+		/// <param name="Logger">Logger for output</param>
 		/// <param name="bSilent"></param>
 		/// <returns></returns>
-		public int Execute(DirectoryReference WorkingDir, string Command, bool bSilent = false)
+		public int Execute(DirectoryReference WorkingDir, string Command, ILogger Logger, bool bSilent = false)
 		{
-			return Execute(GetRemotePath(WorkingDir), Command, bSilent);
+			return Execute(GetRemotePath(WorkingDir), Command, Logger, bSilent);
 		}
 
 		/// <summary>
@@ -1146,15 +1162,16 @@ namespace UnrealBuildTool
 		/// </summary>
 		/// <param name="WorkingDirectory">The remote working directory</param>
 		/// <param name="Command">Command to be executed</param>
+		/// <param name="Logger">Logger for output</param>
 		/// <param name="bSilent">If true, logging is suppressed</param>
 		/// <returns></returns>
-		protected int Execute(string WorkingDirectory, string Command, bool bSilent = false)
+		protected int Execute(string WorkingDirectory, string Command, ILogger Logger, bool bSilent = false)
 		{
 			string FullCommand = String.Format("cd {0} && {1}", EscapeShellArgument(WorkingDirectory), Command);
 			using (Process SSHProcess = new Process())
 			{
-				DataReceivedEventHandler OutputHandler = (E, Args) => { SshOutput(Args, false); };
-				DataReceivedEventHandler ErrorHandler = (E, Args) => { SshOutput(Args, true); };
+				DataReceivedEventHandler OutputHandler = (E, Args) => { SshOutput(Args, false, Logger); };
+				DataReceivedEventHandler ErrorHandler = (E, Args) => { SshOutput(Args, true, Logger); };
 
 				SSHProcess.StartInfo.FileName = SshExe.FullName;
 				SSHProcess.StartInfo.WorkingDirectory = SshExe.Directory.FullName;
@@ -1165,7 +1182,7 @@ namespace UnrealBuildTool
 					SSHProcess.ErrorDataReceived += ErrorHandler;
 				}
 
-				Log.TraceLog("[SSH] {0} {1}", Utils.MakePathSafeToUseWithCommandLine(SSHProcess.StartInfo.FileName), SSHProcess.StartInfo.Arguments);
+				Logger.LogDebug("[SSH] {Exe} {Args}", Utils.MakePathSafeToUseWithCommandLine(SSHProcess.StartInfo.FileName), SSHProcess.StartInfo.Arguments);
 				return Utils.RunLocalProcess(SSHProcess);
 			}
 		}
@@ -1175,18 +1192,19 @@ namespace UnrealBuildTool
 		/// </summary>
 		/// <param name="Args"></param>
 		/// <param name="bStdErr">whether the data was received on stderr</param>
-		private void SshOutput(DataReceivedEventArgs Args, bool bStdErr)
+		/// <param name="Logger">Logger for output</param>
+		private void SshOutput(DataReceivedEventArgs Args, bool bStdErr, ILogger Logger)
 		{
 			if (Args.Data != null)
 			{
 				string FormattedOutput = ConvertRemotePathsToLocal(Args.Data);
 				if (bStdErr)
 				{
-					Log.TraceError("  {0}", FormattedOutput);
+					Logger.LogError("  {Output}", FormattedOutput);
 				}
 				else
 				{
-					Log.TraceInformation("  {0}", FormattedOutput);
+					Logger.LogInformation("  {Output}", FormattedOutput);
 				}
 			}
 		}
@@ -1195,9 +1213,10 @@ namespace UnrealBuildTool
 		/// Execute a remote command, capturing the output text
 		/// </summary>
 		/// <param name="Command">Command to be executed</param>
+		/// <param name="Logger">Logger for output</param>
 		/// <param name="Output">Receives the output text</param>
 		/// <returns></returns>
-		protected int ExecuteAndCaptureOutput(string Command, out StringBuilder Output)
+		protected int ExecuteAndCaptureOutput(string Command, ILogger Logger, out StringBuilder Output)
 		{
 			StringBuilder FullCommand = new StringBuilder();
 			foreach(string CommonSshArgument in CommonSshArguments)
@@ -1219,7 +1238,7 @@ namespace UnrealBuildTool
 				SSHProcess.OutputDataReceived += OutputHandler;
 				SSHProcess.ErrorDataReceived += OutputHandler;
 
-				Log.TraceLog("[SSH] {0} {1}", Utils.MakePathSafeToUseWithCommandLine(SSHProcess.StartInfo.FileName), SSHProcess.StartInfo.Arguments);
+				Logger.LogDebug("[SSH] {Exe} {Args}", Utils.MakePathSafeToUseWithCommandLine(SSHProcess.StartInfo.FileName), SSHProcess.StartInfo.Arguments);
 				return Utils.RunLocalProcess(SSHProcess);
 			}
 		}
diff --git a/Engine/Source/Programs/UnrealBuildTool/ToolChain/UEToolChain.cs b/Engine/Source/Programs/UnrealBuildTool/ToolChain/UEToolChain.cs
--- a/Engine/Source/Programs/UnrealBuildTool/ToolChain/UEToolChain.cs
+++ b/Engine/Source/Programs/UnrealBuildTool/ToolChain/UEToolChain.cs
@@ -4,14 +4,18 @@
 using System.Collections.Generic;
 using System.Text.RegularExpressions;
 using EpicGames.Core;
+using Microsoft.Extensions.Logging;
 using UnrealBuildBase;
 
 namespace UnrealBuildTool
 {
 	abstract class UEToolChain
 	{
-		public UEToolChain()
+		protected readonly ILogger Logger;
+
+		public UEToolChain(ILogger InLogger)
 		{
+			Logger = InLogger;
 		}
 
 		public virtual void SetEnvironmentVariables()
@@ -92,7 +96,7 @@
 		{
 		}
 
-		public virtual void FinalizeOutput(ReadOnlyTargetRules Target, TargetMakefile Makefile)
+		public virtual void FinalizeOutput(ReadOnlyTargetRules Target, TargetMakefileBuilder MakefileBuilder)
 		{
 		}
 
@@ -132,9 +136,9 @@
 		/// <param name="ToolArg">Argument that will be passed to the tool</param>
 		/// <param name="Expression">null, or a Regular expression to capture in the output</param>
 		/// <returns></returns>
-		public static string? RunToolAndCaptureOutput(FileReference Command, string ToolArg, string? Expression = null)
+		protected string? RunToolAndCaptureOutput(FileReference Command, string ToolArg, string? Expression = null)
 		{
-			string ProcessOutput = Utils.RunLocalProcessAndReturnStdOut(Command.FullName, ToolArg);
+			string ProcessOutput = Utils.RunLocalProcessAndReturnStdOut(Command.FullName, ToolArg, Logger);
 
 			if (string.IsNullOrEmpty(Expression))
 			{
@@ -152,7 +156,7 @@
 		/// <param name="VersionArg">Argument that will result in the version string being shown (it's ok this is a byproduct of a command that returns an error)</param>
 		/// <param name="VersionExpression">Regular expression to capture the version. By default we look for four integers separated by periods, with the last two optional</param>
 		/// <returns></returns>
-		public static Version RunToolAndCaptureVersion(FileReference Command, string VersionArg, string VersionExpression = @"(\d+\.\d+(\.\d+)?(\.\d+)?)")
+		public Version RunToolAndCaptureVersion(FileReference Command, string VersionArg, string VersionExpression = @"(\d+\.\d+(\.\d+)?(\.\d+)?)")
 		{
 			string? ProcessOutput = RunToolAndCaptureOutput(Command, VersionArg, VersionExpression);
 
@@ -163,7 +167,7 @@
 				return ToolVersion;
 			}
 
-			Log.TraceWarning("Unable to retrieve version from {0} {1}", Command, VersionArg);
+			Logger.LogWarning("Unable to retrieve version from {Command} {Arg}", Command, VersionArg);
 
 			return new Version(0, 0);
 		}
diff --git a/Engine/Source/Programs/UnrealBuildTool/UnrealBuildTool.cs b/Engine/Source/Programs/UnrealBuildTool/UnrealBuildTool.cs
--- a/Engine/Source/Programs/UnrealBuildTool/UnrealBuildTool.cs
+++ b/Engine/Source/Programs/UnrealBuildTool/UnrealBuildTool.cs
@@ -10,6 +10,7 @@ using System.Text;
 using System.Text.Json;
 using System.Threading;
 using EpicGames.Core;
+using Microsoft.Extensions.Logging;
 using OpenTracing.Util;
 using UnrealBuildBase;
 
@@ -448,7 +449,8 @@ namespace UnrealBuildTool
 			FileReference? RunFile = null;
 			SingleInstanceMutex? Mutex = null;
 			JsonTracer? Tracer = null;
-			
+
+			ILogger Logger = Log.Logger;
 			try
 			{
 				// Start capturing performance info
@@ -480,7 +482,7 @@ namespace UnrealBuildTool
 
 				if (Options.TraceWrites != null)
 				{
-					Log.TraceInformation($"All attempts to write to \"{Options.TraceWrites}\" via WriteFileIfChanged() will be logged");
+					Logger.LogInformation("All attempts to write to \"{TraceWrites}\" via WriteFileIfChanged() will be logged", Options.TraceWrites);
 					Utils.WriteFileIfChangedTrace = Options.TraceWrites;
 				}
 
@@ -511,7 +513,7 @@ namespace UnrealBuildTool
 					// Try to get the correct mode
 					if(!ModeNameToType.TryGetValue(Options.Mode, out ModeType))
 					{
-						Log.TraceError("No mode named '{0}'. Available modes are:\n  {1}", Options.Mode, String.Join("\n  ", ModeNameToType.Keys));
+						Logger.LogError("No mode named '{Name}'. Available modes are:\n  {ModeList}", Options.Mode, String.Join("\n  ", ModeNameToType.Keys));
 						return 1;
 					}
 				}
@@ -543,7 +545,7 @@ namespace UnrealBuildTool
 						using(SingleInstanceMutex XmlConfigMutex = new SingleInstanceMutex(XmlConfigMutexName, true))
 						{
 							FileReference? XmlConfigCache = Arguments.GetFileReferenceOrDefault("-XmlConfigCache=", null);
-							XmlConfig.ReadConfigFiles(XmlConfigCache);
+							XmlConfig.ReadConfigFiles(XmlConfigCache, Logger);
 						}
 					}
 				
@@ -592,21 +594,21 @@ namespace UnrealBuildTool
 				{
 					using (GlobalTracer.Instance.BuildSpan("UEBuildPlatform.RegisterPlatforms()").StartActive())
 					{
-						UEBuildPlatform.RegisterPlatforms(false, false);
+						UEBuildPlatform.RegisterPlatforms(false, false, Logger);
 					}
 				}
 				if ((ModeOptions & ToolModeOptions.BuildPlatformsHostOnly) != 0)
 				{
 					using (GlobalTracer.Instance.BuildSpan("UEBuildPlatform.RegisterPlatforms()").StartActive())
 					{
-						UEBuildPlatform.RegisterPlatforms(false, true);
+						UEBuildPlatform.RegisterPlatforms(false, true, Logger);
 					}
 				}
 				if ((ModeOptions & ToolModeOptions.BuildPlatformsForValidation) != 0)
 				{
 					using (GlobalTracer.Instance.BuildSpan("UEBuildPlatform.RegisterPlatforms()").StartActive())
 					{
-						UEBuildPlatform.RegisterPlatforms(true, false);
+						UEBuildPlatform.RegisterPlatforms(true, false, Logger);
 					}
 				}
 
@@ -614,31 +616,31 @@ namespace UnrealBuildTool
 				ToolMode Mode = (ToolMode)Activator.CreateInstance(ModeType)!;
 
 				// Execute the mode
-				int Result = Mode.Execute(Arguments);
+				int Result = Mode.Execute(Arguments, Logger);
 				if((ModeOptions & ToolModeOptions.ShowExecutionTime) != 0)
 				{
-					Log.TraceInformation("Total execution time: {0:0.00} seconds", Timeline.Elapsed.TotalSeconds);
+					Logger.LogInformation("Total execution time: {Time:0.00} seconds", Timeline.Elapsed.TotalSeconds);
 				}
 				return Result;
 			}
 			catch (CompilationResultException Ex)
 			{
 				// Used to return a propagate a specific exit code after an error has occurred. Does not log any message.
-				Log.TraceLog(ExceptionUtils.FormatExceptionDetails(Ex));
+				Logger.LogDebug(Ex, "{Ex}", ExceptionUtils.FormatExceptionDetails(Ex));
 				return (int)Ex.Result;
 			}
 			catch (BuildException Ex)
 			{
 				// BuildExceptions should have nicely formatted messages. We can log these directly.
-				Log.TraceError(ExceptionUtils.FormatException(Ex));
-				Log.TraceLog(ExceptionUtils.FormatExceptionDetails(Ex));
+				Logger.LogError(Ex, "{Ex}", ExceptionUtils.FormatException(Ex));
+				Logger.LogDebug(Ex, "{Ex}", ExceptionUtils.FormatExceptionDetails(Ex));
 				return (int)CompilationResult.OtherCompilationError;
 			}
 			catch (Exception Ex)
 			{
 				// Unhandled exception.
-				Log.TraceError("Unhandled exception: {0}", ExceptionUtils.FormatException(Ex));
-				Log.TraceLog(ExceptionUtils.FormatExceptionDetails(Ex));
+				Logger.LogError(Ex, "Unhandled exception: {Ex}", ExceptionUtils.FormatException(Ex));
+				Logger.LogDebug(Ex, "Unhandled exception: {Ex}", ExceptionUtils.FormatExceptionDetails(Ex));
 				return (int)CompilationResult.OtherCompilationError;
 			}
 			finally
@@ -649,10 +651,10 @@ namespace UnrealBuildTool
 					FileMetadataPrefetch.Stop();
 				}
 
-				Utils.LogWriteFileIfChangedActivity();
+				Utils.LogWriteFileIfChangedActivity(Logger);
 
 				// Print out all the performance info
-				Timeline.Print(TimeSpan.FromMilliseconds(20.0), LogEventType.Log);
+				Timeline.Print(TimeSpan.FromMilliseconds(20.0), LogLevel.Debug, Logger);
 
 				// Make sure we flush the logs however we exit
 				Trace.Close();
