UnrealBuildTool: Add BinarySuffix to BuildHostPlatform

#rnx

diff --git a/Engine/Source/Programs/UnrealBuildTool/Platform/Linux/LinuxToolChain.cs b/Engine/Source/Programs/UnrealBuildTool/Platform/Linux/LinuxToolChain.cs
--- a/Engine/Source/Programs/UnrealBuildTool/Platform/Linux/LinuxToolChain.cs
+++ b/Engine/Source/Programs/UnrealBuildTool/Platform/Linux/LinuxToolChain.cs
@@ -1,4 +1,4 @@
-﻿// Copyright Epic Games, Inc. All Rights Reserved.
+// Copyright Epic Games, Inc. All Rights Reserved.
 
 using System;
 using System.Collections.Generic;
@@ -89,11 +89,11 @@ namespace UnrealBuildTool
 				ToolchainInfo = String.Format("toolchain located at '{0}'", BaseLinuxPath);
 
 				// set up the path to our toolchain
-				ClangPath = Path.Combine(BaseLinuxPath, @"bin", "clang++" + GetHostPlatformBinarySuffix());
-				ArPath = Path.Combine(Path.Combine(BaseLinuxPath, String.Format("bin/{0}-{1}", Architecture, "ar" + GetHostPlatformBinarySuffix())));
-				LlvmArPath = Path.Combine(Path.Combine(BaseLinuxPath, String.Format("bin/{0}", "llvm-ar" + GetHostPlatformBinarySuffix())));
-				RanlibPath = Path.Combine(Path.Combine(BaseLinuxPath, String.Format("bin/{0}-{1}", Architecture, "ranlib" + GetHostPlatformBinarySuffix())));
-				ObjcopyPath = Path.Combine(Path.Combine(BaseLinuxPath, String.Format("bin/{0}", "llvm-objcopy" + GetHostPlatformBinarySuffix())));
+				ClangPath = Path.Combine(BaseLinuxPath, @"bin", "clang++" + BuildHostPlatform.Current.BinarySuffix);
+				ArPath = Path.Combine(Path.Combine(BaseLinuxPath, String.Format("bin/{0}-{1}", Architecture, "ar" + BuildHostPlatform.Current.BinarySuffix)));
+				LlvmArPath = Path.Combine(Path.Combine(BaseLinuxPath, String.Format("bin/{0}", "llvm-ar" + BuildHostPlatform.Current.BinarySuffix)));
+				RanlibPath = Path.Combine(Path.Combine(BaseLinuxPath, String.Format("bin/{0}-{1}", Architecture, "ranlib" + BuildHostPlatform.Current.BinarySuffix)));
+				ObjcopyPath = Path.Combine(Path.Combine(BaseLinuxPath, String.Format("bin/{0}", "llvm-objcopy" + BuildHostPlatform.Current.BinarySuffix)));
 				StripPath = ObjcopyPath;
 
 				// When cross-compiling on Windows, use old FixDeps. It is slow, but it does not have timing issues
@@ -142,16 +142,6 @@ namespace UnrealBuildTool
 			bPreprocessDepends = Target.LinuxPlatform.bPreprocessDepends;
 		}
 
-		private string GetHostPlatformBinarySuffix()
-		{
-			if (BuildHostPlatform.Current.Platform == UnrealTargetPlatform.Win64)
-			{
-				return ".exe";
-			}
-
-			return "";
-		}
-
 		protected virtual bool CrossCompiling()
 		{
 			return bIsCrossCompiling;
diff --git a/Engine/Source/Programs/UnrealBuildTool/System/BuildHostPlatform.cs b/Engine/Source/Programs/UnrealBuildTool/System/BuildHostPlatform.cs
--- a/Engine/Source/Programs/UnrealBuildTool/System/BuildHostPlatform.cs
+++ b/Engine/Source/Programs/UnrealBuildTool/System/BuildHostPlatform.cs
@@ -80,6 +80,11 @@
 		abstract public ShellType ShellType { get; }
 
 		/// <summary>
+		/// The executable binary suffix for this platform
+		/// </summary>
+		abstract public string BinarySuffix { get; }
+
+		/// <summary>
 		/// Class that holds information about a running process
 		/// </summary>
 		public class ProcessInfo
@@ -228,46 +233,32 @@
 
 	class WindowsBuildHostPlatform : BuildHostPlatform
 	{
-		public override UnrealTargetPlatform Platform
-		{
-			get { return UnrealTargetPlatform.Win64; }
-		}
+		public override UnrealTargetPlatform Platform => UnrealTargetPlatform.Win64;
 
-		public override FileReference Shell
-		{
-			get { return new FileReference(Environment.GetEnvironmentVariable("COMSPEC")!); }
-		}
+		public override FileReference Shell => new FileReference(Environment.GetEnvironmentVariable("COMSPEC")!);
 
-		public override ShellType ShellType
-		{
-			get { return ShellType.Cmd; }
-		}
+		public override ShellType ShellType => ShellType.Cmd;
+
+		public override string BinarySuffix => ".exe";
 
 		internal override IEnumerable<ProjectFileFormat> GetDefaultProjectFileFormats()
 		{
 			yield return ProjectFileFormat.VisualStudio;
-			#if __VPROJECT_AVAILABLE__
-				yield return ProjectFileFormat.VProject;
-			#endif
+#if __VPROJECT_AVAILABLE__
+			yield return ProjectFileFormat.VProject;
+#endif
 		}
 	}
 
 	class MacBuildHostPlatform : BuildHostPlatform
 	{
-		public override UnrealTargetPlatform Platform
-		{
-			get { return UnrealTargetPlatform.Mac; }
-		}
+		public override UnrealTargetPlatform Platform => UnrealTargetPlatform.Mac;
 
-		public override FileReference Shell
-		{
-			get { return new FileReference("/bin/sh"); }
-		}
+		public override FileReference Shell => new FileReference("/bin/sh");
 
-		public override ShellType ShellType
-		{
-			get { return ShellType.Sh; }
-		}
+		public override ShellType ShellType => ShellType.Sh;
+
+		public override string BinarySuffix => string.Empty;
 
 		/// <summary>
 		/// Currently Mono returns incomplete process names in Process.GetProcesses() so we need to parse 'ps' output.
@@ -372,28 +363,21 @@
 		internal override IEnumerable<ProjectFileFormat> GetDefaultProjectFileFormats()
 		{
 			yield return ProjectFileFormat.XCode;
-			#if __VPROJECT_AVAILABLE__
-				yield return ProjectFileFormat.VProject;
-			#endif
+#if __VPROJECT_AVAILABLE__
+			yield return ProjectFileFormat.VProject;
+#endif
 		}
 	}
 
 	class LinuxBuildHostPlatform : BuildHostPlatform
 	{
-		public override UnrealTargetPlatform Platform
-		{
-			get { return UnrealTargetPlatform.Linux; }
-		}
+		public override UnrealTargetPlatform Platform => UnrealTargetPlatform.Linux;
 
-		public override FileReference Shell
-		{
-			get { return new FileReference("/bin/sh"); }
-		}
+		public override FileReference Shell => new FileReference("/bin/sh");
 
-		public override ShellType ShellType
-		{
-			get { return ShellType.Sh; }
-		}
+		public override ShellType ShellType => ShellType.Sh;
+
+		public override string BinarySuffix => string.Empty;
 
 		/// <summary>
 		/// Currently Mono returns incomplete process names in Process.GetProcesses() so we need to use /proc
@@ -423,9 +407,9 @@
 		{
 			yield return ProjectFileFormat.Make;
 			yield return ProjectFileFormat.VisualStudioCode;
-			#if __VPROJECT_AVAILABLE__
-				yield return ProjectFileFormat.VProject;
-			#endif
+#if __VPROJECT_AVAILABLE__
+			yield return ProjectFileFormat.VProject;
+#endif
 		}
 	}
 }
