UBT: Fix exception getting resource compiler path.

diff --git a/Engine/Source/Programs/UnrealBuildTool/Platform/Windows/VCEnvironment.cs b/Engine/Source/Programs/UnrealBuildTool/Platform/Windows/VCEnvironment.cs
--- a/Engine/Source/Programs/UnrealBuildTool/Platform/Windows/VCEnvironment.cs
+++ b/Engine/Source/Programs/UnrealBuildTool/Platform/Windows/VCEnvironment.cs
@@ -131,7 +131,7 @@ namespace UnrealBuildTool
 			LibraryManagerPath = GetLibraryLinkerToolPath(Params.Platform, Compiler, CompilerDir, DefaultLinkerDir);
 
 			// Get the resource compiler path from the Windows SDK
-			ResourceCompilerPath = GetResourceCompilerToolPath(Params.Platform, WindowsSdkDir, WindowsSdkVersion);
+			ResourceCompilerPath = GetResourceCompilerToolPath(Params.Platform, WindowsSdkDir, WindowsSdkVersion, Logger);
 
 			// Get all the system include paths
 			SetupEnvironment(Params.Platform, Logger);
@@ -254,7 +254,7 @@ namespace UnrealBuildTool
 		/// <summary>
 		/// Gets the path to the resource compiler's rc.exe for the specified platform.
 		/// </summary>
-		virtual protected FileReference GetResourceCompilerToolPath(UnrealTargetPlatform Platform, DirectoryReference WindowsSdkDir, VersionNumber WindowsSdkVersion)
+		virtual protected FileReference GetResourceCompilerToolPath(UnrealTargetPlatform Platform, DirectoryReference WindowsSdkDir, VersionNumber WindowsSdkVersion, ILogger Logger)
 		{
 			// 64 bit -- we can use the 32 bit version to target 64 bit on 32 bit OS.
 			FileReference ResourceCompilerPath = FileReference.Combine(WindowsSdkDir, "bin", WindowsSdkVersion.ToString(), "x64", "rc.exe");
