UAT: Fix exception when creating dummy targets in UAT.

diff --git a/Engine/Source/Programs/AutomationTool/AutomationUtils/ProjectUtils.cs b/Engine/Source/Programs/AutomationTool/AutomationUtils/ProjectUtils.cs
--- a/Engine/Source/Programs/AutomationTool/AutomationUtils/ProjectUtils.cs
+++ b/Engine/Source/Programs/AutomationTool/AutomationUtils/ProjectUtils.cs
@@ -10,6 +10,7 @@ using System.Diagnostics;
 using EpicGames.Core;
 using System.Reflection;
 using UnrealBuildBase;
+using System.Runtime.Serialization;
 
 namespace AutomationTool
 {
@@ -783,7 +784,7 @@ namespace AutomationTool
 
 					// Create an instance of this type
 					CommandUtils.LogVerbose("Creating target rules object: {0}", TargetType.Name);
-					TargetRules Rules = Activator.CreateInstance(TargetType, DummyTargetInfo) as TargetRules;
+					TargetRules Rules = TargetRules.Create(TargetType, DummyTargetInfo, null, null, null, Log.Logger);
 					CommandUtils.LogVerbose("Adding target: {0} ({1})", TargetType.Name, Rules.Type);
 
 					SingleTargetProperties TargetData = new SingleTargetProperties();
diff --git a/Engine/Source/Programs/UnrealBuildTool/Configuration/TargetRules.cs b/Engine/Source/Programs/UnrealBuildTool/Configuration/TargetRules.cs
--- a/Engine/Source/Programs/UnrealBuildTool/Configuration/TargetRules.cs
+++ b/Engine/Source/Programs/UnrealBuildTool/Configuration/TargetRules.cs
@@ -8,6 +8,7 @@ using System.Reflection;
 using EpicGames.Core;
 using UnrealBuildBase;
 using Microsoft.Extensions.Logging;
+using System.Runtime.Serialization;
 
 namespace UnrealBuildTool
 {
@@ -1827,6 +1828,53 @@ namespace UnrealBuildTool
 		[ConfigSubObject]
 		public WindowsTargetRules WindowsPlatform; // Requires 'this' parameter; initialized in constructor
 
+		/// <summary>
+		/// Create a TargetRules instance
+		/// </summary>
+		/// <param name="RulesType">Type to create</param>
+		/// <param name="TargetInfo">Target info</param>
+		/// <param name="BaseFile">Path to the file for the rules assembly</param>
+		/// <param name="PlatformFile">Path to the platform specific rules file</param>
+		/// <param name="DefaultBuildSettings"></param>
+		/// <param name="Logger">Logger for the new target rules</param>
+		/// <returns>Target instance</returns>
+		public static TargetRules Create(Type RulesType, TargetInfo TargetInfo, FileReference? BaseFile, FileReference? PlatformFile, BuildSettingsVersion? DefaultBuildSettings, ILogger Logger)
+		{
+			TargetRules Rules = (TargetRules)FormatterServices.GetUninitializedObject(RulesType);
+			if (DefaultBuildSettings.HasValue)
+			{
+				Rules.DefaultBuildSettings = DefaultBuildSettings.Value;
+			}
+
+			// The base target file name: this affects where the resulting build product is created so the platform/group is not desired in this case.
+			Rules.File = BaseFile;
+
+			// The platform/group-specific target file name
+			Rules.TargetSourceFile = PlatformFile;
+
+			// Initialize the logger
+			Rules.Logger = Logger;
+
+			// Find the constructor
+			ConstructorInfo? Constructor = RulesType.GetConstructor(new Type[] { typeof(TargetInfo) });
+			if (Constructor == null)
+			{
+				throw new BuildException("No constructor found on {0} which takes an argument of type TargetInfo.", RulesType.Name);
+			}
+
+			// Invoke the regular constructor
+			try
+			{
+				Constructor.Invoke(Rules, new object[] { TargetInfo });
+			}
+			catch (Exception Ex)
+			{
+				throw new BuildException(Ex, "Unable to instantiate instance of '{0}' object type from compiled assembly '{1}'.  Unreal Build Tool creates an instance of your module's 'Rules' object in order to find out about your module's requirements.  The CLR exception details may provide more information:  {2}", RulesType.Name, Path.GetFileNameWithoutExtension(RulesType.Assembly?.Location), Ex.ToString());
+			}
+
+			return Rules;
+		}
+
 		/// <summary>
 		/// Constructor.
 		/// </summary>
diff --git a/Engine/Source/Programs/UnrealBuildTool/System/RulesAssembly.cs b/Engine/Source/Programs/UnrealBuildTool/System/RulesAssembly.cs
--- a/Engine/Source/Programs/UnrealBuildTool/System/RulesAssembly.cs
+++ b/Engine/Source/Programs/UnrealBuildTool/System/RulesAssembly.cs
@@ -583,37 +583,9 @@ namespace UnrealBuildTool
 
 			// Create an instance of the module's rules object, and set some defaults before calling the constructor.
 			Type RulesType = PlatformRulesType ?? BaseRulesType;
-			TargetRules Rules = (TargetRules)FormatterServices.GetUninitializedObject(RulesType);
-			if (DefaultBuildSettings.HasValue)
-			{
-				Rules.DefaultBuildSettings = DefaultBuildSettings.Value;
-			}
-
-			// The base target file name: this affects where the resulting build product is created so the platform/group is not desired in this case.
-			Rules.File = TargetNameToTargetFile[TargetInfo.Name];
-
-			// The platform/group-specific target file name
-			Rules.TargetSourceFile = TargetNameToTargetFile.TryGetValue(PlatformRulesName, out FileReference? PlatformTargetFile) ? PlatformTargetFile : Rules.File;
-
-			// Initialize the logger
-			Rules.Logger = Logger;
-
-			// Find the constructor
-			ConstructorInfo? Constructor = RulesType.GetConstructor(new Type[] { typeof(TargetInfo) });
-			if(Constructor == null)
-			{
-				throw new BuildException("No constructor found on {0} which takes an argument of type TargetInfo.", RulesType.Name);
-			}
-
-			// Invoke the regular constructor
-			try
-			{
-				Constructor.Invoke(Rules, new object[] { TargetInfo });
-			}
-			catch (Exception Ex)
-			{
-				throw new BuildException(Ex, "Unable to instantiate instance of '{0}' object type from compiled assembly '{1}'.  Unreal Build Tool creates an instance of your module's 'Rules' object in order to find out about your module's requirements.  The CLR exception details may provide more information:  {2}", TypeName, Path.GetFileNameWithoutExtension(CompiledAssembly?.Location), Ex.ToString());
-			}
+			FileReference BaseFile = TargetNameToTargetFile[TargetInfo.Name];
+			FileReference PlatformFile = TargetNameToTargetFile.TryGetValue(PlatformRulesName, out FileReference? PlatformTargetFile) ? PlatformTargetFile : BaseFile;
+			TargetRules Rules = TargetRules.Create(RulesType, TargetInfo, BaseFile, PlatformFile, DefaultBuildSettings, Logger);
 
 			// Set the default overriddes for the configured target type
 			Rules.SetOverridesForTargetType();
