Fixes Unhandled exception: System.MissingMethodException: Constructor on type 'UnrealBuildTool.AndroidProjectGenerator' not found when generating project files

Regression seams to be caused by 20373777

diff --git a/Engine/Source/Programs/UnrealBuildTool/Modes/GenerateProjectFilesMode.cs b/Engine/Source/Programs/UnrealBuildTool/Modes/GenerateProjectFilesMode.cs
--- a/Engine/Source/Programs/UnrealBuildTool/Modes/GenerateProjectFilesMode.cs
+++ b/Engine/Source/Programs/UnrealBuildTool/Modes/GenerateProjectFilesMode.cs
@@ -128,7 +128,7 @@ namespace UnrealBuildTool
 			{
 				if (CheckType.IsClass && !CheckType.IsAbstract && CheckType.IsSubclassOf(typeof(PlatformProjectGenerator)))
 				{
-					PlatformProjectGenerator Generator = (PlatformProjectGenerator)Activator.CreateInstance(CheckType, Arguments)!;
+					PlatformProjectGenerator Generator = (PlatformProjectGenerator)Activator.CreateInstance(CheckType, Arguments, Logger)!;
 					foreach(UnrealTargetPlatform Platform in Generator.GetPlatforms())
 					{
 						if(DisablePlatformProjectGenerators == null || !DisablePlatformProjectGenerators.Any(x => x.Equals(Platform.ToString(), StringComparison.OrdinalIgnoreCase)))
