Enable RayTracing by default for Vualkn RHI
*Only inline ray tracing use cases are currently supported, as in Vulkan RHI on Windows

Fix Linux not loading raytracing extensions if they were enabled

diff --git a/Engine/Config/BaseEngine.ini b/Engine/Config/BaseEngine.ini
--- a/Engine/Config/BaseEngine.ini
+++ b/Engine/Config/BaseEngine.ini
@@ -2757,6 +2757,7 @@
 
 [/Script/LinuxTargetPlatform.LinuxTargetSettings]
 +TargetedRHIs=SF_VULKAN_SM5
+bEnableRayTracing=true
 
 [/Script/MacTargetPlatform.MacTargetSettings]
 MetalLanguageVersion=5
diff --git a/Engine/Source/Programs/UnrealBuildTool/Platform/Linux/UEBuildLinux.cs b/Engine/Source/Programs/UnrealBuildTool/Platform/Linux/UEBuildLinux.cs
--- a/Engine/Source/Programs/UnrealBuildTool/Platform/Linux/UEBuildLinux.cs
+++ b/Engine/Source/Programs/UnrealBuildTool/Platform/Linux/UEBuildLinux.cs
@@ -85,6 +85,12 @@
 		[CommandLine("-PreprocessDepends")]
 		[XmlConfigFile(Category = "BuildConfiguration", Name = "bPreprocessDepends")]
 		public bool bPreprocessDepends = false;
+
+		/// <summary>
+		/// Enables runtime ray tracing support.
+		/// </summary>
+		[ConfigFile(ConfigHierarchyType.Engine, "/Script/LinuxPlatform.LinuxTargetSettings")]
+		public bool bEnableRayTracing = false;
 	}
 
 	/// <summary>
@@ -144,6 +150,11 @@
 			get { return Inner.bPreprocessDepends; }
 		}
 
+		public bool bEnableRayTracing
+		{
+			get { return Inner.bEnableRayTracing; }
+		}
+
 		#pragma warning restore CS1591
 		#endregion
 	}
@@ -511,6 +522,11 @@
 
 			CompileEnvironment.Definitions.Add("INT64_T_TYPES_NOT_LONG_LONG=1");
 
+			if (Target.LinuxPlatform.bEnableRayTracing)
+			{
+				CompileEnvironment.Definitions.Add("RHI_RAYTRACING=1");
+			}
+
 			// link with Linux libraries.
 			LinkEnvironment.SystemLibraries.Add("pthread");
 
diff --git a/Engine/Source/Runtime/VulkanRHI/Private/Linux/VulkanLinuxPlatform.cpp b/Engine/Source/Runtime/VulkanRHI/Private/Linux/VulkanLinuxPlatform.cpp
--- a/Engine/Source/Runtime/VulkanRHI/Private/Linux/VulkanLinuxPlatform.cpp
+++ b/Engine/Source/Runtime/VulkanRHI/Private/Linux/VulkanLinuxPlatform.cpp
@@ -109,6 +109,14 @@ bool FVulkanLinuxPlatform::LoadVulkanInstanceFunctions(VkInstance inInstance)
 		return false;
 	}
 
+#if VULKAN_RHI_RAYTRACING
+	const bool bFoundRayTracingEntries = FVulkanRayTracingPlatform::LoadVulkanInstanceFunctions(inInstance);
+	if (!bFoundRayTracingEntries)
+	{
+		UE_LOG(LogVulkanRHI, Warning, TEXT("Vulkan RHI ray tracing is enabled, but failed to load instance functions."));
+	}
+#endif
+
 	ENUM_VK_ENTRYPOINTS_OPTIONAL_INSTANCE(GETINSTANCE_VK_ENTRYPOINTS);
 	ENUM_VK_ENTRYPOINTS_OPTIONAL_PLATFORM_INSTANCE(GETINSTANCE_VK_ENTRYPOINTS);
 #if UE_BUILD_DEBUG
diff --git a/Engine/Source/Runtime/VulkanRHI/Private/VulkanExtensions.cpp b/Engine/Source/Runtime/VulkanRHI/Private/VulkanExtensions.cpp
--- a/Engine/Source/Runtime/VulkanRHI/Private/VulkanExtensions.cpp
+++ b/Engine/Source/Runtime/VulkanRHI/Private/VulkanExtensions.cpp
@@ -37,9 +37,9 @@ TAutoConsoleVariable<int32> GRHIAllow64bitShaderAtomicsCvar(
 
 TAutoConsoleVariable<int32> GVulkanRayTracingCVar(
 	TEXT("r.Vulkan.RayTracing"),
-	0,
-	TEXT("0: Do not enable Vulkan ray tracing extensions (default)\n")
-	TEXT("1: Enable experimental ray tracing support (for development and testing purposes)"),
+	1,
+	TEXT("0: Do not enable Vulkan ray tracing extensions\n")
+	TEXT("1: Enable experimental ray tracing support (default)"),
 	ECVF_ReadOnly
 );
 
