UBT: Prefer AutoSDK toolchains over manually installed toolchains, within the preferred version range. Sometimes hotfixes to a particular compiler can change behavior, and we want build machines to behave deterministically.

diff --git a/Engine/Source/Programs/UnrealBuildTool/Platform/Windows/MicrosoftPlatformSDK.cs b/Engine/Source/Programs/UnrealBuildTool/Platform/Windows/MicrosoftPlatformSDK.cs
--- a/Engine/Source/Programs/UnrealBuildTool/Platform/Windows/MicrosoftPlatformSDK.cs
+++ b/Engine/Source/Programs/UnrealBuildTool/Platform/Windows/MicrosoftPlatformSDK.cs
@@ -590,6 +590,7 @@
 				.ThenByDescending(x => x.Is64Bit)
 				.ThenBy(x => x.IsPreview)
 				.ThenBy(x => x.FamilyRank)
+				.ThenByDescending(x => x.IsAutoSdk)
 				.ThenByDescending(x => x.Version)
 				.FirstOrDefault();
 
@@ -613,25 +614,25 @@
 					{
 						// Check for a manual installation to the default directory
 						DirectoryReference ManualInstallDir = DirectoryReference.Combine(DirectoryReference.GetSpecialFolder(Environment.SpecialFolder.ProgramFiles)!, "LLVM");
-						AddClangToolChain(ManualInstallDir, ToolChains, Logger);
+						AddClangToolChain(ManualInstallDir, ToolChains, IsAutoSdk: false, Logger);
 
 						// Check for a manual installation to a custom directory
 						string? LlvmPath = Environment.GetEnvironmentVariable("LLVM_PATH");
 						if (!String.IsNullOrEmpty(LlvmPath))
 						{
-							AddClangToolChain(new DirectoryReference(LlvmPath), ToolChains, Logger);
+							AddClangToolChain(new DirectoryReference(LlvmPath), ToolChains, IsAutoSdk: false, Logger);
 						}
 
 						// Check for installations bundled with Visual Studio 2019
 						foreach (VisualStudioInstallation Installation in FindVisualStudioInstallations(WindowsCompiler.VisualStudio2019, Logger))
 						{
-							AddClangToolChain(DirectoryReference.Combine(Installation.BaseDir, "VC", "Tools", "Llvm"), ToolChains, Logger);
+							AddClangToolChain(DirectoryReference.Combine(Installation.BaseDir, "VC", "Tools", "Llvm"), ToolChains, IsAutoSdk: false, Logger);
 						}
 
 						// Check for installations bundled with Visual Studio 2022
 						foreach (VisualStudioInstallation Installation in FindVisualStudioInstallations(WindowsCompiler.VisualStudio2022, Logger))
 						{
-							AddClangToolChain(DirectoryReference.Combine(Installation.BaseDir, "VC", "Tools", "Llvm"), ToolChains, Logger);
+							AddClangToolChain(DirectoryReference.Combine(Installation.BaseDir, "VC", "Tools", "Llvm"), ToolChains, IsAutoSdk: false, Logger);
 						}
 
 						// Check for AutoSDK paths
@@ -643,7 +644,7 @@
 							{
 								foreach (DirectoryReference ToolChainDir in DirectoryReference.EnumerateDirectories(ClangBaseDir))
 								{
-									AddClangToolChain(ToolChainDir, ToolChains, Logger);
+									AddClangToolChain(ToolChainDir, ToolChains, IsAutoSdk: true, Logger);
 								}
 							}
 						}
@@ -651,10 +652,10 @@
 					else if (Compiler == WindowsCompiler.Intel)
 					{
 						DirectoryReference InstallDir = DirectoryReference.Combine(DirectoryReference.GetSpecialFolder(Environment.SpecialFolder.ProgramFiles)!, "Intel", "oneAPI", "compiler");
-						FindIntelOneApiToolChains(InstallDir, ToolChains, Logger);
+						FindIntelOneApiToolChains(InstallDir, ToolChains, IsAutoSdk: false, Logger);
 
 						InstallDir = DirectoryReference.Combine(DirectoryReference.GetSpecialFolder(Environment.SpecialFolder.ProgramFilesX86)!, "Intel", "oneAPI", "compiler");
-						FindIntelOneApiToolChains(InstallDir, ToolChains, Logger);
+						FindIntelOneApiToolChains(InstallDir, ToolChains, IsAutoSdk: false, Logger);
 
 						// Check for AutoSDK paths
 						DirectoryReference? AutoSdkDir;
@@ -663,7 +664,7 @@
 							DirectoryReference IntelBaseDir = DirectoryReference.Combine(AutoSdkDir, "Win64", "Intel");
 							if (DirectoryReference.Exists(IntelBaseDir))
 							{
-								FindIntelOneApiToolChains(IntelBaseDir, ToolChains, Logger);
+								FindIntelOneApiToolChains(IntelBaseDir, ToolChains, IsAutoSdk: true, Logger);
 							}
 						}
 					}
@@ -675,7 +676,7 @@
 						{
 							DirectoryReference ToolChainBaseDir = DirectoryReference.Combine(Installation.BaseDir, "VC", "Tools", "MSVC");
 							DirectoryReference RedistBaseDir = DirectoryReference.Combine(Installation.BaseDir, "VC", "Redist", "MSVC");
-							FindVisualStudioToolChains(ToolChainBaseDir, RedistBaseDir, Installation.bPreview, ToolChains, Logger);
+							FindVisualStudioToolChains(ToolChainBaseDir, RedistBaseDir, Installation.bPreview, ToolChains, IsAutoSdk: false, Logger);
 						}
 
 						// Enumerate all the AutoSDK toolchains
@@ -692,10 +693,10 @@
 							if (!string.IsNullOrEmpty(VSDir))
 							{
 								DirectoryReference ReleaseBaseDir = DirectoryReference.Combine(PlatformDir, "Win64", VSDir);
-								FindVisualStudioToolChains(ReleaseBaseDir, null, false, ToolChains, Logger);
+								FindVisualStudioToolChains(ReleaseBaseDir, null, false, ToolChains, IsAutoSdk: true, Logger);
 
 								DirectoryReference PreviewBaseDir = DirectoryReference.Combine(PlatformDir, "Win64", $"{VSDir}-Preview");
-								FindVisualStudioToolChains(PreviewBaseDir, null, true, ToolChains, Logger);
+								FindVisualStudioToolChains(PreviewBaseDir, null, true, ToolChains, IsAutoSdk: true, Logger);
 							}
 						}
 					}
@@ -809,8 +810,9 @@
 		/// <param name="OptionalRedistDir">Optional directory for redistributable components (DLLs etc)</param>
 		/// <param name="bPreview">Whether this is a preview installation</param>
 		/// <param name="ToolChains">Map of tool chain version to installation info</param>
+		/// <param name="IsAutoSdk">Whether this folder contains AutoSDK entries</param>
 		/// <param name="Logger">Logger for output</param>
-		static void FindVisualStudioToolChains(DirectoryReference BaseDir, DirectoryReference? OptionalRedistDir, bool bPreview, List<ToolChainInstallation> ToolChains, ILogger Logger)
+		static void FindVisualStudioToolChains(DirectoryReference BaseDir, DirectoryReference? OptionalRedistDir, bool bPreview, List<ToolChainInstallation> ToolChains, bool IsAutoSdk, ILogger Logger)
 		{
 			if (DirectoryReference.Exists(BaseDir))
 			{
@@ -821,7 +823,7 @@
 					{
 						DirectoryReference? RedistDir = FindVisualStudioRedistForToolChain(ToolChainDir, OptionalRedistDir, Version);
 
-						AddVisualCppToolChain(Version, bPreview, ToolChainDir, RedistDir, ToolChains, Logger);
+						AddVisualCppToolChain(Version, bPreview, ToolChainDir, RedistDir, ToolChains, IsAutoSdk, Logger);
 					}
 				}
 			}
@@ -872,8 +874,9 @@
 		/// <param name="ToolChainDir"></param>
 		/// <param name="RedistDir"></param>
 		/// <param name="ToolChains"></param>
+		/// <param name="IsAutoSdk"></param>
 		/// <param name="Logger"></param>
-		static void AddVisualCppToolChain(VersionNumber Version, bool bPreview, DirectoryReference ToolChainDir, DirectoryReference? RedistDir, List<ToolChainInstallation> ToolChains, ILogger Logger)
+		static void AddVisualCppToolChain(VersionNumber Version, bool bPreview, DirectoryReference ToolChainDir, DirectoryReference? RedistDir, List<ToolChainInstallation> ToolChains, bool IsAutoSdk, ILogger Logger)
 		{
 			bool Is64Bit = Has64BitToolChain(ToolChainDir);
 
@@ -898,13 +901,13 @@
 			}
 
 			Logger.LogDebug("Found Visual Studio toolchain: {ToolChainDir} (Family={Family}, FamilyRank={FamilyRank}, Version={Version}, Is64Bit={Is64Bit}, Preview={Preview}, Architecture={Arch}, Error={Error}, Redist={RedistDir})", ToolChainDir, Family, FamilyRank, Version, Is64Bit, bPreview, WindowsArchitecture.x64.ToString(), Error != null, RedistDir);
-			ToolChains.Add(new ToolChainInstallation(Family, FamilyRank, Version, Is64Bit, bPreview, WindowsArchitecture.x64, Error, ToolChainDir, RedistDir));
+			ToolChains.Add(new ToolChainInstallation(Family, FamilyRank, Version, Is64Bit, bPreview, WindowsArchitecture.x64, Error, ToolChainDir, RedistDir, IsAutoSdk));
 
 			bool HasArm64 = HasArm64ToolChain(ToolChainDir);
 			if (HasArm64)
 			{
 				Logger.LogDebug("Found Visual Studio toolchain: {ToolChainDir} (Family={Family}, FamilyRank={FamilyRank}, Version={Version}, Is64Bit={Is64Bit}, Preview={Preview}, Architecture={Arch}, Error={Error}, Redist={RedistDir})", ToolChainDir, Family, FamilyRank, Version, Is64Bit, bPreview, WindowsArchitecture.ARM64.ToString(), Error != null, RedistDir);
-				ToolChains.Add(new ToolChainInstallation(Family, FamilyRank, Version, Is64Bit, bPreview, WindowsArchitecture.ARM64, Error, ToolChainDir, RedistDir));
+				ToolChains.Add(new ToolChainInstallation(Family, FamilyRank, Version, Is64Bit, bPreview, WindowsArchitecture.ARM64, Error, ToolChainDir, RedistDir, IsAutoSdk));
 			}
 		}
 
@@ -913,8 +916,9 @@
 		/// </summary>
 		/// <param name="ToolChainDir"></param>
 		/// <param name="ToolChains"></param>
+		/// <param name="IsAutoSdk"></param>
 		/// <param name="Logger"></param>
-		static void AddClangToolChain(DirectoryReference ToolChainDir, List<ToolChainInstallation> ToolChains, ILogger Logger)
+		static void AddClangToolChain(DirectoryReference ToolChainDir, List<ToolChainInstallation> ToolChains, bool IsAutoSdk, ILogger Logger)
 		{
 			FileReference CompilerFile = FileReference.Combine(ToolChainDir, "bin", "clang-cl.exe");
 			if (FileReference.Exists(CompilerFile))
@@ -932,7 +936,7 @@
 				}
 
 				Logger.LogDebug("Found Clang toolchain: {ToolChainDir} (Version={Version}, Is64Bit={Is64Bit}, Rank={Rank})", ToolChainDir, Version, Is64Bit, Rank);
-				ToolChains.Add(new ToolChainInstallation(Version, Rank, Version, Is64Bit, false, WindowsArchitecture.x64, null, ToolChainDir, null));
+				ToolChains.Add(new ToolChainInstallation(Version, Rank, Version, Is64Bit, false, WindowsArchitecture.x64, null, ToolChainDir, null, IsAutoSdk));
 			}
 		}
 
@@ -942,8 +946,9 @@
 		/// </summary>
 		/// <param name="ToolChainDir"></param>
 		/// <param name="ToolChains"></param>
+		/// <param name="IsAutoSdk"></param>
 		/// <param name="Logger"></param>
-		static void AddIntelOneApiToolChain(DirectoryReference ToolChainDir, List<ToolChainInstallation> ToolChains, ILogger Logger)
+		static void AddIntelOneApiToolChain(DirectoryReference ToolChainDir, List<ToolChainInstallation> ToolChains, bool IsAutoSdk, ILogger Logger)
 		{
 			FileReference CompilerFile = FileReference.Combine(ToolChainDir, "windows", "bin", "icx.exe");
 			if (FileReference.Exists(CompilerFile))
@@ -967,7 +972,7 @@
 				}
 
 				Logger.LogDebug("Found Intel OneAPI toolchain: {ToolChainDir} (Version={Version}, Is64Bit={Is64Bit}, Rank={Rank})", ToolChainDir, Version, Is64Bit, Rank);
-				ToolChains.Add(new ToolChainInstallation(Version, Rank, Version, Is64Bit, false, WindowsArchitecture.x64, Error, ToolChainDir, null));
+				ToolChains.Add(new ToolChainInstallation(Version, Rank, Version, Is64Bit, false, WindowsArchitecture.x64, Error, ToolChainDir, null, IsAutoSdk));
 			}
 		}
 
@@ -976,19 +981,19 @@
 		/// </summary>
 		/// <param name="BaseDir">Base directory to search</param>
 		/// <param name="ToolChains">Map of tool chain version to installation info</param>
+		/// <param name="IsAutoSdk"></param>
 		/// <param name="Logger">Logger for output</param>
-		static void FindIntelOneApiToolChains(DirectoryReference BaseDir, List<ToolChainInstallation> ToolChains, ILogger Logger)
+		static void FindIntelOneApiToolChains(DirectoryReference BaseDir, List<ToolChainInstallation> ToolChains, bool IsAutoSdk, ILogger Logger)
 		{
 			if (DirectoryReference.Exists(BaseDir))
 			{
 				foreach (DirectoryReference ToolChainDir in DirectoryReference.EnumerateDirectories(BaseDir))
 				{
-					AddIntelOneApiToolChain(ToolChainDir, ToolChains, Logger);
+					AddIntelOneApiToolChain(ToolChainDir, ToolChains, IsAutoSdk, Logger);
 				}
 			}
 		}
 
-
 		/// <summary>
 		/// Test whether an executable is 64-bit
 		/// </summary>
@@ -1316,6 +1321,11 @@
 		public DirectoryReference? RedistDir { get; }
 
 		/// <summary>
+		/// Whether this toolchain comes from AutoSDK.
+		/// </summary>
+		public bool IsAutoSdk { get; }
+
+		/// <summary>
 		/// Constructor
 		/// </summary>
 		/// <param name="Family"></param>
@@ -1327,7 +1337,8 @@
 		/// <param name="Error"></param>
 		/// <param name="BaseDir">Base directory for the toolchain</param>
 		/// <param name="RedistDir">Optional directory for redistributable components (DLLs etc)</param>
-		public ToolChainInstallation(VersionNumber Family, int FamilyRank, VersionNumber Version, bool Is64Bit, bool IsPreview, WindowsArchitecture Architecture, string? Error, DirectoryReference BaseDir, DirectoryReference? RedistDir)
+		/// <param name="IsAutoSdk">Whether this toolchain comes from AutoSDK</param>
+		public ToolChainInstallation(VersionNumber Family, int FamilyRank, VersionNumber Version, bool Is64Bit, bool IsPreview, WindowsArchitecture Architecture, string? Error, DirectoryReference BaseDir, DirectoryReference? RedistDir, bool IsAutoSdk)
 		{
 			this.Family = Family;
 			this.FamilyRank = FamilyRank;
@@ -1338,6 +1349,7 @@
 			this.Error = Error;
 			this.BaseDir = BaseDir;
 			this.RedistDir = RedistDir;
+			this.IsAutoSdk = IsAutoSdk;
 		}
 	}
 }
