UBT: Allow latest version of the VS2019 compiler in AutoSDK.

diff --git a/Engine/Source/Programs/UnrealBuildTool/Platform/Windows/MicrosoftPlatformSDK.Versions.cs b/Engine/Source/Programs/UnrealBuildTool/Platform/Windows/MicrosoftPlatformSDK.Versions.cs
--- a/Engine/Source/Programs/UnrealBuildTool/Platform/Windows/MicrosoftPlatformSDK.Versions.cs
+++ b/Engine/Source/Programs/UnrealBuildTool/Platform/Windows/MicrosoftPlatformSDK.Versions.cs
@@ -34,7 +34,7 @@ namespace UnrealBuildTool
 		/// </summary>
 		static readonly VersionNumberRange[] PreferredVisualCppVersions = new VersionNumberRange[]
 		{
-			VersionNumberRange.Parse("14.29.30133", "14.29.30136"), // VS2019 16.11.5
+			VersionNumberRange.Parse("14.29.30133", "14.29.30145"), // VS2019 16.11.32428.217
 		};
 
 		/// <summary>
diff --git a/Engine/Source/Programs/UnrealBuildTool/Platform/Windows/MicrosoftPlatformSDK.cs b/Engine/Source/Programs/UnrealBuildTool/Platform/Windows/MicrosoftPlatformSDK.cs
--- a/Engine/Source/Programs/UnrealBuildTool/Platform/Windows/MicrosoftPlatformSDK.cs
+++ b/Engine/Source/Programs/UnrealBuildTool/Platform/Windows/MicrosoftPlatformSDK.cs
@@ -1273,6 +1273,7 @@ namespace UnrealBuildTool
 	/// <summary>
 	/// Information about a particular toolchain installation
 	/// </summary>
+	[DebuggerDisplay("{BaseDir}")]
 	internal class ToolChainInstallation
 	{
 		/// <summary>
