deprecate UnrealPlatformGroup.HoloLens

#rnx

diff --git a/Engine/Platforms/Hololens/Source/Programs/UnrealBuildTool/HoloLensToolChain.cs b/Engine/Platforms/Hololens/Source/Programs/UnrealBuildTool/HoloLensToolChain.cs
--- a/Engine/Platforms/Hololens/Source/Programs/UnrealBuildTool/HoloLensToolChain.cs
+++ b/Engine/Platforms/Hololens/Source/Programs/UnrealBuildTool/HoloLensToolChain.cs
@@ -34,6 +34,11 @@
 			CommandLine.ParseArguments(Environment.GetCommandLineArgs(), this, Logger);
 		}
 
+		protected override bool UseWindowsArchitecture(UnrealTargetPlatform Platform)
+		{
+			return true;
+		}
+
 		protected override void AppendCLArguments_Global(CppCompileEnvironment CompileEnvironment, List<string> Arguments)
 		{
 			base.AppendCLArguments_Global(CompileEnvironment, Arguments);
diff --git a/Engine/Platforms/Hololens/Source/Programs/UnrealBuildTool/UEBuildHoloLens.cs b/Engine/Platforms/Hololens/Source/Programs/UnrealBuildTool/UEBuildHoloLens.cs
--- a/Engine/Platforms/Hololens/Source/Programs/UnrealBuildTool/UEBuildHoloLens.cs
+++ b/Engine/Platforms/Hololens/Source/Programs/UnrealBuildTool/UEBuildHoloLens.cs
@@ -14,6 +14,15 @@ using Microsoft.Extensions.Logging;
 
 namespace UnrealBuildTool
 {
+	partial struct UnrealPlatformGroup
+	{
+		/// <summary>
+		/// this group was for grouping HoloLens32 and HoloLens64 into HoloLens directories but is no longer relevant
+		/// </summary>
+		[Obsolete("UnrealPlatformGroup.HoloLens is deprecated - please change the code to use ModuleName_HoloLens.build.cs etc or reference the UnrealTargetPlatform.HoloLens directly")]
+		public static UnrealPlatformGroup HoloLens = FindOrAddByName("HoloLens");
+	}
+
 	/// <summary>
 	/// HoloLens-specific target settings
 	/// </summary>
@@ -814,8 +823,10 @@ namespace UnrealBuildTool
 
 				UEBuildPlatform.RegisterBuildPlatform(new HoloLensPlatform(SDK, Logger), Logger);
 				UEBuildPlatform.RegisterPlatformWithGroup(UnrealTargetPlatform.HoloLens, UnrealPlatformGroup.Microsoft);
+#pragma warning disable CS0618 // Type or member is obsolete
 				UEBuildPlatform.RegisterPlatformWithGroup(UnrealTargetPlatform.HoloLens, UnrealPlatformGroup.HoloLens);
+#pragma warning restore CS0618
+				}
 			}
 		}
-	}
 }
diff --git a/Engine/Source/Programs/UnrealBuildTool/Configuration/UEBuildTarget.cs b/Engine/Source/Programs/UnrealBuildTool/Configuration/UEBuildTarget.cs
--- a/Engine/Source/Programs/UnrealBuildTool/Configuration/UEBuildTarget.cs
+++ b/Engine/Source/Programs/UnrealBuildTool/Configuration/UEBuildTarget.cs
@@ -524,11 +524,6 @@
 		public static UnrealPlatformGroup Windows = FindOrAddByName("Windows");
 
 		/// <summary>
-		/// this group is just to lump HoloLens32 and HoloLens64 into HoloLens directories
-		/// </summary>
-		public static UnrealPlatformGroup HoloLens = FindOrAddByName("HoloLens");
-
-		/// <summary>
 		/// Microsoft platforms
 		/// </summary>
 		public static UnrealPlatformGroup Microsoft = FindOrAddByName("Microsoft");
diff --git a/Engine/Source/Programs/UnrealBuildTool/Platform/Windows/VCToolChain.cs b/Engine/Source/Programs/UnrealBuildTool/Platform/Windows/VCToolChain.cs
--- a/Engine/Source/Programs/UnrealBuildTool/Platform/Windows/VCToolChain.cs
+++ b/Engine/Source/Programs/UnrealBuildTool/Platform/Windows/VCToolChain.cs
@@ -1006,7 +1006,7 @@
 			//
 			//	PC
 			//
-			if (LinkEnvironment.Platform.IsInGroup(UnrealPlatformGroup.Windows) || LinkEnvironment.Platform.IsInGroup(UnrealPlatformGroup.HoloLens))
+			if (UseWindowsArchitecture(LinkEnvironment.Platform))
 			{
 				Arguments.Add($"/MACHINE:{WindowsExports.GetArchitectureSubpath(Target.WindowsPlatform.Architecture)}");
 
@@ -1167,7 +1167,7 @@
 			//
 			//	PC
 			//
-			if (LinkEnvironment.Platform.IsInGroup(UnrealPlatformGroup.Windows) || LinkEnvironment.Platform.IsInGroup(UnrealPlatformGroup.HoloLens))
+			if (UseWindowsArchitecture(LinkEnvironment.Platform))
 			{
 				Arguments.Add($"/MACHINE:{WindowsExports.GetArchitectureSubpath(Target.WindowsPlatform.Architecture)}");
 
@@ -1547,6 +1547,11 @@
 			return WindowsPlatform.GetWindowsApplicationIcon(ProjectFile);
 		}
 
+		protected virtual bool UseWindowsArchitecture(UnrealTargetPlatform Platform)
+		{
+			return Platform.IsInGroup(UnrealPlatformGroup.Windows);
+		}
+
 		public override CPPOutput CompileRCFiles(CppCompileEnvironment CompileEnvironment, List<FileItem> InputFiles, DirectoryReference OutputDir, IActionGraphBuilder Graph)
 		{
 			CPPOutput Result = new CPPOutput();
diff --git a/Engine/Source/Programs/UnrealBuildTool/ProjectFiles/VisualStudio/VCProject.cs b/Engine/Source/Programs/UnrealBuildTool/ProjectFiles/VisualStudio/VCProject.cs
--- a/Engine/Source/Programs/UnrealBuildTool/ProjectFiles/VisualStudio/VCProject.cs
+++ b/Engine/Source/Programs/UnrealBuildTool/ProjectFiles/VisualStudio/VCProject.cs
@@ -1821,33 +1821,10 @@
 				{
 					TargetRules TargetRulesObject = Combination.ProjectTarget.TargetRules!;
 
-					if (Combination.Platform!.Value.IsInGroup(UnrealPlatformGroup.HoloLens)) // this logic should be moved into HoloLensProjectGenerator.GetVisualStudioUserFileStrings
-					{
-						VCUserFileContent.AppendLine("  <PropertyGroup {0}>", ConditionString);
-						if (TargetRulesObject.Type != TargetType.Game)
-						{
-							string DebugOptions = "";
-
-							if (IsForeignProject)
-							{
-								DebugOptions += UProjectPath;
-								DebugOptions += " -skipcompile";
-							}
-							else if (TargetRulesObject.Type == TargetType.Editor && ProjectName != ProjectFileGenerator.EngineProjectFileNameBase)
-							{
-								DebugOptions += ProjectName;
-							}
-
-							VCUserFileContent.AppendLine("    <LocalDebuggerCommandArguments>{0}</LocalDebuggerCommandArguments>", DebugOptions);
-						}
-						VCUserFileContent.AppendLine("    <DebuggerFlavor>{0}</DebuggerFlavor>", "AppHostLocalDebugger ");
-						VCUserFileContent.AppendLine("  </PropertyGroup>");
-					}
-
 					if (ProjGenerator != null)
 					{
 						string? ForeignUProjectPath = (IsForeignProject && !string.IsNullOrEmpty(UProjectPath)) ? UProjectPath : null;
-						VCUserFileContent.Append(ProjGenerator.GetVisualStudioUserFileStrings(Combination.Platform.Value, Configuration, ConditionString, TargetRulesObject, Combination.ProjectTarget.TargetFilePath, ProjectFilePath, ProjectName, ForeignUProjectPath));
+						VCUserFileContent.Append(ProjGenerator.GetVisualStudioUserFileStrings(Combination.Platform!.Value, Configuration, ConditionString, TargetRulesObject, Combination.ProjectTarget.TargetFilePath, ProjectFilePath, ProjectName, ForeignUProjectPath));
 					}
 				}
 			}
