Remove ShaderCompileWorker from the allowed list for Apple Silicon.

- Regressions in ASTC texture compression caused us to have to roll-back ISPC and ISPCTexComp binary rebuilds. Consequently, we have to disallow SCW until we can either remove the dependency or fix-up the regression and check in updated ISPCTexComp binaires.

#rnx

diff --git a/Engine/Source/Programs/UnrealBuildTool/Platform/Mac/MacExports.cs b/Engine/Source/Programs/UnrealBuildTool/Platform/Mac/MacExports.cs
--- a/Engine/Source/Programs/UnrealBuildTool/Platform/Mac/MacExports.cs
+++ b/Engine/Source/Programs/UnrealBuildTool/Platform/Mac/MacExports.cs
@@ -62,7 +62,7 @@
 					"BenchmarkTool",
 					"BlankProgram",
 					"BuildPatchTool",
-					"ShaderCompileWorker",
+					/* "ShaderCompileWorker", // blocked on ISPCTexComp */
 					"UnrealClient",
 					"UnrealGame", 
 					"UnrealServer",
