Ignore compile actions that aren't using the c++ compiler when doing live coding.  Resolves issues where changes that would contribute to the resource file (i.e. copyright notice) would cause live coding to fail.

diff --git a/Engine/Source/Programs/UnrealBuildTool/System/HotReload.cs b/Engine/Source/Programs/UnrealBuildTool/System/HotReload.cs
--- a/Engine/Source/Programs/UnrealBuildTool/System/HotReload.cs
+++ b/Engine/Source/Programs/UnrealBuildTool/System/HotReload.cs
@@ -592,7 +592,7 @@
 				}
 
 				// Filter the prerequisite actions down to just the compile actions, then recompute all the actions to execute
-				PrerequisiteActions = new List<LinkedAction>(TargetActionsToExecute.Where(x => x.ActionType == ActionType.Compile));
+				PrerequisiteActions = new List<LinkedAction>(TargetActionsToExecute.Where(x => IsLiveCodingAction(x)));
 				TargetActionsToExecute = ActionGraph.GetActionsToExecute(PrerequisiteActions, CppDependencies, History, BuildConfiguration.bIgnoreOutdatedImportLibraries, Logger);
 
 				// Update the action graph with these new paths
@@ -803,6 +803,20 @@
 		}
 
 		/// <summary>
+		/// Test to see if the action is an action live coding supports.  All other actions will be filtered
+		/// </summary>
+		/// <param name="Action">Action in question</param>
+		/// <returns>True if the action is a compile action for the compiler.  This filters out RC compiles.</returns>
+		static bool IsLiveCodingAction(LinkedAction Action)
+		{
+			return Action.ActionType == ActionType.Compile &&
+				(Action.CommandPath.GetFileName().Equals("cl-filter.exe", StringComparison.OrdinalIgnoreCase)
+					|| Action.CommandPath.GetFileName().Equals("cl.exe", StringComparison.OrdinalIgnoreCase)
+					|| Action.CommandPath.GetFileName().Equals("clang-cl.exe", StringComparison.OrdinalIgnoreCase)
+				);
+		}
+
+		/// <summary>
 		/// Patches a set of actions for use with live coding. The new action list will output object files to a different location.
 		/// </summary>
 		/// <param name="Actions">Set of actions</param>
