Created a UBT commandline argument to suppress the SDK warnings when running UBT from a Warning level to Log. The suppression occurs during when UEBuildPlatformSDK.cs is printing SDK info to console/log.

There are steps in Horde where we generate project files for Verse, and are seeing the warning for all the Apple platforms. This gives us a way to suppress those.

diff --git a/Engine/Source/Programs/Shared/EpicGames.Build/UEBuildPlatformSDK.cs b/Engine/Source/Programs/Shared/EpicGames.Build/UEBuildPlatformSDK.cs
--- a/Engine/Source/Programs/Shared/EpicGames.Build/UEBuildPlatformSDK.cs
+++ b/Engine/Source/Programs/Shared/EpicGames.Build/UEBuildPlatformSDK.cs
@@ -840,6 +840,12 @@
 						}
 					}
 
+					// Reducing warnings to log to help prevent warnings locally or in Horde about SDKs we might not currently be concerned about
+					if (bSuppressSDKWarnings)
+					{
+						ErrorVerbosity = LogEventType.Log;
+					}
+
 					// always print errors to the screen
 					Log.WriteLine(ErrorVerbosity, ErrorOptions, Msg.ToString());
 				}
diff --git a/Engine/Source/Programs/UnrealBuildTool/UnrealBuildTool.cs b/Engine/Source/Programs/UnrealBuildTool/UnrealBuildTool.cs
--- a/Engine/Source/Programs/UnrealBuildTool/UnrealBuildTool.cs
+++ b/Engine/Source/Programs/UnrealBuildTool/UnrealBuildTool.cs
@@ -290,6 +290,12 @@
 			public bool bLogFromMsBuild = false;
 
 			/// <summary>
+			/// Whether or not to suppress warnings of missing SDKs from warnings to LogEventType.Log in UEBuildPlatformSDK.cs 
+			/// </summary>
+			[CommandLine(Prefix = "-SuppressSDKWarnings", Description = "Missing SDKs error verbosity level will be reduced from warning to log")]
+			public bool bShouldSuppressSDKWarnings = false;
+
+			/// <summary>
 			/// Whether to write progress markup in a format that can be parsed by other programs
 			/// </summary>
 			[CommandLine(Prefix = "-Progress", Description = "Write progress messages in a format that can be parsed by other programs")]
@@ -480,6 +486,12 @@
 				Log.IncludeTimestamps = Options.bLogTimestamps;
 				Log.IncludeProgramNameWithSeverityPrefix = Options.bLogFromMsBuild;
 
+				// Reducing SDK warning events in the log to LogEventType.Log
+				if (Options.bShouldSuppressSDKWarnings)
+				{
+					UEBuildPlatformSDK.bSuppressSDKWarnings = true;
+				}
+
 				if (Options.TraceWrites != null)
 				{
 					Logger.LogInformation("All attempts to write to \"{TraceWrites}\" via WriteFileIfChanged() will be logged", Options.TraceWrites);
