Manual resolve of merge from cine-staging at 20378345

diff --git a/Engine/Source/Programs/UnrealBuildTool/Configuration/TargetRules.cs b/Engine/Source/Programs/UnrealBuildTool/Configuration/TargetRules.cs
--- a/Engine/Source/Programs/UnrealBuildTool/Configuration/TargetRules.cs
+++ b/Engine/Source/Programs/UnrealBuildTool/Configuration/TargetRules.cs
@@ -564,7 +564,7 @@
 		[CommandLine("-NoUseVerse", Value = "false")]
 		[CommandLine("-UseVerse", Value = "true")]
 		[XmlConfigFile(Category = "BuildConfiguration")]
-		public bool bUseVerse = false;
+		public bool bUseVerse = true;
 
 		/// <summary>
 		/// Whether to compile the Chaos physics plugin.
