Undo default enabling of flag unintentionally introduced in CL# 20395952
#rnx

diff --git a/Engine/Source/Programs/UnrealBuildTool/Configuration/TargetRules.cs b/Engine/Source/Programs/UnrealBuildTool/Configuration/TargetRules.cs
--- a/Engine/Source/Programs/UnrealBuildTool/Configuration/TargetRules.cs
+++ b/Engine/Source/Programs/UnrealBuildTool/Configuration/TargetRules.cs
@@ -564,7 +564,7 @@ namespace UnrealBuildTool
 		[CommandLine("-NoUseVerse", Value = "false")]
 		[CommandLine("-UseVerse", Value = "true")]
 		[XmlConfigFile(Category = "BuildConfiguration")]
-		public bool bUseVerse = true;
+		public bool bUseVerse = false;
 
 		/// <summary>
 		/// Whether to compile the Chaos physics plugin.
