UnrealBuildTool: Ignore build outputs in manifest when doing a single file compile

#rnx

diff --git a/Engine/Source/Programs/UnrealBuildTool/Configuration/TargetRules.cs b/Engine/Source/Programs/UnrealBuildTool/Configuration/TargetRules.cs
--- a/Engine/Source/Programs/UnrealBuildTool/Configuration/TargetRules.cs
+++ b/Engine/Source/Programs/UnrealBuildTool/Configuration/TargetRules.cs
@@ -1517,6 +1517,11 @@ namespace UnrealBuildTool
 		[CommandLine("-NoLink")]
 		public bool bDisableLinking = false;
 
+		/// <summary>
+		/// Whether to ignore tracking build outputs for this target.
+		/// </summary>
+		public bool bIgnoreBuildOutputs = false;
+
 		/// <summary>
 		/// Indicates that this is a formal build, intended for distribution. This flag is automatically set to true when Build.version has a changelist set.
 		/// The only behavior currently bound to this flag is to compile the default resource file separately for each binary so that the OriginalFilename field is set correctly.
@@ -3030,6 +3035,11 @@ namespace UnrealBuildTool
 			get { return Inner.bDisableLinking; }
 		}
 
+		public bool bIgnoreBuildOutputs
+		{
+			get { return Inner.bIgnoreBuildOutputs; }
+		}
+
 		public bool bFormalBuild
 		{
 			get { return Inner.bFormalBuild; }
diff --git a/Engine/Source/Programs/UnrealBuildTool/Configuration/UEBuildTarget.cs b/Engine/Source/Programs/UnrealBuildTool/Configuration/UEBuildTarget.cs
--- a/Engine/Source/Programs/UnrealBuildTool/Configuration/UEBuildTarget.cs
+++ b/Engine/Source/Programs/UnrealBuildTool/Configuration/UEBuildTarget.cs
@@ -786,6 +786,7 @@ namespace UnrealBuildTool
 				RulesObject.bForceUnityBuild = false;
 				RulesObject.bUsePCHFiles = false;
 				RulesObject.bDisableLinking = true;
+				RulesObject.bIgnoreBuildOutputs = true;
 			}
 
 			// If the files to compile contains headers, figure out the cpp files that include them
@@ -1443,6 +1444,13 @@ namespace UnrealBuildTool
 				Manifest.BuildProducts.RemoveAll(x => !new FileReference(x).IsUnderDirectory(ForeignPluginDir));
 			}
 
+			// Completely ignore build outputs if requested
+			if (Rules.bIgnoreBuildOutputs)
+			{
+				Manifest.BuildProducts.Clear();
+				Manifest.DeployTargetFiles.Clear();
+			}
+
 			Manifest.BuildProducts.Sort();
 			Manifest.DeployTargetFiles.Sort();
 
