- FIxed TargetArchitecture string parsing

diff --git a/Engine/Config/BaseEngine.ini b/Engine/Config/BaseEngine.ini
--- a/Engine/Config/BaseEngine.ini
+++ b/Engine/Config/BaseEngine.ini
@@ -103,22 +103,6 @@ EngineModuleReflectedUninitializedPropertyVerbosity=Error
 ProjectModuleReflectedUninitializedPropertyVerbosity=Error
 ObjectReferenceReflectedUninitializedPropertyVerbosity=Error
 
-[CoreUObject.ShortTypeNameInMetaDataCheck]
-;+MetaDataKeys=AllowedClasses
-;+MetaDataKeys=DisallowedClasses
-;+MetaDataKeys=BaseStruct
-;+MetaDataKeys=BaseClass
-;+MetaDataKeys=HasNativeMake
-;+MetaDataKeys=HasnativeBreak
-;+MetaDataKeys=RowType
-;+MetaDataKeys=BitmaskEnum
-;+MetaDataKeys=ActorClassToPoll
-;+MetaDataKeys=ComponentClass
-;+MetaDataKeys=MustImplement
-;+MetaDataKeys=MetaClass
-;+MetaDataKeys=MetaStruct
-;+MetaDataKeys=RequiredAssetDataTags
-
 [/Script/Engine.Engine]
 ConsoleClassName=/Script/Engine.Console
 GameViewportClientClassName=/Script/Engine.GameViewportClient
@@ -2780,8 +2764,8 @@ MetalLanguageVersion=5
 +TargetedRHIs=SF_METAL_SM5
 UseFastIntrinsics=False
 EnableMathOptimisations=True
-TargetArchitecture=Universal
-EditorTargetArchitecture=Intel
+TargetArchitecture=MacTargetArchitectureIntel
+EditorTargetArchitecture=MacTargetArchitectureIntel
 
 [HMDPluginPriority]
 ; Since SteamVR also works with the Oculus Rift, give priority to the native Oculus plugin before trying SteamVR.
diff --git a/Engine/Source/Programs/UnrealBuildTool/ProjectFiles/Xcode/XcodeProject.cs b/Engine/Source/Programs/UnrealBuildTool/ProjectFiles/Xcode/XcodeProject.cs
--- a/Engine/Source/Programs/UnrealBuildTool/ProjectFiles/Xcode/XcodeProject.cs
+++ b/Engine/Source/Programs/UnrealBuildTool/ProjectFiles/Xcode/XcodeProject.cs
@@ -1,4 +1,4 @@
-// Copyright Epic Games, Inc. All Rights Reserved.
+﻿// Copyright Epic Games, Inc. All Rights Reserved.
 
 using System;
 using System.Collections.Generic;
@@ -1025,15 +1025,15 @@ namespace UnrealBuildTool
 						string Key = Config.ProjectTarget!.TargetRules!.Type == TargetType.Editor ? "EditorTargetArchitecture" : "TargetArchitecture";
 						if (EngineIni.GetString("/Script/MacTargetPlatform.MacTargetSettings", Key, out TargetArchitecture))
 						{
-							if (TargetArchitecture.Equals("Universal", StringComparison.OrdinalIgnoreCase))
+							if (TargetArchitecture.Contains("Universal", StringComparison.OrdinalIgnoreCase))
 							{
 								TargetArchitectures = AllArchitectures;
 							}
-							else if (TargetArchitecture.Equals("Intel", StringComparison.OrdinalIgnoreCase))
+							else if (TargetArchitecture.Contains("Intel", StringComparison.OrdinalIgnoreCase))
 							{
 								TargetArchitectures = new[] { MacExports.IntelArchitecture };
 							}
-							else if (TargetArchitecture.Equals("Apple", StringComparison.OrdinalIgnoreCase))
+							else if (TargetArchitecture.Contains("Apple", StringComparison.OrdinalIgnoreCase))
 							{
 								TargetArchitectures = new[] { MacExports.AppleArchitecture };
 							}
