UnrealBuildTool PreprocessorExpression: Merge identical handling of keywords into a single block to reduce duplication.
#rnx

diff --git a/Engine/Source/Programs/UnrealBuildTool/Preprocessor/PreprocessorExpression.cs b/Engine/Source/Programs/UnrealBuildTool/Preprocessor/PreprocessorExpression.cs
--- a/Engine/Source/Programs/UnrealBuildTool/Preprocessor/PreprocessorExpression.cs
+++ b/Engine/Source/Programs/UnrealBuildTool/Preprocessor/PreprocessorExpression.cs
@@ -238,97 +238,12 @@ namespace UnrealBuildTool
 					{
 						throw new NotImplementedException(Text.ToString());
 					}
-					else if (Text == Identifiers.__has_builtin)
-					{
-						if (Tokens[Idx + 1].Type != TokenType.LeftParen || Tokens[Idx + 3].Type != TokenType.RightParen)
-						{
-							throw new NotImplementedException(Text.ToString());
-						}
-						Idx += 4;
-						return 0;
-					}
-					else if(Text == Identifiers.__has_feature)
-					{
-						if(Tokens[Idx + 1].Type != TokenType.LeftParen || Tokens[Idx + 3].Type != TokenType.RightParen)
-						{
-							throw new NotImplementedException(Text.ToString());
-						}
-						Idx += 4;
-						return 0;
-					}
-					else if (Text == Identifiers.__has_warning)
-					{
-						if (Tokens[Idx + 1].Type != TokenType.LeftParen || Tokens[Idx + 3].Type != TokenType.RightParen)
-						{
-							throw new NotImplementedException(Text.ToString());
-						}
-						Idx += 4;
-						return 0;
-					}
-					else if(Text == Identifiers.__building_module)
-					{
-						if(Tokens[Idx + 1].Type != TokenType.LeftParen || Tokens[Idx + 3].Type != TokenType.RightParen)
-						{
-							throw new NotImplementedException(Text.ToString());
-						}
-						Idx += 4;
-						return 0;
-					}
-					else if (Text == Identifiers.__pragma)
-					{
-						if (Tokens[Idx + 1].Type != TokenType.LeftParen || Tokens[Idx + 3].Type != TokenType.RightParen)
-						{
-							throw new NotImplementedException(Text.ToString());
-						}
-						Idx += 4;
-						return 0;
-					}
-					else if (Text == Identifiers.__builtin_return_address)
-					{
-						if (Tokens[Idx + 1].Type != TokenType.LeftParen || Tokens[Idx + 3].Type != TokenType.RightParen)
-						{
-							throw new NotImplementedException(Text.ToString());
-						}
-						Idx += 4;
-						return 0;
-					}
-					else if (Text == Identifiers.__builtin_frame_address)
-					{
-						if (Tokens[Idx + 1].Type != TokenType.LeftParen || Tokens[Idx + 3].Type != TokenType.RightParen)
-						{
-							throw new NotImplementedException(Text.ToString());
-						}
-						Idx += 4;
-						return 0;
-					}
-					else if (Text == Identifiers.__has_keyword)
-					{
-						if (Tokens[Idx + 1].Type != TokenType.LeftParen || Tokens[Idx + 3].Type != TokenType.RightParen)
-						{
-							throw new NotImplementedException(Text.ToString());
-						}
-						Idx += 4;
-						return 0;
-					}
-					else if (Text == Identifiers.__has_extension)
-					{
-						if (Tokens[Idx + 1].Type != TokenType.LeftParen || Tokens[Idx + 3].Type != TokenType.RightParen)
-						{
-							throw new NotImplementedException(Text.ToString());
-						}
-						Idx += 4;
-						return 0;
-					}
-					else if (Text == Identifiers.__is_target_arch)
-					{
-						if (Tokens[Idx + 1].Type != TokenType.LeftParen || Tokens[Idx + 3].Type != TokenType.RightParen)
-						{
-							throw new NotImplementedException(Text.ToString());
-						}
-						Idx += 4;
-						return 0;
-					}
-					else if (Text == Identifiers.__is_identifier)
+					else if (Text == Identifiers.__has_builtin || Text == Identifiers.__has_feature
+						|| Text == Identifiers.__has_warning || Text == Identifiers.__building_module
+						|| Text == Identifiers.__pragma || Text == Identifiers.__builtin_return_address
+						|| Text == Identifiers.__builtin_frame_address || Text == Identifiers.__has_keyword
+						|| Text == Identifiers.__has_extension || Text == Identifiers.__is_target_arch
+						|| Text == Identifiers.__is_identifier)
 					{
 						if (Tokens[Idx + 1].Type != TokenType.LeftParen || Tokens[Idx + 3].Type != TokenType.RightParen)
 						{
