UBT: Fix missing event markup when executing directly inside UBT. Event parsers are only currently run when using old logging methods.

diff --git a/Engine/Source/Programs/UnrealBuildTool/Executors/ParallelExecutor.cs b/Engine/Source/Programs/UnrealBuildTool/Executors/ParallelExecutor.cs
--- a/Engine/Source/Programs/UnrealBuildTool/Executors/ParallelExecutor.cs
+++ b/Engine/Source/Programs/UnrealBuildTool/Executors/ParallelExecutor.cs
@@ -400,7 +400,7 @@
 
 				s_previousLineLength = message.Length;
 
-				Logger.LogInformation("{Message}", message);
+				Log.TraceInformation("{0}", message); // Need to send this through registered event parser; using old logger
 				foreach (string Line in LogLines.Skip(Action.bShouldOutputStatusDescription ? 0 : 1))
 				{
 					// suppress library creation messages when writing compact output
@@ -409,7 +409,7 @@
 						continue;
 					}
 
-					Logger.LogInformation("{Message}", Line);
+					Log.TraceInformation("{0}", Line);
 
 					// Prevent overwriting of logged lines
 					s_previousLineLength = -1;
diff --git a/Engine/Source/Programs/UnrealBuildTool/Executors/SNDBS.cs b/Engine/Source/Programs/UnrealBuildTool/Executors/SNDBS.cs
--- a/Engine/Source/Programs/UnrealBuildTool/Executors/SNDBS.cs
+++ b/Engine/Source/Programs/UnrealBuildTool/Executors/SNDBS.cs
@@ -297,7 +297,7 @@
 						// Suppress redundant tool output of status we already printed (e.g., msvc cl prints compile unit name always)
 						if (!Text.Equals(CurrentStatus))
 						{
-							Logger.LogInformation("{Output}", Text);
+							Log.TraceInformation("{0}", Text);
 						}
 					}
 				};
diff --git a/Engine/Source/Programs/UnrealBuildTool/Executors/XGE.cs b/Engine/Source/Programs/UnrealBuildTool/Executors/XGE.cs
--- a/Engine/Source/Programs/UnrealBuildTool/Executors/XGE.cs
+++ b/Engine/Source/Programs/UnrealBuildTool/Executors/XGE.cs
@@ -763,7 +763,7 @@
 							ProgressText = string.Empty;
 							return;
 						}
-						Logger.LogInformation("{Status}", Text);
+						Log.TraceInformation("{0}", Text); // Using old log function to pick up registered event parsers
 					}
 				};
 
