UnrealBuildTool: Allow NETFXSDK 4.6.2 from AutoSDK

#rnx

diff --git a/Engine/Source/Programs/UnrealBuildTool/Platform/Windows/MicrosoftPlatformSDK.cs b/Engine/Source/Programs/UnrealBuildTool/Platform/Windows/MicrosoftPlatformSDK.cs
--- a/Engine/Source/Programs/UnrealBuildTool/Platform/Windows/MicrosoftPlatformSDK.cs
+++ b/Engine/Source/Programs/UnrealBuildTool/Platform/Windows/MicrosoftPlatformSDK.cs
@@ -307,25 +307,21 @@ namespace UnrealBuildTool
 		public static bool TryGetNetFxSdkInstallDir([NotNullWhen(true)] out DirectoryReference? OutInstallDir)
 		{
 			DirectoryReference? HostAutoSdkDir;
+			string[] PreferredVersions = new string[] { "4.6.2", "4.6.1", "4.6" };
 			if (UEBuildPlatformSDK.TryGetHostPlatformAutoSDKDir(out HostAutoSdkDir))
 			{
-				DirectoryReference NetFxDir_4_6 = DirectoryReference.Combine(HostAutoSdkDir, "Win64", "Windows Kits", "NETFXSDK", "4.6");
-				if (FileReference.Exists(FileReference.Combine(NetFxDir_4_6, "Include", "um", "mscoree.h")))
-				{
-					OutInstallDir = NetFxDir_4_6;
-					return true;
-				}
-
-				DirectoryReference NetFxDir_4_6_1 = DirectoryReference.Combine(HostAutoSdkDir, "Win64", "Windows Kits", "NETFXSDK", "4.6.1");
-				if (FileReference.Exists(FileReference.Combine(NetFxDir_4_6_1, "Include", "um", "mscoree.h")))
+				foreach (string PreferredVersion in PreferredVersions)
 				{
-					OutInstallDir = NetFxDir_4_6_1;
-					return true;
+					DirectoryReference NetFxDir = DirectoryReference.Combine(HostAutoSdkDir, "Win64", "Windows Kits", "NETFXSDK", PreferredVersion);
+					if (FileReference.Exists(FileReference.Combine(NetFxDir, "Include", "um", "mscoree.h")))
+					{
+						OutInstallDir = NetFxDir;
+						return true;
+					}
 				}
 			}
 
 			string NetFxSDKKeyName = "Microsoft\\Microsoft SDKs\\NETFXSDK";
-			string[] PreferredVersions = new string[] { "4.6.2", "4.6.1", "4.6" };
 			foreach (string PreferredVersion in PreferredVersions)
 			{
 				if (TryReadInstallDirRegistryKey32(NetFxSDKKeyName + "\\" + PreferredVersion, "KitsInstallationFolder", out OutInstallDir))
