UnrealBuildTool: Move diagnostics flags to ClangToolChain

#rnx

diff --git a/Engine/Source/Programs/UnrealBuildTool/Platform/IOS/IOSToolChain.cs b/Engine/Source/Programs/UnrealBuildTool/Platform/IOS/IOSToolChain.cs
--- a/Engine/Source/Programs/UnrealBuildTool/Platform/IOS/IOSToolChain.cs
+++ b/Engine/Source/Programs/UnrealBuildTool/Platform/IOS/IOSToolChain.cs
@@ -214,11 +214,6 @@
 		{
 			base.GetCompileArguments_Global(CompileEnvironment, Arguments);
 
-			if (RuntimePlatform.IsWindows)
-			{
-				Arguments.Add("-fdiagnostics-format=msvc");
-			}
-
 			// Optionally enable exception handling (off by default since it generates extra code needed to propagate exceptions)
 			if (CompileEnvironment.bEnableExceptions)
 			{
diff --git a/Engine/Source/Programs/UnrealBuildTool/ToolChain/ClangToolChain.cs b/Engine/Source/Programs/UnrealBuildTool/ToolChain/ClangToolChain.cs
--- a/Engine/Source/Programs/UnrealBuildTool/ToolChain/ClangToolChain.cs
+++ b/Engine/Source/Programs/UnrealBuildTool/ToolChain/ClangToolChain.cs
@@ -1,4 +1,4 @@
-﻿// Copyright Epic Games, Inc. All Rights Reserved.
+// Copyright Epic Games, Inc. All Rights Reserved.
 
 using EpicGames.Core;
 using System;
@@ -355,7 +355,19 @@
 			}
 
 			// always use absolute paths for errors, this can help IDEs go to the error properly
-			Arguments.Add("-fdiagnostics-absolute-paths");
+			Arguments.Add("-fdiagnostics-absolute-paths");  // https://clang.llvm.org/docs/ClangCommandLineReference.html#cmdoption-clang-fdiagnostics-absolute-paths
+
+			// https://clang.llvm.org/docs/ClangCommandLineReference.html#cmdoption-clang-fcolor-diagnostics
+			if (Log.ColorConsoleOutput)
+			{
+				Arguments.Add("-fdiagnostics-color");
+			}
+
+			// https://clang.llvm.org/docs/ClangCommandLineReference.html#cmdoption-clang-fdiagnostics-format
+			if (RuntimePlatform.IsWindows)
+			{
+				Arguments.Add("-fdiagnostics-format=msvc");
+			}
 		}
 
 		/// <summary>
