- Missed the UBT change to enable CEF on M1 Macs now that the libs are in

diff --git a/Engine/Source/Programs/UnrealBuildTool/Platform/Mac/UEBuildMac.cs b/Engine/Source/Programs/UnrealBuildTool/Platform/Mac/UEBuildMac.cs
--- a/Engine/Source/Programs/UnrealBuildTool/Platform/Mac/UEBuildMac.cs
+++ b/Engine/Source/Programs/UnrealBuildTool/Platform/Mac/UEBuildMac.cs
@@ -1,4 +1,4 @@
-// Copyright Epic Games, Inc. All Rights Reserved.
+﻿// Copyright Epic Games, Inc. All Rights Reserved.
 
 
 using System;
@@ -134,7 +134,6 @@ namespace UnrealBuildTool
 			if (bCompilingForArm)
 			{
 				Target.bCompileISPC = false;
-				Target.bCompileCEF3 = false;
 				Target.DisablePlugins.AddRange(new string[]
 				{
 					// onnx:
