UnrealBuildTool: Move clang debug and optimization args to virtual functions

#rnx

diff --git a/Engine/Source/Programs/UnrealBuildTool/Platform/IOS/IOSToolChain.cs b/Engine/Source/Programs/UnrealBuildTool/Platform/IOS/IOSToolChain.cs
--- a/Engine/Source/Programs/UnrealBuildTool/Platform/IOS/IOSToolChain.cs
+++ b/Engine/Source/Programs/UnrealBuildTool/Platform/IOS/IOSToolChain.cs
@@ -210,21 +210,35 @@ namespace UnrealBuildTool
 		}
 
 		/// <inheritdoc/>
-		protected override void GetCompileArguments_Global(CppCompileEnvironment CompileEnvironment, List<string> Arguments)
+		protected override void GetCompileArguments_Optimizations(CppCompileEnvironment CompileEnvironment, List<string> Arguments)
 		{
-			base.GetCompileArguments_Global(CompileEnvironment, Arguments);
+			base.GetCompileArguments_Optimizations(CompileEnvironment, Arguments);
 
-			// Optionally enable exception handling (off by default since it generates extra code needed to propagate exceptions)
-			if (CompileEnvironment.bEnableExceptions)
+			if (GetClangVersion().Major >= 12)
 			{
-				Arguments.Add("-fexceptions");
-				Arguments.Add("-DPLATFORM_EXCEPTIONS_DISABLED=0");
+				// We have 'this' vs nullptr comparisons that get optimized away for newer versions of Clang, which is undesirable until we refactor these checks.
+				Arguments.Add("-fno-delete-null-pointer-checks");
 			}
-			else
+
+			// use LTO if desired (like VCToolchain does)
+			if (CompileEnvironment.bAllowLTCG)
 			{
-				Arguments.Add("-fno-exceptions");
-				Arguments.Add("-DPLATFORM_EXCEPTIONS_DISABLED=1");
+				Arguments.Add("-flto");
 			}
+		}
+
+		/// <inheritdoc/>
+		protected override void GetCompileArguments_Debugging(CppCompileEnvironment CompileEnvironment, List<string> Arguments)
+		{
+			base.GetCompileArguments_Debugging(CompileEnvironment, Arguments);
+
+			Arguments.Add("-fvisibility=hidden"); // hides the linker warnings with PhysX
+		}
+
+		/// <inheritdoc/>
+		protected override void GetCompileArguments_Global(CppCompileEnvironment CompileEnvironment, List<string> Arguments)
+		{
+			base.GetCompileArguments_Global(CompileEnvironment, Arguments);
 
 			string? SanitizerMode = Environment.GetEnvironmentVariable("ENABLE_ADDRESS_SANITIZER");
 			if ((SanitizerMode != null && SanitizerMode == "YES") || (Options.HasFlag(ClangToolChainOptions.EnableAddressSanitizer)))
@@ -244,19 +258,6 @@ namespace UnrealBuildTool
 			}
 
 			Arguments.Add(GetRTTIFlag(CompileEnvironment));
-			Arguments.Add("-fvisibility=hidden"); // hides the linker warnings with PhysX
-
-			// use LTO if desired (like VCToolchain does)
-			if (CompileEnvironment.bAllowLTCG)
-			{
-				Arguments.Add("-flto");
-			}
-
-			if (GetClangVersion().Major >= 12)
-			{
-				// We have 'this' vs nullptr comparisons that get optimized away for newer versions of Clang, which is undesirable until we refactor these checks.
-				Arguments.Add("-fno-delete-null-pointer-checks");
-			}
 
 			if (IsBitcodeCompilingEnabled(CompileEnvironment.Configuration))
 			{
@@ -270,41 +271,6 @@ namespace UnrealBuildTool
 
 			Arguments.Add("-m" + GetXcodeMinVersionParam() + "=" + ProjectSettings.RuntimeVersion);
 
-			bool bStaticAnalysis = false;
-			string? StaticAnalysisMode = Environment.GetEnvironmentVariable("CLANG_STATIC_ANALYZER_MODE");
-			if (StaticAnalysisMode != null && StaticAnalysisMode != "")
-			{
-				bStaticAnalysis = true;
-			}
-
-			// Optimize non- debug builds.
-			if (CompileEnvironment.bOptimizeCode && !bStaticAnalysis)
-			{
-				if (CompileEnvironment.bOptimizeForSize)
-				{
-					Arguments.Add("-Oz");
-				}
-				else
-				{
-					Arguments.Add("-O3");
-				}
-			}
-			else
-			{
-				Arguments.Add("-O0");
-			}
-
-			if (!CompileEnvironment.bUseInlining)
-			{
-				Arguments.Add("-fno-inline-functions");
-			}
-
-			// Create DWARF format debug info if wanted,
-			if (CompileEnvironment.bCreateDebugInfo)
-			{
-				Arguments.Add("-gdwarf-2");
-			}
-
 			// Add additional frameworks so that their headers can be found
 			foreach (UEBuildFramework Framework in CompileEnvironment.AdditionalFrameworks)
 			{
diff --git a/Engine/Source/Programs/UnrealBuildTool/Platform/Linux/LinuxToolChain.cs b/Engine/Source/Programs/UnrealBuildTool/Platform/Linux/LinuxToolChain.cs
--- a/Engine/Source/Programs/UnrealBuildTool/Platform/Linux/LinuxToolChain.cs
+++ b/Engine/Source/Programs/UnrealBuildTool/Platform/Linux/LinuxToolChain.cs
@@ -414,90 +414,11 @@
 		}
 
 		/// <inheritdoc/>
-		protected override void GetCompileArguments_Global(CppCompileEnvironment CompileEnvironment, List<string> Arguments)
+		protected override void GetCompileArguments_Optimizations(CppCompileEnvironment CompileEnvironment, List<string> Arguments)
 		{
-			base.GetCompileArguments_Global(CompileEnvironment, Arguments);
+			base.GetCompileArguments_Optimizations(CompileEnvironment, Arguments);
 
-			// build up the commandline common to C and C++
-
-			if (ShouldUseLibcxx(CompileEnvironment.Architecture))
-			{
-				Arguments.Add("-nostdinc++");
-				Arguments.Add("-I" + "ThirdParty/Unix/LibCxx/include/");
-				Arguments.Add("-I" + "ThirdParty/Unix/LibCxx/include/c++/v1");
-			}
-
-			// ASan
-			if (Options.HasFlag(ClangToolChainOptions.EnableAddressSanitizer))
-			{
-				// Force using the ANSI allocator if ASan is enabled
-				Arguments.Add("-fsanitize=address -DFORCE_ANSI_ALLOCATOR=1");
-			}
-
-			// TSan
-			if (Options.HasFlag(ClangToolChainOptions.EnableThreadSanitizer))
-			{
-				// Force using the ANSI allocator if TSan is enabled
-				Arguments.Add("-fsanitize=thread -DFORCE_ANSI_ALLOCATOR=1");
-			}
-
-			// UBSan
-			if (Options.HasFlag(ClangToolChainOptions.EnableUndefinedBehaviorSanitizer))
-			{
-				Arguments.Add("-fsanitize=undefined -fno-sanitize=vptr");
-			}
-
-			// MSan
-			if (Options.HasFlag(ClangToolChainOptions.EnableMemorySanitizer))
-			{
-				// Force using the ANSI allocator if MSan is enabled
-				// -fsanitize-memory-track-origins adds a 1.5x-2.5x slow down ontop of MSan normal amount of overhead
-				// -fsanitize-memory-track-origins=1 is faster but collects only allocation points but not intermediate stores
-				Arguments.Add("-fsanitize=memory -fsanitize-memory-track-origins -DFORCE_ANSI_ALLOCATOR=1");
-			}
-
-			if (CompilerVersionGreaterOrEqual(12, 0, 0))
-			{
-				Arguments.Add("-fbinutils-version=2.36");
-			}
-
-			if (!CompileEnvironment.Architecture.StartsWith("x86_64"))
-			{
-				Arguments.Add("-funwind-tables");               // generate unwind tables as they are needed for backtrace (on x86(64) they are generated implicitly)
-			}
-
-			Arguments.Add(ArchitectureSpecificSwitches(CompileEnvironment.Architecture));
-
-			Arguments.Add("-fno-math-errno");               // do not assume that math ops have side effects
-
-			Arguments.Add(GetRTTIFlag(CompileEnvironment)); // flag for run-time type info
-
-			if (CompileEnvironment.Architecture.StartsWith("x86_64"))
-			{
-				Arguments.Add("-mssse3"); // enable ssse3 by default for x86. This is default on for MSVC so lets reflect that here
-			}
-
-			if (CompileEnvironment.bHideSymbolsByDefault)
-			{
-				Arguments.Add("-fvisibility-ms-compat");
-				Arguments.Add("-fvisibility-inlines-hidden");
-			}
-
-			// Profile Guided Optimization (PGO) and Link Time Optimization (LTO)
-			// Whether we actually can enable that is checked in CanUseAdvancedLinkerFeatures() earlier
-			if (CompileEnvironment.bPGOOptimize)
-			{
-				Log.TraceInformationOnce("Enabling Profile Guided Optimization (PGO). Linking will take a while.");
-				Arguments.Add(string.Format(" -fprofile-instr-use=\"{0}\"", Path.Combine(CompileEnvironment.PGODirectory!, CompileEnvironment.PGOFilenamePrefix!)));
-			}
-			else if (CompileEnvironment.bPGOProfile)
-			{
-				Log.TraceInformationOnce("Enabling Profile Guided Instrumentation (PGI). Linking will take a while.");
-				Arguments.Add("-fprofile-generate");
-			}
-
 			// Unlike on other platforms, allow LTO be specified independently of PGO
-			// Whether we actually can enable that is checked in CanUseAdvancedLinkerFeatures() earlier
 			if (CompileEnvironment.bAllowLTCG)
 			{
 				if ((Options & ClangToolChainOptions.EnableThinLTO) != 0)
@@ -510,7 +431,30 @@
 				}
 			}
 
-			//Arguments.Add("-DOPERATOR_NEW_INLINE=FORCENOINLINE");
+			// optimization level
+			if (!CompileEnvironment.bOptimizeCode)
+			{
+				Arguments.Add("-O0");
+			}
+			else
+			{
+				// Don't over optimise if using Address/MemorySanitizer or you'll get false positive errors due to erroneous optimisation of necessary Address/MemorySanitizer instrumentation.
+				if (Options.HasFlag(ClangToolChainOptions.EnableAddressSanitizer) || Options.HasFlag(ClangToolChainOptions.EnableMemorySanitizer))
+				{
+					Arguments.Add("-O1 -g");
 
+					// This enables __asan_default_options() in UnixCommonStartup.h which disables the leak detector
+					Arguments.Add("-DDISABLE_ASAN_LEAK_DETECTOR=1");
+				}
+				else if (Options.HasFlag(ClangToolChainOptions.EnableThreadSanitizer))
+				{
+					Arguments.Add("-O1 -g");
+				}
+				else
+				{
+					Arguments.Add("-O3");
+				}
+			}
+
 			bool bRetainFramePointers = CompileEnvironment.bRetainFramePointers
 				|| Options.HasFlag(ClangToolChainOptions.EnableAddressSanitizer) || Options.HasFlag(ClangToolChainOptions.EnableMemorySanitizer)
@@ -528,13 +472,19 @@
 			{
 				Arguments.Add("-fno-inline");                   // disable inlining for better debuggability (e.g. callstacks, "skip file" in gdb)
 				Arguments.Add("-fstack-protector");             // detect stack smashing
-																//Arguments.Add("-fsanitize=address");            // detect address based errors (support properly and link to libasan)
 			}
 
 			if (bRetainFramePointers)
 			{
-				Arguments.Add("-fno-optimize-sibling-calls -fno-omit-frame-pointer");
+				Arguments.Add("-fno-optimize-sibling-calls");
+				Arguments.Add("-fno-omit-frame-pointer");
 			}
+		}
 
+		/// <inheritdoc/>
+		protected override void GetCompileArguments_Debugging(CppCompileEnvironment CompileEnvironment, List<string> Arguments)
+		{
+			base.GetCompileArguments_Debugging(CompileEnvironment, Arguments);
+
 			// debug info
 			// bCreateDebugInfo is normally set for all configurations, including Shipping - this is needed to enable callstack in Shipping builds (proper resolution: UEPLAT-205, separate files with debug info)
@@ -555,36 +505,79 @@
 				}
 			}
 
-			// optimization level
-			if (!CompileEnvironment.bOptimizeCode)
+			if (CompileEnvironment.bHideSymbolsByDefault)
 			{
-				Arguments.Add("-O0");
+				Arguments.Add("-fvisibility-ms-compat");
+				Arguments.Add("-fvisibility-inlines-hidden");
 			}
-			else
+		}
+
+		/// <inheritdoc/>
+		protected override void GetCompileArguments_Global(CppCompileEnvironment CompileEnvironment, List<string> Arguments)
+		{
+			base.GetCompileArguments_Global(CompileEnvironment, Arguments);
+
+			// build up the commandline common to C and C++
+
+			if (ShouldUseLibcxx(CompileEnvironment.Architecture))
 			{
-				// Don't over optimise if using Address/MemorySanitizer or you'll get false positive errors due to erroneous optimisation of necessary Address/MemorySanitizer instrumentation.
-				if (Options.HasFlag(ClangToolChainOptions.EnableAddressSanitizer) || Options.HasFlag(ClangToolChainOptions.EnableMemorySanitizer))
-				{
-					Arguments.Add("-O1 -g");
-
-					// This enables __asan_default_options() in UnixCommonStartup.h which disables the leak detector
-					Arguments.Add("-DDISABLE_ASAN_LEAK_DETECTOR=1");
-				}
-				else if (Options.HasFlag(ClangToolChainOptions.EnableThreadSanitizer))
-				{
-					Arguments.Add("-O1 -g");
-				}
-				else
-				{
-					Arguments.Add("-O3");
-				}
+				Arguments.Add("-nostdinc++");
+				Arguments.Add("-I" + "ThirdParty/Unix/LibCxx/include/");
+				Arguments.Add("-I" + "ThirdParty/Unix/LibCxx/include/c++/v1");
 			}
 
-			if (!CompileEnvironment.bUseInlining)
+			// ASan
+			if (Options.HasFlag(ClangToolChainOptions.EnableAddressSanitizer))
 			{
-				Arguments.Add("-fno-inline-functions");
+				// Force using the ANSI allocator if ASan is enabled
+				Arguments.Add("-fsanitize=address -DFORCE_ANSI_ALLOCATOR=1");
 			}
 
+			// TSan
+			if (Options.HasFlag(ClangToolChainOptions.EnableThreadSanitizer))
+			{
+				// Force using the ANSI allocator if TSan is enabled
+				Arguments.Add("-fsanitize=thread -DFORCE_ANSI_ALLOCATOR=1");
+			}
+
+			// UBSan
+			if (Options.HasFlag(ClangToolChainOptions.EnableUndefinedBehaviorSanitizer))
+			{
+				Arguments.Add("-fsanitize=undefined -fno-sanitize=vptr");
+			}
+
+			// MSan
+			if (Options.HasFlag(ClangToolChainOptions.EnableMemorySanitizer))
+			{
+				// Force using the ANSI allocator if MSan is enabled
+				// -fsanitize-memory-track-origins adds a 1.5x-2.5x slow down ontop of MSan normal amount of overhead
+				// -fsanitize-memory-track-origins=1 is faster but collects only allocation points but not intermediate stores
+				Arguments.Add("-fsanitize=memory -fsanitize-memory-track-origins -DFORCE_ANSI_ALLOCATOR=1");
+			}
+
+			if (CompilerVersionGreaterOrEqual(12, 0, 0))
+			{
+				Arguments.Add("-fbinutils-version=2.36");
+			}
+
+			if (!CompileEnvironment.Architecture.StartsWith("x86_64"))
+			{
+				Arguments.Add("-funwind-tables");               // generate unwind tables as they are needed for backtrace (on x86(64) they are generated implicitly)
+			}
+
+			Arguments.Add(ArchitectureSpecificSwitches(CompileEnvironment.Architecture));
+
+			Arguments.Add("-fno-math-errno");               // do not assume that math ops have side effects
+
+			Arguments.Add(GetRTTIFlag(CompileEnvironment)); // flag for run-time type info
+
+			if (CompileEnvironment.Architecture.StartsWith("x86_64"))
+			{
+				Arguments.Add("-mssse3"); // enable ssse3 by default for x86. This is default on for MSVC so lets reflect that here
+			}
+
+			//Arguments.Add("-DOPERATOR_NEW_INLINE=FORCENOINLINE");
+
 			if (CompileEnvironment.bIsBuildingDLL)
 			{
 				Arguments.Add("-fPIC");
@@ -598,17 +591,6 @@
 				Arguments.Add("-fdata-sections");
 			}
 
-			if (CompileEnvironment.bEnableExceptions)
-			{
-				Arguments.Add("-fexceptions");
-				Arguments.Add("-DPLATFORM_EXCEPTIONS_DISABLED=0");
-			}
-			else
-			{
-				Arguments.Add("-fno-exceptions");               // no exceptions
-				Arguments.Add("-DPLATFORM_EXCEPTIONS_DISABLED=1");
-			}
-
 			if (bSuppressPIE && !CompileEnvironment.bIsBuildingDLL)
 			{
 				Arguments.Add("-fno-PIE");
diff --git a/Engine/Source/Programs/UnrealBuildTool/Platform/Mac/MacToolChain.cs b/Engine/Source/Programs/UnrealBuildTool/Platform/Mac/MacToolChain.cs
--- a/Engine/Source/Programs/UnrealBuildTool/Platform/Mac/MacToolChain.cs
+++ b/Engine/Source/Programs/UnrealBuildTool/Platform/Mac/MacToolChain.cs
@@ -164,20 +164,33 @@
 		}
 
 		/// <inheritdoc/>
-		protected override void GetCompileArguments_Global(CppCompileEnvironment CompileEnvironment, List<string> Arguments)
+		protected override void GetCompileArguments_Debugging(CppCompileEnvironment CompileEnvironment, List<string> Arguments)
 		{
-			base.GetCompileArguments_Global(CompileEnvironment, Arguments);
+			base.GetCompileArguments_Debugging(CompileEnvironment, Arguments);
 
+			// TODO: Mac always enables exceptions, is this correct?
+			if (!CompileEnvironment.bEnableExceptions)
+			{
+				Arguments.Remove("-fno-exceptions");
+				Arguments.Remove("-DPLATFORM_EXCEPTIONS_DISABLED=1");
+			}
 			Arguments.Add("-fexceptions");
 			Arguments.Add("-DPLATFORM_EXCEPTIONS_DISABLED=0");
 
-			Arguments.Add("-fasm-blocks");
-
-			if(CompileEnvironment.bHideSymbolsByDefault)
+			if (CompileEnvironment.bHideSymbolsByDefault)
 			{
 				Arguments.Add("-fvisibility-ms-compat");
 				Arguments.Add("-fvisibility-inlines-hidden");
 			}
+		}
+
+		/// <inheritdoc/>
+		protected override void GetCompileArguments_Global(CppCompileEnvironment CompileEnvironment, List<string> Arguments)
+		{
+			base.GetCompileArguments_Global(CompileEnvironment, Arguments);
+
+			Arguments.Add("-fasm-blocks");
+
 			if (Options.HasFlag(ClangToolChainOptions.EnableAddressSanitizer))
 			{
 				Arguments.Add("-fsanitize=address");
@@ -189,7 +202,7 @@
 			if (Options.HasFlag(ClangToolChainOptions.EnableUndefinedBehaviorSanitizer))
 			{
 				Arguments.Add("-fsanitize=undefined");
-			}					
+			}
 
 			if (CompileEnvironment.bEnableOSX109Support)
 			{
@@ -200,50 +213,6 @@
 			Arguments.Add("" + FormatArchitectureArg(CompileEnvironment.Architecture));
 			Arguments.Add(string.Format(" -isysroot \"{0}\"", SDKPath));
 			Arguments.Add("-mmacosx-version-min=" + (CompileEnvironment.bEnableOSX109Support ? "10.9" : Settings.MacOSVersion));
-
-			bool bStaticAnalysis = false;
-			string? StaticAnalysisMode = Environment.GetEnvironmentVariable("CLANG_STATIC_ANALYZER_MODE");
-			if(StaticAnalysisMode != null && StaticAnalysisMode != "")
-			{
-				bStaticAnalysis = true;
-			}
-
-			// Optimize non- debug builds.
-			if (CompileEnvironment.bOptimizeCode && !bStaticAnalysis)
-			{
-				// Don't over optimise if using AddressSanitizer or you'll get false positive errors due to erroneous optimisation of necessary AddressSanitizer instrumentation.
-				if (Options.HasFlag(ClangToolChainOptions.EnableAddressSanitizer))
-				{
-					Arguments.Add("-O1 -g -fno-optimize-sibling-calls -fno-omit-frame-pointer");
-				}
-				else if (Options.HasFlag(ClangToolChainOptions.EnableThreadSanitizer))
-				{
-					Arguments.Add("-O1 -g");
-				}
-				else if (CompileEnvironment.bOptimizeForSize)
-				{
-					Arguments.Add("-Oz");
-				}
-				else
-				{
-					Arguments.Add("-O3");
-				}
-			}
-			else
-			{
-				Arguments.Add("-O0");
-			}
-
-			if (!CompileEnvironment.bUseInlining)
-			{
-				Arguments.Add("-fno-inline-functions");
-			}
-
-			// Create DWARF format debug info if wanted,
-			if (CompileEnvironment.bCreateDebugInfo)
-			{
-				Arguments.Add("-gdwarf-2");
-			}
 		}
 		
 		string AddFrameworkToLinkCommand(string FrameworkName, string Arg = "-framework")
diff --git a/Engine/Source/Programs/UnrealBuildTool/ToolChain/AppleToolChain.cs b/Engine/Source/Programs/UnrealBuildTool/ToolChain/AppleToolChain.cs
--- a/Engine/Source/Programs/UnrealBuildTool/ToolChain/AppleToolChain.cs
+++ b/Engine/Source/Programs/UnrealBuildTool/ToolChain/AppleToolChain.cs
@@ -272,6 +272,62 @@
 		}
 
 		/// <inheritdoc/>
+		protected override void GetCompileArguments_Optimizations(CppCompileEnvironment CompileEnvironment, List<string> Arguments)
+		{
+			base.GetCompileArguments_Optimizations(CompileEnvironment, Arguments);
+
+			bool bStaticAnalysis = false;
+			string? StaticAnalysisMode = Environment.GetEnvironmentVariable("CLANG_STATIC_ANALYZER_MODE");
+			if (!string.IsNullOrEmpty(StaticAnalysisMode))
+			{
+				bStaticAnalysis = true;
+			}
+
+			// Optimize non- debug builds.
+			if (CompileEnvironment.bOptimizeCode && !bStaticAnalysis)
+			{
+				// Don't over optimise if using AddressSanitizer or you'll get false positive errors due to erroneous optimisation of necessary AddressSanitizer instrumentation.
+				if (Options.HasFlag(ClangToolChainOptions.EnableAddressSanitizer))
+				{
+					Arguments.Add("-O1");
+					Arguments.Add("-g");
+					Arguments.Add("-fno-optimize-sibling-calls");
+					Arguments.Add("-fno-omit-frame-pointer");
+				}
+				else if (Options.HasFlag(ClangToolChainOptions.EnableThreadSanitizer))
+				{
+					Arguments.Add("-O1");
+					Arguments.Add("-g");
+				}
+				else if (CompileEnvironment.bOptimizeForSize)
+				{
+					Arguments.Add("-Oz");
+				}
+				else
+				{
+					Arguments.Add("-O3");
+				}
+			}
+			else
+			{
+				Arguments.Add("-O0");
+			}
+		}
+
+
+		/// <inheritdoc/>
+		protected override void GetCompileArguments_Debugging(CppCompileEnvironment CompileEnvironment, List<string> Arguments)
+		{
+			base.GetCompileArguments_Debugging(CompileEnvironment, Arguments);
+
+			// Create DWARF format debug info if wanted,
+			if (CompileEnvironment.bCreateDebugInfo)
+			{
+				Arguments.Add("-gdwarf-2");
+			}
+		}
+
+		/// <inheritdoc/>
 		protected override void GetCompileArguments_Global(CppCompileEnvironment CompileEnvironment, List<string> Arguments)
 		{
 			// Ensure Clang version is set for base.GetCompileArguments_Global()
diff --git a/Engine/Source/Programs/UnrealBuildTool/ToolChain/ClangToolChain.cs b/Engine/Source/Programs/UnrealBuildTool/ToolChain/ClangToolChain.cs
--- a/Engine/Source/Programs/UnrealBuildTool/ToolChain/ClangToolChain.cs
+++ b/Engine/Source/Programs/UnrealBuildTool/ToolChain/ClangToolChain.cs
@@ -6,6 +6,7 @@ using System.Collections.Generic;
 using System.Linq;
 using UnrealBuildBase;
 using Microsoft.Extensions.Logging;
+using System.IO;
 
 namespace UnrealBuildTool
 {
@@ -113,7 +114,7 @@ namespace UnrealBuildTool
 
 		protected ClangToolChainOptions Options;
 
-		public ClangToolChain(ClangToolChainOptions InOptions, ILogger InLogger) 
+		public ClangToolChain(ClangToolChainOptions InOptions, ILogger InLogger)
 			: base(InLogger)
 		{
 			Options = InOptions;
@@ -357,7 +358,7 @@ namespace UnrealBuildTool
 			// always use absolute paths for errors, this can help IDEs go to the error properly
 			Arguments.Add("-fdiagnostics-absolute-paths");  // https://clang.llvm.org/docs/ClangCommandLineReference.html#cmdoption-clang-fdiagnostics-absolute-paths
 
-			// https://clang.llvm.org/docs/ClangCommandLineReference.html#cmdoption-clang-fcolor-diagnostics
+			// https://clang.llvm.org/docs/ClangCommandLineReference.html#cmdoption-cla-fcolor-diagnostics
 			if (Log.ColorConsoleOutput)
 			{
 				Arguments.Add("-fdiagnostics-color");
@@ -370,6 +371,52 @@ namespace UnrealBuildTool
 			}
 		}
 
+		/// <summary>
+		/// Compile arguments for optimization settings, such as profile guided optimization and link time optimization
+		/// </summary>
+		/// <param name="CompileEnvironment"></param>
+		/// <param name="Arguments"></param>
+		protected virtual void GetCompileArguments_Optimizations(CppCompileEnvironment CompileEnvironment, List<string> Arguments)
+		{
+			// Profile Guided Optimization (PGO) and Link Time Optimization (LTO)
+			if (CompileEnvironment.bPGOOptimize)
+			{
+				Log.TraceInformationOnce("Enabling Profile Guided Optimization (PGO). Linking will take a while.");
+				Arguments.Add(string.Format("-fprofile-instr-use=\"{0}\"", Path.Combine(CompileEnvironment.PGODirectory!, CompileEnvironment.PGOFilenamePrefix!)));
+			}
+			else if (CompileEnvironment.bPGOProfile)
+			{
+				// Always enable LTO when generating PGO profile data.
+				Log.TraceInformationOnce("Enabling Profile Guided Instrumentation (PGI). Linking will take a while.");
+				Arguments.Add("-fprofile-instr-generate");
+			}
+
+			if (!CompileEnvironment.bUseInlining)
+			{
+				Arguments.Add("-fno-inline-functions");
+			}
+		}
+
+		/// <summary>
+		/// Compile arguments for debug settings
+		/// </summary>
+		/// <param name="CompileEnvironment"></param>
+		/// <param name="Arguments"></param>
+		protected virtual void GetCompileArguments_Debugging(CppCompileEnvironment CompileEnvironment, List<string> Arguments)
+		{
+			// Optionally enable exception handling (off by default since it generates extra code needed to propagate exceptions)
+			if (CompileEnvironment.bEnableExceptions)
+			{
+				Arguments.Add("-fexceptions");
+				Arguments.Add("-DPLATFORM_EXCEPTIONS_DISABLED=0");
+			}
+			else
+			{
+				Arguments.Add("-fno-exceptions");
+				Arguments.Add("-DPLATFORM_EXCEPTIONS_DISABLED=1");
+			}
+		}
+
 		/// <summary>
 		/// Common compile arguments for all files in a module.
 		/// Override and call base.GetCompileArguments_Global() in derived classes.
@@ -391,6 +438,12 @@ namespace UnrealBuildTool
 
 			// Add warning and error flags to the argument list.
 			GetCompileArguments_WarningsAndErrors(CompileEnvironment, Arguments);
+
+			// Add optimization flags to the argument list.
+			GetCompileArguments_Optimizations(CompileEnvironment, Arguments);
+
+			// Add debugging flags to the argument list.
+			GetCompileArguments_Debugging(CompileEnvironment, Arguments);
 		}
 	}
 }
\ No newline at end of file
