- Removed the version specific portion of Mac SDK path, since the there is always a softlink called MacOSX.sdk that points to the versioned directory

diff --git a/Engine/Source/Programs/UnrealBuildTool/Platform/Mac/MacToolChain.cs b/Engine/Source/Programs/UnrealBuildTool/Platform/Mac/MacToolChain.cs
--- a/Engine/Source/Programs/UnrealBuildTool/Platform/Mac/MacToolChain.cs
+++ b/Engine/Source/Programs/UnrealBuildTool/Platform/Mac/MacToolChain.cs
@@ -88,7 +88,7 @@ namespace UnrealBuildTool
 
 		public static string SDKPath
 		{
-			get { return Settings.BaseSDKDir + "/MacOSX" + Settings.MacOSSDKVersion + ".sdk"; }
+			get { return Settings.BaseSDKDir + "/MacOSX.sdk"; }
 		}
 
 		/// <summary>
