UnrealBuildTool: Move sanitizer flags to ClangToolChain

diff --git a/Engine/Source/Programs/UnrealBuildTool/Platform/IOS/IOSToolChain.cs b/Engine/Source/Programs/UnrealBuildTool/Platform/IOS/IOSToolChain.cs
--- a/Engine/Source/Programs/UnrealBuildTool/Platform/IOS/IOSToolChain.cs
+++ b/Engine/Source/Programs/UnrealBuildTool/Platform/IOS/IOSToolChain.cs
@@ -236,9 +236,10 @@ namespace UnrealBuildTool
 		}
 
 		/// <inheritdoc/>
-		protected override void GetCompileArguments_Global(CppCompileEnvironment CompileEnvironment, List<string> Arguments)
+		protected override void GetCompilerArguments_Sanitizers(CppCompileEnvironment CompileEnvironment, List<string> Arguments)
 		{
-			base.GetCompileArguments_Global(CompileEnvironment, Arguments);
+			// TODO: Reconcile with base
+			//base.GetCompilerArguments_Sanitizers(CompileEnvironment, Arguments);
 
 			string? SanitizerMode = Environment.GetEnvironmentVariable("ENABLE_ADDRESS_SANITIZER");
 			if ((SanitizerMode != null && SanitizerMode == "YES") || (Options.HasFlag(ClangToolChainOptions.EnableAddressSanitizer)))
@@ -257,6 +258,17 @@ namespace UnrealBuildTool
 				Arguments.Add("-fsanitize=thread");
 			}
 
+			if (Options.HasFlag(ClangToolChainOptions.EnableMemorySanitizer))
+			{
+				Arguments.Add("-fsanitize=memory");
+			}
+		}
+
+		/// <inheritdoc/>
+		protected override void GetCompileArguments_Global(CppCompileEnvironment CompileEnvironment, List<string> Arguments)
+		{
+			base.GetCompileArguments_Global(CompileEnvironment, Arguments);
+
 			Arguments.Add(GetRTTIFlag(CompileEnvironment));
 
 			if (IsBitcodeCompilingEnabled(CompileEnvironment.Configuration))
diff --git a/Engine/Source/Programs/UnrealBuildTool/Platform/Linux/LinuxToolChain.cs b/Engine/Source/Programs/UnrealBuildTool/Platform/Linux/LinuxToolChain.cs
--- a/Engine/Source/Programs/UnrealBuildTool/Platform/Linux/LinuxToolChain.cs
+++ b/Engine/Source/Programs/UnrealBuildTool/Platform/Linux/LinuxToolChain.cs
@@ -513,37 +513,28 @@ namespace UnrealBuildTool
 		}
 
 		/// <inheritdoc/>
-		protected override void GetCompileArguments_Global(CppCompileEnvironment CompileEnvironment, List<string> Arguments)
+		protected override void GetCompilerArguments_Sanitizers(CppCompileEnvironment CompileEnvironment, List<string> Arguments)
 		{
-			base.GetCompileArguments_Global(CompileEnvironment, Arguments);
-
-			// build up the commandline common to C and C++
-
-			if (ShouldUseLibcxx(CompileEnvironment.Architecture))
-			{
-				Arguments.Add("-nostdinc++");
-				Arguments.Add("-I" + "ThirdParty/Unix/LibCxx/include/");
-				Arguments.Add("-I" + "ThirdParty/Unix/LibCxx/include/c++/v1");
-			}
+			base.GetCompilerArguments_Sanitizers(CompileEnvironment, Arguments);
 
 			// ASan
 			if (Options.HasFlag(ClangToolChainOptions.EnableAddressSanitizer))
 			{
 				// Force using the ANSI allocator if ASan is enabled
-				Arguments.Add("-fsanitize=address -DFORCE_ANSI_ALLOCATOR=1");
+				Arguments.Add("-DFORCE_ANSI_ALLOCATOR=1");
 			}
 
 			// TSan
 			if (Options.HasFlag(ClangToolChainOptions.EnableThreadSanitizer))
 			{
 				// Force using the ANSI allocator if TSan is enabled
-				Arguments.Add("-fsanitize=thread -DFORCE_ANSI_ALLOCATOR=1");
+				Arguments.Add("-DFORCE_ANSI_ALLOCATOR=1");
 			}
 
 			// UBSan
 			if (Options.HasFlag(ClangToolChainOptions.EnableUndefinedBehaviorSanitizer))
 			{
-				Arguments.Add("-fsanitize=undefined -fno-sanitize=vptr");
+				Arguments.Add("-fno-sanitize=vptr");
 			}
 
 			// MSan
@@ -552,7 +543,23 @@ namespace UnrealBuildTool
 				// Force using the ANSI allocator if MSan is enabled
 				// -fsanitize-memory-track-origins adds a 1.5x-2.5x slow down ontop of MSan normal amount of overhead
 				// -fsanitize-memory-track-origins=1 is faster but collects only allocation points but not intermediate stores
-				Arguments.Add("-fsanitize=memory -fsanitize-memory-track-origins -DFORCE_ANSI_ALLOCATOR=1");
+				Arguments.Add("-fsanitize-memory-track-origins");
+				Arguments.Add("-DFORCE_ANSI_ALLOCATOR=1");
+			}
+		}
+
+		/// <inheritdoc/>
+		protected override void GetCompileArguments_Global(CppCompileEnvironment CompileEnvironment, List<string> Arguments)
+		{
+			base.GetCompileArguments_Global(CompileEnvironment, Arguments);
+
+			// build up the commandline common to C and C++
+
+			if (ShouldUseLibcxx(CompileEnvironment.Architecture))
+			{
+				Arguments.Add("-nostdinc++");
+				Arguments.Add("-I" + "ThirdParty/Unix/LibCxx/include/");
+				Arguments.Add("-I" + "ThirdParty/Unix/LibCxx/include/c++/v1");
 			}
 
 			if (CompilerVersionGreaterOrEqual(12, 0, 0))
diff --git a/Engine/Source/Programs/UnrealBuildTool/Platform/Mac/MacToolChain.cs b/Engine/Source/Programs/UnrealBuildTool/Platform/Mac/MacToolChain.cs
--- a/Engine/Source/Programs/UnrealBuildTool/Platform/Mac/MacToolChain.cs
+++ b/Engine/Source/Programs/UnrealBuildTool/Platform/Mac/MacToolChain.cs
@@ -191,19 +191,6 @@ namespace UnrealBuildTool
 
 			Arguments.Add("-fasm-blocks");
 
-			if (Options.HasFlag(ClangToolChainOptions.EnableAddressSanitizer))
-			{
-				Arguments.Add("-fsanitize=address");
-			}
-			if (Options.HasFlag(ClangToolChainOptions.EnableThreadSanitizer))
-			{
-				Arguments.Add("-fsanitize=thread");
-			}
-			if (Options.HasFlag(ClangToolChainOptions.EnableUndefinedBehaviorSanitizer))
-			{
-				Arguments.Add("-fsanitize=undefined");
-			}
-
 			if (CompileEnvironment.bEnableOSX109Support)
 			{
 				Arguments.Add("-faligned-new"); // aligned operator new is supported only on macOS 10.14 and above
diff --git a/Engine/Source/Programs/UnrealBuildTool/ToolChain/ClangToolChain.cs b/Engine/Source/Programs/UnrealBuildTool/ToolChain/ClangToolChain.cs
--- a/Engine/Source/Programs/UnrealBuildTool/ToolChain/ClangToolChain.cs
+++ b/Engine/Source/Programs/UnrealBuildTool/ToolChain/ClangToolChain.cs
@@ -417,6 +417,38 @@ namespace UnrealBuildTool
 			}
 		}
 
+		/// <summary>
+		/// Compile arguments for sanitizers
+		/// </summary>
+		/// <param name="CompileEnvironment"></param>
+		/// <param name="Arguments"></param>
+		protected virtual void GetCompilerArguments_Sanitizers(CppCompileEnvironment CompileEnvironment, List<string> Arguments)
+		{
+			// ASan
+			if (Options.HasFlag(ClangToolChainOptions.EnableAddressSanitizer))
+			{
+				Arguments.Add("-fsanitize=address");
+			}
+
+			// TSan
+			if (Options.HasFlag(ClangToolChainOptions.EnableThreadSanitizer))
+			{
+				Arguments.Add("-fsanitize=thread");
+			}
+
+			// UBSan
+			if (Options.HasFlag(ClangToolChainOptions.EnableUndefinedBehaviorSanitizer))
+			{
+				Arguments.Add("-fsanitize=undefined");
+			}
+
+			// MSan
+			if (Options.HasFlag(ClangToolChainOptions.EnableMemorySanitizer))
+			{
+				Arguments.Add("-fsanitize=memory");
+			}
+		}
+
 		/// <summary>
 		/// Common compile arguments for all files in a module.
 		/// Override and call base.GetCompileArguments_Global() in derived classes.
@@ -444,6 +476,9 @@ namespace UnrealBuildTool
 
 			// Add debugging flags to the argument list.
 			GetCompileArguments_Debugging(CompileEnvironment, Arguments);
+
+			// Add sanitizer flags to the argument list
+			GetCompilerArguments_Sanitizers(CompileEnvironment, Arguments);
 		}
 	}
 }
\ No newline at end of file
