UnrealBuildTool: Remove old debugging code

#rnx

diff --git a/Engine/Source/Programs/UnrealBuildTool/Platform/Linux/LinuxToolChain.cs b/Engine/Source/Programs/UnrealBuildTool/Platform/Linux/LinuxToolChain.cs
--- a/Engine/Source/Programs/UnrealBuildTool/Platform/Linux/LinuxToolChain.cs
+++ b/Engine/Source/Programs/UnrealBuildTool/Platform/Linux/LinuxToolChain.cs
@@ -1101,7 +1101,6 @@ namespace UnrealBuildTool
 		/// </summary>
 		public FileItem CreateArchiveAndIndex(LinkEnvironment LinkEnvironment, IActionGraphBuilder Graph, ILogger Logger)
 		{
-			Debugger.Launch();
 			// Create an archive action
 			Action ArchiveAction = Graph.CreateAction(ActionType.Link);
 			ArchiveAction.WorkingDirectory = Unreal.EngineSourceDirectory;
