UnrealBuildTool: Handle modules with the same name in different paths when generating projects

#rnx

diff --git a/Engine/Source/Programs/UnrealBuildTool/ProjectFiles/ProjectFileGenerator.cs b/Engine/Source/Programs/UnrealBuildTool/ProjectFiles/ProjectFileGenerator.cs
--- a/Engine/Source/Programs/UnrealBuildTool/ProjectFiles/ProjectFileGenerator.cs
+++ b/Engine/Source/Programs/UnrealBuildTool/ProjectFiles/ProjectFileGenerator.cs
@@ -377,9 +377,9 @@
 		protected bool bPrimaryProjectNameFromFolder = false;
 
 		/// <summary>
-		/// Maps all module names that were included in generated project files, to actual project file objects.
+		/// Maps all module files that were included in generated project files, to actual project file objects.
 		/// </summary>
-		protected Dictionary<string, ProjectFile> ModuleToProjectFileMap = new Dictionary<string, ProjectFile>(StringComparer.InvariantCultureIgnoreCase);
+		protected Dictionary<FileReference, ProjectFile> ModuleToProjectFileMap = new Dictionary<FileReference, ProjectFile>();
 
 		/// <summary>
 		/// If generating project files for a single project, the path to its .uproject file.
@@ -1962,7 +1962,7 @@
 									foreach (UEBuildModuleCPP Module in Binary.Modules.OfType<UEBuildModuleCPP>())
 									{
 										ProjectFile? ProjectFileForIDE;
-										if (ModuleToProjectFileMap.TryGetValue(Module.Name, out ProjectFileForIDE) && ProjectFileForIDE == TargetProjectFile)
+										if (ModuleToProjectFileMap.TryGetValue(Module.RulesFile, out ProjectFileForIDE) && ProjectFileForIDE == TargetProjectFile)
 										{
 											Utils.WriteFileIfChangedContext = $"{Target.TargetName} {Target.Configuration} {Target.Platform} {Binary.OutputFilePaths[0].GetFileName()} {Module.Name}";
 
@@ -2167,7 +2167,7 @@
 					ProjectFile ProjectFile = FindProjectForModule(CurModuleFile, AllGames, ProgramProjects, ModProjects, AdditionalSearchPaths, out BaseFolder);
 
 					// Update our module map
-					ModuleToProjectFileMap[ModuleName] = ProjectFile;
+					ModuleToProjectFileMap[CurModuleFile] = ProjectFile;
 					ProjectFile.IsGeneratedProject = true;
 
 					// Only search subdirectories for non-external modules.  We don't want to add all of the source and header files
