UBT: Remove async keyword from Main to fix warning.

diff --git a/Engine/Source/Programs/UnrealBuildTool/UnrealBuildTool.cs b/Engine/Source/Programs/UnrealBuildTool/UnrealBuildTool.cs
--- a/Engine/Source/Programs/UnrealBuildTool/UnrealBuildTool.cs
+++ b/Engine/Source/Programs/UnrealBuildTool/UnrealBuildTool.cs
@@ -451,7 +451,7 @@ namespace UnrealBuildTool
 		/// </summary>
 		/// <param name="ArgumentsArray">Command line arguments</param>
 		/// <returns>Zero on success, non-zero on error</returns>
-		private static async Task<int> Main(string[] ArgumentsArray)
+		private static int Main(string[] ArgumentsArray)
 		{
 			FileReference? RunFile = null;
 			SingleInstanceMutex? Mutex = null;
