- Re-enabled bBuildISPC flag on m1

diff --git a/Engine/Source/Programs/UnrealBuildTool/Platform/Mac/UEBuildMac.cs b/Engine/Source/Programs/UnrealBuildTool/Platform/Mac/UEBuildMac.cs
--- a/Engine/Source/Programs/UnrealBuildTool/Platform/Mac/UEBuildMac.cs
+++ b/Engine/Source/Programs/UnrealBuildTool/Platform/Mac/UEBuildMac.cs
@@ -133,7 +133,6 @@ namespace UnrealBuildTool
 
 			if (bCompilingForArm)
 			{
-				Target.bCompileISPC = false;
 				Target.DisablePlugins.AddRange(new string[]
 				{
 					// onnx:
