UnrealBuildTool: Only mark marketplace plugin as readonly for installed builds

This should be tested against the fix for UE-148690 to ensure there's no regression

diff --git a/Engine/Source/Programs/UnrealBuildTool/System/RulesCompiler.cs b/Engine/Source/Programs/UnrealBuildTool/System/RulesCompiler.cs
--- a/Engine/Source/Programs/UnrealBuildTool/System/RulesCompiler.cs
+++ b/Engine/Source/Programs/UnrealBuildTool/System/RulesCompiler.cs
@@ -98,7 +98,7 @@
 
 				if (MarketplacePlugins.Count > 0)
 				{
-					EngineRulesAssembly = CreateMarketplaceRulesAssembly(MarketplacePlugins, bSkipCompile, bForceCompile, EngineRulesAssembly, Logger);
+					EngineRulesAssembly = CreateMarketplaceRulesAssembly(MarketplacePlugins, Unreal.IsEngineInstalled() || bUsePrecompiled, bSkipCompile, bForceCompile, EngineRulesAssembly, Logger);
 				}
 			}
 			return EngineRulesAssembly;
@@ -187,12 +187,13 @@
 		/// Creates a rules assembly
 		/// </summary>
 		/// <param name="Plugins">List of plugins to include in this assembly</param>
+		/// <param name="bReadOnly">Whether the assembly should be marked as installed</param>
 		/// <param name="bSkipCompile">Whether to skip compilation for this assembly</param>
 		/// <param name="bForceCompile">Whether to always compile this assembly</param>
 		/// <param name="Parent">The parent rules assembly</param>
 		/// <param name="Logger"></param>
 		/// <returns>New rules assembly</returns>
-		private static RulesAssembly CreateMarketplaceRulesAssembly(IReadOnlyList<PluginInfo> Plugins, bool bSkipCompile, bool bForceCompile, RulesAssembly Parent, ILogger Logger)
+		private static RulesAssembly CreateMarketplaceRulesAssembly(IReadOnlyList<PluginInfo> Plugins, bool bReadOnly, bool bSkipCompile, bool bForceCompile, RulesAssembly Parent, ILogger Logger)
 		{
 			RulesScope MarketplaceScope = new RulesScope("Marketplace", Parent.Scope);
 
@@ -209,7 +210,7 @@
 			if (ModuleFileToContext.Count > 0)
 			{
 				FileReference AssemblyFileName = FileReference.Combine(UnrealBuildTool.WritableEngineDirectory, "Intermediate", "Build", "BuildRules", "MarketplaceRules.dll");
-				Result = new RulesAssembly(MarketplaceScope, new List<DirectoryReference> { DirectoryReference.Combine(Unreal.EngineDirectory, "Plugins", "Marketplace") }, Plugins, ModuleFileToContext, new List<FileReference>(), AssemblyFileName, bContainsEngineModules: true, DefaultBuildSettings: BuildSettingsVersion.Latest, bReadOnly: true, bSkipCompile: bSkipCompile, bForceCompile: bForceCompile, Parent: Parent, Logger: Logger);
+				Result = new RulesAssembly(MarketplaceScope, new List<DirectoryReference> { DirectoryReference.Combine(Unreal.EngineDirectory, "Plugins", "Marketplace") }, Plugins, ModuleFileToContext, new List<FileReference>(), AssemblyFileName, bContainsEngineModules: true, DefaultBuildSettings: BuildSettingsVersion.Latest, bReadOnly: bReadOnly, bSkipCompile: bSkipCompile, bForceCompile: bForceCompile, Parent: Parent, Logger: Logger);
 			}
 			return Result;
 		}
