Add file/line logging for UnrealPluginLanguage (UPL) warnings and errors

#android

diff --git a/Engine/Source/Programs/UnrealBuildTool/System/UnrealPluginLanguage.cs b/Engine/Source/Programs/UnrealBuildTool/System/UnrealPluginLanguage.cs
--- a/Engine/Source/Programs/UnrealBuildTool/System/UnrealPluginLanguage.cs
+++ b/Engine/Source/Programs/UnrealBuildTool/System/UnrealPluginLanguage.cs
@@ -487,6 +487,7 @@
 
 		/** Trace flag to enable debugging */
 		static private bool bGlobalTrace = false;
+		static private bool bGlobalTraceFilename = false;
 
 		/** Project file reference */
 		private FileReference? ProjectFile;
@@ -495,6 +496,8 @@
 
 		private class UPLContext
 		{
+			public string Filename;
+
 			/** Variable state */
 			public Dictionary<string, bool> BoolVariables;
 			public Dictionary<string, int> IntVariables;
@@ -503,9 +506,12 @@
 
 			/** Local context trace */
 			public bool bTrace;
+			public bool bTraceFilename;
 
-			public UPLContext(string Architecture, string PluginDir)
+			public UPLContext(string inFilename, string Architecture, string PluginDir)
 			{
+				Filename = inFilename;
+
 				BoolVariables = new Dictionary<string, bool>();
 				IntVariables = new Dictionary<string, int>();
 				StringVariables = new Dictionary<string, string>();
@@ -521,6 +527,7 @@
 				StringVariables["Architecture"] = Architecture;
 
 				bTrace = false;
+				bTraceFilename = false;
 			}
 		}
 
@@ -538,7 +545,7 @@
 			LastError = null;
 
 			Contexts = new Dictionary<string, UPLContext>();
-			GlobalContext = new UPLContext("", "");
+			GlobalContext = new UPLContext("", "", "");
 			ContextIndex = 0;
 
 			XMLNameSpace = InXMLNameSpace;
@@ -557,8 +564,8 @@
 					string PluginDir = Path.GetDirectoryName(Filename)!;
 					try
 					{
-						XDocument MergeDoc = XDocument.Load(Filename);
-						MergeXML(MergeDoc, PluginDir, InArchitectures);
+						XDocument MergeDoc = XDocument.Load(Filename, LoadOptions.SetLineInfo);
+						MergeXML(MergeDoc, Filename, PluginDir, InArchitectures);
 					}
 					catch (Exception e)
 					{
@@ -580,13 +587,16 @@
 		public bool GetTrace() { return bGlobalTrace; }
 		public void SetTrace() { bGlobalTrace = true; }
 		public void ClearTrace() { bGlobalTrace = false; }
+		public bool GetTraceFilename() { return bGlobalTraceFilename; }
+		public void SetTraceFilename() { bGlobalTraceFilename = true; }
+		public void ClearTraceFilename() { bGlobalTraceFilename = false; }
 
 		public string GetUPLHash()
 		{
 			return ContentHash.MD5(XDoc.ToString()).ToString();
 		}
 
-		public bool MergeXML(XDocument MergeDoc, string PluginDir, List<string> Architectures)
+		public bool MergeXML(XDocument MergeDoc, string Filename, string PluginDir, List<string> Architectures)
 		{
 			if (MergeDoc == null)
 			{
@@ -597,10 +607,19 @@
 			ContextIndex++;
 			foreach (string Architecture in Architectures)
 			{
-				UPLContext Context = new UPLContext(Architecture, PluginDir);
+				UPLContext Context = new UPLContext(Filename, Architecture, PluginDir);
 				Contexts[Architecture + "_" + ContextIndex] = Context;
 			}
 
+			// line numbers are lost in the merge so add them as an attribute for later tracing of warnings/errors
+			foreach (XElement Element in MergeDoc.Root!.Descendants())
+			{
+				if (((IXmlLineInfo)Element).HasLineInfo())
+				{
+					Element.Add(new XAttribute("__line", ((IXmlLineInfo)Element).LineNumber));
+				}
+			}
+
 			// merge in the nodes
 			foreach (XElement Element in MergeDoc.Root!.Elements())
 			{
@@ -677,7 +696,7 @@
 			{
 				if (!GlobalContext.BoolVariables.TryGetValue(Condition, out Result))
 				{
-					Logger.LogWarning("\nMissing condition '{Condition}' in '{Node}' (skipping instruction)", Condition, TraceNodeString(Node));
+					Logger.LogWarning("\nMissing condition '{Condition}' in '{Node}' (skipping instruction)", Condition, TraceNodeString(Context, Node));
 					return false;
 				}
 			}
@@ -794,13 +813,32 @@
 			return Result;
 		}
 
-		private string TraceNodeString(XElement Node)
+		private string TraceNodeString(UPLContext Context, XElement Node, bool ShowFilename = true)
 		{
+			XAttribute? LineAttrib = null;
 			string Result = Node.Name.ToString();
 			foreach (XAttribute Attrib in Node.Attributes())
 			{
-				Result += " " + Attrib.ToString();
+				string AttribStr = Attrib.ToString();
+				if (AttribStr.StartsWith("__line="))
+				{
+					LineAttrib = Attrib;
+				}
+				else
+				{
+					Result += " " + AttribStr;
+				}
 			}
+
+			if (ShowFilename)
+			{
+				Result += ", File: " + Context.Filename;
+			}
+			if (LineAttrib != null)
+			{
+				Result += ", Line: " + LineAttrib.Value;
+			}
+
 			return Result;
 		}
 
@@ -814,12 +852,12 @@
 			return !(Input.Equals("0") || Input.Equals("false") || Input.Equals("off") || Input.Equals("no"));
 		}
 
-		private int StringToInt(string? Input, XElement Node)
+		private int StringToInt(UPLContext Context, XElement Node, string? Input)
 		{
 			int Result = 0;
 			if (!int.TryParse(Input, out Result))
 			{
-				Logger.LogWarning("\nInvalid integer '{Input}' in '{Node}' (defaulting to 0)", Input, TraceNodeString(Node));
+				Logger.LogWarning("\nInvalid integer '{Input}' in '{Node}' (defaulting to 0)", Input, TraceNodeString(Context, Node));
 			}
 			return Result;
 		}
@@ -832,7 +870,7 @@
 			{
 				if (bRequired)
 				{
-					Logger.LogWarning("\nMissing attribute '{Attribute}' in '{Node}' (skipping instruction)", AttributeName, TraceNodeString(Node));
+					Logger.LogWarning("\nMissing attribute '{Attribute}' in '{Node}' (skipping instruction)", AttributeName, TraceNodeString(Context, Node));
 				}
 				return Fallback;
 			}
@@ -848,7 +886,7 @@
 			{
 				if (bRequired)
 				{
-					Logger.LogWarning("\nMissing attribute '{Attribute}' in '{Node}' (skipping instruction)", AttributeName, TraceNodeString(Node));
+					Logger.LogWarning("\nMissing attribute '{Attribute}' in '{Node}' (skipping instruction)", AttributeName, TraceNodeString(Context, Node));
 				}
 				return Fallback;
 			}
@@ -1068,12 +1106,17 @@
 				XElement Node = ExecutionStack.Pop();
 				if (bGlobalTrace || CurrentContext.bTrace)
 				{
-					Logger.LogInformation("Execute: '{Node}'", TraceNodeString(Node));
+					Logger.LogInformation("Execute: '{Node}'", TraceNodeString(CurrentContext, Node, bGlobalTraceFilename || CurrentContext.bTraceFilename));
 				}
 				switch (Node.Name.ToString())
 				{
 					case "trace":
 						{
+							string? TraceFilename = GetAttribute(CurrentContext, Node, "filename");
+							if (TraceFilename != null)
+							{
+								CurrentContext.bTraceFilename = StringToBool(TraceFilename);
+							}
 							string? Enable = GetAttribute(CurrentContext, Node, "enable");
 							if (Enable != null)
 							{
@@ -1275,7 +1318,7 @@
 										{
 											if (!GlobalContext.ElementVariables.TryGetValue(Tag.Substring(1), out Target))
 											{
-												Logger.LogWarning("Missing element variable '{Tag}' in '{Node}' (skipping instruction)", Tag, TraceNodeString(Node));
+												Logger.LogWarning("Missing element variable '{Tag}' in '{Node}' (skipping instruction)", Tag, TraceNodeString(CurrentContext, Node));
 												continue;
 											}
 										}
@@ -1312,7 +1355,7 @@
 										{
 											if (!GlobalContext.ElementVariables.TryGetValue(Tag.Substring(1), out Target))
 											{
-												Logger.LogInformation("\nMissing element variable '{Tag}' in '{Node}' (skipping instruction)", Tag, TraceNodeString(Node));
+												Logger.LogInformation("\nMissing element variable '{Tag}' in '{Node}' (skipping instruction)", Tag, TraceNodeString(CurrentContext, Node));
 												continue;
 											}
 										}
@@ -1548,7 +1591,7 @@
 								{
 									if (!GlobalContext.ElementVariables.TryGetValue(Name, out Element))
 									{
-										Logger.LogWarning("Missing element variable '{Name}' in '{Node}' (skipping instruction)", Name, TraceNodeString(Node));
+										Logger.LogWarning("Missing element variable '{Name}' in '{Node}' (skipping instruction)", Name, TraceNodeString(CurrentContext, Node));
 										continue;
 									}
 								}
@@ -1561,7 +1604,7 @@
 										{
 											if (!GlobalContext.ElementVariables.TryGetValue(Tag.Substring(1), out Target))
 											{
-												Logger.LogWarning("Missing element variable '{Name}' in '{Node}' (skipping instruction)", Tag, TraceNodeString(Node));
+												Logger.LogWarning("Missing element variable '{Name}' in '{Node}' (skipping instruction)", Tag, TraceNodeString(CurrentContext, Node));
 												continue;
 											}
 										}
@@ -1609,7 +1652,7 @@
 										{
 											if (!GlobalContext.ElementVariables.TryGetValue(Tag.Substring(1), out Target))
 											{
-												Logger.LogWarning("Missing element variable '{Tag}' in '{Node}' (skipping instruction)", Tag, TraceNodeString(Node));
+												Logger.LogWarning("Missing element variable '{Tag}' in '{Node}' (skipping instruction)", Tag, TraceNodeString(CurrentContext, Node));
 												continue;
 											}
 										}
@@ -1977,7 +2020,7 @@
 							string? Arg2 = GetAttribute(CurrentContext, Node, "arg2", true, false);
 							if (Result != null)
 							{
-								CurrentContext.BoolVariables[Result] = (StringToInt(Arg1, Node) < StringToInt(Arg2, Node));
+								CurrentContext.BoolVariables[Result] = (StringToInt(CurrentContext, Node, Arg1) < StringToInt(CurrentContext, Node, Arg2));
 							}
 						}
 						break;
@@ -1989,7 +2032,7 @@
 							string? Arg2 = GetAttribute(CurrentContext, Node, "arg2", true, false);
 							if (Result != null)
 							{
-								CurrentContext.BoolVariables[Result] = (StringToInt(Arg1, Node) <= StringToInt(Arg2, Node));
+								CurrentContext.BoolVariables[Result] = (StringToInt(CurrentContext, Node, Arg1) <= StringToInt(CurrentContext, Node, Arg2));
 							}
 						}
 						break;
@@ -2001,7 +2044,7 @@
 							string? Arg2 = GetAttribute(CurrentContext, Node, "arg2", true, false);
 							if (Result != null)
 							{
-								CurrentContext.BoolVariables[Result] = (StringToInt(Arg1, Node) > StringToInt(Arg2, Node));
+								CurrentContext.BoolVariables[Result] = (StringToInt(CurrentContext, Node, Arg1) > StringToInt(CurrentContext, Node, Arg2));
 							}
 						}
 						break;
@@ -2013,7 +2056,7 @@
 							string? Arg2 = GetAttribute(CurrentContext, Node, "arg2", true, false);
 							if (Result != null)
 							{
-								CurrentContext.BoolVariables[Result] = (StringToInt(Arg1, Node) >= StringToInt(Arg2, Node));
+								CurrentContext.BoolVariables[Result] = (StringToInt(CurrentContext, Node, Arg1) >= StringToInt(CurrentContext, Node, Arg2));
 							}
 						}
 						break;
@@ -2035,7 +2078,7 @@
 							string Value = GetAttribute(CurrentContext, Node, "value", true, false, "0");
 							if (Result != null)
 							{
-								CurrentContext.IntVariables[Result] = StringToInt(Value, Node);
+								CurrentContext.IntVariables[Result] = StringToInt(CurrentContext, Node, Value);
 							}
 						}
 						break;
@@ -2047,7 +2090,7 @@
 							if (Result != null)
 							{
 								Value = ExpandVariables(CurrentContext, "$I(" + Value + ")");
-								CurrentContext.IntVariables[Result] = StringToInt(Value, Node);
+								CurrentContext.IntVariables[Result] = StringToInt(CurrentContext, Node, Value);
 							}
 						}
 						break;
@@ -2061,7 +2104,7 @@
 							string DefaultVal = GetAttribute(CurrentContext, Node, "default", true, false, "0");
 							if (Result != null && Ini != null && Section != null && Property != null)
 							{
-								int Value = StringToInt(DefaultVal, Node);
+								int Value = StringToInt(CurrentContext, Node, DefaultVal);
 
 								ConfigCacheIni_UPL ConfigIni = GetConfigCacheIni_UPL(Ini);
 								if (ConfigIni != null)
@@ -2080,7 +2123,7 @@
 							string Arg2 = GetAttribute(CurrentContext, Node, "arg2", true, false, "0");
 							if (Result != null)
 							{
-								CurrentContext.IntVariables[Result] = StringToInt(Arg1, Node) + StringToInt(Arg2, Node);
+								CurrentContext.IntVariables[Result] = StringToInt(CurrentContext, Node, Arg1) + StringToInt(CurrentContext, Node, Arg2);
 							}
 						}
 						break;
@@ -2092,7 +2135,7 @@
 							string Arg2 = GetAttribute(CurrentContext, Node, "arg2", true, false, "0");
 							if (Result != null)
 							{
-								CurrentContext.IntVariables[Result] = StringToInt(Arg1, Node) - StringToInt(Arg2, Node);
+								CurrentContext.IntVariables[Result] = StringToInt(CurrentContext, Node, Arg1) - StringToInt(CurrentContext, Node, Arg2);
 							}
 						}
 						break;
@@ -2104,7 +2147,7 @@
 							string Arg2 = GetAttribute(CurrentContext, Node, "arg2", true, false, "1");
 							if (Result != null)
 							{
-								CurrentContext.IntVariables[Result] = StringToInt(Arg1, Node) * StringToInt(Arg2, Node);
+								CurrentContext.IntVariables[Result] = StringToInt(CurrentContext, Node, Arg1) * StringToInt(CurrentContext, Node, Arg2);
 							}
 						}
 						break;
@@ -2116,14 +2159,14 @@
 							string Arg2 = GetAttribute(CurrentContext, Node, "arg2", true, false, "1");
 							if (Result != null)
 							{
-								int Denominator = StringToInt(Arg2, Node);
+								int Denominator = StringToInt(CurrentContext, Node, Arg2);
 								if (Denominator == 0)
 								{
-									CurrentContext.IntVariables[Result] = StringToInt(Arg1, Node);
+									CurrentContext.IntVariables[Result] = StringToInt(CurrentContext, Node, Arg1);
 								}
 								else
 								{
-									CurrentContext.IntVariables[Result] = StringToInt(Arg1, Node) / Denominator;
+									CurrentContext.IntVariables[Result] = StringToInt(CurrentContext, Node, Arg1) / Denominator;
 								}
 							}
 						}
@@ -2208,7 +2251,7 @@
 										{
 											if (!GlobalContext.ElementVariables.TryGetValue(Tag.Substring(1), out Element))
 											{
-												Logger.LogWarning("Missing element variable '{Tag}' in '{Node}' (skipping instruction)", Tag, TraceNodeString(Node));
+												Logger.LogWarning("Missing element variable '{Tag}' in '{Node}' (skipping instruction)", Tag, TraceNodeString(CurrentContext, Node));
 												continue;
 											}
 										}
@@ -2235,7 +2278,7 @@
 										{
 											if (!GlobalContext.ElementVariables.TryGetValue(Tag.Substring(1), out Element))
 											{
-												Logger.LogWarning("Missing element variable '{Tag}' in '{Node}' (skipping instruction)", Tag, TraceNodeString(Node));
+												Logger.LogWarning("Missing element variable '{Tag}' in '{Node}' (skipping instruction)", Tag, TraceNodeString(CurrentContext, Node));
 												continue;
 											}
 										}
@@ -2274,7 +2317,7 @@
 										{
 											if (!GlobalContext.ElementVariables.TryGetValue(Tag.Substring(1), out Element))
 											{
-												Logger.LogWarning("Missing element variable '{Tag}' in '{Node}' (skipping instruction)", Tag, TraceNodeString(Node));
+												Logger.LogWarning("Missing element variable '{Tag}' in '{Node}' (skipping instruction)", Tag, TraceNodeString(CurrentContext, Node));
 												continue;
 											}
 										}
@@ -2283,7 +2326,7 @@
 
 								if (Element.Value == null)
 								{
-									Logger.LogWarning("Expected text in element '{Element}' in '{Node}' but found none (skipping instruction)", Element.Name.ToString(), TraceNodeString(Node));
+									Logger.LogWarning("Expected text in element '{Element}' in '{Node}' but found none (skipping instruction)", Element.Name.ToString(), TraceNodeString(CurrentContext, Node));
 									continue;
 								}
 
@@ -2353,8 +2396,8 @@
 							string Length = GetAttribute(CurrentContext, Node, "length", true, false, "0");
 							if (Result != null && Source != null)
 							{
-								int Index = StringToInt(Start, Node);
-								int Count = StringToInt(Length, Node);
+								int Index = StringToInt(CurrentContext, Node, Start);
+								int Count = StringToInt(CurrentContext, Node, Length);
 								Index = (Index < 0) ? 0 : (Index > Source.Length) ? Source.Length : Index;
 								Count = (Index + Count > Source.Length) ? Source.Length - Index : Count;
 								string Value = Source.Substring(Index, Count);
