Making it so UBT accounts for project level content-only plugins in content-only projects (the project could still enable/disable those plugins and the rules assembly needs to be aware of them).

diff --git a/Engine/Source/Programs/UnrealBuildTool/System/RulesCompiler.cs b/Engine/Source/Programs/UnrealBuildTool/System/RulesCompiler.cs
--- a/Engine/Source/Programs/UnrealBuildTool/System/RulesCompiler.cs
+++ b/Engine/Source/Programs/UnrealBuildTool/System/RulesCompiler.cs
@@ -294,9 +294,9 @@ namespace UnrealBuildTool
 
 				// Compile the assembly. If there are no module or target files, just use the parent assembly.
 				FileReference AssemblyFileName = FileReference.Combine(MainProjectDirectory, "Intermediate", "Build", "BuildRules", ProjectFileName.GetFileNameWithoutExtension() + "ModuleRules" + FrameworkAssemblyExtension);
-				if(ModuleFiles.Count == 0 && TargetFiles.Count == 0)
+				if(ModuleFiles.Count == 0 && TargetFiles.Count == 0 && ProjectPlugins.Count == 0)  
 				{
-					ProjectRulesAssembly = Parent;
+					ProjectRulesAssembly = Parent; 
 				}
 				else
 				{
