UnrealBuildTool: Add verbose logging to RulesCompiler to track which plugins and modules are found

#rnx

diff --git a/Engine/Source/Programs/UnrealBuildTool/System/RulesCompiler.cs b/Engine/Source/Programs/UnrealBuildTool/System/RulesCompiler.cs
--- a/Engine/Source/Programs/UnrealBuildTool/System/RulesCompiler.cs
+++ b/Engine/Source/Programs/UnrealBuildTool/System/RulesCompiler.cs
@@ -230,6 +230,8 @@ namespace UnrealBuildTool
 			RulesAssembly? ProjectRulesAssembly;
 			if (!LoadedAssemblyMap.TryGetValue(ProjectFileName, out ProjectRulesAssembly))
 			{
+				Logger.LogTrace("Creating project rules assembly for {Project}...", ProjectFileName);
+
 				ProjectDescriptor Project = ProjectDescriptor.FromFile(ProjectFileName);
 
 				// Create the parent assembly
@@ -281,9 +283,18 @@ namespace UnrealBuildTool
 					}
 				}
 
+				Logger.LogTrace(" Found {Count} Plugins:", ProjectPlugins.Count);
+				ProjectPlugins.ForEach(x => { Logger.LogTrace("  {Plugin}", x.File); });
+
 				// Find all the plugin module rules
 				FindModuleRulesForPlugins(ProjectPlugins, DefaultModuleContext, ModuleFiles);
 
+				Logger.LogTrace(" Found {Count} Modules:", ModuleFiles.Count);
+				foreach (var Item in ModuleFiles)
+				{
+					Logger.LogTrace("  {Module}", Item.Key);
+				}
+
 				// Add the games project's intermediate source folder
 				DirectoryReference ProjectIntermediateSourceDirectory = DirectoryReference.Combine(MainProjectDirectory, "Intermediate", "Source");
 				if (DirectoryReference.Exists(ProjectIntermediateSourceDirectory))
@@ -323,6 +334,8 @@ namespace UnrealBuildTool
 			RulesAssembly? PluginRulesAssembly;
 			if (!LoadedAssemblyMap.TryGetValue(PluginFileName, out PluginRulesAssembly))
 			{
+				Logger.LogTrace("Creating plugin rules assembly for {Plugin}...", PluginFileName);
+
 				// Find all the rules source files
 				Dictionary<FileReference, ModuleRulesContext> ModuleFiles = new Dictionary<FileReference, ModuleRulesContext>();
 				List<FileReference> TargetFiles = new List<FileReference>();
@@ -339,6 +352,9 @@ namespace UnrealBuildTool
 					ForeignPlugins.Add(ForeignPluginInfo);
 				}
 
+				Logger.LogTrace(" Found {Count} ForeignPlugins:", ForeignPlugins.Count);
+				ForeignPlugins.ForEach(x => { Logger.LogTrace("  {ForeignPlugin}", x.File); });
+
 				// Create a new scope for the plugin. It should not reference anything else.
 				RulesScope Scope = new RulesScope("Plugin", Parent.Scope);
 
@@ -347,6 +363,12 @@ namespace UnrealBuildTool
 				PluginModuleContext.bClassifyAsGameModuleForUHT = !bContainsEngineModules;
 				FindModuleRulesForPlugins(ForeignPlugins, PluginModuleContext, ModuleFiles);
 
+				Logger.LogTrace(" Found {Count} Modules:", ModuleFiles.Count);
+				foreach (var Item in ModuleFiles)
+				{
+					Logger.LogTrace("  {Module}", Item.Key);
+				}
+
 				// Compile the assembly
 				FileReference AssemblyFileName = FileReference.Combine(PluginFileName.Directory, "Intermediate", "Build", "BuildRules", Path.GetFileNameWithoutExtension(PluginFileName.FullName) + "ModuleRules" + FrameworkAssemblyExtension);
 				PluginRulesAssembly = new RulesAssembly(Scope, new List<DirectoryReference> { PluginFileName.Directory }, ForeignPlugins, ModuleFiles, TargetFiles, AssemblyFileName, bContainsEngineModules, DefaultBuildSettings: null, bReadOnly: false, bSkipCompile: bSkipCompile, bForceCompile: bForceCompile, Parent: Parent, Logger: Logger);
