PR #8999: Fix detection of PVS-Studio installation (Contributed by [REDACTED_PERSON])

diff --git a/Engine/Source/Programs/UnrealBuildTool/Platform/Windows/PVSToolChain.cs b/Engine/Source/Programs/UnrealBuildTool/Platform/Windows/PVSToolChain.cs
--- a/Engine/Source/Programs/UnrealBuildTool/Platform/Windows/PVSToolChain.cs
+++ b/Engine/Source/Programs/UnrealBuildTool/Platform/Windows/PVSToolChain.cs
@@ -360,7 +360,7 @@
 			if (!FileReference.Exists(AnalyzerFile))
 			{
 				FileReference InstalledAnalyzerFile = FileReference.Combine(new DirectoryReference(Environment.GetFolderPath(Environment.SpecialFolder.ProgramFilesX86)), "PVS-Studio", "x64", "PVS-Studio.exe");
-				if (FileReference.Exists(AnalyzerFile))
+				if (FileReference.Exists(InstalledAnalyzerFile))
 				{
 					AnalyzerFile = InstalledAnalyzerFile;
 				}
