UnrealBuildTool: Fix include path using SN-DBS with clang

#8575
#rnx

diff --git a/Engine/Source/Programs/UnrealBuildTool/Executors/SNDBS.cs b/Engine/Source/Programs/UnrealBuildTool/Executors/SNDBS.cs
--- a/Engine/Source/Programs/UnrealBuildTool/Executors/SNDBS.cs
+++ b/Engine/Source/Programs/UnrealBuildTool/Executors/SNDBS.cs
@@ -517,6 +517,7 @@ family=clang-cl
 include_path01=..\include
 include_path02=..\include\c++\v1
 include_path03=..\lib\clang\*\include
+include_path04=..\usr\include
 
 [files]
 main=clang++.exe
