UnrealBuildTool: Support -Timing for additional clang platforms

#rnx

diff --git a/Engine/Source/Programs/UnrealBuildTool/Platform/IOS/IOSToolChain.cs b/Engine/Source/Programs/UnrealBuildTool/Platform/IOS/IOSToolChain.cs
--- a/Engine/Source/Programs/UnrealBuildTool/Platform/IOS/IOSToolChain.cs
+++ b/Engine/Source/Programs/UnrealBuildTool/Platform/IOS/IOSToolChain.cs
@@ -605,6 +605,14 @@ namespace UnrealBuildTool
 				// Add the source file path to the command-line.
 				FileArguments.Add(string.Format("\"{0}\"", NormalizeCommandLinePath(SourceFile)));
 
+				// Generate the timing info
+				if (CompileEnvironment.bPrintTimingInfo)
+				{
+					FileItem TraceFile = FileItem.GetItemByFileReference(FileReference.Combine(OutputDir, Path.GetFileName(SourceFile.AbsolutePath) + ".json"));
+					FileArguments.Add("-ftime-trace");
+					CompileAction.ProducedItems.Add(TraceFile);
+				}
+
 				// Generate the included header dependency list
 				if (CompileEnvironment.bGenerateDependenciesFile)
 				{
