UnrealBuildTool: Don't run dotnet through bash in the generated makefile

#rnx

diff --git a/Engine/Source/Programs/UnrealBuildTool/ProjectFiles/Make/MakefileGenerator.cs b/Engine/Source/Programs/UnrealBuildTool/ProjectFiles/Make/MakefileGenerator.cs
--- a/Engine/Source/Programs/UnrealBuildTool/ProjectFiles/Make/MakefileGenerator.cs
+++ b/Engine/Source/Programs/UnrealBuildTool/ProjectFiles/Make/MakefileGenerator.cs
@@ -69,7 +69,7 @@
 			{
 				GameProjectFile = OnlyGameProject!.FullName;
 				MakeGameProjectFile = "GAMEPROJECTFILE =" + GameProjectFile + "\n";
-				ProjectBuildCommand = $"PROJECTBUILD = bash \"$(UNREALROOTPATH)/Engine/{Unreal.RelativeDotnetDirectory}/dotnet\" \"$(UNREALROOTPATH)/Engine/Binaries/DotNET/UnrealBuildTool/UnrealBuildTool.dll\"\n";
+				ProjectBuildCommand = $"PROJECTBUILD = \"$(UNREALROOTPATH)/Engine/{Unreal.RelativeDotnetDirectory}/dotnet\" \"$(UNREALROOTPATH)/Engine/Binaries/DotNET/UnrealBuildTool/UnrealBuildTool.dll\"\n";
 			}
 
 			BuildCommand = "BUILD = bash \"$(UNREALROOTPATH)/Engine/Build/BatchFiles/Linux/Build.sh\"\n";
