Added .so and .dylib files as files to consider as library files to help reduce relinking on Mac and Linux.

diff --git a/Engine/Source/Programs/UnrealBuildTool/System/ActionGraph.cs b/Engine/Source/Programs/UnrealBuildTool/System/ActionGraph.cs
--- a/Engine/Source/Programs/UnrealBuildTool/System/ActionGraph.cs
+++ b/Engine/Source/Programs/UnrealBuildTool/System/ActionGraph.cs
@@ -775,7 +775,7 @@
 		{
 			if (PrerequisiteAction.bProducesImportLibrary)
 			{
-				return PrerequisiteAction.ProducedItems.All(I => I.Location.HasExtension(".lib") || !RootAction.PrerequisiteItems.Contains(I));
+				return PrerequisiteAction.ProducedItems.All(I => IsLibraryFile(I.Location) || !RootAction.PrerequisiteItems.Contains(I));
 			}
 			else
 			{
@@ -791,7 +791,7 @@
 		/// <returns>True if the only dependency between two actions is for an import library</returns>
 		static bool IsImportLibraryDependency(LinkedAction RootAction, FileItem PrerequisiteItem)
 		{
-			if (PrerequisiteItem.Location.HasExtension(".lib"))
+			if (IsLibraryFile(PrerequisiteItem.Location))
 			{
 				foreach (LinkedAction PrerequisiteAction in RootAction.PrerequisiteActions)
 				{
@@ -806,6 +806,16 @@
 		}
 
 		/// <summary>
+		/// Test to see if the given file is a library file (or in the case of linux/mac, a dynamic library)
+		/// </summary>
+		/// <param name="Location">File to test</param>
+		/// <returns>True if the file is a library file</returns>
+		static bool IsLibraryFile(FileReference Location)
+		{
+			return Location.HasExtension(".lib") || Location.HasExtension(".so") || Location.HasExtension(".dylib");
+		}
+
+		/// <summary>
 		/// Builds a dictionary containing the actions from AllActions that are outdated by calling
 		/// IsActionOutdated.
 		/// </summary>
