PR #9076: Fix invalid JSON keys in VSCode project generation (Contributed by [REDACTED_PERSON])

diff --git a/Engine/Source/Programs/UnrealBuildTool/ProjectFiles/VisualStudioCode/VSCodeProjectFileGenerator.cs b/Engine/Source/Programs/UnrealBuildTool/ProjectFiles/VisualStudioCode/VSCodeProjectFileGenerator.cs
--- a/Engine/Source/Programs/UnrealBuildTool/ProjectFiles/VisualStudioCode/VSCodeProjectFileGenerator.cs
+++ b/Engine/Source/Programs/UnrealBuildTool/ProjectFiles/VisualStudioCode/VSCodeProjectFileGenerator.cs
@@ -168,7 +168,7 @@
 
 			public void BeginObject(string? Name = null)
 			{
-				string Prefix = Name == null ? "" : Quoted(Name) + ": ";
+				string Prefix = Name == null ? "" : Quoted(JsonWriter.EscapeString(Name)) + ": ";
 				Lines.Add(TabString + Prefix + "{");
 				TabString += "\t";
 			}
@@ -182,7 +182,7 @@
 
 			public void BeginArray(string? Name = null)
 			{
-				string Prefix = Name == null ? "" : Quoted(Name) + ": ";
+				string Prefix = Name == null ? "" : Quoted(JsonWriter.EscapeString(Name)) + ": ";
 				Lines.Add(TabString + Prefix + "[");
 				TabString += "\t";
 			}
@@ -196,12 +196,12 @@
 
 			public void AddField(string Name, bool Value)
 			{
-				Lines.Add(TabString + Quoted(Name) + ": " + Value.ToString().ToLower() + ",");
+				Lines.Add(TabString + Quoted(JsonWriter.EscapeString(Name)) + ": " + Value.ToString().ToLower() + ",");
 			}
 
 			public void AddField(string Name, string Value)
 			{
-				Lines.Add(TabString + Quoted(Name) + ": " + Quoted(JsonWriter.EscapeString(Value)) + ",");
+				Lines.Add(TabString + Quoted(JsonWriter.EscapeString(Name)) + ": " + Quoted(JsonWriter.EscapeString(Value)) + ",");
 			}
 
 			public void AddUnnamedField(string Value)
