PR #9128: Added option to remove programs when generating project file (Contributed by [REDACTED_PERSON])

diff --git a/Engine/Source/Programs/UnrealBuildTool/ProjectFiles/ProjectFileGenerator.cs b/Engine/Source/Programs/UnrealBuildTool/ProjectFiles/ProjectFileGenerator.cs
--- a/Engine/Source/Programs/UnrealBuildTool/ProjectFiles/ProjectFileGenerator.cs
+++ b/Engine/Source/Programs/UnrealBuildTool/ProjectFiles/ProjectFileGenerator.cs
@@ -1304,6 +1304,11 @@ namespace UnrealBuildTool
 								bGatherThirdPartySource = true;
 								break;
 
+							case "-NOPROGRAMS":
+								bIncludeEnginePrograms = false;
+								bIncludeDotNetPrograms = false;
+								break;
+
 							case "-GAME":
 								// Generates project files for a single game
 								bIncludeDotNetPrograms = false;
