PR #9159: Fix small typo in UnrealPluginLanguage.cs documentation (Contributed by [REDACTED_PERSON])

diff --git a/Engine/Source/Programs/UnrealBuildTool/System/UnrealPluginLanguage.cs b/Engine/Source/Programs/UnrealBuildTool/System/UnrealPluginLanguage.cs
--- a/Engine/Source/Programs/UnrealBuildTool/System/UnrealPluginLanguage.cs
+++ b/Engine/Source/Programs/UnrealBuildTool/System/UnrealPluginLanguage.cs
@@ -438,7 +438,7 @@ namespace UnrealBuildTool
 	 * <log text=""/>
 	 * <insert> </insert>
 	 * <insertValue value=""/>
-	 * <replace find="" with""/>
+	 * <replace find="" with=""/>
 	 * <copyFile src="" dst=""/>
 	 * <copyDir src="" dst=""/>
 	 * <loadLibrary name="" failmsg=""/>
