PR #9262: Update valid range version maximum to allow for Xcode 14 beta. (Contributed by [REDACTED_PERSON])

#rnx

diff --git a/Engine/Source/Programs/UnrealBuildTool/Platform/Mac/ApplePlatformSDK.Versions.cs b/Engine/Source/Programs/UnrealBuildTool/Platform/Mac/ApplePlatformSDK.Versions.cs
--- a/Engine/Source/Programs/UnrealBuildTool/Platform/Mac/ApplePlatformSDK.Versions.cs
+++ b/Engine/Source/Programs/UnrealBuildTool/Platform/Mac/ApplePlatformSDK.Versions.cs
@@ -22,13 +22,13 @@ namespace UnrealBuildTool
 			if (RuntimePlatform.IsMac)
 			{
 				MinVersion = "12.4.0";
-				MaxVersion = "13.9.9";
+				MaxVersion = "14.9.9";
 			}
 			else
 			{
 				// @todo turnkey: these are MobileDevice .dll versions in Windows - to get the iTunes app version (12.3.4.1 etc) would need to hunt down the .exe
 				MinVersion = "1100.0.0.0";
-				MaxVersion = "1399.0";
+				MaxVersion = "1499.0";
 			}
 		}
 
@@ -37,4 +37,4 @@ namespace UnrealBuildTool
 			MinVersion = MaxVersion = null;
 		}
 	}
-}
\ No newline at end of file
+}
