Added the ability to compile specific files (i.e. -singlefile) while a live coding session is active.  Output is sent to a different location to prevent clashes with overriting object files for the live coding session.

diff --git a/Engine/Source/Programs/UnrealBuildTool/System/HotReload.cs b/Engine/Source/Programs/UnrealBuildTool/System/HotReload.cs
--- a/Engine/Source/Programs/UnrealBuildTool/System/HotReload.cs
+++ b/Engine/Source/Programs/UnrealBuildTool/System/HotReload.cs
@@ -24,7 +24,8 @@ namespace UnrealBuildTool
 		Disabled,
 		FromIDE,
 		FromEditor,
-		LiveCoding
+		LiveCoding,
+		LiveCodingPassThrough, // Special mode for specific file compiles but live coding is currently active
 	}
 
 	/// <summary>
@@ -138,10 +139,11 @@ namespace UnrealBuildTool
 		/// Getts the default hot reload mode for the given target
 		/// </summary>
 		/// <param name="TargetDescriptor">The target being built</param>
+		/// <param name="Makefile">Makefile for the target</param>
 		/// <param name="BuildConfiguration">Global build configuration</param>
 		/// <param name="Logger">Logger for output</param>
 		/// <returns>Default hotreload mode</returns>
-		public static HotReloadMode GetDefaultMode(TargetDescriptor TargetDescriptor, BuildConfiguration BuildConfiguration, ILogger Logger)
+		public static HotReloadMode GetDefaultMode(TargetDescriptor TargetDescriptor, TargetMakefile Makefile, BuildConfiguration BuildConfiguration, ILogger Logger)
 		{
 			if (TargetDescriptor.HotReloadModuleNameToSuffix.Count > 0 && TargetDescriptor.ForeignPlugin == null)
 			{
@@ -151,6 +153,11 @@ namespace UnrealBuildTool
 			{
 				return HotReloadMode.FromIDE;
 			}
+			else if (TargetDescriptor.SpecificFilesToCompile.Count > 0 && IsLiveCodingSessionActive(Makefile, Logger))
+			{
+				Logger.LogWarning("Live coding session active. Actions will be limited to compilation of specified files.  Output will be sent to a temporary location.");
+				return HotReloadMode.LiveCodingPassThrough;
+			}
 			else
 			{
 				return HotReloadMode.Disabled;
@@ -161,7 +168,7 @@ namespace UnrealBuildTool
 		/// Sets the appropriate hot reload mode for a target, and cleans up old state.
 		/// </summary>
 		/// <param name="TargetDescriptor">The target being built</param>
-		/// <param name="Makefile">Makefile for the targe</param>
+		/// <param name="Makefile">Makefile for the target</param>
 		/// <param name="Actions">Actions for this target</param>
 		/// <param name="BuildConfiguration">Global build configuration</param>
 		/// <param name="Logger">Logger for output</param>
@@ -170,7 +177,7 @@ namespace UnrealBuildTool
 			Dictionary<FileReference, FileReference>? PatchedOldLocationToNewLocation = null;
 
 			// Get the hot-reload mode
-			if (TargetDescriptor.HotReloadMode == HotReloadMode.LiveCoding)
+			if (TargetDescriptor.HotReloadMode == HotReloadMode.LiveCoding || TargetDescriptor.HotReloadMode == HotReloadMode.LiveCodingPassThrough)
 			{
 				// In some instances such as packaged builds, we might not have hot reload modules names.
 				// We don't want to lose the live coding setting in that case.
@@ -181,7 +188,7 @@ namespace UnrealBuildTool
 			}
 			else if (TargetDescriptor.HotReloadMode == HotReloadMode.Default)
 			{
-				TargetDescriptor.HotReloadMode = GetDefaultMode(TargetDescriptor, BuildConfiguration, Logger);
+				TargetDescriptor.HotReloadMode = GetDefaultMode(TargetDescriptor, Makefile, BuildConfiguration, Logger);
 			}
 
 			// Apply the previous hot reload state
@@ -217,7 +224,9 @@ namespace UnrealBuildTool
 		public static void CheckForLiveCodingSessionActive(TargetDescriptor TargetDescriptor, TargetMakefile Makefile, BuildConfiguration BuildConfiguration, ILogger Logger)
 		{
 			// Guard against a live coding session for this target being active
-			if (BuildConfiguration.bAllowHotReloadFromIDE && TargetDescriptor.HotReloadMode != HotReloadMode.LiveCoding && TargetDescriptor.ForeignPlugin == null && HotReload.IsLiveCodingSessionActive(Makefile, Logger))
+			if (BuildConfiguration.bAllowHotReloadFromIDE && TargetDescriptor.ForeignPlugin == null &&
+				TargetDescriptor.HotReloadMode != HotReloadMode.LiveCoding && TargetDescriptor.HotReloadMode != HotReloadMode.LiveCodingPassThrough && 
+				HotReload.IsLiveCodingSessionActive(Makefile, Logger))
 			{
 				throw new BuildException("Unable to build while Live Coding is active. Exit the editor and game, or press Ctrl+Alt+F11 if iterating on code in the editor or game");
 			}
@@ -545,7 +554,7 @@ namespace UnrealBuildTool
 				History.Mount(TargetDescriptor.ProjectFile.Directory);
 			}
 
-			if (TargetDescriptor.HotReloadMode == HotReloadMode.LiveCoding)
+			if (TargetDescriptor.HotReloadMode == HotReloadMode.LiveCoding || TargetDescriptor.HotReloadMode == HotReloadMode.LiveCodingPassThrough)
 			{
 				CompilationResult Result = CompilationResult.Succeeded;
 
@@ -597,7 +606,7 @@ namespace UnrealBuildTool
 
 				// Update the action graph with these new paths
 				Dictionary<FileReference, FileReference> OriginalFileToPatchedFile = new Dictionary<FileReference, FileReference>();
-				HotReload.PatchActionGraphForLiveCoding(PrerequisiteActions, OriginalFileToPatchedFile, Logger);
+				HotReload.PatchActionGraphForLiveCoding(PrerequisiteActions, OriginalFileToPatchedFile, TargetDescriptor.HotReloadMode, Logger);
 
 				// Get a new list of actions to execute now that the graph has been modified
 				TargetActionsToExecute = ActionGraph.GetActionsToExecute(PrerequisiteActions, CppDependencies, History, BuildConfiguration.bIgnoreOutdatedImportLibraries, Logger);
@@ -821,9 +830,16 @@ namespace UnrealBuildTool
 		/// </summary>
 		/// <param name="Actions">Set of actions</param>
 		/// <param name="OriginalFileToPatchedFile">Dictionary that receives a map of original object file to patched object file</param>
+		/// <param name="hotReloadMode">Requested hot reload mode</param>
 		/// <param name="Logger"></param>
-		public static void PatchActionGraphForLiveCoding(IEnumerable<LinkedAction> Actions, Dictionary<FileReference, FileReference> OriginalFileToPatchedFile, ILogger Logger)
+		public static void PatchActionGraphForLiveCoding(IEnumerable<LinkedAction> Actions, Dictionary<FileReference, FileReference> OriginalFileToPatchedFile, HotReloadMode hotReloadMode, ILogger Logger)
 		{
+			string dependencyFileExtension = hotReloadMode == HotReloadMode.LiveCoding ? ".lc.response" : ".lcpt.response";
+			string responseFileExtension = hotReloadMode == HotReloadMode.LiveCoding ? ".lc" : ".lcpt";
+			string objectFileExtension = hotReloadMode == HotReloadMode.LiveCoding ? ".lc.obj" : ".lcpt.obj";
+			string clSourceDepFileExtension = hotReloadMode == HotReloadMode.LiveCoding ? ".lc.json" : ".lcpt.json";
+			string clangSourceDepFileExtension = hotReloadMode == HotReloadMode.LiveCoding ? ".lc.d" : ".lcpt.d";
+
 			foreach (LinkedAction Action in Actions)
 			{
 				if(Action.ActionType == ActionType.Compile)
@@ -870,7 +886,7 @@ namespace UnrealBuildTool
 						FileItem OldDependenciesFileItem = Action.ProducedItems.First(x => x.Location == OldDependenciesFile);
 						NewAction.ProducedItems.Remove(OldDependenciesFileItem);
 
-						FileReference NewDependenciesFile = OldDependenciesFile.ChangeExtension(".lc.response");
+						FileReference NewDependenciesFile = OldDependenciesFile.ChangeExtension(dependencyFileExtension);
 						FileItem NewDependenciesFileItem = FileItem.GetItemByFileReference(NewDependenciesFile);
 						NewAction.ProducedItems.Add(NewDependenciesFileItem);
 						NewAction.DependencyListFile = NewDependenciesFileItem;
@@ -892,7 +908,7 @@ namespace UnrealBuildTool
 					}
 
 					FileReference OldResponseFile = new FileReference(Arguments[ResponseFileIdx].Substring(1).Trim('\"'));
-					FileReference NewResponseFile = new FileReference(OldResponseFile.FullName + ".lc");
+					FileReference NewResponseFile = new FileReference(OldResponseFile.FullName + responseFileExtension);
 
 					NewAction.PrerequisiteItems.Remove(FileItem.GetItemByFileReference(OldResponseFile));
 					NewAction.PrerequisiteItems.Add(FileItem.GetItemByFileReference(NewResponseFile));
@@ -909,7 +925,7 @@ namespace UnrealBuildTool
 							FileItem OldOutputFileItem = Action.ProducedItems.First(x => x.Location == OldOutputFile);
 							NewAction.ProducedItems.Remove(OldOutputFileItem);
 
-							FileReference NewOutputFile = OldOutputFile.ChangeExtension(".lc.obj");
+							FileReference NewOutputFile = OldOutputFile.ChangeExtension(objectFileExtension);
 							FileItem NewOutputFileItem = FileItem.GetItemByFileReference(NewOutputFile);
 							NewAction.ProducedItems.Add(NewOutputFileItem);
 
@@ -925,7 +941,7 @@ namespace UnrealBuildTool
 						Action.CommandPath.GetFileName().Equals("clang-cl.exe", StringComparison.OrdinalIgnoreCase))
 					{
 						string SourceDependencyPrefix = Action.CommandPath.GetFileName().Equals("cl.exe", StringComparison.OrdinalIgnoreCase) ? "/sourceDependencies" : "/clang:-MD /clang:-MF";
-						string NewExtension = Action.CommandPath.GetFileName().Equals("cl.exe", StringComparison.OrdinalIgnoreCase) ? ".lc.json" : ".lc.d";
+						string NewExtension = Action.CommandPath.GetFileName().Equals("cl.exe", StringComparison.OrdinalIgnoreCase) ? clSourceDepFileExtension : clangSourceDepFileExtension;
 						for (int Idx = 0; Idx < ResponseLines.Length; Idx++)
 						{
 							string ResponseLine = ResponseLines[Idx];
