Inject LLDB startup command that loads data types visualizer script for Android
#[review]-20669220 Chris.Babcock

diff --git a/Engine/Source/Programs/UnrealBuildTool/Platform/Android/AndroidProjectGenerator.cs b/Engine/Source/Programs/UnrealBuildTool/Platform/Android/AndroidProjectGenerator.cs
--- a/Engine/Source/Programs/UnrealBuildTool/Platform/Android/AndroidProjectGenerator.cs
+++ b/Engine/Source/Programs/UnrealBuildTool/Platform/Android/AndroidProjectGenerator.cs
@@ -7,6 +7,7 @@
 using System.IO;
 using EpicGames.Core;
 using Microsoft.Extensions.Logging;
+using UnrealBuildBase;
 
 namespace UnrealBuildTool
 {
@@ -746,6 +747,22 @@
 			// do not need to check InPlatform since it will always be UnrealTargetPlatform.Android
 			return (AGDEInstalled ? " -Architectures=arm64 -ForceAPKGeneration" : "") + base.GetExtraBuildArguments(InPlatform, InConfiguration);
 		}
+
+		public override string GetVisualStudioUserFileStrings(UnrealTargetPlatform InPlatform, UnrealTargetConfiguration InConfiguration,
+			string InConditionString, TargetRules InTargetRules, FileReference TargetRulesPath, FileReference ProjectFilePath)
+		{
+			if (AGDEInstalled 
+				&& (InPlatform == UnrealTargetPlatform.Android)
+				&& ((InTargetRules.Type == TargetRules.TargetType.Client) || (InTargetRules.Type == TargetRules.TargetType.Game)))
+			{
+				string UserFileEntry = "<PropertyGroup " + InConditionString + ">\n";
+				UserFileEntry		+= "	<AndroidLldbStartupCommands>command script import \"" + Path.Combine(Unreal.EngineDirectory.FullName, "Extras", "LLDBDataFormatters", "UEDataFormatters_2ByteChars.py") + "\"</AndroidLldbStartupCommands>\n";
+				UserFileEntry		+= "</PropertyGroup>\n";
+				return UserFileEntry;
+			}
+
+			return base.GetVisualStudioUserFileStrings(InPlatform, InConfiguration, InConditionString, InTargetRules, TargetRulesPath, ProjectFilePath);
+		}
 	}
 
 
