* Added debug pch info.

diff --git a/Engine/Source/Programs/UnrealBuildTool/Configuration/UEBuildModuleCPP.cs b/Engine/Source/Programs/UnrealBuildTool/Configuration/UEBuildModuleCPP.cs
--- a/Engine/Source/Programs/UnrealBuildTool/Configuration/UEBuildModuleCPP.cs
+++ b/Engine/Source/Programs/UnrealBuildTool/Configuration/UEBuildModuleCPP.cs
@@ -519,6 +519,10 @@
 
 			// Configure the precompiled headers for this module
 			CompileEnvironment = SetupPrecompiledHeaders(Target, ToolChain, CompileEnvironment, LinkInputFiles, Graph);
+			if (CompileEnvironment.PrecompiledHeaderFile != null)
+			{
+				Logger.LogDebug("Module '{ModuleName}' uses PCH '{PCHIncludeFilename}'", this.Name, CompileEnvironment.PrecompiledHeaderFile);
+			}
 
 			// Write all the definitions to a separate file
 			CreateHeaderForDefinitions(CompileEnvironment, IntermediateDirectory, null, Graph);
