UBT: Fix UHT output not being run through log parser.

diff --git a/Engine/Source/Programs/UnrealBuildTool/System/ExternalExecution.cs b/Engine/Source/Programs/UnrealBuildTool/System/ExternalExecution.cs
--- a/Engine/Source/Programs/UnrealBuildTool/System/ExternalExecution.cs
+++ b/Engine/Source/Programs/UnrealBuildTool/System/ExternalExecution.cs
@@ -920,7 +920,7 @@
 			// print anything for that event.
 			if (!String.IsNullOrEmpty(Event.Data))
 			{
-				Logger.LogInformation("{Output}", Event.Data);
+				Log.TraceInformation("{0}", Event.Data); // Output must be routed through event matchers; can't use Logger directly.
 			}
 		}
 
