UnrealBuildTool: Option to aggregate clang timing data and write csv to Intermediate directory

#rnx

diff --git a/Engine/Source/Programs/UnrealBuildTool/Modes/AggregateClangTimingInfo.cs b/Engine/Source/Programs/UnrealBuildTool/Modes/AggregateClangTimingInfo.cs
new file mode 100644
--- /dev/null
+++ b/Engine/Source/Programs/UnrealBuildTool/Modes/AggregateClangTimingInfo.cs
@@ -0,0 +1,144 @@
+// Copyright Epic Games, Inc. All Rights Reserved.
+
+using EpicGames.Core;
+using Microsoft.Extensions.Logging;
+using System;
+using System.Collections.Concurrent;
+using System.Collections.Generic;
+using System.IO;
+using System.IO.Compression;
+using System.Linq;
+using System.Reflection;
+using System.Text.Json;
+using System.Threading.Tasks;
+
+namespace UnrealBuildTool
+{
+	/// <summary>
+	/// Aggregates clang timing information files together into one monolithic breakdown file.
+	/// </summary>
+	[ToolMode("AggregateClangTimingInfo", ToolModeOptions.None)]
+	class AggregateClangTimingInfo : ToolMode
+	{
+		public class ClangTrace
+		{
+			public class TraceEvent
+			{
+				public long pid { get; set; }
+				public long tid { get; set; }
+				public string? ph { get; set; }
+				public long ts { get; set; }
+				public long dur { get; set; }
+				public string? name { get; set; }
+				public Dictionary<string, object>? args { get; set; }
+			}
+
+			public List<TraceEvent>? traceEvents { get; set; }
+			public long beginningOfTime { get; set; }
+		}
+
+		public class TraceData
+		{
+			FileReference SourceFile { get; init; }
+			public string Module => SourceFile.Directory.GetDirectoryName();
+			public string Name => SourceFile.GetFileName();
+			public long TotalExecuteCompilerMs { get; init; }
+			public long TotalFrontendMs { get; init; }
+			public long TotalBackendMs { get; init; }
+			public long TotalSourceMs { get; init; }
+			public long SourceEntries { get; init; }
+			public long ObjectBytes { get; init; }
+			public long DependencyIncludes { get; init; }
+
+			public TraceData(FileReference inputFile, ClangTrace? trace)
+			{
+				SourceFile = inputFile;
+				TotalExecuteCompilerMs = trace?.traceEvents?.FindLast(x => string.Equals(x.name, "Total ExecuteCompiler"))?.dur ?? 0;
+				TotalFrontendMs = trace?.traceEvents?.FindLast(x => string.Equals(x.name, "Total Frontend"))?.dur ?? 0;
+				TotalBackendMs = trace?.traceEvents?.FindLast(x => string.Equals(x.name, "Total Backend"))?.dur ?? 0;
+				TotalSourceMs = trace?.traceEvents?.FindLast(x => string.Equals(x.name, "Total Source"))?.dur ?? 0;
+				SourceEntries = trace?.traceEvents?.Where(x => string.Equals(x.name, "Source")).LongCount() ?? 0;
+				ObjectBytes = GetObjectSize();
+				DependencyIncludes = CountIncludes();
+			}
+
+			public static string CsvHeader => "Module,Name,TotalExecuteCompilerMs,TotalFrontendMs,TotalBackendMs,TotalSourceMs,SourceEntries,ObjectBytes,DependencyIncludes";
+			public string CsvLine => $"{Module},{Name},{TotalExecuteCompilerMs},{TotalFrontendMs},{TotalBackendMs},{TotalSourceMs},{SourceEntries},{ObjectBytes},{DependencyIncludes}";
+
+			private long GetObjectSize()
+			{
+				FileReference ObjectFile = new FileReference($"{SourceFile.FullName}.{(SourceFile.HasExtension(".h") ? "gch" : "o")}");
+				if (!FileReference.Exists(ObjectFile))
+				{
+					return 0;
+				}
+				return ObjectFile.ToFileInfo().Length;
+			}
+
+			private long CountIncludes()
+			{
+				FileReference DependsFile = new FileReference($"{SourceFile.FullName}.d");
+				if (!FileReference.Exists(DependsFile))
+				{
+					return 0;
+				}
+				// Subtract 1 for the header line
+				return Math.Max(0, File.ReadLines(DependsFile.FullName).Count() - 1);
+			}
+		}
+
+		public override int Execute(CommandLineArguments Arguments, ILogger Logger)
+		{
+			FileReference ManifestFile = Arguments.GetFileReference("-ManifestFile=");
+			IEnumerable<FileReference> SourceFiles = FileReference.ReadAllLines(ManifestFile).Select(x => new FileReference(x));
+
+			// Create aggregate summary.
+			FileReference? AggregateFile = Arguments.GetFileReferenceOrDefault("-AggregateFile=", null);
+			if (AggregateFile != null)
+			{
+				var Tasks = Task.WhenAll(SourceFiles.Select(x => ParseTimingDataFile(x, Logger)));
+				Tasks.Wait();
+
+				string TempFilePath = Path.Join(Path.GetTempPath(), AggregateFile.GetFileName() + ".tmp");
+				using (StreamWriter Writer = new StreamWriter(TempFilePath))
+				{
+					Writer.WriteLine(TraceData.CsvHeader);
+					foreach (TraceData Data in Tasks.Result.OrderBy(x => x.Module).ThenBy(x => x.Name))
+					{
+						Writer.WriteLine(Data.CsvLine);
+					}
+				}
+				File.Move(TempFilePath, AggregateFile.FullName, true);
+			}
+
+			// Write out aggregate archive if requested.
+			FileReference? ArchiveFile = Arguments.GetFileReferenceOrDefault("-ArchiveFile=", null);
+			if (ArchiveFile != null)
+			{
+				Logger.LogDebug("Writing {OutputFile} Archive", ArchiveFile);
+				string TempFilePath = Path.Join(Path.GetTempPath(), ArchiveFile.GetFileName() + ".tmp");
+				using (ZipArchive ZipArchive = new ZipArchive(File.Open(TempFilePath, FileMode.Create), ZipArchiveMode.Create))
+				{
+					foreach (FileReference SourceFile in SourceFiles)
+					{
+						FileReference JsonFile = new FileReference($"{SourceFile.FullName}.json");
+						string EntryName = $"{JsonFile.Directory.GetDirectoryName()}/{JsonFile.GetFileName()}";
+						ZipArchive.CreateEntryFromFile_CrossPlatform(JsonFile.FullName, EntryName, CompressionLevel.Optimal);
+					}
+				}
+				File.Move(TempFilePath, ArchiveFile.FullName, true);
+			}
+
+			return 0;
+		}
+
+		private static async Task<TraceData> ParseTimingDataFile(FileReference SourceFile, ILogger Logger)
+		{
+			Logger.LogDebug("Parsing {SourceFile}", SourceFile.FullName);
+			FileReference JsonFile = new FileReference($"{SourceFile.FullName}.json");
+
+			ClangTrace? Trace = await JsonSerializer.DeserializeAsync<ClangTrace>(File.OpenRead(JsonFile.FullName));
+			return new TraceData(SourceFile, Trace);
+		}
+	}
+}
diff --git a/Engine/Source/Programs/UnrealBuildTool/Platform/Mac/MacToolChain.cs b/Engine/Source/Programs/UnrealBuildTool/Platform/Mac/MacToolChain.cs
--- a/Engine/Source/Programs/UnrealBuildTool/Platform/Mac/MacToolChain.cs
+++ b/Engine/Source/Programs/UnrealBuildTool/Platform/Mac/MacToolChain.cs
@@ -1320,6 +1320,8 @@ namespace UnrealBuildTool
 
 		public override void FinalizeOutput(ReadOnlyTargetRules Target, TargetMakefileBuilder MakefileBuilder)
 		{
+			base.FinalizeOutput(Target, MakefileBuilder);
+
 			TargetMakefile Makefile = MakefileBuilder.Makefile;
 
 			// Re-add any .dSYM files that may have been stripped out.
diff --git a/Engine/Source/Programs/UnrealBuildTool/ToolChain/ClangToolChain.cs b/Engine/Source/Programs/UnrealBuildTool/ToolChain/ClangToolChain.cs
--- a/Engine/Source/Programs/UnrealBuildTool/ToolChain/ClangToolChain.cs
+++ b/Engine/Source/Programs/UnrealBuildTool/ToolChain/ClangToolChain.cs
@@ -120,6 +120,45 @@ namespace UnrealBuildTool
 			Options = InOptions;
 		}
 
+		public override void FinalizeOutput(ReadOnlyTargetRules Target, TargetMakefileBuilder MakefileBuilder)
+		{
+			if (Target.bPrintToolChainTimingInfo && Target.bParseTimingInfoForTracing)
+			{
+				TargetMakefile Makefile = MakefileBuilder.Makefile;
+				List<IExternalAction> CompileActions = Makefile.Actions.Where(x => x.ActionType == ActionType.Compile && x.ProducedItems.Any(i => i.HasExtension(".json"))).ToList();
+				List<FileItem> TimingJsonFiles = CompileActions.SelectMany(a => a.ProducedItems.Where(i => i.HasExtension(".json"))).ToList();
+
+				// Handing generating aggregate timing information if we compiled more than one file.
+				if (TimingJsonFiles.Count > 0)
+				{
+					// Generate the file manifest for the aggregator.
+					FileReference ManifestFile = FileReference.Combine(Makefile.ProjectIntermediateDirectory, $"{Target.Name}TimingManifest.csv");
+					if (!DirectoryReference.Exists(ManifestFile.Directory))
+					{
+						DirectoryReference.CreateDirectory(ManifestFile.Directory);
+					}
+					File.WriteAllLines(ManifestFile.FullName, TimingJsonFiles.Select(f => f.FullName.Remove(f.FullName.Length - ".json".Length)));
+
+					FileItem AggregateOutputFile = FileItem.GetItemByFileReference(FileReference.Combine(Makefile.ProjectIntermediateDirectory, $"{Target.Name}.trace.csv"));
+					List<string> ActionArgs = new List<string>()
+					{
+						$"-AggregateFile={AggregateOutputFile.FullName}",
+						$"-ManifestFile={ManifestFile.FullName}",
+					};
+
+					Action AggregateTimingInfoAction = MakefileBuilder.CreateRecursiveAction<AggregateClangTimingInfo>(ActionType.ParseTimingInfo, string.Join(" ", ActionArgs));
+					AggregateTimingInfoAction.WorkingDirectory = Unreal.EngineSourceDirectory;
+					AggregateTimingInfoAction.StatusDescription = $"Aggregating {TimingJsonFiles.Count} Timing File(s)";
+					AggregateTimingInfoAction.bCanExecuteRemotely = false;
+					AggregateTimingInfoAction.bCanExecuteRemotelyWithSNDBS = false;
+					AggregateTimingInfoAction.PrerequisiteItems.AddRange(TimingJsonFiles);
+
+					AggregateTimingInfoAction.ProducedItems.Add(AggregateOutputFile);
+					Makefile.OutputItems.Add(AggregateOutputFile);
+				}
+			}
+		}
+
 		/// <summary>
 		/// Sanitizes a preprocessor definition argument if needed.
 		/// </summary>
