UnrealBuildTool: Clang -Tracing also output csv of all headers parse time

#rnx

diff --git a/Engine/Source/Programs/UnrealBuildTool/Modes/AggregateClangTimingInfo.cs b/Engine/Source/Programs/UnrealBuildTool/Modes/AggregateClangTimingInfo.cs
--- a/Engine/Source/Programs/UnrealBuildTool/Modes/AggregateClangTimingInfo.cs
+++ b/Engine/Source/Programs/UnrealBuildTool/Modes/AggregateClangTimingInfo.cs
@@ -87,6 +87,8 @@ namespace UnrealBuildTool
 			}
 		}
 
+		private ConcurrentDictionary<FileReference, ClangTrace> ClangTraceCache = new();
+
 		public override int Execute(CommandLineArguments Arguments, ILogger Logger)
 		{
 			FileReference ManifestFile = Arguments.GetFileReference("-ManifestFile=");
@@ -94,16 +96,19 @@ namespace UnrealBuildTool
 
 			// Create aggregate summary.
 			FileReference? AggregateFile = Arguments.GetFileReferenceOrDefault("-AggregateFile=", null);
+			FileReference? HeadersFile = Arguments.GetFileReferenceOrDefault("-HeadersFile=", null);
+
 			if (AggregateFile != null)
 			{
-				var Tasks = Task.WhenAll(SourceFiles.Select(x => ParseTimingDataFile(x, Logger)));
+				var Tasks = Task.WhenAll(SourceFiles.Select(x => GetTraceData(x, Logger)));
 				Tasks.Wait();
+				List<TraceData> TraceDatas = Tasks.Result.OrderBy(x => x.Module).ThenBy(x => x.Name).ToList();
 
 				string TempFilePath = Path.Join(Path.GetTempPath(), AggregateFile.GetFileName() + ".tmp");
 				using (StreamWriter Writer = new StreamWriter(TempFilePath))
 				{
 					Writer.WriteLine(TraceData.CsvHeader);
-					foreach (TraceData Data in Tasks.Result.OrderBy(x => x.Module).ThenBy(x => x.Name))
+					foreach (TraceData Data in TraceDatas)
 					{
 						Writer.WriteLine(Data.CsvLine);
 					}
@@ -111,6 +116,39 @@ namespace UnrealBuildTool
 				File.Move(TempFilePath, AggregateFile.FullName, true);
 			}
 
+			if (HeadersFile != null)
+			{
+				var Tasks = Task.WhenAll(SourceFiles.Select(x => ParseTimingDataFile(x, Logger)));
+				Tasks.Wait();
+				List<ClangTrace> ClangTraces = Tasks.Result.ToList();
+
+				Dictionary<FileReference, List<long>> Sources = new Dictionary<FileReference, List<long>>();
+				foreach (ClangTrace ClangTrace in ClangTraces.Where(x => x.traceEvents != null))
+				{
+					foreach (ClangTrace.TraceEvent Event in ClangTrace.traceEvents!.Where(x => string.Equals(x.name, "Source") && x.args?.ContainsKey("detail") == true))
+					{
+						FileReference SourceFile = new FileReference(Event.args!["detail"].ToString()!);
+						if (!Sources.ContainsKey(SourceFile))
+						{
+							Sources.Add(SourceFile, new List<long>());
+						}
+
+						Sources[SourceFile].Add(Event.dur);
+					}
+				}
+
+				string TempFilePath = Path.Join(Path.GetTempPath(), HeadersFile.GetFileName() + ".tmp");
+				using (StreamWriter Writer = new StreamWriter(TempFilePath))
+				{
+					Writer.WriteLine("Source,Count,TotalMs,MinMs,MaxMs,AverageMs");
+					foreach (var Data in Sources.OrderBy(x => x.Key.FullName).Where(x => x.Key.HasExtension(".h")))
+					{
+						Writer.WriteLine($"{Data.Key.FullName},{Data.Value.Count},{Data.Value.Sum()},{Data.Value.Min()},{Data.Value.Max()},{Data.Value.Average()}");
+					}
+				}
+				File.Move(TempFilePath, HeadersFile.FullName, true);
+			}
+
 			// Write out aggregate archive if requested.
 			FileReference? ArchiveFile = Arguments.GetFileReferenceOrDefault("-ArchiveFile=", null);
 			if (ArchiveFile != null)
@@ -132,12 +170,27 @@ namespace UnrealBuildTool
 			return 0;
 		}
 
-		private static async Task<TraceData> ParseTimingDataFile(FileReference SourceFile, ILogger Logger)
+
+		private async Task<ClangTrace> ParseTimingDataFile(FileReference SourceFile, ILogger Logger)
 		{
-			Logger.LogDebug("Parsing {SourceFile}", SourceFile.FullName);
-			FileReference JsonFile = new FileReference($"{SourceFile.FullName}.json");
+			if (!ClangTraceCache.ContainsKey(SourceFile))
+			{
+				Logger.LogDebug("Parsing {SourceFile}", SourceFile.FullName);
+				FileReference JsonFile = new FileReference($"{SourceFile.FullName}.json");
+				ClangTrace? Trace = await JsonSerializer.DeserializeAsync<ClangTrace>(File.OpenRead(JsonFile.FullName));
+				if (Trace == null)
+				{
+					throw new NullReferenceException($"Unable to deserialize {JsonFile}");
+				}
+				ClangTraceCache[SourceFile] = Trace;
+			}
 
-			ClangTrace? Trace = await JsonSerializer.DeserializeAsync<ClangTrace>(File.OpenRead(JsonFile.FullName));
+			return ClangTraceCache[SourceFile];
+		}
+
+		private async Task<TraceData> GetTraceData(FileReference SourceFile, ILogger Logger)
+		{
+			ClangTrace Trace = await ParseTimingDataFile(SourceFile, Logger);
 			return new TraceData(SourceFile, Trace);
 		}
 	}
diff --git a/Engine/Source/Programs/UnrealBuildTool/ToolChain/ClangToolChain.cs b/Engine/Source/Programs/UnrealBuildTool/ToolChain/ClangToolChain.cs
--- a/Engine/Source/Programs/UnrealBuildTool/ToolChain/ClangToolChain.cs
+++ b/Engine/Source/Programs/UnrealBuildTool/ToolChain/ClangToolChain.cs
@@ -140,10 +140,12 @@ namespace UnrealBuildTool
 					File.WriteAllLines(ManifestFile.FullName, TimingJsonFiles.Select(f => f.FullName.Remove(f.FullName.Length - ".json".Length)));
 
 					FileItem AggregateOutputFile = FileItem.GetItemByFileReference(FileReference.Combine(Makefile.ProjectIntermediateDirectory, $"{Target.Name}.trace.csv"));
+					FileItem HeadersOutputFile = FileItem.GetItemByFileReference(FileReference.Combine(Makefile.ProjectIntermediateDirectory, $"{Target.Name}.headers.csv"));
 					List<string> ActionArgs = new List<string>()
 					{
 						$"-AggregateFile={AggregateOutputFile.FullName}",
 						$"-ManifestFile={ManifestFile.FullName}",
+						$"-HeadersFile={HeadersOutputFile.FullName}",
 					};
 
 					Action AggregateTimingInfoAction = MakefileBuilder.CreateRecursiveAction<AggregateClangTimingInfo>(ActionType.ParseTimingInfo, string.Join(" ", ActionArgs));
@@ -154,7 +156,9 @@ namespace UnrealBuildTool
 					AggregateTimingInfoAction.PrerequisiteItems.AddRange(TimingJsonFiles);
 
 					AggregateTimingInfoAction.ProducedItems.Add(AggregateOutputFile);
+					AggregateTimingInfoAction.ProducedItems.Add(HeadersOutputFile);
 					Makefile.OutputItems.Add(AggregateOutputFile);
+					Makefile.OutputItems.Add(HeadersOutputFile);
 				}
 			}
 		}
