UnrealBuildTool: Add additional data to clang trace csv

#rnx

diff --git a/Engine/Source/Programs/UnrealBuildTool/Modes/AggregateClangTimingInfo.cs b/Engine/Source/Programs/UnrealBuildTool/Modes/AggregateClangTimingInfo.cs
--- a/Engine/Source/Programs/UnrealBuildTool/Modes/AggregateClangTimingInfo.cs
+++ b/Engine/Source/Programs/UnrealBuildTool/Modes/AggregateClangTimingInfo.cs
@@ -42,28 +42,94 @@ namespace UnrealBuildTool
 			FileReference SourceFile { get; init; }
 			public string Module => SourceFile.Directory.GetDirectoryName();
 			public string Name => SourceFile.GetFileName();
+
 			public long TotalExecuteCompilerMs { get; init; }
+
+			// Subset of execute compiler
 			public long TotalFrontendMs { get; init; }
 			public long TotalBackendMs { get; init; }
+
+			// Subset of frontend
 			public long TotalSourceMs { get; init; }
+			public long TotalInstantiateFunctionMs { get; init; }
+			public long TotalCodeGenFunctionMs { get; init; }
+
+			// Subset of backend
+			public long TotalModuleToFunctionPassAdaptorMs { get; init; }
+			public long TotalModuleInlinerWrapperPassMs { get; init; }
+			public long TotalOptModuleMs { get; init; }
+
+			// Frontend entry counts
 			public long SourceEntries { get; init; }
+			public long InstantiateFunctionEntries { get; init; }
+			public long CodeGenFunctionEntries { get; init; }
+
+			// Other
 			public long ObjectBytes { get; init; }
 			public long DependencyIncludes { get; init; }
 
 			public TraceData(FileReference inputFile, ClangTrace? trace)
 			{
 				SourceFile = inputFile;
+
 				TotalExecuteCompilerMs = trace?.traceEvents?.FindLast(x => string.Equals(x.name, "Total ExecuteCompiler"))?.dur ?? 0;
+
+				// Subset of execute compiler
 				TotalFrontendMs = trace?.traceEvents?.FindLast(x => string.Equals(x.name, "Total Frontend"))?.dur ?? 0;
 				TotalBackendMs = trace?.traceEvents?.FindLast(x => string.Equals(x.name, "Total Backend"))?.dur ?? 0;
+
+				// Subset of frontend
 				TotalSourceMs = trace?.traceEvents?.FindLast(x => string.Equals(x.name, "Total Source"))?.dur ?? 0;
+				TotalInstantiateFunctionMs = trace?.traceEvents?.FindLast(x => string.Equals(x.name, "Total InstantiateFunction"))?.dur ?? 0;
+				TotalCodeGenFunctionMs = trace?.traceEvents?.FindLast(x => string.Equals(x.name, "Total CodeGen Function"))?.dur ?? 0;
+
+				// Subset of backend
+				TotalModuleToFunctionPassAdaptorMs = trace?.traceEvents?.FindLast(x => string.Equals(x.name, "Total ModuleToFunctionPassAdaptor"))?.dur ?? 0;
+				TotalModuleInlinerWrapperPassMs = trace?.traceEvents?.FindLast(x => string.Equals(x.name, "Total ModuleInlinerWrapperPass"))?.dur ?? 0;
+				TotalOptModuleMs = trace?.traceEvents?.FindLast(x => string.Equals(x.name, "Total OptModule"))?.dur ?? 0;
+
+				// Frontend entry counts
 				SourceEntries = trace?.traceEvents?.Where(x => string.Equals(x.name, "Source")).LongCount() ?? 0;
+				InstantiateFunctionEntries = trace?.traceEvents?.Where(x => string.Equals(x.name, "InstantiateFunction")).LongCount() ?? 0;
+				CodeGenFunctionEntries = trace?.traceEvents?.Where(x => string.Equals(x.name, "CodeGen Function")).LongCount() ?? 0;
+
+				// Other
 				ObjectBytes = GetObjectSize();
 				DependencyIncludes = CountIncludes();
 			}
 
-			public static string CsvHeader => "Module,Name,TotalExecuteCompilerMs,TotalFrontendMs,TotalBackendMs,TotalSourceMs,SourceEntries,ObjectBytes,DependencyIncludes";
-			public string CsvLine => $"{Module},{Name},{TotalExecuteCompilerMs},{TotalFrontendMs},{TotalBackendMs},{TotalSourceMs},{SourceEntries},{ObjectBytes},{DependencyIncludes}";
+			static readonly string[] CsvColumns = new string[] {
+				"Module",
+				"Name",
+
+				"TotalExecuteCompilerMs",
+
+				// Subset of execute compiler
+				"TotalFrontendMs",
+				"TotalBackendMs",
+
+				// Subset of frontend
+				"TotalSourceMs",
+				"TotalInstantiateFunctionMs",
+				"TotalCodeGenFunctionMs",
+
+				// Subset of backend
+				"TotalModuleToFunctionPassAdaptorMs",
+				"TotalModuleInlinerWrapperPassMs",
+				"TotalOptModuleMs",
+
+				// Frontend entry counts
+				"SourceEntries",
+				"InstantiateFunctionEntries",
+				"CodeGenFunctionEntries",
+
+				// Other
+				"ObjectBytes",
+				"DependencyIncludes",
+			};
+
+			public static string CsvHeader => string.Join(',', CsvColumns);
+			public string CsvLine => string.Join(',', CsvColumns.Select(x => GetType().GetProperty(x)!.GetValue(this)!.ToString()));
 
 			private long GetObjectSize()
 			{
