UnrealBuildTool: Remove setting -fcrash-diagnostics-dir via environment variable

#rnx

diff --git a/Engine/Source/Programs/UnrealBuildTool/ToolChain/ClangToolChain.cs b/Engine/Source/Programs/UnrealBuildTool/ToolChain/ClangToolChain.cs
--- a/Engine/Source/Programs/UnrealBuildTool/ToolChain/ClangToolChain.cs
+++ b/Engine/Source/Programs/UnrealBuildTool/ToolChain/ClangToolChain.cs
@@ -411,6 +411,21 @@ namespace UnrealBuildTool
 			{
 				Arguments.Add("-fdiagnostics-format=msvc");
 			}
+
+			// Set the output directory for crashes
+			// https://clang.llvm.org/docs/ClangCommandLineReference.html#cmdoption-clang-fcrash-diagnostics-dir
+			DirectoryReference? CrashDiagnosticDirectory = DirectoryReference.FromString(CompileEnvironment.CrashDiagnosticDirectory);
+			if (CrashDiagnosticDirectory != null)
+			{
+				if (DirectoryReference.Exists(CrashDiagnosticDirectory))
+				{
+					Arguments.Add($"-fcrash-diagnostics-dir=\"{NormalizeCommandLinePath(CrashDiagnosticDirectory)}\"");
+				}
+				else
+				{
+					Log.TraceWarningOnce($"CrashDiagnosticDirectory has been specified but directory \"{CrashDiagnosticDirectory}\" does not exist. Compiler argument \"-fcrash-diagnostics-dir\" has been discarded.");
+				}
+			}
 		}
 
 		/// <summary>
