* Combined generated cpp files and module cpp files into the same unity files.

ShooterGame Editor with 8 processes:
Before:
3706 compile tasks
1517.98 seconds

After:
3372 compile tasks
1428.56 seconds

Game Editor with 64 processes:
Before:
6323 compile tasks
1090.78 seconds

After:
5666 compile tasks
1063.90 seconds

diff --git a/Engine/Source/Programs/UnrealBuildTool/Configuration/TargetRules.cs b/Engine/Source/Programs/UnrealBuildTool/Configuration/TargetRules.cs
--- a/Engine/Source/Programs/UnrealBuildTool/Configuration/TargetRules.cs
+++ b/Engine/Source/Programs/UnrealBuildTool/Configuration/TargetRules.cs
@@ -1137,6 +1137,13 @@ namespace UnrealBuildTool
 		[XmlConfigFile(Category = "BuildConfiguration")]
 		public bool bForceUnityBuild = false;
 
+		/// <summary>
+		/// Whether to merge module and generated unity files for faster compilation.
+		/// </summary>
+		[CommandLine("-DisableMergingUnityFiles", Value = "false")]
+		[XmlConfigFile(Category = "BuildConfiguration")]
+		public bool bMergeModuleAndGeneratedUnityFiles = true;
+
 		/// <summary>
 		/// Use a heuristic to determine which files are currently being iterated on and exclude them from unity blobs, result in faster
 		/// incremental compile times. The current implementation uses the read-only flag to distinguish the working set, assuming that files will
@@ -2802,6 +2809,10 @@ namespace UnrealBuildTool
 			get { return Inner.bForceUnityBuild; }
 		}
 
+		public bool bMergeModuleAndGeneratedUnityFiles
+		{
+			get { return Inner.bMergeModuleAndGeneratedUnityFiles; }
+		}
 		public bool bAdaptiveUnityDisablesOptimizations
 		{
 			get { return Inner.bAdaptiveUnityDisablesOptimizations; }
diff --git a/Engine/Source/Programs/UnrealBuildTool/Configuration/UEBuildModuleCPP.cs b/Engine/Source/Programs/UnrealBuildTool/Configuration/UEBuildModuleCPP.cs
--- a/Engine/Source/Programs/UnrealBuildTool/Configuration/UEBuildModuleCPP.cs
+++ b/Engine/Source/Programs/UnrealBuildTool/Configuration/UEBuildModuleCPP.cs
@@ -530,23 +530,7 @@ namespace UnrealBuildTool
 			// Mapping of source file to unity file. We output this to intermediate directories for other tools (eg. live coding) to use.
 			Dictionary<FileItem, FileItem> SourceFileToUnityFile = new Dictionary<FileItem, FileItem>();
 
-			// Compile CPP files
-			if (bModuleUsesUnityBuild)
-			{
-				Unity.GenerateUnityCPPs(Target, InputFiles.CPPFiles, InputFiles.HeaderFiles, CompileEnvironment, WorkingSet, Rules.ShortName ?? Name, IntermediateDirectory, Graph, SourceFileToUnityFile, 
-					out List<FileItem> NormalFiles, out List<FileItem> AdaptiveFiles);
-				LinkInputFiles.AddRange(CompileFilesWithToolChain(Target, ToolChain, CompileEnvironment, ModuleCompileEnvironment, NormalFiles, AdaptiveFiles, Graph, Logger).ObjectFiles);
-			}
-			else if (SpecificFilesToCompile.Count == 0)
-			{
-				Unity.GetAdaptiveFiles(Target, InputFiles.CPPFiles, InputFiles.HeaderFiles, CompileEnvironment, WorkingSet, Rules.ShortName ?? Name, IntermediateDirectory, Graph, 
-					out List<FileItem> NormalFiles, out List<FileItem> AdaptiveFiles);
-				LinkInputFiles.AddRange(CompileFilesWithToolChain(Target, ToolChain, CompileEnvironment, ModuleCompileEnvironment, NormalFiles, AdaptiveFiles, Graph, Logger).ObjectFiles);
-			}
-			else
-			{
-				LinkInputFiles.AddRange(CompileFilesWithToolChain(Target, ToolChain, CompileEnvironment, ModuleCompileEnvironment, InputFiles.CPPFiles, new List<FileItem>(), Graph, Logger).ObjectFiles);
-			}
+			List<FileItem> CPPFiles = new List<FileItem>(InputFiles.CPPFiles);
 
 			// Compile all the generated CPP files
 			if (GeneratedCppDirectories != null && !CompileEnvironment.bHackHeaderGenerator && SpecificFilesToCompile.Count == 0)
@@ -560,27 +544,31 @@ namespace UnrealBuildTool
 					}
 				}
 
-				if(GeneratedFiles.Count > 0)
+				if (GeneratedFiles.Count > 0)
 				{
+					bool bMergeUnityFiles = Target.bMergeModuleAndGeneratedUnityFiles;
+
 					// Create a compile environment for the generated files. We can disable creating debug info here to improve link times.
 					CppCompileEnvironment GeneratedCPPCompileEnvironment = CompileEnvironment;
-					if(GeneratedCPPCompileEnvironment.bCreateDebugInfo && Target.bDisableDebugInfoForGeneratedCode)
+					if (GeneratedCPPCompileEnvironment.bCreateDebugInfo && Target.bDisableDebugInfoForGeneratedCode)
 					{
 						GeneratedCPPCompileEnvironment = new CppCompileEnvironment(GeneratedCPPCompileEnvironment);
 						GeneratedCPPCompileEnvironment.bCreateDebugInfo = false;
+						bMergeUnityFiles = false;
 					}
 
 					// Always force include the PCH, even if PCHs are disabled, for generated code. Legacy code can rely on PCHs being included to compile correctly, and this used to be done by UHT manually including it.
-					if(GeneratedCPPCompileEnvironment.PrecompiledHeaderFile == null && Rules.PrivatePCHHeaderFile != null && Rules.PCHUsage != ModuleRules.PCHUsageMode.UseExplicitOrSharedPCHs)
+					if (GeneratedCPPCompileEnvironment.PrecompiledHeaderFile == null && Rules.PrivatePCHHeaderFile != null && Rules.PCHUsage != ModuleRules.PCHUsageMode.UseExplicitOrSharedPCHs)
 					{
 						FileItem PrivatePchFileItem = FileItem.GetItemByFileReference(FileReference.Combine(ModuleDirectory, Rules.PrivatePCHHeaderFile));
-						if(!PrivatePchFileItem.Exists)
+						if (!PrivatePchFileItem.Exists)
 						{
 							throw new BuildException("Unable to find private PCH file '{0}', referenced by '{1}'", PrivatePchFileItem.Location, RulesFile);
 						}
 
 						GeneratedCPPCompileEnvironment = new CppCompileEnvironment(GeneratedCPPCompileEnvironment);
 						GeneratedCPPCompileEnvironment.ForceIncludeFiles.Add(PrivatePchFileItem);
+						bMergeUnityFiles = false;
 					}
 
 					// Compile all the generated files
@@ -596,9 +584,16 @@ namespace UnrealBuildTool
 
 					if (bModuleUsesUnityBuild)
 					{
-						Unity.GenerateUnityCPPs(Target, GeneratedFileItems, new List<FileItem>(), CompileEnvironment, WorkingSet, (Rules.ShortName ?? Name) + ".gen", IntermediateDirectory, Graph, SourceFileToUnityFile, 
-							out List<FileItem> NormalGeneratedFiles, out List<FileItem> AdaptiveGeneratedFiles);
-						LinkInputFiles.AddRange(CompileFilesWithToolChain(Target, ToolChain, GeneratedCPPCompileEnvironment, ModuleCompileEnvironment, NormalGeneratedFiles, AdaptiveGeneratedFiles, Graph, Logger).ObjectFiles);
+						if (bMergeUnityFiles)
+						{
+							CPPFiles.AddRange(GeneratedFileItems);
+						}
+						else
+						{
+							Unity.GenerateUnityCPPs(Target, GeneratedFileItems, new List<FileItem>(), CompileEnvironment, WorkingSet, (Rules.ShortName ?? Name) + ".gen", IntermediateDirectory, Graph, SourceFileToUnityFile,
+								out List<FileItem> NormalGeneratedFiles, out List<FileItem> AdaptiveGeneratedFiles);
+							LinkInputFiles.AddRange(CompileFilesWithToolChain(Target, ToolChain, GeneratedCPPCompileEnvironment, ModuleCompileEnvironment, NormalGeneratedFiles, AdaptiveGeneratedFiles, Graph, Logger).ObjectFiles);
+						}
 					}
 					else
 					{
@@ -607,6 +602,24 @@ namespace UnrealBuildTool
 				}
 			}
 
+			// Compile CPP files
+			if (bModuleUsesUnityBuild)
+			{
+				Unity.GenerateUnityCPPs(Target, CPPFiles, InputFiles.HeaderFiles, CompileEnvironment, WorkingSet, Rules.ShortName ?? Name, IntermediateDirectory, Graph, SourceFileToUnityFile, 
+					out List<FileItem> NormalFiles, out List<FileItem> AdaptiveFiles);
+				LinkInputFiles.AddRange(CompileFilesWithToolChain(Target, ToolChain, CompileEnvironment, ModuleCompileEnvironment, NormalFiles, AdaptiveFiles, Graph, Logger).ObjectFiles);
+			}
+			else if (SpecificFilesToCompile.Count == 0)
+			{
+				Unity.GetAdaptiveFiles(Target, CPPFiles, InputFiles.HeaderFiles, CompileEnvironment, WorkingSet, Rules.ShortName ?? Name, IntermediateDirectory, Graph, 
+					out List<FileItem> NormalFiles, out List<FileItem> AdaptiveFiles);
+				LinkInputFiles.AddRange(CompileFilesWithToolChain(Target, ToolChain, CompileEnvironment, ModuleCompileEnvironment, NormalFiles, AdaptiveFiles, Graph, Logger).ObjectFiles);
+			}
+			else
+			{
+				LinkInputFiles.AddRange(CompileFilesWithToolChain(Target, ToolChain, CompileEnvironment, ModuleCompileEnvironment, CPPFiles, new List<FileItem>(), Graph, Logger).ObjectFiles);
+			}
+
 			// Compile ISPC files directly
 			if (InputFiles.ISPCFiles.Count > 0)
 			{
diff --git a/Engine/Source/Programs/UnrealBuildTool/System/Unity.cs b/Engine/Source/Programs/UnrealBuildTool/System/Unity.cs
--- a/Engine/Source/Programs/UnrealBuildTool/System/Unity.cs
+++ b/Engine/Source/Programs/UnrealBuildTool/System/Unity.cs
@@ -198,8 +198,23 @@
 				// Note that we're relying on this not only sorting files within each directory, but also the directories
 				// themselves, so the whole list of file paths is the same across computers.
 				// Case-insensitive file path compare, because you never know what is going on with local file systems.
-				List<FileItem> SortedCPPFiles = CPPFiles.OrderBy(File => File.AbsolutePath, StringComparer.OrdinalIgnoreCase).ToList();
+				List<FileItem> SortedCPPFiles = new List<FileItem>(CPPFiles);
+				SortedCPPFiles.Sort((A, B) =>
+				{
+					// Generated files from UHT need to be first in the list because they implement templated functions that aren't
+					// declared in the header but are required to link. If they are placed later in the list, you will see
+					// compile errors because the templated function is instantiated but is defined later in the same translation unit
+					// which results in 'error C2908: explicit specialization; '*****' has already been instantiated'
+					bool bAIsGenerated = A.AbsolutePath.EndsWith(".gen.cpp");
+					bool bBIsGenerated = B.AbsolutePath.EndsWith(".gen.cpp");
+					if (bAIsGenerated && !bBIsGenerated)
+						return -1;
+					if (!bAIsGenerated && bBIsGenerated)
+						return 1;
 
+					return String.Compare(A.AbsolutePath, B.AbsolutePath, StringComparison.OrdinalIgnoreCase);
+				});
+
 				HashSet<FileItem> AdaptiveFileSet = new HashSet<FileItem>(AdaptiveFiles);
 				UnityFileBuilder CPPUnityFileBuilder = new UnityFileBuilder(bForceIntoSingleUnityFile ? -1 : Target.NumIncludedBytesPerUnityCPP);
 				foreach (FileItem CPPFile in SortedCPPFiles)
