UnrealBuildTool: Add aggregate csvs to clang trace archive

#rnx

diff --git a/Engine/Source/Programs/UnrealBuildTool/Modes/AggregateClangTimingInfo.cs b/Engine/Source/Programs/UnrealBuildTool/Modes/AggregateClangTimingInfo.cs
--- a/Engine/Source/Programs/UnrealBuildTool/Modes/AggregateClangTimingInfo.cs
+++ b/Engine/Source/Programs/UnrealBuildTool/Modes/AggregateClangTimingInfo.cs
@@ -145,6 +145,10 @@ namespace UnrealBuildTool
 			{
 				FileReference DependsFile = new FileReference($"{SourceFile.FullName}.d");
 				if (!FileReference.Exists(DependsFile))
+				{
+					DependsFile = new FileReference($"{SourceFile.FullName}.txt"); // ispc depends
+				}
+				if (!FileReference.Exists(DependsFile))
 				{
 					return 0;
 				}
@@ -223,6 +227,15 @@ namespace UnrealBuildTool
 				string TempFilePath = Path.Join(Path.GetTempPath(), ArchiveFile.GetFileName() + ".tmp");
 				using (ZipArchive ZipArchive = new ZipArchive(File.Open(TempFilePath, FileMode.Create), ZipArchiveMode.Create))
 				{
+					if (AggregateFile != null)
+					{
+						ZipArchive.CreateEntryFromFile_CrossPlatform(AggregateFile.FullName, AggregateFile.GetFileName(), CompressionLevel.Optimal);
+					}
+					if (HeadersFile != null)
+					{
+						ZipArchive.CreateEntryFromFile_CrossPlatform(HeadersFile.FullName, HeadersFile.GetFileName(), CompressionLevel.Optimal);
+					}
+					List<string> Metadata = new List<string>();
 					foreach (FileReference SourceFile in SourceFiles)
 					{
 						FileReference JsonFile = new FileReference($"{SourceFile.FullName}.json");
