UnrealBuildTool: Reference msvc system includes with /external:I

#rnx

diff --git a/Engine/Source/Programs/UnrealBuildTool/Platform/Windows/VCToolChain.cs b/Engine/Source/Programs/UnrealBuildTool/Platform/Windows/VCToolChain.cs
--- a/Engine/Source/Programs/UnrealBuildTool/Platform/Windows/VCToolChain.cs
+++ b/Engine/Source/Programs/UnrealBuildTool/Platform/Windows/VCToolChain.cs
@@ -165,7 +165,7 @@ namespace UnrealBuildTool
 
 		private static void AddIncludePath(List<string> Arguments, DirectoryReference IncludePath, WindowsCompiler Compiler, bool bSystemInclude)
 		{
-			// Try to use a relative path to shorten command line length. Always need the full path when preprocessing because the output file will be in a different place, where include paths cannot be relative.
+			// Try to use a relative path to shorten command line length.
 			string IncludePathString = NormalizeCommandLinePath(IncludePath);
 
 			if (Compiler.IsClang() && bSystemInclude)
@@ -174,6 +174,15 @@ namespace UnrealBuildTool
 				// and warnings are disabled to allow compiler toolchains to be upgraded separately.
 				Arguments.Add("/imsvc " + Utils.MakePathSafeToUseWithCommandLine(IncludePathString));
 			}
+			else if (Compiler.IsMSVC() && Compiler >= WindowsCompiler.VisualStudio2022 && bSystemInclude)
+			{
+				if (!Arguments.Contains("/external:W0"))
+				{
+					Arguments.Add("/external:W0");
+				}
+				// Defines a root directory that contains external headers.
+				Arguments.Add("/external:I " + Utils.MakePathSafeToUseWithCommandLine(IncludePathString));
+			}
 			else
 			{
 				Arguments.Add("/I " + Utils.MakePathSafeToUseWithCommandLine(IncludePathString));
