For UBT, mimicing changes made on the UE side in CL 20764079.
Adding support for managing multiple plugins, and specifying explicit versions to enable.

diff --git a/Engine/Source/Programs/UnrealBuildTool/Configuration/UEBuildTarget.cs b/Engine/Source/Programs/UnrealBuildTool/Configuration/UEBuildTarget.cs
--- a/Engine/Source/Programs/UnrealBuildTool/Configuration/UEBuildTarget.cs
+++ b/Engine/Source/Programs/UnrealBuildTool/Configuration/UEBuildTarget.cs
@@ -2974,7 +2974,7 @@ namespace UnrealBuildTool
 			}
 
 			// Add all the plugin modules that need to be compiled
-			List<PluginInfo> Plugins = RulesAssembly.EnumeratePlugins().ToList();
+			List<PluginInfo> Plugins = RulesAssembly.EnumeratePlugins().Where(x => x.ChoiceVersion != null).Select(x => x.ChoiceVersion!).ToList();
 			foreach (PluginInfo Plugin in Plugins)
 			{
 				// Ignore plugins which are specifically disabled by this target
@@ -3292,7 +3292,7 @@ namespace UnrealBuildTool
 		public void SetupPlugins(ILogger Logger)
 		{
 			// Find all the valid plugins
-			Dictionary<string, PluginInfo> NameToInfo = RulesAssembly.EnumeratePlugins().ToDictionary(x => x.Name, x => x, StringComparer.InvariantCultureIgnoreCase);
+			Dictionary<string, PluginSet> NameToInfos = RulesAssembly.EnumeratePlugins().ToDictionary(x => x.Name, x => x, StringComparer.InvariantCultureIgnoreCase);
 
 			// Remove any plugins for platforms we don't have
 			List<UnrealTargetPlatform> MissingPlatforms = new List<UnrealTargetPlatform>();
@@ -3317,7 +3317,7 @@ namespace UnrealBuildTool
 			if (ForeignPlugin != null)
 			{
 				PluginReferenceDescriptor PluginReference = new PluginReferenceDescriptor(ForeignPlugin.GetFileNameWithoutExtension(), null, true);
-				AddPlugin(PluginReference, "command line", ExcludeFolders, NameToInstance, NameToInfo, Logger);
+				AddPlugin(PluginReference, "command line", ExcludeFolders, NameToInstance, NameToInfos, Logger);
 			}
 
 			// Configure plugins explicitly enabled via target settings
@@ -3326,7 +3326,7 @@ namespace UnrealBuildTool
 				if (ReferencedNames.Add(PluginName))
 				{
 					PluginReferenceDescriptor PluginReference = new PluginReferenceDescriptor(PluginName, null, true);
-					AddPlugin(PluginReference, "target settings", ExcludeFolders, NameToInstance, NameToInfo, Logger);
+					AddPlugin(PluginReference, "target settings", ExcludeFolders, NameToInstance, NameToInfos, Logger);
 				}
 			}
 
@@ -3336,7 +3336,7 @@ namespace UnrealBuildTool
 				if (ReferencedNames.Add(PluginName))
 				{
 					PluginReferenceDescriptor PluginReference = new PluginReferenceDescriptor(PluginName, null, false);
-					AddPlugin(PluginReference, "target settings", ExcludeFolders, NameToInstance, NameToInfo, Logger);
+					AddPlugin(PluginReference, "target settings", ExcludeFolders, NameToInstance, NameToInfos, Logger);
 				}
 			}
 
@@ -3360,7 +3360,7 @@ namespace UnrealBuildTool
 							}
 							else
 							{
-								AddPlugin(PluginReference, ProjectReferenceChain, ExcludeFolders, NameToInstance, NameToInfo, Logger);
+								AddPlugin(PluginReference, ProjectReferenceChain, ExcludeFolders, NameToInstance, NameToInfos, Logger);
 							}
 						}
 					}
@@ -3370,9 +3370,10 @@ namespace UnrealBuildTool
 			// Also synthesize references for plugins which are enabled by default
 			if (Rules.bCompileAgainstEngine || Rules.bCompileWithPluginSupport)
 			{
-				foreach (PluginInfo Plugin in NameToInfo.Values)
+				foreach (PluginSet PluginVersions in NameToInfos.Values)
 				{
-					if (Plugin.IsEnabledByDefault(bAllowEnginePluginsEnabledByDefault) && !ReferencedNames.Contains(Plugin.Name))
+					PluginInfo? Plugin = PluginVersions.ChoiceVersion;
+					if (Plugin != null && Plugin.IsEnabledByDefault(bAllowEnginePluginsEnabledByDefault) && !ReferencedNames.Contains(Plugin.Name))
 					{
 						ReferencedNames.Add(Plugin.Name);
 
@@ -3386,7 +3387,7 @@ namespace UnrealBuildTool
 							PluginReference.PlatformAllowList = PluginReference.SupportedTargetPlatforms; //synthesize allow list if it must be explicit
 						}
 
-						AddPlugin(PluginReference, "default plugins", ExcludeFolders, NameToInstance, NameToInfo, Logger);
+						AddPlugin(PluginReference, "default plugins", ExcludeFolders, NameToInstance, NameToInfos, Logger);
 					}
 				}
 			}
@@ -3404,7 +3405,7 @@ namespace UnrealBuildTool
 						if (ReferencedNames.Add(PluginName))
 						{
 							PluginReferenceDescriptor PluginReference = new PluginReferenceDescriptor(PluginName, null, true);
-							AddPlugin(PluginReference, "DefaultEngine.ini", ExcludeFolders, NameToInstance, NameToInfo, Logger);
+							AddPlugin(PluginReference, "DefaultEngine.ini", ExcludeFolders, NameToInstance, NameToInfos, Logger);
 						}
 					}
 				}
@@ -3419,7 +3420,7 @@ namespace UnrealBuildTool
 				if (ReferencedNames.Add(PluginName))
 				{
 					PluginReferenceDescriptor PluginReference = new PluginReferenceDescriptor(PluginName, null, true);
-					AddPlugin(PluginReference, "target settings", ExcludeFolders, NameToInstance, NameToInfo, Logger);
+					AddPlugin(PluginReference, "target settings", ExcludeFolders, NameToInstance, NameToInfos, Logger);
 				}
 			}
 
@@ -3437,10 +3438,10 @@ namespace UnrealBuildTool
 		/// <param name="ReferenceChain">Textual representation of the chain of references, for error reporting</param>
 		/// <param name="ExcludeFolders">Array of folder names to be excluded</param>
 		/// <param name="NameToInstance">Map from plugin name to instance of it</param>
-		/// <param name="NameToInfo">Map from plugin name to information</param>
+		/// <param name="NameToInfos">Map from plugin name to information</param>
 		/// <param name="Logger">Logger for diagnostic output</param>
 		/// <returns>Instance of the plugin, or null if it should not be used</returns>
-		private UEBuildPlugin? AddPlugin(PluginReferenceDescriptor Reference, string ReferenceChain, string[] ExcludeFolders, Dictionary<string, UEBuildPlugin> NameToInstance, Dictionary<string, PluginInfo> NameToInfo, ILogger Logger)
+		private UEBuildPlugin? AddPlugin(PluginReferenceDescriptor Reference, string ReferenceChain, string[] ExcludeFolders, Dictionary<string, UEBuildPlugin> NameToInstance, Dictionary<string, PluginSet> NameToInfos, ILogger Logger)
 		{
 			// Ignore disabled references
 			if (!Reference.bEnabled)
@@ -3461,7 +3462,7 @@ namespace UnrealBuildTool
 						foreach (PluginReferenceDescriptor NextReference in Instance.Descriptor.Plugins)
 						{
 							string NextReferenceChain = String.Format("{0} -> {1}", ReferenceChain, Instance.File.GetFileName());
-							AddPlugin(NextReference, NextReferenceChain, ExcludeFolders, NameToInstance, NameToInfo, Logger);
+							AddPlugin(NextReference, NextReferenceChain, ExcludeFolders, NameToInstance, NameToInfos, Logger);
 						}
 					}
 				}
@@ -3495,8 +3496,31 @@ namespace UnrealBuildTool
 				}
 
 				// Find the plugin being enabled
-				PluginInfo? Info;
-				if (!NameToInfo.TryGetValue(Reference.Name, out Info))
+				PluginInfo? Info = null;
+
+				PluginSet? PluginVersions;
+				if (NameToInfos.TryGetValue(Reference.Name, out PluginVersions))
+				{
+					if (Reference.RequestedVersion.HasValue)
+					{
+						Info = PluginVersions.KnownVersions.Find(x => x.Descriptor.Version == Reference.RequestedVersion.Value);
+						if (Info != null)
+						{
+							PluginVersions.ChoiceVersion = Info;
+						}
+						else
+						{
+							Logger.LogWarning("Failed to find specific version (v{RequestedVersion}) of plugin '{ReferenceName}' (referenced via {ReferenceChain}). Other versions exist, but the explicit version requested could not be found.",
+								Reference.RequestedVersion.Value, Reference.Name, ReferenceChain);
+						}
+					}
+					else
+					{
+						Info = PluginVersions.ChoiceVersion;
+					}
+				}
+
+				if (Info == null)
 				{
 					if (Reference.bOptional)
 					{
@@ -3566,7 +3590,7 @@ namespace UnrealBuildTool
 				{
 					foreach (PluginReferenceDescriptor NextReference in Info.Descriptor.Plugins)
 					{
-						UEBuildPlugin? NextInstance = AddPlugin(NextReference, PluginReferenceChain, ExcludeFolders, NameToInstance, NameToInfo, Logger);
+						UEBuildPlugin? NextInstance = AddPlugin(NextReference, PluginReferenceChain, ExcludeFolders, NameToInstance, NameToInfos, Logger);
 						if (NextInstance != null)
 						{
 							Dependencies.Add(NextInstance);
diff --git a/Engine/Source/Programs/UnrealBuildTool/System/PluginReferenceDescriptor.cs b/Engine/Source/Programs/UnrealBuildTool/System/PluginReferenceDescriptor.cs
--- a/Engine/Source/Programs/UnrealBuildTool/System/PluginReferenceDescriptor.cs
+++ b/Engine/Source/Programs/UnrealBuildTool/System/PluginReferenceDescriptor.cs
@@ -82,6 +82,11 @@ namespace UnrealBuildTool
 		/// </summary>
 		public bool bHasExplicitPlatforms;
 
+		/// <summary>
+		/// When set, specifies a specific version of the plugin that this references.
+		/// </summary>
+		public int? RequestedVersion;
+
 
 		/// <summary>
 		/// Constructor
@@ -149,6 +154,10 @@ namespace UnrealBuildTool
 			{
 				Writer.WriteValue("HasExplicitPlatforms",bHasExplicitPlatforms);
 			}
+			if (bEnabled && RequestedVersion != null)
+			{
+				Writer.WriteValue("Version", RequestedVersion.Value);
+			}
 			Writer.WriteObjectEnd();
 		}
 
@@ -214,6 +223,12 @@ namespace UnrealBuildTool
 				RawObject.TryGetStringArrayField("SupportedTargetPlatforms", out Descriptor.SupportedTargetPlatforms);
 				RawObject.TryGetBoolField("HasExplicitPlatforms", out Descriptor.bHasExplicitPlatforms);
 
+				int RequestedVersion = -1;
+				if (RawObject.TryGetIntegerField("Version", out RequestedVersion))
+				{
+					Descriptor.RequestedVersion = RequestedVersion;
+				}
+
 				Descriptor.VerifyPlatformNames(Descriptor.PlatformAllowList);
 				Descriptor.VerifyPlatformNames(Descriptor.PlatformDenyList);
 				Descriptor.VerifyPlatformNames(Descriptor.SupportedTargetPlatforms);
diff --git a/Engine/Source/Programs/UnrealBuildTool/System/Plugins.cs b/Engine/Source/Programs/UnrealBuildTool/System/Plugins.cs
--- a/Engine/Source/Programs/UnrealBuildTool/System/Plugins.cs
+++ b/Engine/Source/Programs/UnrealBuildTool/System/Plugins.cs
@@ -150,6 +150,84 @@
 	}
 
 	/// <summary>
+	/// Represents a group of plugins all sharing the same name -- notionally all different versions of the same plugin.
+	/// Since UE can only manage a single plugin of a given name (to avoid module conflicts, etc.), this object 
+	/// prioritizes and bubbles up a single "choice" plugin from the set -- see `ChoiceVersion`.
+	/// </summary>
+	public class PluginSet
+	{
+		/// <summary>
+		/// Shared name that all the plugins in the set go by.
+		/// </summary>
+		public readonly string Name;
+
+		/// <summary>
+		/// Unordered list of all the discovered plugins matching the above name.
+		/// </summary>
+		public readonly List<PluginInfo> KnownVersions = new List<PluginInfo>();
+
+		/// <summary>
+		/// Constructor which takes the specified PluginInfo and adds it to the set.
+		/// </summary>
+		/// <param name="FirstPlugin">The fist plugin discovered for this set -- will automatically be prioritized as the "choice" plugin.</param>
+		public PluginSet(PluginInfo FirstPlugin)
+		{
+			Name = FirstPlugin.Name;
+			Add(FirstPlugin, /*bPromoteToChoiceVersion =*/true);
+		}
+
+		/// <summary>
+		/// Pushes the specified PluginInfo into the set, and optionally promotes it to the new "choice" version.
+		/// </summary>
+		/// <param name="Plugin">The new plugin to add to this set (its name should match this set's name)</param>
+		/// <param name="bPromoteToChoiceVersion">Whether or not to make this the new prioritized "choice" plugin in the set.</param>
+		public void Add(PluginInfo Plugin, bool bPromoteToChoiceVersion = true)
+		{
+			if (bPromoteToChoiceVersion || KnownVersions.Count() == 0)
+			{
+				IndexOfChoiceVersion = KnownVersions.Count();
+			}
+			KnownVersions.Add(Plugin);
+		}
+
+		/// <summary>
+		/// The single plugin, out of the entire set, that we prioritize to be built and utilized by UE.
+		/// </summary>
+		public PluginInfo? ChoiceVersion
+		{
+			get
+			{
+				if (IndexOfChoiceVersion >= 0 && IndexOfChoiceVersion < KnownVersions.Count())
+				{
+					return KnownVersions[IndexOfChoiceVersion];
+				}
+				return null;
+			}
+
+			set
+			{
+				if (value == null)
+				{
+					throw new BuildException("Cannot manually set a PluginSet's ChoiceVersion to be null.");
+				}
+				else
+				{
+					IndexOfChoiceVersion = KnownVersions.FindIndex(x => x.File == value.File);
+					if (IndexOfChoiceVersion == -1)
+					{
+						Add(value, /*bPromoteToChoiceVersion =*/ true);
+					}
+				}
+			}
+		}
+		
+		/// <summary>
+		/// An index that we use to identify which plugin in the list is the prioritized "choice" version.
+		/// </summary>
+		private int IndexOfChoiceVersion = -1;
+	}
+
+	/// <summary>
 	/// Class for enumerating plugin metadata
 	/// </summary>
 	public static class Plugins
@@ -176,26 +254,27 @@
 		/// </summary>
 		/// <param name="Plugins">List of plugins to filter</param>
 		/// <returns>Filtered Dictionary of plugins</returns>
-		public static Dictionary<string, PluginInfo> ToFilteredDictionary(IEnumerable<PluginInfo> Plugins)
+		public static Dictionary<string, PluginSet> ToFilteredDictionary(IEnumerable<PluginInfo> Plugins)
 		{
-			Dictionary<string, PluginInfo> NameToPluginInfo = new Dictionary<string, PluginInfo>(StringComparer.InvariantCultureIgnoreCase);
+			Dictionary<string, PluginSet> NameToPluginInfos = new Dictionary<string, PluginSet>(StringComparer.InvariantCultureIgnoreCase);
 			foreach (PluginInfo Plugin in Plugins)
 			{
-				PluginInfo? ExistingPluginInfo;
-				if (!NameToPluginInfo.TryGetValue(Plugin.Name, out ExistingPluginInfo))
+				PluginSet? ExistingPluginSet;
+				if (!NameToPluginInfos.TryGetValue(Plugin.Name, out ExistingPluginSet))
 				{
-					NameToPluginInfo.Add(Plugin.Name, Plugin);
+					NameToPluginInfos.Add(Plugin.Name, new PluginSet(Plugin) );
 				}
-				else if (Plugin.Type > ExistingPluginInfo.Type)
+				else if (Plugin.Type > ExistingPluginSet.ChoiceVersion?.Type)
 				{
-					NameToPluginInfo[Plugin.Name] = Plugin;
+					ExistingPluginSet.Add(Plugin, /*bPromoteToChoiceVersion =*/true);
 				}
-				else if (Plugin.Type == ExistingPluginInfo.Type)
+				else
 				{
-					throw new BuildException(String.Format("Found '{0}' plugin in two locations ({1} and {2}). Plugin names must be unique.", Plugin.Name, ExistingPluginInfo.File, Plugin.File));
+					bool bPromoteToChoiceVersion = Plugin.Descriptor.Version > ExistingPluginSet.ChoiceVersion?.Descriptor.Version;
+					ExistingPluginSet.Add(Plugin, bPromoteToChoiceVersion);
 				}
 			}
-			return NameToPluginInfo;
+			return NameToPluginInfos;
 		}
 
 		/// <summary>
@@ -204,10 +283,10 @@
 		/// </summary>
 		/// <param name="Plugins">List of plugins to filter</param>
 		/// <returns>Filtered list of plugins in the original order</returns>
-		public static IEnumerable<PluginInfo> FilterPlugins(IEnumerable<PluginInfo> Plugins)
+		public static IEnumerable<PluginSet> FilterPlugins(IEnumerable<PluginInfo> Plugins)
 		{
-			Dictionary<string, PluginInfo> NameToPluginInfo = ToFilteredDictionary(Plugins);
-			return Plugins.Where(x => NameToPluginInfo[x.Name] == x);
+			Dictionary<string, PluginSet> NameToPluginInfos = ToFilteredDictionary(Plugins);
+			return NameToPluginInfos.Values.AsEnumerable();
 		}
 
 		/// <summary>
diff --git a/Engine/Source/Programs/UnrealBuildTool/System/RulesAssembly.cs b/Engine/Source/Programs/UnrealBuildTool/System/RulesAssembly.cs
--- a/Engine/Source/Programs/UnrealBuildTool/System/RulesAssembly.cs
+++ b/Engine/Source/Programs/UnrealBuildTool/System/RulesAssembly.cs
@@ -652,7 +652,7 @@ namespace UnrealBuildTool
 				List<string> DisallowedPlatformsAndGroups = Utils.MakeListOfUnsupportedPlatforms(Rules.OptedInModulePlatforms.ToList(), false, Logger);
 
 				// look in all plugins' paths to see if any disallowed
-				IEnumerable<PluginInfo> DisallowedPlugins = EnumeratePlugins().Where(Plugin =>
+				IEnumerable<PluginInfo> DisallowedPlugins = EnumeratePlugins().Where(x => x.ChoiceVersion != null).Select(x => x.ChoiceVersion!).Where(Plugin =>
 					Plugin.File.ContainsAnyNames(DisallowedPlatformsAndGroups, Unreal.EngineDirectory) ||
 					(Rules.ProjectFile != null && Plugin.File.ContainsAnyNames(DisallowedPlatformsAndGroups, Rules.ProjectFile.Directory)));
 				// log out the plugins we are disabling
@@ -796,7 +796,7 @@ namespace UnrealBuildTool
 		/// Enumerates all the plugins that are available
 		/// </summary>
 		/// <returns></returns>
-		public IEnumerable<PluginInfo> EnumeratePlugins()
+		public IEnumerable<PluginSet> EnumeratePlugins()
 		{
 			return global::UnrealBuildTool.Plugins.FilterPlugins(EnumeratePluginsInternal());
 		}
diff --git a/Engine/Source/Programs/UnrealBuildTool/System/RulesCompiler.cs b/Engine/Source/Programs/UnrealBuildTool/System/RulesCompiler.cs
--- a/Engine/Source/Programs/UnrealBuildTool/System/RulesCompiler.cs
+++ b/Engine/Source/Programs/UnrealBuildTool/System/RulesCompiler.cs
@@ -343,7 +343,7 @@ namespace UnrealBuildTool
 				// Create a list of plugins for this assembly. We need to override the parent plugin, if it exists, due to overriding the
 				// setting for bClassifyAsGameModuleForUHT below.
 				List<PluginInfo> ForeignPlugins = new List<PluginInfo>();
-				if (!Parent.EnumeratePlugins().Any(x => x.File == PluginFileName && x.Type == PluginType.Engine))
+				if (!Parent.EnumeratePlugins().Any(x => x.ChoiceVersion != null && x.ChoiceVersion.File == PluginFileName && x.ChoiceVersion.Type == PluginType.Engine))
 				{
 					PluginInfo ForeignPluginInfo = new PluginInfo(PluginFileName, PluginType.External)
 					{
