Modular builds: Clang: Added a toolchain variable indicating if the current target uses the modular link mode.

It's a preparation step as some platforms will require special linker switches for modular builds.

#rnx

diff --git a/Engine/Source/Programs/UnrealBuildTool/ToolChain/ClangToolChain.cs b/Engine/Source/Programs/UnrealBuildTool/ToolChain/ClangToolChain.cs
--- a/Engine/Source/Programs/UnrealBuildTool/ToolChain/ClangToolChain.cs
+++ b/Engine/Source/Programs/UnrealBuildTool/ToolChain/ClangToolChain.cs
@@ -101,7 +101,12 @@ namespace UnrealBuildTool
 		/// <summary>
 		/// Enables dead code/data stripping and common code folding.
 		/// </summary>
-		EnableDeadStripping = 1 << 15
+		EnableDeadStripping = 1 << 15,
+
+		/// <summary>
+		/// Indicates that the target is a moduler build i.e. Target.LinkType == TargetLinkType.Modular
+		/// </summary>
+		ModularBuild = 1 << 16
 	}
 
 	abstract class ClangToolChain : ISPCToolChain
