UnrealBuildTool: Support aggregating clang-cl traces

#rnx

diff --git a/Engine/Source/Programs/UnrealBuildTool/Modes/AggregateClangTimingInfo.cs b/Engine/Source/Programs/UnrealBuildTool/Modes/AggregateClangTimingInfo.cs
--- a/Engine/Source/Programs/UnrealBuildTool/Modes/AggregateClangTimingInfo.cs
+++ b/Engine/Source/Programs/UnrealBuildTool/Modes/AggregateClangTimingInfo.cs
@@ -128,17 +128,28 @@ namespace UnrealBuildTool
 				"DependencyIncludes",
 			};
 
+			static readonly string[] ObjectExtensions = new string[]
+			{
+				".o",
+				".obj",
+				".gch",
+				".pch",
+			};
+
 			public static string CsvHeader => string.Join(',', CsvColumns);
 			public string CsvLine => string.Join(',', CsvColumns.Select(x => GetType().GetProperty(x)!.GetValue(this)!.ToString()));
 
 			private long GetObjectSize()
 			{
-				FileReference ObjectFile = new FileReference($"{SourceFile.FullName}.{(SourceFile.HasExtension(".h") ? "gch" : "o")}");
-				if (!FileReference.Exists(ObjectFile))
+				foreach (string Extension in ObjectExtensions)
 				{
-					return 0;
+					FileReference ObjectFile = new FileReference($"{SourceFile.FullName}{Extension}");
+					if (FileReference.Exists(ObjectFile))
+					{
+						return ObjectFile.ToFileInfo().Length;
+					}
 				}
-				return ObjectFile.ToFileInfo().Length;
+				return 0;
 			}
 
 			private long CountIncludes()
diff --git a/Engine/Source/Programs/UnrealBuildTool/Platform/Windows/VCToolChain.cs b/Engine/Source/Programs/UnrealBuildTool/Platform/Windows/VCToolChain.cs
--- a/Engine/Source/Programs/UnrealBuildTool/Platform/Windows/VCToolChain.cs
+++ b/Engine/Source/Programs/UnrealBuildTool/Platform/Windows/VCToolChain.cs
@@ -1299,7 +1299,7 @@ namespace UnrealBuildTool
 					}
 
 					// Experimental: support for JSON output of timing data
-					if (Target.WindowsPlatform.Compiler.IsClang() && Target.WindowsPlatform.bClangTimeTrace)
+					if (Target.WindowsPlatform.Compiler.IsClang() && (Target.bPrintToolChainTimingInfo || Target.WindowsPlatform.bClangTimeTrace))
 					{
 						CompileAction.Arguments.Add("-Xclang -ftime-trace");
 						CompileAction.AdditionalProducedItems.Add(FileItem.GetItemByFileReference(ObjectFile.Location.ChangeExtension(".json")));
@@ -1400,7 +1400,7 @@ namespace UnrealBuildTool
 					CompileAction.CompiledModuleInterfaceFile = IfcFile;
 				}
 
-				if (Target.bPrintToolChainTimingInfo || Target.WindowsPlatform.bCompilerTrace)
+				if ((Target.bPrintToolChainTimingInfo || Target.WindowsPlatform.bCompilerTrace) && Target.WindowsPlatform.Compiler.IsMSVC())
 				{
 					CompileAction.ForceClFilter = true;
 					CompileAction.TimingFile = FileItem.GetItemByFileReference(FileReference.Combine(OutputDir, $"{SourceFile.Location.GetFileName()}.timing"));
@@ -1482,8 +1482,19 @@ namespace UnrealBuildTool
 			{
 				TargetMakefile Makefile = MakefileBuilder.Makefile;
 
-				List<IExternalAction> ParseTimingActions = Makefile.Actions.Where(x => x.ActionType == ActionType.ParseTimingInfo).ToList();
-				List<FileItem> TimingJsonFiles = ParseTimingActions.SelectMany(a => a.ProducedItems.Where(i => i.HasExtension(".cta"))).ToList();
+				List<FileItem> TimingJsonFiles = new List<FileItem>();
+
+				if (Target.WindowsPlatform.Compiler.IsMSVC())
+				{
+					List<IExternalAction> ParseTimingActions = Makefile.Actions.Where(x => x.ActionType == ActionType.ParseTimingInfo).ToList();
+					TimingJsonFiles = ParseTimingActions.SelectMany(a => a.ProducedItems.Where(i => i.HasExtension(".cta"))).ToList();
+				}
+				else if (Target.WindowsPlatform.Compiler.IsClang())
+				{
+					List<IExternalAction> CompileActions = Makefile.Actions.Where(x => x.ActionType == ActionType.Compile && x.ProducedItems.Any(i => i.HasExtension(".json"))).ToList();
+					TimingJsonFiles = CompileActions.SelectMany(a => a.ProducedItems.Where(i => i.HasExtension(".json"))).ToList();
+				}
+
 				Makefile.OutputItems.AddRange(TimingJsonFiles);
 
 				// Handing generating aggregate timing information if we compiled more than one file.
@@ -1495,26 +1506,55 @@ namespace UnrealBuildTool
 					{
 						DirectoryReference.CreateDirectory(ManifestFile.Directory);
 					}
-					File.WriteAllLines(ManifestFile.FullName, TimingJsonFiles.Select(f => f.AbsolutePath));
 
-					FileReference ExpectedCompileTimeFile = FileReference.FromString(Path.Combine(Makefile.ProjectIntermediateDirectory.FullName, $"{Target.Name}.json"));
-					List<string> ActionArgs = new List<string>()
+					if (Target.WindowsPlatform.Compiler.IsMSVC())
 					{
-						$"-Name={Target.Name}",
-						$"-ManifestFile={ManifestFile.FullName}",
-						$"-CompileTimingFile={ExpectedCompileTimeFile}",
-					};
+						File.WriteAllLines(ManifestFile.FullName, TimingJsonFiles.Select(f => f.AbsolutePath));
 
-					Action AggregateTimingInfoAction = MakefileBuilder.CreateRecursiveAction<AggregateParsedTimingInfo>(ActionType.ParseTimingInfo, string.Join(" ", ActionArgs));
-					AggregateTimingInfoAction.WorkingDirectory = Unreal.EngineSourceDirectory;
-					AggregateTimingInfoAction.StatusDescription = $"Aggregating {TimingJsonFiles.Count} Timing File(s)";
-					AggregateTimingInfoAction.bCanExecuteRemotely = false;
-					AggregateTimingInfoAction.bCanExecuteRemotelyWithSNDBS = false;
-					AggregateTimingInfoAction.PrerequisiteItems.AddRange(TimingJsonFiles);
+						FileReference ExpectedCompileTimeFile = FileReference.FromString(Path.Combine(Makefile.ProjectIntermediateDirectory.FullName, $"{Target.Name}.json"));
+						List<string> ActionArgs = new List<string>()
+						{
+							$"-Name={Target.Name}",
+							$"-ManifestFile={ManifestFile.FullName}",
+							$"-CompileTimingFile={ExpectedCompileTimeFile}",
+						};
+
+						Action AggregateTimingInfoAction = MakefileBuilder.CreateRecursiveAction<AggregateParsedTimingInfo>(ActionType.ParseTimingInfo, string.Join(" ", ActionArgs));
+						AggregateTimingInfoAction.WorkingDirectory = Unreal.EngineSourceDirectory;
+						AggregateTimingInfoAction.StatusDescription = $"Aggregating {TimingJsonFiles.Count} Timing File(s)";
+						AggregateTimingInfoAction.bCanExecuteRemotely = false;
+						AggregateTimingInfoAction.bCanExecuteRemotelyWithSNDBS = false;
+						AggregateTimingInfoAction.PrerequisiteItems.AddRange(TimingJsonFiles);
+
+						FileItem AggregateOutputFile = FileItem.GetItemByFileReference(FileReference.Combine(Makefile.ProjectIntermediateDirectory, $"{Target.Name}.cta"));
+						AggregateTimingInfoAction.ProducedItems.Add(AggregateOutputFile);
+						Makefile.OutputItems.Add(AggregateOutputFile);
+					}
+					else if (Target.WindowsPlatform.Compiler.IsClang())
+					{
+						File.WriteAllLines(ManifestFile.FullName, TimingJsonFiles.Select(f => f.FullName.Remove(f.FullName.Length - ".json".Length)));
 
-					FileItem AggregateOutputFile = FileItem.GetItemByFileReference(FileReference.Combine(Makefile.ProjectIntermediateDirectory, $"{Target.Name}.cta"));
-					AggregateTimingInfoAction.ProducedItems.Add(AggregateOutputFile);
-					Makefile.OutputItems.Add(AggregateOutputFile);
+						FileItem AggregateOutputFile = FileItem.GetItemByFileReference(FileReference.Combine(Makefile.ProjectIntermediateDirectory, $"{Target.Name}.trace.csv"));
+						FileItem HeadersOutputFile = FileItem.GetItemByFileReference(FileReference.Combine(Makefile.ProjectIntermediateDirectory, $"{Target.Name}.headers.csv"));
+						List<string> ActionArgs = new List<string>()
+						{
+							$"-AggregateFile={AggregateOutputFile.FullName}",
+							$"-ManifestFile={ManifestFile.FullName}",
+							$"-HeadersFile={HeadersOutputFile.FullName}",
+						};
+
+						Action AggregateTimingInfoAction = MakefileBuilder.CreateRecursiveAction<AggregateClangTimingInfo>(ActionType.ParseTimingInfo, string.Join(" ", ActionArgs));
+						AggregateTimingInfoAction.WorkingDirectory = Unreal.EngineSourceDirectory;
+						AggregateTimingInfoAction.StatusDescription = $"Aggregating {TimingJsonFiles.Count} Timing File(s)";
+						AggregateTimingInfoAction.bCanExecuteRemotely = false;
+						AggregateTimingInfoAction.bCanExecuteRemotelyWithSNDBS = false;
+						AggregateTimingInfoAction.PrerequisiteItems.AddRange(TimingJsonFiles);
+
+						AggregateTimingInfoAction.ProducedItems.Add(AggregateOutputFile);
+						AggregateTimingInfoAction.ProducedItems.Add(HeadersOutputFile);
+						Makefile.OutputItems.Add(AggregateOutputFile);
+						Makefile.OutputItems.Add(HeadersOutputFile);
+					}
 				}
 			}
 		}
