UnrealBuildTool: Suppress duplicate key exception when generating vs code projects

#rnx

diff --git a/Engine/Source/Programs/UnrealBuildTool/ProjectFiles/VisualStudioCode/VSCodeProjectFileGenerator.cs b/Engine/Source/Programs/UnrealBuildTool/ProjectFiles/VisualStudioCode/VSCodeProjectFileGenerator.cs
--- a/Engine/Source/Programs/UnrealBuildTool/ProjectFiles/VisualStudioCode/VSCodeProjectFileGenerator.cs
+++ b/Engine/Source/Programs/UnrealBuildTool/ProjectFiles/VisualStudioCode/VSCodeProjectFileGenerator.cs
@@ -389,7 +389,7 @@ namespace UnrealBuildTool
 						CommandBuilder.AppendFormat("{0} \"{1}\" {2}", UsingClang ? "-I" : "/I", IncludePath, Environment.NewLine);
 					}
 
-					ModuleDirectoryToCompileCommand.Add(Module.ModuleDirectory, CommandBuilder.ToString());
+					ModuleDirectoryToCompileCommand.TryAdd(Module.ModuleDirectory, CommandBuilder.ToString());
 				}
 			}
 
