UnrealBuildTool: Ignore VCProjectFileSettings.MaxSharedIncludePaths for VS2022

#rnx

diff --git a/Engine/Source/Programs/UnrealBuildTool/ProjectFiles/VisualStudio/VCProject.cs b/Engine/Source/Programs/UnrealBuildTool/ProjectFiles/VisualStudio/VCProject.cs
--- a/Engine/Source/Programs/UnrealBuildTool/ProjectFiles/VisualStudio/VCProject.cs
+++ b/Engine/Source/Programs/UnrealBuildTool/ProjectFiles/VisualStudio/VCProject.cs
@@ -891,12 +891,12 @@ namespace UnrealBuildTool
 				}
 
 				// Append the most common include paths to the search list.
-				if (Settings.MaxSharedIncludePaths > 0)
+				if (Settings.MaxSharedIncludePaths > 0 || ProjectFileFormat >= VCProjectFileFormat.VisualStudio2022)
 				{
 					foreach (DirectoryReference IncludePath in IncludePathToCount.OrderByDescending(x => x.Value).Select(x => x.Key))
 					{
 						string RelativePath = NormalizeProjectPath(IncludePath);
-						if (SharedIncludeSearchPaths.Length + RelativePath.Length >= Settings.MaxSharedIncludePaths)
+						if (ProjectFileFormat < VCProjectFileFormat.VisualStudio2022 && SharedIncludeSearchPaths.Length + RelativePath.Length >= Settings.MaxSharedIncludePaths)
 						{
 							break;
 						}
diff --git a/Engine/Source/Programs/UnrealBuildTool/ProjectFiles/VisualStudio/VCProjectFileGenerator.cs b/Engine/Source/Programs/UnrealBuildTool/ProjectFiles/VisualStudio/VCProjectFileGenerator.cs
--- a/Engine/Source/Programs/UnrealBuildTool/ProjectFiles/VisualStudio/VCProjectFileGenerator.cs
+++ b/Engine/Source/Programs/UnrealBuildTool/ProjectFiles/VisualStudio/VCProjectFileGenerator.cs
@@ -31,7 +31,7 @@
 		/// <summary>
 		/// Puts the most common include paths in the IncludePath property in the MSBuild project. This significantly reduces Visual Studio
 		/// memory usage (measured 1.1GB -> 500mb), but seems to be causing issues with Visual Assist. Value here specifies maximum length
-		/// of the include path list in KB.
+		/// of the include path list in KB. Ignored for VS2022.
 		/// </summary>
 		[XmlConfigFile(Category = "VCProjectFileGenerator")]
 		public int MaxSharedIncludePaths = 24 * 1024;
