Raise minimum Android version supported to 8.0 (level 26)

#android

diff --git a/Engine/Config/BaseEngine.ini b/Engine/Config/BaseEngine.ini
--- a/Engine/Config/BaseEngine.ini
+++ b/Engine/Config/BaseEngine.ini
@@ -2673,7 +2673,7 @@ StoreVersionOffsetArm64=0
 StoreVersionOffsetX86=0
 StoreVersionOffsetX8664=0
 VersionDisplayName=1.0
-MinSDKVersion=19
+MinSDKVersion=26
 TargetSDKVersion=28
 bEnableGradle=true
 bEnableLint=false
@@ -2711,7 +2711,7 @@ bEnableDomStorage=False
 
 [/Script/AndroidPlatformEditor.AndroidSDKSettings]
 SDKAPILevel=latest
-NDKAPILevel=android-21
+NDKAPILevel=android-26
 
 [/Script/UnrealEd.CookerSettings]
 DefaultPVRTCQuality=1
diff --git a/Engine/Source/Programs/UnrealBuildTool/Platform/Android/AndroidToolChain.cs b/Engine/Source/Programs/UnrealBuildTool/Platform/Android/AndroidToolChain.cs
--- a/Engine/Source/Programs/UnrealBuildTool/Platform/Android/AndroidToolChain.cs
+++ b/Engine/Source/Programs/UnrealBuildTool/Platform/Android/AndroidToolChain.cs
@@ -15,6 +15,9 @@
 {
 	class AndroidToolChain : ClangToolChain, IAndroidToolChain
 	{
+		// Minimum NDK API level to allow
+		public const int MinimumNDKAPILevel = 26;
+
 		public static readonly string[] AllCpuSuffixes =
 		{
 			"-arm64",
@@ -261,14 +264,8 @@
 				ArPathx64 = ArPathArm64;
 			}
 
-			// NDK setup (use no less than 21 for 64-bit targets)
-			int NDKApiLevel64Int = GetNdkApiLevelInt();
-			//string NDKApiLevel64Bit = GetNdkApiLevel();
-			if (NDKApiLevel64Int < 21)
-			{
-				NDKApiLevel64Int = 21;
-				//NDKApiLevel64Bit = "android-21";
-			}
+			// NDK setup (enforce minimum API level)
+			int NDKApiLevel64Int = GetNdkApiLevelInt(MinimumNDKAPILevel);
 
 			string GCCToolchainPath = Path.Combine(NDKPath, "toolchains", "llvm", ArchitecturePath);
 			string SysrootPath = Path.Combine(NDKPath, "toolchains", "llvm", ArchitecturePath, "sysroot");
@@ -391,7 +388,7 @@
 			return Arches!;
 		}
 
-		public int GetNdkApiLevelInt(int MinNdk = 21)
+		public int GetNdkApiLevelInt(int MinNdk = MinimumNDKAPILevel)
 		{
 			string NDKVersion = GetNdkApiLevel();
 			int NDKVersionInt = MinNdk;
@@ -447,7 +444,7 @@
 		}
 		
 		//This doesn't take into account SDK version overrides in packaging
-		public int GetMinSdkVersion(int MinSdk = 21)
+		public int GetMinSdkVersion(int MinSdk = MinimumNDKAPILevel)
 		{
 			int MinSDKVersion = MinSdk;
 			ConfigHierarchy Ini = ConfigCache.ReadHierarchy(ConfigHierarchyType.Engine, DirectoryReference.FromFile(ProjectFile), UnrealTargetPlatform.Android);
@@ -1138,13 +1135,13 @@
 			}
 			*/
 
-			// NDK setup (use no less than 21 for 64-bit targets)
-			int NDKApiLevel64Int = GetNdkApiLevelInt();
+			// NDK setup (enforce minimum API level)
+			int NDKApiLevel64Int = GetNdkApiLevelInt(21);		// deliberately use 21 minimum to force NDKApiLevel64Bit to update if below minimum
 			string NDKApiLevel64Bit = GetNdkApiLevel();
-			if (NDKApiLevel64Int < 21)
+			if (NDKApiLevel64Int < MinimumNDKAPILevel)
 			{
-				NDKApiLevel64Int = 21;
-				NDKApiLevel64Bit = "android-21";
+				NDKApiLevel64Int = MinimumNDKAPILevel;
+				NDKApiLevel64Bit = "android-" + MinimumNDKAPILevel;
 			}
 
 			Log.TraceInformationOnce("Compiling Native 64-bit code with NDK API '{0}'", NDKApiLevel64Bit);
diff --git a/Engine/Source/Programs/UnrealBuildTool/Platform/Android/UEDeployAndroid.cs b/Engine/Source/Programs/UnrealBuildTool/Platform/Android/UEDeployAndroid.cs
--- a/Engine/Source/Programs/UnrealBuildTool/Platform/Android/UEDeployAndroid.cs
+++ b/Engine/Source/Programs/UnrealBuildTool/Platform/Android/UEDeployAndroid.cs
@@ -3423,32 +3423,25 @@
 				Logger.LogInformation("Fixing minSdkVersion; requires minSdkVersion of {MinVer} with Gradle based on active plugins", MinimumSDKLevelForGradle);
 			}
 
-			// 64-bit targets must be android-21 or higher
+			// Get NDK API level and enforce minimum
 			NDKLevelInt = ToolChain.GetNdkApiLevelInt();
-			if (NDKLevelInt < 21)
+			if (NDKLevelInt < AndroidToolChain.MinimumNDKAPILevel)
 			{
-				// 21 is requred for GL ES3.1
-				NDKLevelInt = 21;
+				// 21 is required for GL ES3.1, 26 for ANativeWindow_setBuffersTransform
+				NDKLevelInt = AndroidToolChain.MinimumNDKAPILevel;
 			}
 
-			// fix up the MinSdkVersion
-			if (NDKLevelInt > 19)
+			// fix up the MinSdkVersion to be at least NDKLevelInt
+			if (MinSDKVersion < NDKLevelInt)
 			{
-				if (MinSDKVersion < 21)
-				{
-					MinSDKVersion = 21;
-					Logger.LogInformation("Fixing minSdkVersion; NDK level above 19 requires minSdkVersion of 21 (arch={Arch})", Arch.Substring(1));
-				}
+				Logger.LogInformation("Fixing minSdkVersion; NDK level is {NDKLevelInt} which is above minSdkVersion {MinSDKVersion}.", NDKLevelInt, MinSDKVersion);
+				MinSDKVersion = NDKLevelInt;
+			}
 
-				if (MinSDKVersion < NDKLevelInt)
-                {
-                    Logger.LogInformation("Fixing minSdkVersion; NDK level is {NdkLevel} which is above minSdkVersion {MinSdkVersion}.", NDKLevelInt, MinSDKVersion);
-					MinSDKVersion = NDKLevelInt;
-				}
-            }
-
+			// fix up the TargetSDK to be at least MinSdkVersion
 			if (TargetSDKVersion < MinSDKVersion)
 			{
+				Logger.LogInformation("Fixing targetSdkVersion; minSdkVersion is {MinSDKVersion} which is above targetSdkVersion {TargetSDKVersion}.", MinSDKVersion, TargetSDKVersion);
 				TargetSDKVersion = MinSDKVersion;
 			}
 		}
diff --git a/Engine/Source/Runtime/Android/AndroidRuntimeSettings/Classes/AndroidRuntimeSettings.h b/Engine/Source/Runtime/Android/AndroidRuntimeSettings/Classes/AndroidRuntimeSettings.h
--- a/Engine/Source/Runtime/Android/AndroidRuntimeSettings/Classes/AndroidRuntimeSettings.h
+++ b/Engine/Source/Runtime/Android/AndroidRuntimeSettings/Classes/AndroidRuntimeSettings.h
@@ -165,10 +165,6 @@ public:
 	UPROPERTY(GlobalConfig, EditAnywhere, Category = "APK Packaging", Meta = (DisplayName = "Store Version (1-2147483647)", ClampMin="1", ClampMax="2147483647"))
 	int32 StoreVersion;
 
-	// Offset to add to store version for APKs generated for armv7
-	UPROPERTY(GlobalConfig, EditAnywhere, Category = "APK Packaging", meta = (DisplayName = "Store Version offset (armv7)"))
-	int32 StoreVersionOffsetArmV7;
-
 	// Offset to add to store version for APKs generated for arm64
 	UPROPERTY(GlobalConfig, EditAnywhere, Category = "APK Packaging", meta = (DisplayName = "Store Version offset (arm64)"))
 	int32 StoreVersionOffsetArm64;
@@ -185,12 +181,12 @@ public:
 	UPROPERTY(GlobalConfig, EditAnywhere, Category = "APK Packaging", Meta = (DisplayName = "Version Display Name (usually x.y)"))
 	FString VersionDisplayName;
 
-	// What OS version the app is allowed to be installed on (do not set this lower than 19)
-	UPROPERTY(GlobalConfig, EditAnywhere, Category = "APK Packaging", Meta = (DisplayName = "Minimum SDK Version (19=KitKat, 21=Lollipop)"))
+	// What OS version the app is allowed to be installed on (do not set this lower than 26)
+	UPROPERTY(GlobalConfig, EditAnywhere, Category = "APK Packaging", Meta = (DisplayName = "Minimum SDK Version (26=8.0.0, 27=8.1.0, 28=9, 29=10, 30=11, 31=12)"))
 	int32 MinSDKVersion;
 	
-	// What OS version the app is expected to run on (do not set this lower than 19)
-	UPROPERTY(GlobalConfig, EditAnywhere, Category = "APK Packaging", Meta = (DisplayName = "Target SDK Version (19=KitKat, 21=Lollipop)"))
+	// What OS version the app is expected to run on (do not set this lower than 26)
+	UPROPERTY(GlobalConfig, EditAnywhere, Category = "APK Packaging", Meta = (DisplayName = "Target SDK Version (26=8.0.0, 27=8.1.0, 28=9, 29=10, 30=11, 31=12)"))
 	int32 TargetSDKVersion;
 
 	// Preferred install location for the application
