UnrealBuildTool: Rename msvc dependency file to .dep.json to prevent conflict with clang-cl timing file (which can't be renamed)

#rnx

diff --git a/Engine/Source/Programs/UnrealBuildTool/Platform/Windows/VCToolChain.cs b/Engine/Source/Programs/UnrealBuildTool/Platform/Windows/VCToolChain.cs
--- a/Engine/Source/Programs/UnrealBuildTool/Platform/Windows/VCToolChain.cs
+++ b/Engine/Source/Programs/UnrealBuildTool/Platform/Windows/VCToolChain.cs
@@ -1411,7 +1411,7 @@ namespace UnrealBuildTool
 				{
 					if (Target.WindowsPlatform.Compiler.IsMSVC() && !CompileAction.ForceClFilter)
 					{
-						CompileAction.DependencyListFile = FileItem.GetItemByFileReference(FileReference.Combine(OutputDir, $"{SourceFile.Location.GetFileName()}.json"));
+						CompileAction.DependencyListFile = FileItem.GetItemByFileReference(FileReference.Combine(OutputDir, $"{SourceFile.Location.GetFileName()}.dep.json"));
 					}
 					else if (Target.WindowsPlatform.Compiler.IsClang())
 					{
