UnrealBuildTool: UnsafeTypeCastWarnings for msvc clang

#rnx

diff --git a/Engine/Source/Programs/UnrealBuildTool/Platform/Windows/VCToolChain.cs b/Engine/Source/Programs/UnrealBuildTool/Platform/Windows/VCToolChain.cs
--- a/Engine/Source/Programs/UnrealBuildTool/Platform/Windows/VCToolChain.cs
+++ b/Engine/Source/Programs/UnrealBuildTool/Platform/Windows/VCToolChain.cs
@@ -877,18 +877,46 @@ namespace UnrealBuildTool
 
 				if (CompileEnvironment.ShadowVariableWarningLevel != WarningLevel.Off)
 				{
-					Arguments.Add("-Wshadow");
-					if(CompileEnvironment.ShadowVariableWarningLevel == WarningLevel.Warning)
-					{
-						Arguments.Add("-Wno-error=shadow");
-					}
+					Arguments.Add("-Wshadow" + ((CompileEnvironment.ShadowVariableWarningLevel == WarningLevel.Error) ? "" : " -Wno-error=shadow"));
 				}
 
 				if (CompileEnvironment.bEnableUndefinedIdentifierWarnings)
 				{
 					Arguments.Add(" -Wundef" + (CompileEnvironment.bUndefinedIdentifierWarningsAsErrors ? "" : " -Wno-error=undef"));
 				}
-				
+
+				// Note: This should be kept in sync with PRAGMA_DISABLE_UNSAFE_TYPECAST_WARNINGS in ClangPlatformCompilerPreSetup.h
+				string[] UnsafeTypeCastWarningList = {
+					"float-conversion",
+					"implicit-float-conversion",
+					"implicit-int-conversion",
+					"c++11-narrowing"
+					//"shorten-64-to-32",	<-- too many hits right now, probably want it *soon*
+					//"sign-conversion",	<-- too many hits right now, probably want it eventually
+				};
+
+				if (CompileEnvironment.UnsafeTypeCastWarningLevel == WarningLevel.Error)
+				{
+					foreach (string Warning in UnsafeTypeCastWarningList)
+					{
+						Arguments.Add("-W" + Warning);
+					}
+				}
+				else if (CompileEnvironment.UnsafeTypeCastWarningLevel == WarningLevel.Warning)
+				{
+					foreach (string Warning in UnsafeTypeCastWarningList)
+					{
+						Arguments.Add("-W" + Warning + " -Wno-error=" + Warning);
+					}
+				}
+				else
+				{
+					foreach (string Warning in UnsafeTypeCastWarningList)
+					{
+						Arguments.Add("-Wno-" + Warning);
+					}
+				}
+
 				// This is disabled because clang explicitly warns about changing pack alignment in a header and not
 				// restoring it afterwards, which is something we do with the Pre/PostWindowsApi.h headers.
 				Arguments.Add("-Wno-pragma-pack");
