Force 5.0.3 to support only v20 LLVM 13.0.1

Due to v19 no linking correctly as well as older LLVM versions that dont support C++17 anyway

diff --git a/Engine/Source/Programs/UnrealBuildTool/Platform/Linux/LinuxPlatformSDK.Versions.cs b/Engine/Source/Programs/UnrealBuildTool/Platform/Linux/LinuxPlatformSDK.Versions.cs
--- a/Engine/Source/Programs/UnrealBuildTool/Platform/Linux/LinuxPlatformSDK.Versions.cs
+++ b/Engine/Source/Programs/UnrealBuildTool/Platform/Linux/LinuxPlatformSDK.Versions.cs
@@ -14,7 +14,7 @@ namespace UnrealBuildTool
 		protected override void GetValidVersionRange(out string MinVersion, out string MaxVersion)
 		{
 			// all that matters is the number after the v, according to TryConvertVersionToInt()
-			MinVersion = "v10_clang-5.0.0-centos7";
+			MinVersion = "v20_clang-13.0.1-centos7";
 			MaxVersion = "v20_clang-13.0.1-centos7";
 		}
 
diff --git a/Engine/Source/Programs/UnrealBuildTool/Platform/Linux/LinuxToolChain.cs b/Engine/Source/Programs/UnrealBuildTool/Platform/Linux/LinuxToolChain.cs
--- a/Engine/Source/Programs/UnrealBuildTool/Platform/Linux/LinuxToolChain.cs
+++ b/Engine/Source/Programs/UnrealBuildTool/Platform/Linux/LinuxToolChain.cs
@@ -109,10 +109,10 @@ namespace UnrealBuildTool
 			// refuse to use compilers that we know won't work
 
 			// prevent unknown clangs since the build is likely to fail on too old or too new compilers
-			if ((ClangVersionMajor * 10 + ClangVersionMinor) > 130 || (ClangVersionMajor * 10 + ClangVersionMinor) < 100)
+			if ((ClangVersionMajor * 10 + ClangVersionMinor) > 130 || (ClangVersionMajor * 10 + ClangVersionMinor) < 130)
 			{
 				throw new BuildException(
-					string.Format("This version of the Unreal Engine can only be compiled with clang 13.0, 12.0, 11.0, and 10.0. clang {0} may not build it - please use a different version.",
+					string.Format("This version of the Unreal Engine can only be compiled with clang 13.0. clang {0} may not build it - please use a different version.",
 						ClangVersionString)
 					);
 			}
