UnrealBuildTool: Use filesystem references for sysroot and clang tool paths, instead of strings

#rnx

diff --git a/Engine/Source/Programs/UnrealBuildTool/Platform/Linux/LinuxPlatformSDK.cs b/Engine/Source/Programs/UnrealBuildTool/Platform/Linux/LinuxPlatformSDK.cs
--- a/Engine/Source/Programs/UnrealBuildTool/Platform/Linux/LinuxPlatformSDK.cs
+++ b/Engine/Source/Programs/UnrealBuildTool/Platform/Linux/LinuxPlatformSDK.cs
@@ -32,21 +32,21 @@ namespace UnrealBuildTool
 		{
 			// @todo turnkey: ForceUseSystemCompiler() returns true, we should probably run system clang -V or similar to get version
 
-			string? SDKDir = GetSDKLocation();
-			if (string.IsNullOrEmpty(SDKDir))
+			DirectoryReference? SDKDir = GetSDKLocation();
+			if (SDKDir == null)
 			{
 				return null;
 			}
 
-			string VersionFile = Path.Combine(SDKDir, SDKVersionFileName());
+			FileReference VersionFile = FileReference.Combine(SDKDir, SDKVersionFileName());
 
-			if (!File.Exists(VersionFile))
+			if (!FileReference.Exists(VersionFile))
 			{
 				// ErrorMessage = "Cannot use an old toolchain (missing " + PlatformSDK.SDKVersionFileName() + " file, assuming version earlier than v11)";
 				return null;
 			}
 
-			StreamReader SDKVersionFile = new StreamReader(VersionFile);
+			StreamReader SDKVersionFile = new StreamReader(VersionFile.FullName);
 			string? SDKVersionString = SDKVersionFile.ReadLine();
 			SDKVersionFile.Close();
 
@@ -196,12 +196,12 @@ namespace UnrealBuildTool
 		/// WARNING: Do not cache this value - it may be changed after sourcing OutputEnvVars.txt
 		/// </summary>
 		/// <returns>Valid SDK string</returns>
-		public virtual string? GetSDKLocation()
+		public virtual DirectoryReference? GetSDKLocation()
 		{
 			// if new multi-arch toolchain is used, prefer it
-			string? MultiArchRoot = Environment.GetEnvironmentVariable("LINUX_MULTIARCH_ROOT");
+			DirectoryReference? MultiArchRoot = DirectoryReference.FromString(Environment.GetEnvironmentVariable("LINUX_MULTIARCH_ROOT"));
 
-			if (String.IsNullOrEmpty(MultiArchRoot))
+			if (MultiArchRoot == null)
 			{
 				// check if in-tree SDK is available
 				DirectoryReference InTreeSDKVersionRoot = GetInTreeSDKRoot();
@@ -210,7 +210,7 @@ namespace UnrealBuildTool
 					DirectoryReference InTreeSDKVersionPath = DirectoryReference.Combine(InTreeSDKVersionRoot, GetMainVersion());
 					if (DirectoryReference.Exists(InTreeSDKVersionPath))
 					{
-						MultiArchRoot = InTreeSDKVersionPath.FullName;
+						MultiArchRoot = InTreeSDKVersionPath;
 					}
 				}
 			}
@@ -221,21 +221,21 @@ namespace UnrealBuildTool
 		/// Returns the SDK path for a specific architecture
 		/// WARNING: Do not cache this value - it may be changed after sourcing OutputEnvVars.txt
 		/// </summary>
-		/// <returns>Valid SDK string</returns>
-		public virtual string? GetBaseLinuxPathForArchitecture(string Architecture)
+		/// <returns>Valid SDK DirectoryReference</returns>
+		public virtual DirectoryReference? GetBaseLinuxPathForArchitecture(string Architecture)
 		{
 			// if new multi-arch toolchain is used, prefer it
-			string? MultiArchRoot = GetSDKLocation();
-			string? BaseLinuxPath;
+			DirectoryReference? MultiArchRoot = GetSDKLocation();
+			DirectoryReference? BaseLinuxPath;
 
-			if (!String.IsNullOrEmpty(MultiArchRoot))
+			if (MultiArchRoot != null )
 			{
-				BaseLinuxPath = Path.Combine(MultiArchRoot, Architecture);
+				BaseLinuxPath = DirectoryReference.Combine(MultiArchRoot, Architecture);
 			}
 			else
 			{
 				// use cross linux toolchain if LINUX_ROOT is specified
-				BaseLinuxPath = Environment.GetEnvironmentVariable("LINUX_ROOT");
+				BaseLinuxPath = DirectoryReference.FromString(Environment.GetEnvironmentVariable("LINUX_ROOT"));
 			}
 			return BaseLinuxPath;
 		}
@@ -257,7 +257,7 @@ namespace UnrealBuildTool
 			// FIXME: UBT should loop across all the architectures and compile for all the selected ones.
 
 			// do not cache this value - it may be changed after sourcing OutputEnvVars.txt
-			string? BaseLinuxPath = GetBaseLinuxPathForArchitecture(LinuxPlatform.DefaultHostArchitecture);
+			DirectoryReference? BaseLinuxPath = GetBaseLinuxPathForArchitecture(LinuxPlatform.DefaultHostArchitecture);
 
 			if (ForceUseSystemCompiler())
 			{
@@ -266,12 +266,10 @@ namespace UnrealBuildTool
 					return SDKStatus.Valid;
 				}
 			}
-			else if (!String.IsNullOrEmpty(BaseLinuxPath))
+			else if (BaseLinuxPath != null)
 			{
 				// paths to our toolchains if BaseLinuxPath is specified
-				BaseLinuxPath = BaseLinuxPath.Replace("\"", "");
-
-				if (IsValidClangPath(new DirectoryReference(BaseLinuxPath)))
+				if (IsValidClangPath(BaseLinuxPath))
 				{
 					return SDKStatus.Valid;
 				}
diff --git a/Engine/Source/Programs/UnrealBuildTool/Platform/Linux/LinuxToolChain.cs b/Engine/Source/Programs/UnrealBuildTool/Platform/Linux/LinuxToolChain.cs
--- a/Engine/Source/Programs/UnrealBuildTool/Platform/Linux/LinuxToolChain.cs
+++ b/Engine/Source/Programs/UnrealBuildTool/Platform/Linux/LinuxToolChain.cs
@@ -43,20 +43,20 @@ namespace UnrealBuildTool
 			BaseLinuxPath = PlatformSDK.GetBaseLinuxPathForArchitecture(InArchitecture);
 
 			bool bForceUseSystemCompiler = PlatformSDK.ForceUseSystemCompiler();
-			bool bHasValidCompiler = false;
+			bool bHasValidCompiler;
 
 			if (bForceUseSystemCompiler)
 			{
 				// Validate the system toolchain.
-				BaseLinuxPath = "";
-				MultiArchRoot = "";
+				BaseLinuxPath = null;
+				MultiArchRoot = null;
 
 				ToolchainInfo = "system toolchain";
 
 				// use native linux toolchain
-				ClangPath = LinuxCommon.WhichClang(Logger);
-				LlvmArPath = LinuxCommon.Which("llvm-ar", Logger);
-				ObjcopyPath = LinuxCommon.Which("llvm-objcopy", Logger);
+				ClangPath = FileReference.FromString(LinuxCommon.WhichClang(Logger));
+				LlvmArPath = FileReference.FromString(LinuxCommon.Which("llvm-ar", Logger));
+				ObjcopyPath = FileReference.FromString(LinuxCommon.Which("llvm-objcopy", Logger));
 
 				// When compiling on Linux, use a faster way to relink circularly dependent libraries.
 				// Race condition between actions linking to the .so and action overwriting it is avoided thanks to inodes
@@ -68,23 +68,22 @@ namespace UnrealBuildTool
 			}
 			else
 			{
-				if (String.IsNullOrEmpty(BaseLinuxPath))
+				if (BaseLinuxPath == null)
 				{
 					throw new BuildException("LINUX_MULTIARCH_ROOT environment variable is not set; cannot instantiate Linux toolchain");
 				}
-				if (String.IsNullOrEmpty(MultiArchRoot))
+				if (MultiArchRoot == null)
 				{
 					MultiArchRoot = BaseLinuxPath;
 					Logger.LogInformation("Using LINUX_ROOT (deprecated, consider LINUX_MULTIARCH_ROOT)");
 				}
 
-				BaseLinuxPath = BaseLinuxPath.Replace("\"", "").Replace('\\', '/');
-				ToolchainInfo = String.Format("toolchain located at '{0}'", BaseLinuxPath);
+				ToolchainInfo = $"toolchain located at '{BaseLinuxPath}'";
 
 				// set up the path to our toolchain
-				ClangPath = Path.Combine(BaseLinuxPath, @"bin", "clang++" + BuildHostPlatform.Current.BinarySuffix);
-				LlvmArPath = Path.Combine(Path.Combine(BaseLinuxPath, String.Format("bin/{0}", "llvm-ar" + BuildHostPlatform.Current.BinarySuffix)));
-				ObjcopyPath = Path.Combine(Path.Combine(BaseLinuxPath, String.Format("bin/{0}", "llvm-objcopy" + BuildHostPlatform.Current.BinarySuffix)));
+				ClangPath = FileReference.Combine(BaseLinuxPath, "bin", $"clang++{BuildHostPlatform.Current.BinarySuffix}");
+				LlvmArPath = FileReference.Combine(BaseLinuxPath, "bin", $"llvm-ar{BuildHostPlatform.Current.BinarySuffix}" );
+				ObjcopyPath = FileReference.Combine(BaseLinuxPath, "bin", $"llvm-objcopy{BuildHostPlatform.Current.BinarySuffix}");
 
 				// When cross-compiling on Windows, use old FixDeps. It is slow, but it does not have timing issues
 				bUseFixdeps = BuildHostPlatform.Current.Platform == UnrealTargetPlatform.Win64;
@@ -125,8 +124,8 @@ namespace UnrealBuildTool
 			PlatformSDK = InSDK;
 
 			// these are supplied by the engine and do not change depending on the circumstances
-			DumpSymsPath = Path.Combine(Unreal.EngineDirectory.FullName, "Binaries", "Linux", "dump_syms" + BuildHostPlatform.Current.BinarySuffix);
-			BreakpadEncoderPath = Path.Combine(Unreal.EngineDirectory.FullName, "Binaries", "Linux", "BreakpadSymbolEncoder" + BuildHostPlatform.Current.BinarySuffix);
+			DumpSymsPath = FileReference.Combine(Unreal.EngineDirectory, "Binaries", "Linux", $"dump_syms{BuildHostPlatform.Current.BinarySuffix}");
+			BreakpadEncoderPath = FileReference.Combine(Unreal.EngineDirectory, "Binaries", "Linux", $"BreakpadSymbolEncoder{BuildHostPlatform.Current.BinarySuffix}");
 		}
 
 		public override void SetUpGlobalEnvironment(ReadOnlyTargetRules Target)
@@ -284,9 +283,9 @@ namespace UnrealBuildTool
 				Proc.StartInfo.RedirectStandardOutput = true;
 				Proc.StartInfo.RedirectStandardError = true;
 
-				if (!String.IsNullOrEmpty(ClangPath))
+				if (ClangPath != null && FileReference.Exists(ClangPath))
 				{
-					Proc.StartInfo.FileName = ClangPath;
+					Proc.StartInfo.FileName = ClangPath.FullName;
 					Proc.StartInfo.Arguments = " --version";
 
 					Proc.Start();
@@ -334,9 +333,9 @@ namespace UnrealBuildTool
 				Proc.StartInfo.RedirectStandardError = true;
 				Proc.StartInfo.RedirectStandardInput = true;
 
-				if (!String.IsNullOrEmpty(ClangPath) && File.Exists(ClangPath))
+				if (ClangPath != null && FileReference.Exists(ClangPath))
 				{
-					Proc.StartInfo.FileName = ClangPath;
+					Proc.StartInfo.FileName = ClangPath.FullName;
 					Proc.StartInfo.Arguments = " -E -dM -";
 
 					Proc.Start();
@@ -613,9 +612,9 @@ namespace UnrealBuildTool
 			{
 				if (!String.IsNullOrEmpty(CompileEnvironment.Architecture))
 				{
-					Arguments.Add(String.Format("-target {0}", CompileEnvironment.Architecture));        // Set target triple
+					Arguments.Add($"-target {CompileEnvironment.Architecture}");        // Set target triple
 				}
-				Arguments.Add(String.Format("--sysroot=\"{0}\"", BaseLinuxPath));
+				Arguments.Add($"--sysroot=\"{NormalizeCommandLinePath(BaseLinuxPath!)}\"");
 			}
 		}
 
@@ -785,17 +784,17 @@ namespace UnrealBuildTool
 
 			if (CrossCompiling())
 			{
-				Arguments.Add(String.Format(" -target {0}", LinkEnvironment.Architecture));        // Set target triple
-				string SysRootPath = BaseLinuxPath!.TrimEnd(new char[] { '\\', '/' });
-				Arguments.Add(String.Format(" \"--sysroot={0}\"", SysRootPath));
+				Arguments.Add($"-target {LinkEnvironment.Architecture}");        // Set target triple
+				DirectoryReference SysRootPath = BaseLinuxPath!;
+				Arguments.Add($"--sysroot=\"{NormalizeCommandLinePath(SysRootPath)}\"");
 
 				// Linking with the toolchain on linux appears to not search usr/
 				if (BuildHostPlatform.Current.Platform == UnrealTargetPlatform.Linux)
 				{
-					Arguments.Add(String.Format(" -B\"{0}/usr/lib/\"", SysRootPath));
-					Arguments.Add(String.Format(" -B\"{0}/usr/lib64/\"", SysRootPath));
-					Arguments.Add(String.Format(" -L\"{0}/usr/lib/\"", SysRootPath));
-					Arguments.Add(String.Format(" -L\"{0}/usr/lib64/\"", SysRootPath));
+					Arguments.Add($"-B\"{NormalizeCommandLinePath(DirectoryReference.Combine(SysRootPath, "usr", "lib"))}\"");
+					Arguments.Add($"-B\"{NormalizeCommandLinePath(DirectoryReference.Combine(SysRootPath, "usr", "lib64"))}\"");
+					Arguments.Add($"-L\"{NormalizeCommandLinePath(DirectoryReference.Combine(SysRootPath, "usr", "lib"))}\"");
+					Arguments.Add($"-L\"{NormalizeCommandLinePath(DirectoryReference.Combine(SysRootPath, "usr", "lib64"))}\"");
 				}
 			}
 		}
@@ -807,13 +806,13 @@ namespace UnrealBuildTool
 
 		// cache the location of NDK tools
 		protected bool bIsCrossCompiling;
-		protected string? BaseLinuxPath;
-		protected string? ClangPath;
-		protected string? LlvmArPath;
-		protected string? ObjcopyPath;
-		protected string? DumpSymsPath;
-		protected string? BreakpadEncoderPath;
-		protected string? MultiArchRoot;
+		protected DirectoryReference? BaseLinuxPath;
+		protected FileReference? ClangPath;
+		protected FileReference? LlvmArPath;
+		protected FileReference? ObjcopyPath;
+		protected FileReference? DumpSymsPath;
+		protected FileReference? BreakpadEncoderPath;
+		protected DirectoryReference? MultiArchRoot;
 
 		/// <summary>
 		/// Whether to use old, slower way to relink circularly dependent libraries.
@@ -1015,7 +1014,7 @@ namespace UnrealBuildTool
 				}
 
 				CompileAction.WorkingDirectory = Unreal.EngineSourceDirectory;
-				CompileAction.CommandPath = new FileReference(ClangPath!);
+				CompileAction.CommandPath = ClangPath!;
 
 				List<string> ResponseFileContents = new();
 				ResponseFileContents.AddRange(GlobalArguments);
@@ -1100,7 +1099,7 @@ namespace UnrealBuildTool
 			// Create an archive action
 			Action ArchiveAction = Graph.CreateAction(ActionType.Link);
 			ArchiveAction.WorkingDirectory = Unreal.EngineSourceDirectory;
-			ArchiveAction.CommandPath = new FileReference(LlvmArPath!);
+			ArchiveAction.CommandPath = LlvmArPath!;
 
 			// this will produce a final library
 			ArchiveAction.bProducesImportLibrary = true;
@@ -1805,7 +1804,7 @@ namespace UnrealBuildTool
 			}
 
 			ProcessStartInfo StartInfo = new ProcessStartInfo();
-			StartInfo.FileName = ObjcopyPath!;
+			StartInfo.FileName = ObjcopyPath!.FullName;
 			StartInfo.Arguments = "--strip-debug \"" + TargetFile.FullName + "\"";
 			StartInfo.UseShellExecute = false;
 			StartInfo.CreateNoWindow = true;
diff --git a/Engine/Source/Programs/UnrealBuildTool/Platform/Linux/UEBuildLinux.cs b/Engine/Source/Programs/UnrealBuildTool/Platform/Linux/UEBuildLinux.cs
--- a/Engine/Source/Programs/UnrealBuildTool/Platform/Linux/UEBuildLinux.cs
+++ b/Engine/Source/Programs/UnrealBuildTool/Platform/Linux/UEBuildLinux.cs
@@ -464,8 +464,8 @@ namespace UnrealBuildTool
 		public override void SetUpEnvironment(ReadOnlyTargetRules Target, CppCompileEnvironment CompileEnvironment, LinkEnvironment LinkEnvironment)
 		{
 			// During the native builds, check the system includes as well (check toolchain when cross-compiling?)
-			string? BaseLinuxPath = SDK.GetBaseLinuxPathForArchitecture(Target.Architecture);
-			if (BuildHostPlatform.Current.Platform == UnrealTargetPlatform.Linux && String.IsNullOrEmpty(BaseLinuxPath))
+			DirectoryReference? BaseLinuxPath = SDK.GetBaseLinuxPathForArchitecture(Target.Architecture);
+			if (BuildHostPlatform.Current.Platform == UnrealTargetPlatform.Linux && BaseLinuxPath == null)
 			{
 				CompileEnvironment.SystemIncludePaths.Add(new DirectoryReference("/usr/include"));
 			}
