UnrealBuildTool: Clean up how pch is forced included for clang toolchains

#rnx

diff --git a/Engine/Source/Programs/UnrealBuildTool/Platform/IOS/IOSToolChain.cs b/Engine/Source/Programs/UnrealBuildTool/Platform/IOS/IOSToolChain.cs
--- a/Engine/Source/Programs/UnrealBuildTool/Platform/IOS/IOSToolChain.cs
+++ b/Engine/Source/Programs/UnrealBuildTool/Platform/IOS/IOSToolChain.cs
@@ -489,20 +489,6 @@ namespace UnrealBuildTool
 			List<string> GlobalArguments = new();
 
 			GetCompileArguments_Global(CompileEnvironment, GlobalArguments);
-			List<string> PCHArguments = new();
-
-			if (CompileEnvironment.PrecompiledHeaderAction == PrecompiledHeaderAction.Include)
-			{
-				// Add the precompiled header file's path to the include path so GCC can find it.
-				// This needs to be before the other include paths to ensure GCC uses it instead of the source header file.
-				PCHArguments.Add(GetForceIncludeFileArgument(CompileEnvironment.PrecompiledHeaderIncludeFilename!));
-				if (GetClangVersion().Major >= 11)
-				{
-					PCHArguments.Add("-fpch-validate-input-files-content");
-				}
-			}
-
-			GetCompileArguments_ForceInclude(CompileEnvironment, GlobalArguments);
 
 			CPPOutput Result = new CPPOutput();
 			// Create a compile action for each source file.
@@ -544,9 +530,6 @@ namespace UnrealBuildTool
 					GetCompileArguments_CPP(CompileEnvironment, FileArguments);
 					FileArguments.Add(GetRTTIFlag(CompileEnvironment));
 					FileArguments.Add(GetObjCExceptionsFlag(CompileEnvironment));
-
-					// only use PCH for .cpp files
-					FilePCHArguments.AddRange(PCHArguments);
 				}
 
 				// Add the C++ source file and its included files to the prerequisite item list.
diff --git a/Engine/Source/Programs/UnrealBuildTool/Platform/Linux/LinuxToolChain.cs b/Engine/Source/Programs/UnrealBuildTool/Platform/Linux/LinuxToolChain.cs
--- a/Engine/Source/Programs/UnrealBuildTool/Platform/Linux/LinuxToolChain.cs
+++ b/Engine/Source/Programs/UnrealBuildTool/Platform/Linux/LinuxToolChain.cs
@@ -901,7 +901,6 @@ namespace UnrealBuildTool
 		{
 			List<string> GlobalArguments = new();
 			GetCompileArguments_Global(CompileEnvironment, GlobalArguments);
-			List<string> PCHArguments = new();
 
 			//var BuildPlatform = UEBuildPlatform.GetBuildPlatform(CompileEnvironment.Platform);
 
@@ -912,16 +911,6 @@ namespace UnrealBuildTool
 				bHasPrintedBuildDetails = true;
 			}
 
-			if (CompileEnvironment.PrecompiledHeaderAction == PrecompiledHeaderAction.Include)
-			{
-				// Validate PCH inputs by content if mtime check fails
-				if (CompilerVersionGreaterOrEqual(11, 0, 0))
-				{
-					PCHArguments.Add("-fpch-validate-input-files-content");
-				}
-				PCHArguments.Add(GetForceIncludeFileArgument(CompileEnvironment.PrecompiledHeaderIncludeFilename!));
-			}
-
 			// Create a compile action for each source file.
 			CPPOutput Result = new CPPOutput();
 			foreach (FileItem SourceFile in InputFiles)
@@ -956,14 +945,10 @@ namespace UnrealBuildTool
 				}
 				else
 				{
+					// Compile the file as C++ code.
 					GetCompileArguments_CPP(CompileEnvironment, FileArguments);
-
-					// only use PCH for .cpp files
-					FileArguments.AddRange(PCHArguments);
 				}
 
-				GetCompileArguments_ForceInclude(CompileEnvironment, FileArguments);
-
 				// Add the C++ source file and its included files to the prerequisite item list.
 				CompileAction.PrerequisiteItems.Add(SourceFile);
 
diff --git a/Engine/Source/Programs/UnrealBuildTool/Platform/Mac/MacToolChain.cs b/Engine/Source/Programs/UnrealBuildTool/Platform/Mac/MacToolChain.cs
--- a/Engine/Source/Programs/UnrealBuildTool/Platform/Mac/MacToolChain.cs
+++ b/Engine/Source/Programs/UnrealBuildTool/Platform/Mac/MacToolChain.cs
@@ -279,18 +279,6 @@ namespace UnrealBuildTool
 			List<string> GlobalArguments = new();
 
 			GetCompileArguments_Global(CompileEnvironment, GlobalArguments);
-			List<string> PCHArguments = new();
-
-			if (CompileEnvironment.PrecompiledHeaderAction == PrecompiledHeaderAction.Include)
-			{
-				// Add the precompiled header file's path to the include path so GCC can find it.
-				// This needs to be before the other include paths to ensure GCC uses it instead of the source header file.
-				PCHArguments.Add(GetForceIncludeFileArgument(CompileEnvironment.PrecompiledHeaderIncludeFilename!));
-				if (GetClangVersion().Major >= 11)
-				{
-					PCHArguments.Add("-fpch-validate-input-files-content");
-				}
-			}
 
 			List<string> FrameworksSearchPaths = new List<string>();
 			foreach (UEBuildFramework Framework in CompileEnvironment.AdditionalFrameworks)
@@ -341,13 +329,8 @@ namespace UnrealBuildTool
 					// Compile the file as C++ code.
 					GetCompileArguments_CPP(CompileEnvironment, FileArguments);
 					FileArguments.Add(GetRTTIFlag(CompileEnvironment));
-
-					// only use PCH for .cpp files
-					FileArguments.AddRange(PCHArguments);
 				}
 
-				GetCompileArguments_ForceInclude(CompileEnvironment, FileArguments);
-
 				// Add the C++ source file and its included files to the prerequisite item list.
 				CompileAction.PrerequisiteItems.Add(SourceFile);
 
diff --git a/Engine/Source/Programs/UnrealBuildTool/ToolChain/AppleToolChain.cs b/Engine/Source/Programs/UnrealBuildTool/ToolChain/AppleToolChain.cs
--- a/Engine/Source/Programs/UnrealBuildTool/ToolChain/AppleToolChain.cs
+++ b/Engine/Source/Programs/UnrealBuildTool/ToolChain/AppleToolChain.cs
@@ -225,6 +225,12 @@ namespace UnrealBuildTool
 			Arguments.Add("-x objective-c++");
 			GetCppStandardCompileArgument(CompileEnvironment, Arguments);
 			Arguments.Add("-stdlib=libc++");
+
+			// Only force include PCH for c++ files
+			if (CompileEnvironment.PrecompiledHeaderAction == PrecompiledHeaderAction.Include)
+			{
+				Arguments.Remove(GetForceIncludeFileArgument(CompileEnvironment.PrecompiledHeaderIncludeFilename!));
+			}
 		}
 
 		protected override void GetCompileArguments_M(CppCompileEnvironment CompileEnvironment, List<string> Arguments)
@@ -232,6 +238,12 @@ namespace UnrealBuildTool
 			Arguments.Add("-x objective-c");
 			GetCppStandardCompileArgument(CompileEnvironment, Arguments);
 			Arguments.Add("-stdlib=libc++");
+
+			// Only force include PCH for c++ files
+			if (CompileEnvironment.PrecompiledHeaderAction == PrecompiledHeaderAction.Include)
+			{
+				Arguments.Remove(GetForceIncludeFileArgument(CompileEnvironment.PrecompiledHeaderIncludeFilename!));
+			}
 		}
 
 		protected override void GetCompileArguments_PCH(CppCompileEnvironment CompileEnvironment, List<string> Arguments)
diff --git a/Engine/Source/Programs/UnrealBuildTool/ToolChain/ClangToolChain.cs b/Engine/Source/Programs/UnrealBuildTool/ToolChain/ClangToolChain.cs
--- a/Engine/Source/Programs/UnrealBuildTool/ToolChain/ClangToolChain.cs
+++ b/Engine/Source/Programs/UnrealBuildTool/ToolChain/ClangToolChain.cs
@@ -238,12 +238,23 @@ namespace UnrealBuildTool
 		protected virtual void GetCompileArguments_CPP(CppCompileEnvironment CompileEnvironment, List<string> Arguments)
 		{
 			Arguments.Add("-x c++");
+			if (CompileEnvironment.PrecompiledHeaderAction == PrecompiledHeaderAction.Include && CompilerVersionGreaterOrEqual(11, 0, 0))
+			{
+				// Validate PCH inputs by content if mtime check fails
+				Arguments.Add("-fpch-validate-input-files-content");
+			}
 			GetCppStandardCompileArgument(CompileEnvironment, Arguments);
 		}
 
 		protected virtual void GetCompileArguments_C(CppCompileEnvironment CompileEnvironment, List<string> Arguments)
 		{
 			Arguments.Add("-x c");
+
+			// Only force include PCH for c++ files
+			if (CompileEnvironment.PrecompiledHeaderAction == PrecompiledHeaderAction.Include)
+			{
+				Arguments.Remove(GetForceIncludeFileArgument(CompileEnvironment.PrecompiledHeaderIncludeFilename!));
+			}
 		}
 
 		protected virtual void GetCompileArguments_MM(CppCompileEnvironment CompileEnvironment, List<string> Arguments)
@@ -252,6 +263,12 @@ namespace UnrealBuildTool
 			Arguments.Add("-fobjc-abi-version=2");
 			Arguments.Add("-fobjc-legacy-dispatch");
 			GetCppStandardCompileArgument(CompileEnvironment, Arguments);
+
+			// Only force include PCH for c++ files
+			if (CompileEnvironment.PrecompiledHeaderAction == PrecompiledHeaderAction.Include)
+			{
+				Arguments.Remove(GetForceIncludeFileArgument(CompileEnvironment.PrecompiledHeaderIncludeFilename!));
+			}
 		}
 
 		protected virtual void GetCompileArguments_M(CppCompileEnvironment CompileEnvironment, List<string> Arguments)
@@ -260,6 +277,12 @@ namespace UnrealBuildTool
 			Arguments.Add("-fobjc-abi-version=2");
 			Arguments.Add("-fobjc-legacy-dispatch");
 			GetCppStandardCompileArgument(CompileEnvironment, Arguments);
+
+			// Only force include PCH for c++ files
+			if (CompileEnvironment.PrecompiledHeaderAction == PrecompiledHeaderAction.Include)
+			{
+				Arguments.Remove(GetForceIncludeFileArgument(CompileEnvironment.PrecompiledHeaderIncludeFilename!));
+			}
 		}
 
 		protected virtual void GetCompileArguments_PCH(CppCompileEnvironment CompileEnvironment, List<string> Arguments)
@@ -267,6 +290,7 @@ namespace UnrealBuildTool
 			Arguments.Add("-x c++-header");
 			if (CompilerVersionGreaterOrEqual(11, 0, 0))
 			{
+				// Validate PCH inputs by content if mtime check fails
 				Arguments.Add("-fpch-validate-input-files-content");
 				Arguments.Add("-fpch-instantiate-templates");
 			}
@@ -319,6 +343,14 @@ namespace UnrealBuildTool
 
 		protected virtual void GetCompileArguments_ForceInclude(CppCompileEnvironment CompileEnvironment, List<string> Arguments)
 		{
+			// Add the precompiled header file's path to the force include path
+			// This needs to be before the other force include paths to ensure clang uses it instead of the source header file.
+			// Will be removed if compiling non c++ files.
+			if (CompileEnvironment.PrecompiledHeaderAction == PrecompiledHeaderAction.Include)
+			{
+				Arguments.Add(GetForceIncludeFileArgument(CompileEnvironment.PrecompiledHeaderIncludeFilename!));
+			}
+
 			Arguments.AddRange(CompileEnvironment.ForceIncludeFiles.Select(ForceIncludeFile => GetForceIncludeFileArgument(ForceIncludeFile)));
 		}
 
@@ -527,6 +559,9 @@ namespace UnrealBuildTool
 			// Add include paths to the argument list.
 			GetCompileArguments_IncludePaths(CompileEnvironment, Arguments);
 
+			// Add force include paths to the argument list.
+			GetCompileArguments_ForceInclude(CompileEnvironment, Arguments);
+
 			// Add preprocessor definitions to the argument list.
 			GetCompileArguments_PreprocessorDefinitions(CompileEnvironment, Arguments);
 
