Low Level Tests .ini configuration section

diff --git a/Engine/Source/Programs/UnrealBuildTool/Configuration/TestModuleRules.cs b/Engine/Source/Programs/UnrealBuildTool/Configuration/TestModuleRules.cs
--- a/Engine/Source/Programs/UnrealBuildTool/Configuration/TestModuleRules.cs
+++ b/Engine/Source/Programs/UnrealBuildTool/Configuration/TestModuleRules.cs
@@ -15,6 +15,9 @@ namespace UnrealBuildTool
 	/// </summary>
 	public class TestModuleRules : ModuleRules
 	{
+		[ConfigFile(ConfigHierarchyType.Engine, "LowLevelTestsSettings", "UpdateBuildGraphPropertiesFile")]
+		bool bUpdateBuildGraphPropertiesFile = false;
+
 		private readonly XNamespace BuildGraphNamespace = XNamespace.Get("http://www.epicgames.com/BuildGraph");
 		private readonly XNamespace SchemaInstance = XNamespace.Get("http://www.w3.org/2001/XMLSchema-instance");
 		private readonly XNamespace SchemaLocation = XNamespace.Get("http://www.epicgames.com/BuildGraph ../../Build/Graph/Schema.xsd");
@@ -171,6 +174,11 @@ namespace UnrealBuildTool
 		/// </summary>
 		protected void UpdateBuildGraphPropertiesFile(Metadata TestMetadata, bool AddToTestNames = true)
 		{
+			if (!bUpdateBuildGraphPropertiesFile)
+			{
+				return;
+			}
+
 			bool IsPublic = false;
 			string GeneratedPropertiesScriptFile;
 
