UnrealBuildTool: Report in the log when a dependency file is missing instead of printing an invalid timestamp

#rnx

diff --git a/Engine/Source/Programs/UnrealBuildTool/System/ActionGraph.cs b/Engine/Source/Programs/UnrealBuildTool/System/ActionGraph.cs
--- a/Engine/Source/Programs/UnrealBuildTool/System/ActionGraph.cs
+++ b/Engine/Source/Programs/UnrealBuildTool/System/ActionGraph.cs
@@ -703,7 +703,13 @@ namespace UnrealBuildTool
 				{
 					foreach (FileItem DependencyFile in DependencyFiles)
 					{
-						if (!DependencyFile.Exists || DependencyFile.LastWriteTimeUtc > LastExecutionTimeUtc)
+						if (!DependencyFile.Exists)
+						{
+							Logger.LogDebug("{RootAction.StatusDescription}: Dependency {DependencyFile} doesn't exist", RootAction.StatusDescription, DependencyFile.AbsolutePath);
+							bIsOutdated = true;
+							break;
+						}
+						else if (DependencyFile.LastWriteTimeUtc > LastExecutionTimeUtc)
 						{
 							Logger.LogDebug("{RootAction.StatusDescription}: Dependency {DependencyFile} is newer than the last execution of the action: {Message}", RootAction.StatusDescription, DependencyFile.AbsolutePath,
 								$"{DependencyFile.LastWriteTimeUtc.ToLocalTime().ToString(CultureInfo.CurrentCulture)} vs {LastExecutionTimeUtc.LocalDateTime.ToString(CultureInfo.CurrentCulture)}");
