Added Android Toolchain option to disable stack protector

diff --git a/Engine/Source/Programs/UnrealBuildTool/Platform/Android/AndroidToolChain.cs b/Engine/Source/Programs/UnrealBuildTool/Platform/Android/AndroidToolChain.cs
--- a/Engine/Source/Programs/UnrealBuildTool/Platform/Android/AndroidToolChain.cs
+++ b/Engine/Source/Programs/UnrealBuildTool/Platform/Android/AndroidToolChain.cs
@@ -366,6 +366,13 @@ namespace UnrealBuildTool
 			return Ini.GetBool("/Script/AndroidRuntimeSettings.AndroidRuntimeSettings", "bEnableAdvancedBinaryCompression", out bEnableAdvancedBinaryCompression) && bEnableAdvancedBinaryCompression;
 		}
 
+		private bool DisableStackProtector()
+		{
+			ConfigHierarchy Ini = ConfigCache.ReadHierarchy(ConfigHierarchyType.Engine, DirectoryReference.FromFile(ProjectFile), UnrealTargetPlatform.Android);
+			bool bDisableStackProtector = false;
+			return Ini.GetBool("/Script/AndroidRuntimeSettings.AndroidRuntimeSettings", "bDisableStackProtector", out bDisableStackProtector) && bDisableStackProtector;
+		}
+
 		private string GetVersionScriptFilename(LinkEnvironment LinkEnvironment)
 		{
 			return Path.Combine(LinkEnvironment.IntermediateDirectory!.FullName, "ExportSymbols.ldscript");
@@ -723,7 +730,6 @@ namespace UnrealBuildTool
 			if (Architecture == "-arm64")
 			{
 				Result += " -funwind-tables";           // Just generates any needed static data, affects no code
-				Result += " -fstack-protector-strong";  // Emits extra code to check for buffer overflows
 				Result += " -fno-strict-aliasing";      // Prevents unwanted or invalid optimizations that could produce incorrect code
 				Result += " -fPIC";                     // Generates position-independent code (PIC) suitable for use in a shared library
 				Result += " -fno-short-enums";          // Do not allocate to an enum type only as many bytes as it needs for the declared range of possible values
@@ -755,7 +761,6 @@ namespace UnrealBuildTool
 			else if (Architecture == "-x64")
 			{
 				Result += " -funwind-tables";           // Just generates any needed static data, affects no code
-				Result += " -fstack-protector-strong";  // Emits extra code to check for buffer overflows
 				Result += " -fPIC";                     // Generates position-independent code (PIC) suitable for use in a shared library
 				Result += " -fno-omit-frame-pointer";
 				Result += " -fno-strict-aliasing";
@@ -774,6 +779,11 @@ namespace UnrealBuildTool
 				}
 			}
 
+			if (!DisableStackProtector())
+			{
+				Result += " -fstack-protector-strong";  // Emits extra code to check for buffer overflows
+			}
+
 			Result += " -fforce-emit-vtables";      // Helps with devirtualization
 			//Result += " -fstrict-vtable-pointers";	// Assumes that vtable pointer will never change, i.e. object of a different type would not be constracted in place of another object. Helps with devirtualization
 
