Set minimum version to iOS14 in IOSPlatformSDK

diff --git a/Engine/Source/Programs/UnrealBuildTool/Platform/IOS/IOSPlatformSDK.Versions.cs b/Engine/Source/Programs/UnrealBuildTool/Platform/IOS/IOSPlatformSDK.Versions.cs
--- a/Engine/Source/Programs/UnrealBuildTool/Platform/IOS/IOSPlatformSDK.Versions.cs
+++ b/Engine/Source/Programs/UnrealBuildTool/Platform/IOS/IOSPlatformSDK.Versions.cs
@@ -1,21 +1,21 @@
-// Copyright Epic Games, Inc. All Rights Reserved.
-
-using Microsoft.Extensions.Logging;
-
-namespace UnrealBuildTool
-{
-	internal class IOSPlatformSDK : ApplePlatformSDK
-	{
-		public IOSPlatformSDK(ILogger Logger)
-			: base(Logger)
-		{
-		}
-
-		protected override void GetValidSoftwareVersionRange(out string MinVersion, out string? MaxVersion)
-		{
-			// what is our min IOS version?
-			MinVersion = "12.0";
-			MaxVersion = null;
-		}
-	}
-}
+// Copyright Epic Games, Inc. All Rights Reserved.
+
+using Microsoft.Extensions.Logging;
+
+namespace UnrealBuildTool
+{
+	internal class IOSPlatformSDK : ApplePlatformSDK
+	{
+		public IOSPlatformSDK(ILogger Logger)
+			: base(Logger)
+		{
+		}
+
+		protected override void GetValidSoftwareVersionRange(out string MinVersion, out string? MaxVersion)
+		{
+			// what is our min IOS version?
+			MinVersion = "14.0";
+			MaxVersion = null;
+		}
+	}
+}
