UnrealBuildTool: Reconcile file type specific clang arguments

#rnx

diff --git a/Engine/Source/Programs/UnrealBuildTool/Platform/IOS/IOSToolChain.cs b/Engine/Source/Programs/UnrealBuildTool/Platform/IOS/IOSToolChain.cs
--- a/Engine/Source/Programs/UnrealBuildTool/Platform/IOS/IOSToolChain.cs
+++ b/Engine/Source/Programs/UnrealBuildTool/Platform/IOS/IOSToolChain.cs
@@ -269,8 +269,6 @@ namespace UnrealBuildTool
 		{
 			base.GetCompileArguments_Global(CompileEnvironment, Arguments);
 
-			Arguments.Add(GetRTTIFlag(CompileEnvironment));
-
 			if (IsBitcodeCompilingEnabled(CompileEnvironment.Configuration))
 			{
 				Arguments.Add("-fembed-bitcode");
@@ -279,7 +277,7 @@ namespace UnrealBuildTool
 			// What architecture(s) to build for
 			Arguments.Add(GetArchitectureArgument(CompileEnvironment.Configuration, CompileEnvironment.Architecture));
 
-			Arguments.Add(string.Format(" -isysroot \"{0}\"", Settings.Value.GetSDKPath(CompileEnvironment.Architecture)));
+			Arguments.Add($"-isysroot \"{Settings.Value.GetSDKPath(CompileEnvironment.Architecture)}\"");
 
 			Arguments.Add("-m" + GetXcodeMinVersionParam() + "=" + ProjectSettings.RuntimeVersion);
 
@@ -295,11 +293,24 @@ namespace UnrealBuildTool
 					{
 						FrameworkDir = Path.GetDirectoryName(FrameworkDir)!;
 					}
-					Arguments.Add(String.Format(" -F\"{0}\"", FrameworkDir));
+					Arguments.Add($"-F\"{FrameworkDir}\"");
 				}
 			}
 		}
 
+		/// <inheritdoc/>
+		protected override void GetCompileArguments_FileType(CppCompileEnvironment CompileEnvironment, FileItem SourceFile, DirectoryReference OutputDir, List<string> Arguments, Action CompileAction, CPPOutput CompileResult)
+		{
+			base.GetCompileArguments_FileType(CompileEnvironment, SourceFile, OutputDir, Arguments, CompileAction, CompileResult);
+
+			string Extension = Path.GetExtension(SourceFile.AbsolutePath).ToUpperInvariant();
+
+			if (!Extension.Equals(".C"))
+			{
+				Arguments.Add(GetObjCExceptionsFlag(CompileEnvironment));
+			}
+		}
+
 		// Conditionally enable (default disabled) Objective-C exceptions
 		static string GetObjCExceptionsFlag(CppCompileEnvironment CompileEnvironment)
 		{
@@ -495,125 +506,20 @@ namespace UnrealBuildTool
 			foreach (FileItem SourceFile in InputFiles)
 			{
 				Action CompileAction = Graph.CreateAction(ActionType.Compile);
-				List<string> FilePCHArguments = new();
-				List<string> FileArguments = new();
-				string Extension = Path.GetExtension(SourceFile.AbsolutePath).ToUpperInvariant();
 
-				if (CompileEnvironment.PrecompiledHeaderAction == PrecompiledHeaderAction.Create)
-				{
-					// Compile the file as a C++ PCH.
-					GetCompileArguments_PCH(CompileEnvironment, FileArguments);
-					FileArguments.Add(GetRTTIFlag(CompileEnvironment));
-					FileArguments.Add(GetObjCExceptionsFlag(CompileEnvironment));
-				}
-				else if (Extension == ".C")
-				{
-					// Compile the file as C code.
-					GetCompileArguments_C(CompileEnvironment, FileArguments);
-				}
-				else if (Extension == ".MM")
-				{
-					// Compile the file as Objective-C++ code.
-					GetCompileArguments_MM(CompileEnvironment, FileArguments);
-					FileArguments.Add(GetRTTIFlag(CompileEnvironment));
-					FileArguments.Add(GetObjCExceptionsFlag(CompileEnvironment));
-				}
-				else if (Extension == ".M")
-				{
-					// Compile the file as Objective-C code.
-					GetCompileArguments_M(CompileEnvironment, FileArguments);
-					FileArguments.Add(GetObjCExceptionsFlag(CompileEnvironment));
-				}
-				else
-				{
-					// Compile the file as C++ code.
-					GetCompileArguments_CPP(CompileEnvironment, FileArguments);
-					FileArguments.Add(GetRTTIFlag(CompileEnvironment));
-					FileArguments.Add(GetObjCExceptionsFlag(CompileEnvironment));
-				}
-
-				// Add the C++ source file and its included files to the prerequisite item list.
-				CompileAction.PrerequisiteItems.Add(SourceFile);
-
-				// Add the precompiled header
-				if (CompileEnvironment.PrecompiledHeaderAction == PrecompiledHeaderAction.Include)
-				{
-					CompileAction.PrerequisiteItems.Add(CompileEnvironment.PrecompiledHeaderFile!);
-				}
-
-				// Upload the force included files
-				foreach (FileItem ForceIncludeFile in CompileEnvironment.ForceIncludeFiles)
-				{
-					CompileAction.PrerequisiteItems.Add(ForceIncludeFile);
-				}
-
-				string? OutputFilePath = null;
-				if (CompileEnvironment.PrecompiledHeaderAction == PrecompiledHeaderAction.Create)
-				{
-					// Add the precompiled header file to the produced item list.
-					FileItem PrecompiledHeaderFile = FileItem.GetItemByFileReference(FileReference.Combine(OutputDir, Path.GetFileName(SourceFile.AbsolutePath) + ".gch"));
-
-					CompileAction.ProducedItems.Add(PrecompiledHeaderFile);
-					Result.PrecompiledHeaderFile = PrecompiledHeaderFile;
-
-					// Add the parameters needed to compile the precompiled header file to the command-line.
-					FileArguments.Add(string.Format("-o \"{0}\"", NormalizeCommandLinePath(PrecompiledHeaderFile)));
-				}
-				else
-				{
-					if (CompileEnvironment.PrecompiledHeaderAction == PrecompiledHeaderAction.Include)
-					{
-						CompileAction.PrerequisiteItems.Add(CompileEnvironment.PrecompiledHeaderFile!);
-					}
-
-					string ObjectFileExtension;
-					if (CompileEnvironment.AdditionalArguments != null && CompileEnvironment.AdditionalArguments.Contains("-emit-llvm"))
-					{
-						ObjectFileExtension = ".bc";
-					}
-					else
-					{
-						ObjectFileExtension = ".o";
-					}
-
-					// Add the object file to the produced item list.
-					FileItem ObjectFile = FileItem.GetItemByFileReference(FileReference.Combine(OutputDir, Path.GetFileName(SourceFile.AbsolutePath) + ObjectFileExtension));
-
-					CompileAction.ProducedItems.Add(ObjectFile);
-					Result.ObjectFiles.Add(ObjectFile);
-					FileArguments.Add(string.Format("-o \"{0}\"", NormalizeCommandLinePath(ObjectFile)));
-					OutputFilePath = ObjectFile.AbsolutePath;
-				}
-
-				// Add the source file path to the command-line.
-				FileArguments.Add(string.Format("\"{0}\"", NormalizeCommandLinePath(SourceFile)));
-
-				// Generate the timing info
-				if (CompileEnvironment.bPrintTimingInfo)
-				{
-					FileItem TraceFile = FileItem.GetItemByFileReference(FileReference.Combine(OutputDir, Path.GetFileName(SourceFile.AbsolutePath) + ".json"));
-					FileArguments.Add("-ftime-trace");
-					CompileAction.ProducedItems.Add(TraceFile);
-				}
+				List<string> FileArguments = new();
 
-				// Generate the included header dependency list
-				if (CompileEnvironment.bGenerateDependenciesFile)
-				{
-					FileItem DependencyListFile = FileItem.GetItemByFileReference(FileReference.Combine(OutputDir, Path.GetFileName(SourceFile.AbsolutePath) + ".d"));
-					FileArguments.Add(string.Format("-MD -MF\"{0}\"", NormalizeCommandLinePath(DependencyListFile)));
-					CompileAction.DependencyListFile = DependencyListFile;
-					CompileAction.ProducedItems.Add(DependencyListFile);
-				}
+				// Add C or C++ specific compiler arguments.
+				GetCompileArguments_FileType(CompileEnvironment, SourceFile, OutputDir, FileArguments, CompileAction, Result);
 
 				string CompilerPath = Settings.Value.ToolchainDir + IOSCompiler;
 
 				// Gets the target file so we can get the correct output path.
-				FileItem TargetFile = CompileAction.ProducedItems[0];
+				FileItem TargetFile = CompileAction.ProducedItems.First();
 
 				// Creates the path to the response file using the name of the output file and creates its contents.
 				FileReference ResponseFileName = new FileReference(TargetFile.AbsolutePath + ".response");
 				List<string> ResponseFileContents = new();
-				ResponseFileContents.AddRange(FilePCHArguments);
 				ResponseFileContents.AddRange(GlobalArguments);
 				ResponseFileContents.AddRange(FileArguments);
 				if (!string.IsNullOrWhiteSpace(CompileEnvironment.AdditionalArguments))
@@ -639,9 +545,10 @@ namespace UnrealBuildTool
 
 				// Analyze and then compile using the shell to perform the indirection
 				string? StaticAnalysisMode = Environment.GetEnvironmentVariable("CLANG_STATIC_ANALYZER_MODE");
-				if (StaticAnalysisMode != null && StaticAnalysisMode != "" && OutputFilePath != null)
+				if (StaticAnalysisMode != null && StaticAnalysisMode != "")
 				{
-					CompileAction.CommandArguments = "-c \"" + CompilerPath + " " + string.Format("@\"{0}\"", ResponseFileName) + " --analyze -Wno-unused-command-line-argument -Xclang -analyzer-output=html -Xclang -analyzer-config -Xclang path-diagnostics-alternate=true -Xclang -analyzer-config -Xclang report-in-main-source-file=true -Xclang -analyzer-disable-checker -Xclang deadcode.DeadStores -o " + OutputFilePath.Replace(".o", ".html") + "; " + CompilerPath + " " + string.Format("@\"{0}\"", ResponseFileName) + "\"";
+					FileReference ReportFile = new FileReference(SourceFile.AbsolutePath + ".html");
+					CompileAction.CommandArguments = "-c \"" + CompilerPath + " " + string.Format("@\"{0}\"", ResponseFileName) + " --analyze -Wno-unused-command-line-argument -Xclang -analyzer-output=html -Xclang -analyzer-config -Xclang path-diagnostics-alternate=true -Xclang -analyzer-config -Xclang report-in-main-source-file=true -Xclang -analyzer-disable-checker -Xclang deadcode.DeadStores -o " + ReportFile + "; " + CompilerPath + " " + string.Format("@\"{0}\"", ResponseFileName) + "\"";
 					CompilerPath = "/bin/sh";
 				}
 
diff --git a/Engine/Source/Programs/UnrealBuildTool/Platform/Linux/LinuxToolChain.cs b/Engine/Source/Programs/UnrealBuildTool/Platform/Linux/LinuxToolChain.cs
--- a/Engine/Source/Programs/UnrealBuildTool/Platform/Linux/LinuxToolChain.cs
+++ b/Engine/Source/Programs/UnrealBuildTool/Platform/Linux/LinuxToolChain.cs
@@ -24,8 +24,6 @@ namespace UnrealBuildTool
 		/** Pass --gdb-index option to linker to generate .gdb_index section. */
 		protected bool bGdbIndexSection = true;
 
-		protected bool bPreprocessDepends = false;
-
 		/** Allows you to override the maximum binary size allowed to be passed to objcopy.exe when cross building on Windows. */
 		/** Max value is 2GB, due to bat file limitation */
 		protected UInt64 MaxBinarySizeOverrideForObjcopy = 0;
@@ -916,87 +914,14 @@ namespace UnrealBuildTool
 			foreach (FileItem SourceFile in InputFiles)
 			{
 				Action CompileAction = Graph.CreateAction(ActionType.Compile);
-				CompileAction.PrerequisiteItems.AddRange(CompileEnvironment.ForceIncludeFiles);
-				CompileAction.PrerequisiteItems.AddRange(CompileEnvironment.AdditionalPrerequisites);
 
 				List<string> FileArguments = new();
-				string Extension = Path.GetExtension(SourceFile.AbsolutePath).ToUpperInvariant();
 
 				// Add C or C++ specific compiler arguments.
-				if (CompileEnvironment.PrecompiledHeaderAction == PrecompiledHeaderAction.Create)
-				{
-					GetCompileArguments_PCH(CompileEnvironment, FileArguments);
-				}
-				else if (Extension == ".C")
-				{
-					// Compile the file as C code.
-					GetCompileArguments_C(CompileEnvironment, FileArguments);
-				}
-				else if (Extension == ".MM")
-				{
-					// Compile the file as Objective-C++ code.
-					GetCompileArguments_MM(CompileEnvironment, FileArguments);
-					FileArguments.Add(GetRTTIFlag(CompileEnvironment));
-				}
-				else if (Extension == ".M")
-				{
-					// Compile the file as Objective-C code.
-					GetCompileArguments_M(CompileEnvironment, FileArguments);
-				}
-				else
-				{
-					// Compile the file as C++ code.
-					GetCompileArguments_CPP(CompileEnvironment, FileArguments);
-				}
+				GetCompileArguments_FileType(CompileEnvironment, SourceFile, OutputDir, FileArguments, CompileAction, Result);
 
-				// Add the C++ source file and its included files to the prerequisite item list.
-				CompileAction.PrerequisiteItems.Add(SourceFile);
-
-				if (CompileEnvironment.PrecompiledHeaderAction == PrecompiledHeaderAction.Create)
-				{
-					// Add the precompiled header file to the produced item list.
-					FileItem PrecompiledHeaderFile = FileItem.GetItemByFileReference(FileReference.Combine(OutputDir, Path.GetFileName(SourceFile.AbsolutePath) + ".gch"));
-
-					CompileAction.ProducedItems.Add(PrecompiledHeaderFile);
-					Result.PrecompiledHeaderFile = PrecompiledHeaderFile;
-
-					// Add the parameters needed to compile the precompiled header file to the command-line.
-					FileArguments.Add(string.Format("-o \"{0}\"", NormalizeCommandLinePath(PrecompiledHeaderFile)));
-				}
-				else
-				{
-					if (CompileEnvironment.PrecompiledHeaderAction == PrecompiledHeaderAction.Include)
-					{
-						CompileAction.PrerequisiteItems.Add(CompileEnvironment.PrecompiledHeaderFile!);
-					}
-
-					// Add the object file to the produced item list.
-					FileItem ObjectFile = FileItem.GetItemByFileReference(FileReference.Combine(OutputDir, Path.GetFileName(SourceFile.AbsolutePath) + ".o"));
-					CompileAction.ProducedItems.Add(ObjectFile);
-					Result.ObjectFiles.Add(ObjectFile);
-
-					FileArguments.Add(string.Format("-o \"{0}\"", NormalizeCommandLinePath(ObjectFile)));
-				}
-
-				// Add the source file path to the command-line.
-				FileArguments.Add(string.Format("\"{0}\"", NormalizeCommandLinePath(SourceFile)));
-
-				// Generate the timing info
-				if (CompileEnvironment.bPrintTimingInfo)
-				{
-					FileItem TraceFile = FileItem.GetItemByFileReference(FileReference.Combine(OutputDir, Path.GetFileName(SourceFile.AbsolutePath) + ".json"));
-					FileArguments.Add("-ftime-trace");
-					CompileAction.ProducedItems.Add(TraceFile);
-				}
-
-				// Generate the included header dependency list
-				if (!bPreprocessDepends && CompileEnvironment.bGenerateDependenciesFile)
-				{
-					FileItem DependencyListFile = FileItem.GetItemByFileReference(FileReference.Combine(OutputDir, Path.GetFileName(SourceFile.AbsolutePath) + ".d"));
-					FileArguments.Add(string.Format("-MD -MF\"{0}\"", NormalizeCommandLinePath(DependencyListFile)));
-					CompileAction.DependencyListFile = DependencyListFile;
-					CompileAction.ProducedItems.Add(DependencyListFile);
-				}
+				// Gets the target file so we can get the correct output path.
+				FileItem TargetFile = CompileAction.ProducedItems.First();
 
 				CompileAction.WorkingDirectory = Unreal.EngineSourceDirectory;
 				CompileAction.CommandPath = ClangPath!;
@@ -1004,16 +929,11 @@ namespace UnrealBuildTool
 				List<string> ResponseFileContents = new();
 				ResponseFileContents.AddRange(GlobalArguments);
 				ResponseFileContents.AddRange(FileArguments);
-				ResponseFileContents.Add(CompileEnvironment.AdditionalArguments);
-				// all response lines should have / instead of \, but we cannot just bulk-replace it here since some \ are used to escape quotes, e.g. Definitions.Add("FOO=TEXT(\"Bar\")");
-
-
-				Debug.Assert(CompileAction.ProducedItems.Count > 0);
 
-				FileReference CompilerResponseFileName = CompileAction.ProducedItems[0].Location + ".rsp";
+				FileReference CompilerResponseFileName = TargetFile.Location + ".rsp";
 				FileItem CompilerResponseFileItem = Graph.CreateIntermediateTextFile(CompilerResponseFileName, ResponseFileContents);
 
-				CompileAction.CommandArguments = string.Format(" @\"{0}\"", NormalizeCommandLinePath(CompilerResponseFileName));
+				CompileAction.CommandArguments = GetResponseFileArgument(CompilerResponseFileItem);
 				CompileAction.PrerequisiteItems.Add(CompilerResponseFileItem);
 				CompileAction.CommandDescription = "Compile";
 				CompileAction.CommandVersion = ClangVersionString!;
@@ -1044,12 +964,12 @@ namespace UnrealBuildTool
 					PreprocessGlobalArguments.Remove("-c");
 
 					FileItem DependencyListFile = FileItem.GetItemByFileReference(FileReference.Combine(OutputDir, Path.GetFileName(SourceFile.AbsolutePath) + ".d"));
-					PreprocessFileArguments.Add(string.Format("-M -MF\"{0}\"", NormalizeCommandLinePath(DependencyListFile)));
+					PreprocessFileArguments.Add(GetPreprocessDepencenciesListFileArgument(DependencyListFile));
 					PrepassAction.DependencyListFile = DependencyListFile;
 					PrepassAction.ProducedItems.Add(DependencyListFile);
 
 					PreprocessFileArguments.Remove("-ftime-trace");
-					PreprocessFileArguments.Remove(string.Format(" -o \"{0}\"", NormalizeCommandLinePath(CompileAction.ProducedItems.First())));
+					PreprocessFileArguments.Remove(GetOutputFileArgument(CompileAction.ProducedItems.First()));
 
 					PrepassAction.DeleteItems.AddRange(PrepassAction.ProducedItems);
 
@@ -1067,7 +987,7 @@ namespace UnrealBuildTool
 					FileItem PreprocessResponseFileItem = Graph.CreateIntermediateTextFile(PreprocessResponseFileName, PreprocessResponseFileContents);
 					PrepassAction.PrerequisiteItems.Add(PreprocessResponseFileItem);
 
-					PrepassAction.CommandArguments = string.Format("@\"{0}\"", PreprocessResponseFileItem);
+					PrepassAction.CommandArguments = GetResponseFileArgument(PreprocessResponseFileItem);
 					CompileAction.DependencyListFile = DependencyListFile;
 					CompileAction.PrerequisiteItems.Add(DependencyListFile);
 				}
diff --git a/Engine/Source/Programs/UnrealBuildTool/Platform/Mac/MacToolChain.cs b/Engine/Source/Programs/UnrealBuildTool/Platform/Mac/MacToolChain.cs
--- a/Engine/Source/Programs/UnrealBuildTool/Platform/Mac/MacToolChain.cs
+++ b/Engine/Source/Programs/UnrealBuildTool/Platform/Mac/MacToolChain.cs
@@ -113,8 +113,6 @@ namespace UnrealBuildTool
 
 		private static List<FileItem> BundleDependencies = new List<FileItem>();
 
-		private bool bPreprocessDepends = false;
-
 		private static void SetupXcodePaths(bool bVerbose)
 		{
 		}
@@ -198,7 +196,7 @@ namespace UnrealBuildTool
 
 			// Pass through architecture and OS info
 			Arguments.Add("" + FormatArchitectureArg(CompileEnvironment.Architecture));
-			Arguments.Add(string.Format(" -isysroot \"{0}\"", SDKPath));
+			Arguments.Add($"-isysroot \"{SDKPath}\"");
 			Arguments.Add("-mmacosx-version-min=" + (CompileEnvironment.bEnableOSX109Support ? "10.9" : Settings.MacOSVersion));
 		}
 		
@@ -296,90 +294,11 @@ namespace UnrealBuildTool
 			foreach (FileItem SourceFile in InputFiles)
 			{
 				Action CompileAction = Graph.CreateAction(ActionType.Compile);
-				CompileAction.PrerequisiteItems.AddRange(CompileEnvironment.ForceIncludeFiles);
-				CompileAction.PrerequisiteItems.AddRange(CompileEnvironment.AdditionalPrerequisites);
 
 				List<string> FileArguments = new();
-				string Extension = Path.GetExtension(SourceFile.AbsolutePath).ToUpperInvariant();
-
-				if (CompileEnvironment.PrecompiledHeaderAction == PrecompiledHeaderAction.Create)
-				{
-					// Compile the file as a C++ PCH.
-					GetCompileArguments_PCH(CompileEnvironment, FileArguments);
-					FileArguments.Add(GetRTTIFlag(CompileEnvironment));
-				}
-				else if (Extension == ".C")
-				{
-					// Compile the file as C code.
-					GetCompileArguments_C(CompileEnvironment, FileArguments);
-				}
-				else if (Extension == ".MM")
-				{
-					// Compile the file as Objective-C++ code.
-					GetCompileArguments_MM(CompileEnvironment, FileArguments);
-					FileArguments.Add(GetRTTIFlag(CompileEnvironment));
-				}
-				else if (Extension == ".M")
-				{
-					// Compile the file as Objective-C++ code.
-					GetCompileArguments_M(CompileEnvironment, FileArguments);
-				}
-				else
-				{
-					// Compile the file as C++ code.
-					GetCompileArguments_CPP(CompileEnvironment, FileArguments);
-					FileArguments.Add(GetRTTIFlag(CompileEnvironment));
-				}
-
-				// Add the C++ source file and its included files to the prerequisite item list.
-				CompileAction.PrerequisiteItems.Add(SourceFile);
-
-				string? OutputFilePath = null;
-				if (CompileEnvironment.PrecompiledHeaderAction == PrecompiledHeaderAction.Create)
-				{
-					// Add the precompiled header file to the produced item list.
-					FileItem PrecompiledHeaderFile = FileItem.GetItemByFileReference(FileReference.Combine(OutputDir, Path.GetFileName(SourceFile.AbsolutePath) + ".gch"));
-					CompileAction.ProducedItems.Add(PrecompiledHeaderFile);
-					Result.PrecompiledHeaderFile = PrecompiledHeaderFile;
 
-					// Add the parameters needed to compile the precompiled header file to the command-line.
-					FileArguments.Add(string.Format("-o \"{0}\"", NormalizeCommandLinePath(PrecompiledHeaderFile)));
-				}
-				else
-				{
-					if (CompileEnvironment.PrecompiledHeaderAction == PrecompiledHeaderAction.Include)
-					{
-						CompileAction.PrerequisiteItems.Add(CompileEnvironment.PrecompiledHeaderFile!);
-						CompileAction.PrerequisiteItems.Add(FileItem.GetItemByPath(CompileEnvironment.PrecompiledHeaderFile!.FullName.Replace(".h.gch", ".h")));
-					}
-					// Add the object file to the produced item list.
-					FileItem ObjectFile = FileItem.GetItemByFileReference(FileReference.Combine(OutputDir, Path.GetFileName(SourceFile.AbsolutePath) + ".o"));
-
-					CompileAction.ProducedItems.Add(ObjectFile);
-					Result.ObjectFiles.Add(ObjectFile);
-					FileArguments.Add(string.Format("-o \"{0}\"", NormalizeCommandLinePath(ObjectFile)));
-					OutputFilePath = ObjectFile.AbsolutePath;
-				}
-
-				// Add the source file path to the command-line.
-				FileArguments.Add(string.Format("\"{0}\"", NormalizeCommandLinePath(SourceFile)));
-
-				// Generate the timing info
-				if (CompileEnvironment.bPrintTimingInfo)
-				{
-					FileItem TraceFile = FileItem.GetItemByFileReference(FileReference.Combine(OutputDir, Path.GetFileName(SourceFile.AbsolutePath) + ".json"));
-					FileArguments.Add("-ftime-trace");
-					CompileAction.ProducedItems.Add(TraceFile);
-				}
-
-				// Generate the included header dependency list
-				if(!bPreprocessDepends && CompileEnvironment.bGenerateDependenciesFile)
-				{
-					FileItem DependencyListFile = FileItem.GetItemByFileReference(FileReference.Combine(OutputDir, Path.GetFileName(SourceFile.AbsolutePath) + ".d"));
-					FileArguments.Add(string.Format("-MD -MF\"{0}\"", NormalizeCommandLinePath(DependencyListFile)));
-					CompileAction.DependencyListFile = DependencyListFile;
-					CompileAction.ProducedItems.Add(DependencyListFile);
-				}
+				// Add C or C++ specific compiler arguments.
+				GetCompileArguments_FileType(CompileEnvironment, SourceFile, OutputDir, FileArguments, CompileAction, Result);
 
 				string EscapedAdditionalArgs = "";
 				if(!string.IsNullOrWhiteSpace(CompileEnvironment.AdditionalArguments))
@@ -399,7 +318,7 @@ namespace UnrealBuildTool
 				}
 
 				// Gets the target file so we can get the correct output path.
-				FileItem TargetFile = CompileAction.ProducedItems[0];
+				FileItem TargetFile = CompileAction.ProducedItems.First();
 
 				// Creates the path to the response file using the name of the output file and creates its contents.
 				FileReference ResponseFileName = new FileReference(TargetFile.AbsolutePath + ".response");
@@ -417,9 +336,10 @@ namespace UnrealBuildTool
 				
 				// Analyze and then compile using the shell to perform the indirection
 				string? StaticAnalysisMode = Environment.GetEnvironmentVariable("CLANG_STATIC_ANALYZER_MODE");
-				if (StaticAnalysisMode != null && StaticAnalysisMode != "" && OutputFilePath != null)
+				if (StaticAnalysisMode != null && StaticAnalysisMode != "")
 				{
-					CompileAction.CommandArguments = "-c \"" + CompilerPath + " " + string.Format("@\"{0}\"", ResponseFileName) + " --analyze -Wno-unused-command-line-argument -Xclang -analyzer-output=html -Xclang -analyzer-config -Xclang path-diagnostics-alternate=true -Xclang -analyzer-config -Xclang report-in-main-source-file=true -Xclang -analyzer-disable-checker -Xclang deadcode.DeadStores -o " + OutputFilePath.Replace(".o", ".html") + "; " + CompilerPath + " " + string.Format("@\"{0}\"", ResponseFileName) + "\"";
+					FileReference ReportFile = new FileReference(SourceFile.AbsolutePath + ".html");
+					CompileAction.CommandArguments = "-c \"" + CompilerPath + " " + string.Format("@\"{0}\"", ResponseFileName) + " --analyze -Wno-unused-command-line-argument -Xclang -analyzer-output=html -Xclang -analyzer-config -Xclang path-diagnostics-alternate=true -Xclang -analyzer-config -Xclang report-in-main-source-file=true -Xclang -analyzer-disable-checker -Xclang deadcode.DeadStores -o " + ReportFile + "; " + CompilerPath + " " + string.Format("@\"{0}\"", ResponseFileName) + "\"";
 					CompilerPath = "/bin/sh";
 				}
 
diff --git a/Engine/Source/Programs/UnrealBuildTool/ToolChain/AppleToolChain.cs b/Engine/Source/Programs/UnrealBuildTool/ToolChain/AppleToolChain.cs
--- a/Engine/Source/Programs/UnrealBuildTool/ToolChain/AppleToolChain.cs
+++ b/Engine/Source/Programs/UnrealBuildTool/ToolChain/AppleToolChain.cs
@@ -225,12 +225,6 @@ namespace UnrealBuildTool
 			Arguments.Add("-x objective-c++");
 			GetCppStandardCompileArgument(CompileEnvironment, Arguments);
 			Arguments.Add("-stdlib=libc++");
-
-			// Only force include PCH for c++ files
-			if (CompileEnvironment.PrecompiledHeaderAction == PrecompiledHeaderAction.Include)
-			{
-				Arguments.Remove(GetForceIncludeFileArgument(CompileEnvironment.PrecompiledHeaderIncludeFilename!));
-			}
 		}
 
 		protected override void GetCompileArguments_M(CppCompileEnvironment CompileEnvironment, List<string> Arguments)
@@ -238,12 +232,6 @@ namespace UnrealBuildTool
 			Arguments.Add("-x objective-c");
 			GetCppStandardCompileArgument(CompileEnvironment, Arguments);
 			Arguments.Add("-stdlib=libc++");
-
-			// Only force include PCH for c++ files
-			if (CompileEnvironment.PrecompiledHeaderAction == PrecompiledHeaderAction.Include)
-			{
-				Arguments.Remove(GetForceIncludeFileArgument(CompileEnvironment.PrecompiledHeaderIncludeFilename!));
-			}
 		}
 
 		protected override void GetCompileArguments_PCH(CppCompileEnvironment CompileEnvironment, List<string> Arguments)
@@ -347,6 +335,8 @@ namespace UnrealBuildTool
 
 			base.GetCompileArguments_Global(CompileEnvironment, Arguments);
 
+			Arguments.Add(GetRTTIFlag(CompileEnvironment));
+
 			Arguments.Add("-fmessage-length=0");
 			Arguments.Add("-fpascal-strings");
 		}
diff --git a/Engine/Source/Programs/UnrealBuildTool/ToolChain/ClangToolChain.cs b/Engine/Source/Programs/UnrealBuildTool/ToolChain/ClangToolChain.cs
--- a/Engine/Source/Programs/UnrealBuildTool/ToolChain/ClangToolChain.cs
+++ b/Engine/Source/Programs/UnrealBuildTool/ToolChain/ClangToolChain.cs
@@ -119,6 +119,8 @@ namespace UnrealBuildTool
 
 		protected ClangToolChainOptions Options;
 
+		protected bool bPreprocessDepends = false;
+
 		public ClangToolChain(ClangToolChainOptions InOptions, ILogger InLogger)
 			: base(InLogger)
 		{
@@ -249,12 +251,6 @@ namespace UnrealBuildTool
 		protected virtual void GetCompileArguments_C(CppCompileEnvironment CompileEnvironment, List<string> Arguments)
 		{
 			Arguments.Add("-x c");
-
-			// Only force include PCH for c++ files
-			if (CompileEnvironment.PrecompiledHeaderAction == PrecompiledHeaderAction.Include)
-			{
-				Arguments.Remove(GetForceIncludeFileArgument(CompileEnvironment.PrecompiledHeaderIncludeFilename!));
-			}
 		}
 
 		protected virtual void GetCompileArguments_MM(CppCompileEnvironment CompileEnvironment, List<string> Arguments)
@@ -263,12 +259,6 @@ namespace UnrealBuildTool
 			Arguments.Add("-fobjc-abi-version=2");
 			Arguments.Add("-fobjc-legacy-dispatch");
 			GetCppStandardCompileArgument(CompileEnvironment, Arguments);
-
-			// Only force include PCH for c++ files
-			if (CompileEnvironment.PrecompiledHeaderAction == PrecompiledHeaderAction.Include)
-			{
-				Arguments.Remove(GetForceIncludeFileArgument(CompileEnvironment.PrecompiledHeaderIncludeFilename!));
-			}
 		}
 
 		protected virtual void GetCompileArguments_M(CppCompileEnvironment CompileEnvironment, List<string> Arguments)
@@ -277,12 +267,6 @@ namespace UnrealBuildTool
 			Arguments.Add("-fobjc-abi-version=2");
 			Arguments.Add("-fobjc-legacy-dispatch");
 			GetCppStandardCompileArgument(CompileEnvironment, Arguments);
-
-			// Only force include PCH for c++ files
-			if (CompileEnvironment.PrecompiledHeaderAction == PrecompiledHeaderAction.Include)
-			{
-				Arguments.Remove(GetForceIncludeFileArgument(CompileEnvironment.PrecompiledHeaderIncludeFilename!));
-			}
 		}
 
 		protected virtual void GetCompileArguments_PCH(CppCompileEnvironment CompileEnvironment, List<string> Arguments)
@@ -354,6 +338,31 @@ namespace UnrealBuildTool
 			Arguments.AddRange(CompileEnvironment.ForceIncludeFiles.Select(ForceIncludeFile => GetForceIncludeFileArgument(ForceIncludeFile)));
 		}
 
+		protected virtual string GetSourceFileArgument(FileItem SourceFile)
+		{
+			return $"\"{NormalizeCommandLinePath(SourceFile)}\"";
+		}
+
+		protected virtual string GetOutputFileArgument(FileItem ForceIncludeFile)
+		{
+			return $"-o \"{NormalizeCommandLinePath(ForceIncludeFile)}\"";
+		}
+
+		protected virtual string GetDepencenciesListFileArgument(FileItem DependencyListFile)
+		{
+			return $"-MD -MF\"{NormalizeCommandLinePath(DependencyListFile)}\"";
+		}
+
+		protected virtual string GetPreprocessDepencenciesListFileArgument(FileItem DependencyListFile)
+		{
+			return $"-M -MF\"{NormalizeCommandLinePath(DependencyListFile)}\"";
+		}
+
+		protected virtual string GetResponseFileArgument(FileItem ResponseFile)
+		{
+			return $"@\"{NormalizeCommandLinePath(ResponseFile)}\"";
+		}
+
 		/// <summary>
 		/// Common compile arguments that control which warnings are enabled.
 		/// https://clang.llvm.org/docs/DiagnosticsReference.html
@@ -577,5 +586,115 @@ namespace UnrealBuildTool
 			// Add sanitizer flags to the argument list
 			GetCompilerArguments_Sanitizers(CompileEnvironment, Arguments);
 		}
+
+		/// <summary>
+		/// Compile arguments for specific files in a module. Also updates Action and CPPOutput results.
+		/// </summary>
+		/// <param name="CompileEnvironment"></param>
+		/// <param name="SourceFile"></param>
+		/// <param name="OutputDir"></param>
+		/// <param name="Arguments"></param>
+		/// <param name="CompileAction"></param>
+		/// <param name="CompileResult"></param>
+		protected virtual void GetCompileArguments_FileType(CppCompileEnvironment CompileEnvironment, FileItem SourceFile, DirectoryReference OutputDir, List<string> Arguments, Action CompileAction, CPPOutput CompileResult)
+		{
+			// Add the C++ source file and its included files to the prerequisite item list.
+			CompileAction.PrerequisiteItems.AddRange(CompileEnvironment.ForceIncludeFiles);
+			CompileAction.PrerequisiteItems.AddRange(CompileEnvironment.AdditionalPrerequisites);
+			CompileAction.PrerequisiteItems.Add(SourceFile);
+
+			string Extension = Path.GetExtension(SourceFile.AbsolutePath).ToUpperInvariant();
+			bool IsCPlusPlus = false;
+
+			if (CompileEnvironment.PrecompiledHeaderAction == PrecompiledHeaderAction.Create)
+			{
+				GetCompileArguments_PCH(CompileEnvironment, Arguments);
+			}
+			else if (Extension == ".C")
+			{
+				// Compile the file as C code.
+				GetCompileArguments_C(CompileEnvironment, Arguments);
+			}
+			else if (Extension == ".MM")
+			{
+				// Compile the file as Objective-C++ code.
+				GetCompileArguments_MM(CompileEnvironment, Arguments);
+			}
+			else if (Extension == ".M")
+			{
+				// Compile the file as Objective-C code.
+				GetCompileArguments_M(CompileEnvironment, Arguments);
+			}
+			else
+			{
+				// Compile the file as C++ code.
+				GetCompileArguments_CPP(CompileEnvironment, Arguments);
+				IsCPlusPlus = true;
+			}
+
+			if (CompileEnvironment.PrecompiledHeaderAction == PrecompiledHeaderAction.Include)
+			{
+				// Only force include PCH for c++ files
+				if (!IsCPlusPlus)
+				{
+					Arguments.Remove(GetForceIncludeFileArgument(CompileEnvironment.PrecompiledHeaderIncludeFilename!));
+				}
+				else
+				{
+					CompileAction.PrerequisiteItems.Add(CompileEnvironment.PrecompiledHeaderFile!);
+				}
+			}
+
+			FileItem OutputFile;
+			if (CompileEnvironment.PrecompiledHeaderAction == PrecompiledHeaderAction.Create)
+			{
+				OutputFile = FileItem.GetItemByFileReference(FileReference.Combine(OutputDir, Path.GetFileName(SourceFile.AbsolutePath) + ".gch"));
+				CompileResult.PrecompiledHeaderFile = OutputFile;
+			}
+			else if (CompileEnvironment.bPreprocessOnly)
+			{
+				OutputFile = FileItem.GetItemByFileReference(FileReference.Combine(OutputDir, Path.GetFileName(SourceFile.AbsolutePath) + ".cpp.i"));
+				CompileResult.ObjectFiles.Add(OutputFile);
+
+				// Clang does EITHER pre-process or object file.
+				Arguments.Add("-E"); // Only run the preprocessor
+				Arguments.Add("-fuse-line-directives"); // Use #line in preprocessed output
+
+				// this is parsed by external tools wishing to open this file directly.
+				Logger.LogInformation("PreProcessPath: {File}", OutputFile.AbsolutePath);
+			}
+			else
+			{
+				string ObjectFileExtension = (CompileEnvironment.AdditionalArguments != null && CompileEnvironment.AdditionalArguments.Contains("-emit-llvm")) ? ".bc" : ".o";
+				OutputFile = FileItem.GetItemByFileReference(FileReference.Combine(OutputDir, $"{Path.GetFileName(SourceFile.AbsolutePath)}{ObjectFileExtension}"));
+				CompileResult.ObjectFiles.Add(OutputFile);
+			}
+
+			// Add the output file to the produced item list.
+			CompileAction.ProducedItems.Add(OutputFile);
+
+			// Add the source file path to the command-line.
+			Arguments.Add(GetSourceFileArgument(SourceFile));
+
+			// Generate the timing info
+			if (CompileEnvironment.bPrintTimingInfo)
+			{
+				FileItem TraceFile = FileItem.GetItemByFileReference(FileReference.Combine(OutputDir, Path.GetFileName(SourceFile.AbsolutePath) + ".json"));
+				Arguments.Add("-ftime-trace");
+				CompileAction.ProducedItems.Add(TraceFile);
+			}
+
+			// Generate the included header dependency list
+			if (!bPreprocessDepends && CompileEnvironment.bGenerateDependenciesFile)
+			{
+				FileItem DependencyListFile = FileItem.GetItemByFileReference(FileReference.Combine(OutputDir, Path.GetFileName(SourceFile.AbsolutePath) + ".d"));
+				Arguments.Add(GetDepencenciesListFileArgument(DependencyListFile));
+				CompileAction.DependencyListFile = DependencyListFile;
+				CompileAction.ProducedItems.Add(DependencyListFile);
+			}
+
+			// Add the parameters needed to compile the output file to the command-line.
+			Arguments.Add(GetOutputFileArgument(OutputFile));
+		}
 	}
 }
\ No newline at end of file
