UnrealBuildTool: UnsafeTypeCastWarning for clang platforms

#rnx

diff --git a/Engine/Source/Programs/UnrealBuildTool/ToolChain/ClangToolChain.cs b/Engine/Source/Programs/UnrealBuildTool/ToolChain/ClangToolChain.cs
--- a/Engine/Source/Programs/UnrealBuildTool/ToolChain/ClangToolChain.cs
+++ b/Engine/Source/Programs/UnrealBuildTool/ToolChain/ClangToolChain.cs
@@ -443,6 +443,38 @@ namespace UnrealBuildTool
 				Arguments.Add("-Wundef" + (CompileEnvironment.bUndefinedIdentifierWarningsAsErrors ? "" : " -Wno-error=undef"));
 			}
 
+			// Note: This should be kept in sync with PRAGMA_DISABLE_UNSAFE_TYPECAST_WARNINGS in ClangPlatformCompilerPreSetup.h
+			string[] UnsafeTypeCastWarningList = {
+				"float-conversion",
+				"implicit-float-conversion",
+				"implicit-int-conversion",
+				"c++11-narrowing"
+				//"shorten-64-to-32",	<-- too many hits right now, probably want it *soon*
+				//"sign-conversion",	<-- too many hits right now, probably want it eventually
+			};
+
+			if (CompileEnvironment.UnsafeTypeCastWarningLevel == WarningLevel.Error)
+			{
+				foreach (string Warning in UnsafeTypeCastWarningList)
+				{
+					Arguments.Add("-W" + Warning);
+				}
+			}
+			else if (CompileEnvironment.UnsafeTypeCastWarningLevel == WarningLevel.Warning)
+			{
+				foreach (string Warning in UnsafeTypeCastWarningList)
+				{
+					Arguments.Add("-W" + Warning + " -Wno-error=" + Warning);
+				}
+			}
+			else
+			{
+				foreach (string Warning in UnsafeTypeCastWarningList)
+				{
+					Arguments.Add("-Wno-" + Warning);
+				}
+			}
+
 			// always use absolute paths for errors, this can help IDEs go to the error properly
 			Arguments.Add("-fdiagnostics-absolute-paths");  // https://clang.llvm.org/docs/ClangCommandLineReference.html#cmdoption-clang-fdiagnostics-absolute-paths
 
diff --git a/Engine/Source/Programs/UnrealHeaderTool/UnrealHeaderTool.Build.cs b/Engine/Source/Programs/UnrealHeaderTool/UnrealHeaderTool.Build.cs
--- a/Engine/Source/Programs/UnrealHeaderTool/UnrealHeaderTool.Build.cs
+++ b/Engine/Source/Programs/UnrealHeaderTool/UnrealHeaderTool.Build.cs
@@ -31,6 +31,6 @@ public class UnrealHeaderTool : ModuleRules
 		
 		bEnableExceptions = true;
 
-		UnsafeTypeCastWarningLevel = WarningLevel.Warning;
+		//UnsafeTypeCastWarningLevel = WarningLevel.Warning;
 	}
 }
