Fast Android packaging when only libUnreal.so changes if -BypassGradlePackaging option passed to UBT

#android

diff --git a/Engine/Source/Programs/UnrealBuildTool/Platform/Android/UEDeployAndroid.cs b/Engine/Source/Programs/UnrealBuildTool/Platform/Android/UEDeployAndroid.cs
--- a/Engine/Source/Programs/UnrealBuildTool/Platform/Android/UEDeployAndroid.cs
+++ b/Engine/Source/Programs/UnrealBuildTool/Platform/Android/UEDeployAndroid.cs
@@ -7,6 +7,7 @@ using System.Text;
 using System.Xml;
 using System.Diagnostics;
 using System.IO;
+using System.IO.Compression;
 using Microsoft.Win32;
 using System.Xml.Linq;
 using EpicGames.Core;
@@ -66,6 +67,12 @@ namespace UnrealBuildTool
 		[CommandLine("-ForceAPKGeneration", Value = "true")]
 		public bool ForceAPKGeneration = false;
 
+		/// <summary>
+		/// Do not use Gradle if previous APK exists and only libUnreal.so changed
+		/// </summary>
+		[CommandLine("-BypassGradlePackaging", Value = "true")]
+		public bool BypassGradlePackaging = false;
+
 		public UEDeployAndroid(FileReference? InProjectFile, bool InForcePackageData, ILogger InLogger)
 			: base(InLogger)
 		{
@@ -3446,6 +3453,295 @@ namespace UnrealBuildTool
 			}
 		}
 
+		private uint[] CRCTablesSB8 = {
+			0x00000000, 0x77073096, 0xee0e612c, 0x990951ba, 0x076dc419, 0x706af48f, 0xe963a535, 0x9e6495a3, 0x0edb8832, 0x79dcb8a4, 0xe0d5e91e, 0x97d2d988, 0x09b64c2b, 0x7eb17cbd, 0xe7b82d07, 0x90bf1d91,
+			0x1db71064, 0x6ab020f2, 0xf3b97148, 0x84be41de, 0x1adad47d, 0x6ddde4eb, 0xf4d4b551, 0x83d385c7, 0x136c9856, 0x646ba8c0, 0xfd62f97a, 0x8a65c9ec, 0x14015c4f, 0x63066cd9, 0xfa0f3d63, 0x8d080df5,
+			0x3b6e20c8, 0x4c69105e, 0xd56041e4, 0xa2677172, 0x3c03e4d1, 0x4b04d447, 0xd20d85fd, 0xa50ab56b, 0x35b5a8fa, 0x42b2986c, 0xdbbbc9d6, 0xacbcf940, 0x32d86ce3, 0x45df5c75, 0xdcd60dcf, 0xabd13d59,
+			0x26d930ac, 0x51de003a, 0xc8d75180, 0xbfd06116, 0x21b4f4b5, 0x56b3c423, 0xcfba9599, 0xb8bda50f, 0x2802b89e, 0x5f058808, 0xc60cd9b2, 0xb10be924, 0x2f6f7c87, 0x58684c11, 0xc1611dab, 0xb6662d3d,
+			0x76dc4190, 0x01db7106, 0x98d220bc, 0xefd5102a, 0x71b18589, 0x06b6b51f, 0x9fbfe4a5, 0xe8b8d433, 0x7807c9a2, 0x0f00f934, 0x9609a88e, 0xe10e9818, 0x7f6a0dbb, 0x086d3d2d, 0x91646c97, 0xe6635c01,
+			0x6b6b51f4, 0x1c6c6162, 0x856530d8, 0xf262004e, 0x6c0695ed, 0x1b01a57b, 0x8208f4c1, 0xf50fc457, 0x65b0d9c6, 0x12b7e950, 0x8bbeb8ea, 0xfcb9887c, 0x62dd1ddf, 0x15da2d49, 0x8cd37cf3, 0xfbd44c65,
+			0x4db26158, 0x3ab551ce, 0xa3bc0074, 0xd4bb30e2, 0x4adfa541, 0x3dd895d7, 0xa4d1c46d, 0xd3d6f4fb, 0x4369e96a, 0x346ed9fc, 0xad678846, 0xda60b8d0, 0x44042d73, 0x33031de5, 0xaa0a4c5f, 0xdd0d7cc9,
+			0x5005713c, 0x270241aa, 0xbe0b1010, 0xc90c2086, 0x5768b525, 0x206f85b3, 0xb966d409, 0xce61e49f, 0x5edef90e, 0x29d9c998, 0xb0d09822, 0xc7d7a8b4, 0x59b33d17, 0x2eb40d81, 0xb7bd5c3b, 0xc0ba6cad,
+			0xedb88320, 0x9abfb3b6, 0x03b6e20c, 0x74b1d29a, 0xead54739, 0x9dd277af, 0x04db2615, 0x73dc1683, 0xe3630b12, 0x94643b84, 0x0d6d6a3e, 0x7a6a5aa8, 0xe40ecf0b, 0x9309ff9d, 0x0a00ae27, 0x7d079eb1,
+			0xf00f9344, 0x8708a3d2, 0x1e01f268, 0x6906c2fe, 0xf762575d, 0x806567cb, 0x196c3671, 0x6e6b06e7, 0xfed41b76, 0x89d32be0, 0x10da7a5a, 0x67dd4acc, 0xf9b9df6f, 0x8ebeeff9, 0x17b7be43, 0x60b08ed5,
+			0xd6d6a3e8, 0xa1d1937e, 0x38d8c2c4, 0x4fdff252, 0xd1bb67f1, 0xa6bc5767, 0x3fb506dd, 0x48b2364b, 0xd80d2bda, 0xaf0a1b4c, 0x36034af6, 0x41047a60, 0xdf60efc3, 0xa867df55, 0x316e8eef, 0x4669be79,
+			0xcb61b38c, 0xbc66831a, 0x256fd2a0, 0x5268e236, 0xcc0c7795, 0xbb0b4703, 0x220216b9, 0x5505262f, 0xc5ba3bbe, 0xb2bd0b28, 0x2bb45a92, 0x5cb36a04, 0xc2d7ffa7, 0xb5d0cf31, 0x2cd99e8b, 0x5bdeae1d,
+			0x9b64c2b0, 0xec63f226, 0x756aa39c, 0x026d930a, 0x9c0906a9, 0xeb0e363f, 0x72076785, 0x05005713, 0x95bf4a82, 0xe2b87a14, 0x7bb12bae, 0x0cb61b38, 0x92d28e9b, 0xe5d5be0d, 0x7cdcefb7, 0x0bdbdf21,
+			0x86d3d2d4, 0xf1d4e242, 0x68ddb3f8, 0x1fda836e, 0x81be16cd, 0xf6b9265b, 0x6fb077e1, 0x18b74777, 0x88085ae6, 0xff0f6a70, 0x66063bca, 0x11010b5c, 0x8f659eff, 0xf862ae69, 0x616bffd3, 0x166ccf45,
+			0xa00ae278, 0xd70dd2ee, 0x4e048354, 0x3903b3c2, 0xa7672661, 0xd06016f7, 0x4969474d, 0x3e6e77db, 0xaed16a4a, 0xd9d65adc, 0x40df0b66, 0x37d83bf0, 0xa9bcae53, 0xdebb9ec5, 0x47b2cf7f, 0x30b5ffe9,
+			0xbdbdf21c, 0xcabac28a, 0x53b39330, 0x24b4a3a6, 0xbad03605, 0xcdd70693, 0x54de5729, 0x23d967bf, 0xb3667a2e, 0xc4614ab8, 0x5d681b02, 0x2a6f2b94, 0xb40bbe37, 0xc30c8ea1, 0x5a05df1b, 0x2d02ef8d
+		};
+
+		private bool CopyAPKAndReplaceSO(ILogger Logger, string SourceAPK, string DestAPK, string SourceSOFile, string DestSOFile)
+		{
+			BinaryReader SourceReader;
+			BinaryWriter DestWriter;
+			BinaryWriter DirectoryWriter;
+
+			string DirectoryFilename = DestAPK + ".dir";
+
+			try
+			{
+				SourceReader = new BinaryReader(new FileStream(SourceAPK, FileMode.Open, FileAccess.Read));
+			}
+			catch (IOException)
+			{
+				return false;
+			}
+
+			try
+			{
+				DestWriter = new BinaryWriter(new FileStream(DestAPK, FileMode.Create));
+			}
+			catch (IOException)
+			{
+				return false;
+			}
+
+			try
+			{
+				DirectoryWriter = new BinaryWriter(new FileStream(DirectoryFilename, FileMode.Create));
+			}
+			catch (IOException)
+			{
+				return false;
+			}
+
+			byte[] CopyBuffer = new byte[65536];
+			ushort TotalEntries = 0;
+			uint Remaining;
+
+			uint Signature = 0x02014b50;
+			ushort DirVersion = 0;
+			ushort DiskNumber = 0;
+			ushort DiskStart = 0;
+			ushort InternalAttr = 0;
+			uint ExternalAttr = 0;
+			ushort CommentLen = 0;
+
+			// copy files
+			while (true)
+			{
+				uint Header = SourceReader.ReadUInt32();
+				if (Header != 0x04034b50)
+				{
+					SourceReader.Close();
+					break;
+				}
+
+				ushort Version = SourceReader.ReadUInt16();
+				ushort Flags = SourceReader.ReadUInt16();
+				ushort Compression = SourceReader.ReadUInt16();
+				ushort ModTime = SourceReader.ReadUInt16();
+				ushort ModDate = SourceReader.ReadUInt16();
+				uint CRC32Value = SourceReader.ReadUInt32();
+				uint CompressedSize = SourceReader.ReadUInt32();
+				uint UncompressedSize = SourceReader.ReadUInt32();
+				ushort FilenameLen = SourceReader.ReadUInt16();
+				ushort ExtraLen = SourceReader.ReadUInt16();
+				byte[] Filename = new byte[FilenameLen];
+				SourceReader.BaseStream.Read(Filename, 0, FilenameLen);
+				byte[] Extra = new byte[ExtraLen];
+				SourceReader.BaseStream.Read(Extra, 0, ExtraLen);
+
+				string FilenameStr = Encoding.UTF8.GetString(Filename, 0, FilenameLen);
+				if (!FilenameStr.Equals(DestSOFile))
+				{
+					TotalEntries++;
+					uint Location = (uint)DestWriter.BaseStream.Position;
+
+					// calculate location after header and any additional alignment needed (only if uncompressed)
+					uint HeaderSize = (uint)(4 + 2 + 2 + 2 + 2 + 2 + 4 + 4 + 4 + 2 + 2 + FilenameLen + ExtraLen);
+					uint HeaderEnd = Location + HeaderSize;
+					uint DataStart = (uint)((HeaderEnd + 3) & ~3);
+					ushort ExtraAlign = (ushort)(Compression != 0 ? 0 : (DataStart - HeaderEnd));
+
+					// write local header
+					DestWriter.Write(Header);
+					DestWriter.Write(Version);
+					DestWriter.Write(Flags);
+					DestWriter.Write(Compression);
+					DestWriter.Write(ModTime);
+					DestWriter.Write(ModDate);
+					DestWriter.Write(CRC32Value);
+					DestWriter.Write(CompressedSize);
+					DestWriter.Write(UncompressedSize);
+					DestWriter.Write(FilenameLen);
+					DestWriter.Write((ushort)(ExtraLen + ExtraAlign));
+					DestWriter.BaseStream.Write(Filename, 0, FilenameLen);
+					DestWriter.BaseStream.Write(Extra, 0, ExtraLen);
+					while (ExtraAlign-- > 0)
+					{
+						DestWriter.Write((byte)0);
+					}
+
+					// copy file data as-is
+					Remaining = CompressedSize;
+					while (Remaining > 0)
+					{
+						int CopySize = Remaining < 65536 ? (int)Remaining : 65536;
+						int BytesRead = SourceReader.BaseStream.Read(CopyBuffer, 0, CopySize);
+						DestWriter.BaseStream.Write(CopyBuffer, 0, BytesRead);
+						Remaining -= (uint)BytesRead;
+					}
+
+					// write central directory entry for this file
+					DirectoryWriter.Write(Signature);
+					DirectoryWriter.Write(DirVersion);
+					DirectoryWriter.Write(Version);
+					DirectoryWriter.Write(Flags);
+					DirectoryWriter.Write(Compression);
+					DirectoryWriter.Write(ModTime);
+					DirectoryWriter.Write(ModDate);
+					DirectoryWriter.Write(CRC32Value);
+					DirectoryWriter.Write(CompressedSize);
+					DirectoryWriter.Write(UncompressedSize);
+					DirectoryWriter.Write(FilenameLen);
+					DirectoryWriter.Write(ExtraLen);
+					DirectoryWriter.Write(CommentLen);
+					DirectoryWriter.Write(DiskStart);
+					DirectoryWriter.Write(InternalAttr);
+					DirectoryWriter.Write(ExternalAttr);
+					DirectoryWriter.Write(Location);
+					DirectoryWriter.BaseStream.Write(Filename, 0, FilenameLen);
+					DirectoryWriter.BaseStream.Write(Extra, 0, ExtraLen);
+					// write comment here if there was one
+				}
+				else
+				{
+					// replace with the new .SO
+					try
+					{
+						FileStream SourceFile = File.OpenRead(SourceSOFile);
+
+						UncompressedSize = (uint)SourceFile.Length;
+						CompressedSize = 0;
+						Compression = 8;        // deflate
+
+						// calculate CRC32
+						CRC32Value = 0x0;
+						CRC32Value = ~CRC32Value;
+						Remaining = UncompressedSize;
+						while (Remaining > 0)
+						{
+							int CopySize = Remaining < 65536 ? (int)Remaining : 65536;
+							int BytesRead = SourceFile.Read(CopyBuffer, 0, CopySize);
+							Remaining -= (uint)BytesRead;
+							for (uint Index = 0; Index < BytesRead; Index++)
+							{
+								CRC32Value = (CRC32Value >> 8) ^ CRCTablesSB8[(CRC32Value ^ (uint)CopyBuffer[Index]) & 0xff];
+							}
+						}
+						CRC32Value = ~CRC32Value;
+						SourceFile.Position = 0;
+
+						TotalEntries++;
+						uint Location = (uint)DestWriter.BaseStream.Position;
+
+						using (MemoryStream CompressedStream = new MemoryStream())
+						{
+							using (DeflateStream Compressor = new DeflateStream(CompressedStream, CompressionMode.Compress, leaveOpen: true))
+							{
+								SourceFile.CopyTo(Compressor);
+								SourceFile.Close();
+							}
+
+							byte[] CompressedBuffer = CompressedStream.GetBuffer();
+							CompressedSize = (uint)CompressedStream.Length;
+
+							// calculate location after header and any additional alignment needed (only if uncompressed)
+							uint HeaderSize = (uint)(4 + 2 + 2 + 2 + 2 + 2 + 4 + 4 + 4 + 2 + 2 + FilenameLen + ExtraLen);
+							uint HeaderEnd = Location + HeaderSize;
+							uint DataStart = (uint)((HeaderEnd + 3) & ~3);
+							ushort ExtraAlign = (ushort)(Compression != 0 ? 0 : (DataStart - HeaderEnd));
+
+							// write local header
+							DestWriter.Write(Header);
+							DestWriter.Write(Version);
+							DestWriter.Write(Flags);
+							DestWriter.Write(Compression);
+							DestWriter.Write(ModTime);
+							DestWriter.Write(ModDate);
+							DestWriter.Write(CRC32Value);
+							DestWriter.Write(CompressedSize);
+							DestWriter.Write(UncompressedSize);
+							DestWriter.Write(FilenameLen);
+							DestWriter.Write((ushort)(ExtraLen + ExtraAlign));
+							DestWriter.BaseStream.Write(Filename, 0, FilenameLen);
+							DestWriter.BaseStream.Write(Extra, 0, ExtraLen);
+							while (ExtraAlign-- > 0)
+							{
+								DestWriter.Write((byte)0);
+							}
+
+							// write compressed data
+							DestWriter.BaseStream.Write(CompressedBuffer, 0, (int)CompressedSize);
+						}
+
+						// write central directory entry for this file
+						DirectoryWriter.Write(Signature);
+						DirectoryWriter.Write(DirVersion);
+						DirectoryWriter.Write(Version);
+						DirectoryWriter.Write(Flags);
+						DirectoryWriter.Write(Compression);
+						DirectoryWriter.Write(ModTime);
+						DirectoryWriter.Write(ModDate);
+						DirectoryWriter.Write(CRC32Value);
+						DirectoryWriter.Write(CompressedSize);
+						DirectoryWriter.Write(UncompressedSize);
+						DirectoryWriter.Write(FilenameLen);
+						DirectoryWriter.Write(ExtraLen);
+						DirectoryWriter.Write(CommentLen);
+						DirectoryWriter.Write(DiskStart);
+						DirectoryWriter.Write(InternalAttr);
+						DirectoryWriter.Write(ExternalAttr);
+						DirectoryWriter.Write(Location);
+						DirectoryWriter.BaseStream.Write(Filename, 0, FilenameLen);
+						DirectoryWriter.BaseStream.Write(Extra, 0, ExtraLen);
+						// write comment here if there was one
+					}
+					catch (IOException e)
+					{
+						Logger.LogInformation("Failed to add {SourceSOFile} to APK: Reason = {Reason}", SourceSOFile, e.ToString());
+						DirectoryWriter.Close();
+						File.Delete(DirectoryFilename);
+						DestWriter.Close();
+						File.Delete(DestAPK);
+						return false;
+					}
+				}
+			}
+
+			uint CentralDirectoryLocation = (uint)DestWriter.BaseStream.Position;
+
+			// copy the directory to the end of the APK
+			uint DirectorySize = (uint)DirectoryWriter.BaseStream.Position;
+			DirectoryWriter.BaseStream.Position = 0;
+			Remaining = DirectorySize;
+			while (Remaining > 0)
+			{
+				int CopySize = Remaining < 65536 ? (int)Remaining : 65536;
+				int BytesRead = DirectoryWriter.BaseStream.Read(CopyBuffer, 0, CopySize);
+				DestWriter.BaseStream.Write(CopyBuffer, 0, BytesRead);
+				Remaining -= (uint)BytesRead;
+			}
+			DirectoryWriter.Close();
+			File.Delete(DirectoryFilename);
+
+			// write end of central directory record
+			Signature = 0x06054b50;
+			DestWriter.Write(Signature);
+			DestWriter.Write(DiskNumber);
+			DestWriter.Write(DiskStart);
+			DestWriter.Write(TotalEntries);
+			DestWriter.Write(TotalEntries);
+			DestWriter.Write(DirectorySize);
+			DestWriter.Write(CentralDirectoryLocation);
+			DestWriter.Write(CommentLen);
+			DestWriter.Close();
+			return true;
+		}
+
 		private void CreateGradlePropertiesFiles(string Arch, int MinSDKVersion, int TargetSDKVersion, string CompileSDKVersion, string BuildToolsVersion, string PackageName,
 			string? DestApkName, string NDKArch,	string UnrealBuildFilesPath, string GameBuildFilesPath, string UnrealBuildGradleAppPath, string UnrealBuildPath, string UnrealBuildGradlePath,
 			bool bForDistribution, bool bIsEmbedded, List<string> OBBFiles)
@@ -4066,6 +4362,180 @@ namespace UnrealBuildTool
 				return;
 			}
 
+			if (BypassGradlePackaging && !bEnableBundle && !bSkipGradleBuild)
+			{
+				Logger.LogInformation("Attemping BypassGradlePackaging");
+				int BuildListComboRemaining = BuildListComboTotal;
+
+				foreach (Tuple<string, string> build in BuildList)
+				{
+					string Arch = build.Item1;
+					string Manifest = build.Item2;
+					string NDKArch = GetNDKArch(Arch);
+					
+					string UnrealBuildPath = Path.Combine(IntermediateAndroidPath, Arch.Substring(1).Replace("-", "_"));
+
+					Logger.LogInformation("\n===={Time}====PREPARING NATIVE CODE====={Arch}============================================================", DateTime.Now.ToString(), Arch);
+
+					string DestApkDirectory = Path.Combine(ProjectDirectory, "Binaries", "Android");
+					string SourceSOName = AndroidToolChain.InlineArchName(OutputPath, Arch);
+					// if the source binary was UnrealGame, replace it with the new project name, when re-packaging a binary only build
+					string ApkFilename = Path.GetFileNameWithoutExtension(OutputPath).Replace("UnrealGame", ProjectName);
+					string DestApkName = Path.Combine(DestApkDirectory, ApkFilename + ".apk");
+
+					// As we are always making seperate APKs we need to put the architecture into the name
+					DestApkName = AndroidToolChain.InlineArchName(DestApkName, Arch);
+
+					if (!File.Exists(SourceSOName))
+					{
+						throw new BuildException("Can't make an APK without the compiled .so [{0}]", SourceSOName);
+					}
+
+					if (!File.Exists(DestApkName))
+					{
+						Logger.LogInformation("Output .apk [DestApkName] does not exist, will do full packaging with Gradle", DestApkName);
+						continue;
+					}
+
+					// verify the .so actually changed (this does involve an extra copy, we could always assume it changed for a little more speed)
+					bool bVerifySOChanged = true;
+					if (bVerifySOChanged)
+					{
+						string FinalSOName = UnrealBuildPath + "/jni/" + NDKArch + "/libUnreal.so";
+						if (!CopyIfDifferent(SourceSOName, FinalSOName, false, false))
+						{
+							Logger.LogInformation("{FinalSOName} unchanged, skipping repackage of {DestApkName}", FinalSOName, DestApkName);
+							BuildListComboRemaining--;
+							continue;
+						}
+					}
+
+					// strip symbols to make libUnreal.so small enough for packaging
+					string StrippedSOName = SourceSOName + ".stripped";
+					StripDebugSymbols(SourceSOName, StrippedSOName, Arch, Logger, true);
+
+					Logger.LogInformation("\n===={Time}====PERFORMING FINAL APK PACKAGE OPERATION====={Arch}===========================================", DateTime.Now.ToString(), Arch);
+
+					ConfigHierarchy Ini = GetConfigCacheIni(ConfigHierarchyType.Engine);
+
+					string TempAPK = Path.Combine(DestApkDirectory, "TEMP_" + Path.GetFileName(DestApkName));
+					CopyAPKAndReplaceSO(Logger, DestApkName, TempAPK, StrippedSOName, "lib/" + NDKArch + "/libUnreal.so");
+
+					// defaults for non-distribution build (debug)
+					string KeyStore = Path.Combine(Environment.GetFolderPath(Environment.SpecialFolder.UserProfile), ".android", "debug.keystore");
+					string KeyStorePassword = "android";
+					string KeyAlias = "androiddebugkey";
+					string KeyPassword = "android";
+
+					// read project settings for keystore information if distribution build (release)
+					if (bForDistribution)
+					{
+						Ini.GetString("/Script/AndroidRuntimeSettings.AndroidRuntimeSettings", "KeyStore", out KeyStore);
+						Ini.GetString("/Script/AndroidRuntimeSettings.AndroidRuntimeSettings", "KeyAlias", out KeyAlias);
+						Ini.GetString("/Script/AndroidRuntimeSettings.AndroidRuntimeSettings", "KeyStorePassword", out KeyStorePassword);
+						Ini.GetString("/Script/AndroidRuntimeSettings.AndroidRuntimeSettings", "KeyPassword", out KeyPassword);
+
+						if (string.IsNullOrEmpty(KeyStore) || string.IsNullOrEmpty(KeyAlias) || string.IsNullOrEmpty(KeyStorePassword))
+						{
+							throw new BuildException("DistributionSigning settings are not all set. Check the DistributionSettings section in the Android tab of Project Settings");
+						}
+
+						if (string.IsNullOrEmpty(KeyPassword) || KeyPassword == "_sameaskeystore_")
+						{
+							KeyPassword = KeyStorePassword;
+						}
+
+						KeyStore = Path.Combine(UnrealBuildPath, KeyStore);
+					}
+
+					// Make sure the keystore file exists
+					if (!File.Exists(KeyStore))
+					{
+						throw new BuildException("Keystore file is missing. Check the DistributionSettings section in the Android tab of Project Settings");
+					}
+
+					// sign the APK
+					string BuildToolsPath = Path.Combine(Environment.ExpandEnvironmentVariables("%ANDROID_HOME%"), "Build-Tools", BuildToolsVersion);
+					string APKSignerExecutable = Path.Combine(BuildToolsPath, "apksigner" + (RuntimePlatform.IsWindows ? ".bat" : ""));
+
+					ProcessStartInfo StartInfo = new ProcessStartInfo();
+					StartInfo.WorkingDirectory = IntermediateAndroidPath;
+					if (RuntimePlatform.IsWindows)
+					{
+						StartInfo.FileName = "cmd.exe";
+						StartInfo.ArgumentList.Add("/c");
+						StartInfo.ArgumentList.Add(APKSignerExecutable);
+						StartInfo.ArgumentList.Add("sign");
+						StartInfo.ArgumentList.Add("--ks");
+						StartInfo.ArgumentList.Add(KeyStore);
+						StartInfo.ArgumentList.Add("--ks-pass");
+						StartInfo.ArgumentList.Add("pass:" + KeyStorePassword);
+						StartInfo.ArgumentList.Add("--ks-key-alias");
+						StartInfo.ArgumentList.Add(KeyAlias);
+						StartInfo.ArgumentList.Add("--key-pass");
+						StartInfo.ArgumentList.Add("pass:" + KeyPassword);
+						StartInfo.ArgumentList.Add("--out");
+						StartInfo.ArgumentList.Add(DestApkName);
+						StartInfo.ArgumentList.Add(TempAPK);
+					}
+					else
+					{
+						StartInfo.FileName = "/bin/sh";
+						StartInfo.ArgumentList.Add("-c");
+						StartInfo.ArgumentList.Add("\"" + APKSignerExecutable + "\" sign--ks \"" + KeyStore + "\" --ks-pass pass:" + KeyStorePassword + " --ks-key-alias " + KeyAlias + " --key-pass pass:" + KeyPassword + " --out \"" + DestApkName + "2\" \"" + TempAPK + "\"");
+					}
+					StartInfo.UseShellExecute = false;
+					StartInfo.WindowStyle = ProcessWindowStyle.Minimized;
+
+					Logger.LogInformation("Applying apksigner to apk...");
+
+					Process Proc = new Process();
+					Proc.StartInfo = StartInfo;
+					Proc.Start();
+					Proc.WaitForExit();
+
+					// clean up work files
+					SafeDeleteFile(TempAPK);
+					SafeDeleteFile(DestApkName + ".idsig");
+					SafeDeleteFile(StrippedSOName);
+
+					if (Proc.ExitCode != 0)
+					{
+						string Args = "";
+						foreach (string Arg in StartInfo.ArgumentList)
+						{
+							Args += Arg + " ";
+						}
+						throw new BuildException("{0} failed with args {1}", StartInfo.FileName, Args);
+					}
+
+					// copy .so with symbols if requested
+					bool bBuildWithHiddenSymbolVisibility = false;
+					bool bSaveSymbols = false;
+					Ini.GetBool("/Script/AndroidRuntimeSettings.AndroidRuntimeSettings", "bBuildWithHiddenSymbolVisibility", out bBuildWithHiddenSymbolVisibility);
+					Ini.GetBool("/Script/AndroidRuntimeSettings.AndroidRuntimeSettings", "bSaveSymbols", out bSaveSymbols);
+					if (bSaveSymbols || (Configuration == UnrealTargetConfiguration.Shipping && bBuildWithHiddenSymbolVisibility))
+					{
+						// Copy .so with symbols to 
+						int StoreVersion = GetStoreVersion(Arch);
+						string SymbolSODirectory = Path.Combine(DestApkDirectory, ProjectName + "_Symbols_v" + StoreVersion + "/" + ProjectName + Arch);
+						string SymbolifiedSOPath = Path.Combine(SymbolSODirectory, Path.GetFileName(SourceSOName));
+						MakeDirectoryIfRequired(SymbolifiedSOPath);
+						Logger.LogInformation("Writing symbols to {SymbolifiedSOPath}", SymbolifiedSOPath);
+
+						File.Copy(SourceSOName, SymbolifiedSOPath, true);
+					}
+
+					BuildListComboRemaining--;
+				}
+
+				if (BuildListComboRemaining == 0)
+				{
+					Logger.LogInformation("\n===={Time}====COMPLETED MAKE APK=======================================================================", DateTime.Now.ToString());
+					return;
+				}
+			}
+
 			// at this point, we can write out the cached build settings to compare for a next build
 			File.WriteAllText(BuildSettingsCacheFile, CurrentBuildSettings);
 			if (File.Exists(CommandLineSourceFileName))
