UnrealBuildTool: Move Mac framework args to GetCompileArguments_Global

#rnx

diff --git a/Engine/Source/Programs/UnrealBuildTool/Platform/Mac/MacToolChain.cs b/Engine/Source/Programs/UnrealBuildTool/Platform/Mac/MacToolChain.cs
--- a/Engine/Source/Programs/UnrealBuildTool/Platform/Mac/MacToolChain.cs
+++ b/Engine/Source/Programs/UnrealBuildTool/Platform/Mac/MacToolChain.cs
@@ -198,6 +198,17 @@ namespace UnrealBuildTool
 			Arguments.Add("" + FormatArchitectureArg(CompileEnvironment.Architecture));
 			Arguments.Add($"-isysroot \"{SDKPath}\"");
 			Arguments.Add("-mmacosx-version-min=" + (CompileEnvironment.bEnableOSX109Support ? "10.9" : Settings.MacOSVersion));
+
+			List<string> FrameworksSearchPaths = new List<string>();
+			foreach (UEBuildFramework Framework in CompileEnvironment.AdditionalFrameworks)
+			{
+				FileReference FrameworkPath = new FileReference(Path.GetFullPath(Framework.Name));
+				if (!FrameworksSearchPaths.Contains(FrameworkPath.Directory.FullName))
+				{
+					Arguments.Add($"-F \"{NormalizeCommandLinePath(FrameworkPath.Directory)}\"");
+					FrameworksSearchPaths.Add(FrameworkPath.Directory.FullName);
+				}
+			}
 		}
 		
 		string AddFrameworkToLinkCommand(string FrameworkName, string Arg = "-framework")
@@ -278,17 +289,6 @@ namespace UnrealBuildTool
 
 			GetCompileArguments_Global(CompileEnvironment, GlobalArguments);
 
-			List<string> FrameworksSearchPaths = new List<string>();
-			foreach (UEBuildFramework Framework in CompileEnvironment.AdditionalFrameworks)
-			{
-				FileReference FrameworkPath = new FileReference(Path.GetFullPath(Framework.Name));
-				if (!FrameworksSearchPaths.Contains(FrameworkPath.Directory.FullName))
-				{
-					GlobalArguments.Add($"-F \"{NormalizeCommandLinePath(FrameworkPath.Directory)}\"");
-					FrameworksSearchPaths.Add(FrameworkPath.Directory.FullName);
-				}
-			}
-
 			CPPOutput Result = new CPPOutput();
 			// Create a compile action for each source file.
 			foreach (FileItem SourceFile in InputFiles)
