UnrealBuildTool: Override for clang additional arguments

#rnx

diff --git a/Engine/Source/Programs/UnrealBuildTool/Platform/IOS/IOSToolChain.cs b/Engine/Source/Programs/UnrealBuildTool/Platform/IOS/IOSToolChain.cs
--- a/Engine/Source/Programs/UnrealBuildTool/Platform/IOS/IOSToolChain.cs
+++ b/Engine/Source/Programs/UnrealBuildTool/Platform/IOS/IOSToolChain.cs
@@ -522,21 +522,7 @@ namespace UnrealBuildTool
 				List<string> ResponseFileContents = new();
 				ResponseFileContents.AddRange(GlobalArguments);
 				ResponseFileContents.AddRange(FileArguments);
-				if (!string.IsNullOrWhiteSpace(CompileEnvironment.AdditionalArguments))
-				{
-					ResponseFileContents.Add(CompileEnvironment.AdditionalArguments);
-				}
 				ResponseFileContents = ResponseFileContents.Select(x => Utils.ExpandVariables(x)).ToList();
-				/*				string SourceText = System.IO.File.ReadAllText(SourceFile.AbsolutePath);
-								if (CompileEnvironment.bOptimizeForSize && (SourceFile.AbsolutePath.Contains("ElementBatcher.cpp") || SourceText.Contains("ElementBatcher.cpp") || SourceFile.AbsolutePath.Contains("AnimationRuntime.cpp") || SourceText.Contains("AnimationRuntime.cpp")
-									|| SourceFile.AbsolutePath.Contains("AnimEncoding.cpp") || SourceText.Contains("AnimEncoding.cpp") || SourceFile.AbsolutePath.Contains("TextRenderComponent.cpp") || SourceText.Contains("TextRenderComponent.cpp")
-									|| SourceFile.AbsolutePath.Contains("SWidget.cpp") || SourceText.Contains("SWidget.cpp") || SourceFile.AbsolutePath.Contains("SCanvas.cpp") || SourceText.Contains("SCanvas.cpp") || SourceFile.AbsolutePath.Contains("ShaderCore.cpp") || SourceText.Contains("ShaderCore.cpp")
-									|| SourceFile.AbsolutePath.Contains("ParticleSystemRender.cpp") || SourceText.Contains("ParticleSystemRender.cpp")))
-								{
-									Logger.LogInformation("Forcing {SourceFileAbsolutePath} to --O3!", SourceFile.AbsolutePath);
-
-									AllArgs = AllArgs.Replace("-Oz", "-O3");
-								}*/
 
 				// Adds the response file to the compiler input.
 				FileItem ResponseFileItem = Graph.CreateIntermediateTextFile(ResponseFileName, ResponseFileContents);
diff --git a/Engine/Source/Programs/UnrealBuildTool/Platform/Linux/LinuxToolChain.cs b/Engine/Source/Programs/UnrealBuildTool/Platform/Linux/LinuxToolChain.cs
--- a/Engine/Source/Programs/UnrealBuildTool/Platform/Linux/LinuxToolChain.cs
+++ b/Engine/Source/Programs/UnrealBuildTool/Platform/Linux/LinuxToolChain.cs
@@ -981,7 +981,6 @@ namespace UnrealBuildTool
 					List<string> PreprocessResponseFileContents = new();
 					PreprocessResponseFileContents.AddRange(PreprocessGlobalArguments);
 					PreprocessResponseFileContents.AddRange(PreprocessFileArguments);
-					PreprocessResponseFileContents.Add(CompileEnvironment.AdditionalArguments);
 
 					// Adds the response file to the compiler input.
 					FileItem PreprocessResponseFileItem = Graph.CreateIntermediateTextFile(PreprocessResponseFileName, PreprocessResponseFileContents);
diff --git a/Engine/Source/Programs/UnrealBuildTool/Platform/Mac/MacToolChain.cs b/Engine/Source/Programs/UnrealBuildTool/Platform/Mac/MacToolChain.cs
--- a/Engine/Source/Programs/UnrealBuildTool/Platform/Mac/MacToolChain.cs
+++ b/Engine/Source/Programs/UnrealBuildTool/Platform/Mac/MacToolChain.cs
@@ -182,6 +182,32 @@ namespace UnrealBuildTool
 			}
 		}
 
+		/// <inheritdoc/>
+		protected override void GetCompileArguments_AdditionalArgs(CppCompileEnvironment CompileEnvironment, List<string> Arguments)
+		{
+			if (!string.IsNullOrWhiteSpace(CompileEnvironment.AdditionalArguments))
+			{
+				string EscapedAdditionalArgs = string.Empty;
+				foreach (string AdditionalArg in CompileEnvironment.AdditionalArguments.Split(new char[] { ' ' }, StringSplitOptions.RemoveEmptyEntries))
+				{
+					Match DefinitionMatch = Regex.Match(AdditionalArg, "-D\"?(?<Name>.*)=(?<Value>.*)\"?");
+					if (DefinitionMatch.Success)
+					{
+						EscapedAdditionalArgs += string.Format(" -D{0}=\"{1}\"", DefinitionMatch.Groups["Name"].Value, DefinitionMatch.Groups["Value"].Value);
+					}
+					else
+					{
+						EscapedAdditionalArgs += " " + AdditionalArg;
+					}
+				}
+
+				if (!string.IsNullOrWhiteSpace(EscapedAdditionalArgs))
+				{
+					Arguments.Add(EscapedAdditionalArgs);
+				}
+			}
+		}
+
 		/// <inheritdoc/>
 		protected override void GetCompileArguments_Global(CppCompileEnvironment CompileEnvironment, List<string> Arguments)
 		{
@@ -300,23 +326,6 @@ namespace UnrealBuildTool
 				// Add C or C++ specific compiler arguments.
 				GetCompileArguments_FileType(CompileEnvironment, SourceFile, OutputDir, FileArguments, CompileAction, Result);
 
-				string EscapedAdditionalArgs = "";
-				if(!string.IsNullOrWhiteSpace(CompileEnvironment.AdditionalArguments))
-				{
-					foreach(string AdditionalArg in CompileEnvironment.AdditionalArguments.Split(new char[] { ' ' }, StringSplitOptions.RemoveEmptyEntries))
-					{
-						Match DefinitionMatch = Regex.Match(AdditionalArg, "-D\"?(?<Name>.*)=(?<Value>.*)\"?");
-						if (DefinitionMatch.Success)
-						{
-							EscapedAdditionalArgs += string.Format(" -D{0}=\"{1}\"", DefinitionMatch.Groups["Name"].Value, DefinitionMatch.Groups["Value"].Value);
-						}
-						else
-						{
-							EscapedAdditionalArgs += " " + AdditionalArg;
-						}
-					}
-				}
-
 				// Gets the target file so we can get the correct output path.
 				FileItem TargetFile = CompileAction.ProducedItems.First();
 
@@ -325,7 +334,6 @@ namespace UnrealBuildTool
 				List<string> ResponseFileContents = new();
 				ResponseFileContents.AddRange(GlobalArguments);
 				ResponseFileContents.AddRange(FileArguments);
-				ResponseFileContents.Add(EscapedAdditionalArgs);
 
 				// Adds the response file to the compiler input.
 				FileItem ResponseFileItem = Graph.CreateIntermediateTextFile(ResponseFileName, ResponseFileContents);
@@ -391,7 +399,6 @@ namespace UnrealBuildTool
 					List<string> PreprocessResponseFileContents = new();
 					PreprocessResponseFileContents.AddRange(PreprocessGlobalArguments);
 					PreprocessResponseFileContents.AddRange(PreprocessFileArguments);
-					PreprocessResponseFileContents.Add(EscapedAdditionalArgs);
 
 					// Adds the response file to the compiler input.
 					FileItem PreprocessResponseFileItem = Graph.CreateIntermediateTextFile(PreprocessResponseFileName, PreprocessResponseFileContents);
diff --git a/Engine/Source/Programs/UnrealBuildTool/ToolChain/ClangToolChain.cs b/Engine/Source/Programs/UnrealBuildTool/ToolChain/ClangToolChain.cs
--- a/Engine/Source/Programs/UnrealBuildTool/ToolChain/ClangToolChain.cs
+++ b/Engine/Source/Programs/UnrealBuildTool/ToolChain/ClangToolChain.cs
@@ -584,6 +584,19 @@ namespace UnrealBuildTool
 			}
 		}
 
+		/// <summary>
+		/// Additional compile arguments.
+		/// </summary>
+		/// <param name="CompileEnvironment"></param>
+		/// <param name="Arguments"></param>
+		protected virtual void GetCompileArguments_AdditionalArgs(CppCompileEnvironment CompileEnvironment, List<string> Arguments)
+		{
+			if (!string.IsNullOrWhiteSpace(CompileEnvironment.AdditionalArguments))
+			{
+				Arguments.Add(CompileEnvironment.AdditionalArguments);
+			}
+		}
+
 		/// <summary>
 		/// Common compile arguments for all files in a module.
 		/// Override and call base.GetCompileArguments_Global() in derived classes.
@@ -615,8 +628,11 @@ namespace UnrealBuildTool
 			// Add debugging flags to the argument list.
 			GetCompileArguments_Debugging(CompileEnvironment, Arguments);
 
-			// Add sanitizer flags to the argument list
+			// Add sanitizer flags to the argument list.
 			GetCompilerArguments_Sanitizers(CompileEnvironment, Arguments);
+
+			// Add additional arguments to the argument list.
+			GetCompileArguments_AdditionalArgs(CompileEnvironment, Arguments);
 		}
 
 		/// <summary>
