UnrealBuildTool: Reconcile compile action logic to ClangToolChain.CompileCPPFile

#rnx

diff --git a/Engine/Source/Programs/UnrealBuildTool/Platform/IOS/IOSToolChain.cs b/Engine/Source/Programs/UnrealBuildTool/Platform/IOS/IOSToolChain.cs
--- a/Engine/Source/Programs/UnrealBuildTool/Platform/IOS/IOSToolChain.cs
+++ b/Engine/Source/Programs/UnrealBuildTool/Platform/IOS/IOSToolChain.cs
@@ -501,62 +501,36 @@ namespace UnrealBuildTool
 
 			GetCompileArguments_Global(CompileEnvironment, GlobalArguments);
 
+			FileReference CompilerPath = new FileReference(Settings.Value.ToolchainDir + IOSCompiler);
+			string CompilerVersion = GetFullClangVersion();
+
+			List<FileItem> FrameworkTokenFiles = new List<FileItem>();
+			foreach (UEBuildFramework Framework in CompileEnvironment.AdditionalFrameworks)
+			{
+				if (Framework.ZipFile != null)
+				{
+					FileItem ExtractedTokenFile = ExtractFramework(Framework, Graph, Logger);
+					FrameworkTokenFiles.Add(ExtractedTokenFile);
+				}
+			}
+
 			CPPOutput Result = new CPPOutput();
 			// Create a compile action for each source file.
 			foreach (FileItem SourceFile in InputFiles)
 			{
-				Action CompileAction = Graph.CreateAction(ActionType.Compile);
-
-				List<string> FileArguments = new();
-
-				// Add C or C++ specific compiler arguments.
-				GetCompileArguments_FileType(CompileEnvironment, SourceFile, OutputDir, FileArguments, CompileAction, Result);
-
-				string CompilerPath = Settings.Value.ToolchainDir + IOSCompiler;
-
-				// Gets the target file so we can get the correct output path.
-				FileItem TargetFile = CompileAction.ProducedItems.First();
-
-				// Creates the path to the response file using the name of the output file and creates its contents.
-				FileReference ResponseFileName = new FileReference(TargetFile.AbsolutePath + ".response");
-				List<string> ResponseFileContents = new();
-				ResponseFileContents.AddRange(GlobalArguments);
-				ResponseFileContents.AddRange(FileArguments);
-				ResponseFileContents = ResponseFileContents.Select(x => Utils.ExpandVariables(x)).ToList();
-
-				// Adds the response file to the compiler input.
-				FileItem ResponseFileItem = Graph.CreateIntermediateTextFile(ResponseFileName, ResponseFileContents);
-				CompileAction.CommandArguments += string.Format("@\"{0}\"", NormalizeCommandLinePath(ResponseFileItem));
-				CompileAction.PrerequisiteItems.Add(ResponseFileItem);
+				Action CompileAction = CompileCPPFile(CompileEnvironment, SourceFile, OutputDir, ModuleName, Graph, GlobalArguments, CompilerPath, CompilerVersion, Result);
 
 				// Analyze and then compile using the shell to perform the indirection
 				string? StaticAnalysisMode = Environment.GetEnvironmentVariable("CLANG_STATIC_ANALYZER_MODE");
 				if (StaticAnalysisMode != null && StaticAnalysisMode != "")
 				{
 					FileReference ReportFile = new FileReference(SourceFile.AbsolutePath + ".html");
-					CompileAction.CommandArguments = "-c \"" + CompilerPath + " " + string.Format("@\"{0}\"", ResponseFileName) + " --analyze -Wno-unused-command-line-argument -Xclang -analyzer-output=html -Xclang -analyzer-config -Xclang path-diagnostics-alternate=true -Xclang -analyzer-config -Xclang report-in-main-source-file=true -Xclang -analyzer-disable-checker -Xclang deadcode.DeadStores -o " + ReportFile + "; " + CompilerPath + " " + string.Format("@\"{0}\"", ResponseFileName) + "\"";
-					CompilerPath = "/bin/sh";
+					string Arguments = CompileAction.CommandArguments;
+					CompileAction.CommandArguments = "-c \"" + CompilerPath + " " + Arguments + " --analyze -Wno-unused-command-line-argument -Xclang -analyzer-output=html -Xclang -analyzer-config -Xclang path-diagnostics-alternate=true -Xclang -analyzer-config -Xclang report-in-main-source-file=true -Xclang -analyzer-disable-checker -Xclang deadcode.DeadStores -o " + ReportFile + "; " + CompilerPath + " " + Arguments + "\"";
+					CompileAction.CommandPath = new FileReference("/bin/sh");
 				}
 
-				// RPC utility parameters are in terms of the Mac side
-				CompileAction.WorkingDirectory = GetMacDevSrcRoot();
-				CompileAction.CommandPath = new FileReference(CompilerPath);
-				CompileAction.CommandDescription = "Compile";
-				CompileAction.StatusDescription = string.Format("{0}", Path.GetFileName(SourceFile.AbsolutePath));
-				CompileAction.bIsGCCCompiler = true;
-				// We're already distributing the command by execution on Mac.
-				CompileAction.bCanExecuteRemotely = CompileEnvironment.PrecompiledHeaderAction != PrecompiledHeaderAction.Create;
-				CompileAction.bShouldOutputStatusDescription = true;
-				CompileAction.CommandVersion = GetFullClangVersion();
-
-				foreach (UEBuildFramework Framework in CompileEnvironment.AdditionalFrameworks)
-				{
-					if (Framework.ZipFile != null)
-					{
-						FileItem ExtractedTokenFile = ExtractFramework(Framework, Graph, Logger);
-						CompileAction.PrerequisiteItems.Add(ExtractedTokenFile);
-					}
-				}
+				CompileAction.PrerequisiteItems.AddRange(FrameworkTokenFiles);
 			}
 			return Result;
 		}
diff --git a/Engine/Source/Programs/UnrealBuildTool/Platform/Linux/LinuxToolChain.cs b/Engine/Source/Programs/UnrealBuildTool/Platform/Linux/LinuxToolChain.cs
--- a/Engine/Source/Programs/UnrealBuildTool/Platform/Linux/LinuxToolChain.cs
+++ b/Engine/Source/Programs/UnrealBuildTool/Platform/Linux/LinuxToolChain.cs
@@ -913,83 +913,7 @@ namespace UnrealBuildTool
 			CPPOutput Result = new CPPOutput();
 			foreach (FileItem SourceFile in InputFiles)
 			{
-				Action CompileAction = Graph.CreateAction(ActionType.Compile);
-
-				List<string> FileArguments = new();
-
-				// Add C or C++ specific compiler arguments.
-				GetCompileArguments_FileType(CompileEnvironment, SourceFile, OutputDir, FileArguments, CompileAction, Result);
-
-				// Gets the target file so we can get the correct output path.
-				FileItem TargetFile = CompileAction.ProducedItems.First();
-
-				CompileAction.WorkingDirectory = Unreal.EngineSourceDirectory;
-				CompileAction.CommandPath = ClangPath!;
-
-				List<string> ResponseFileContents = new();
-				ResponseFileContents.AddRange(GlobalArguments);
-				ResponseFileContents.AddRange(FileArguments);
-
-				FileReference CompilerResponseFileName = TargetFile.Location + ".rsp";
-				FileItem CompilerResponseFileItem = Graph.CreateIntermediateTextFile(CompilerResponseFileName, ResponseFileContents);
-
-				CompileAction.CommandArguments = GetResponseFileArgument(CompilerResponseFileItem);
-				CompileAction.PrerequisiteItems.Add(CompilerResponseFileItem);
-				CompileAction.CommandDescription = "Compile";
-				CompileAction.CommandVersion = ClangVersionString!;
-				CompileAction.StatusDescription = Path.GetFileName(SourceFile.AbsolutePath);
-				CompileAction.bIsGCCCompiler = true;
-
-				// Don't farm out creation of pre-compiled headers as it is the critical path task.
-				CompileAction.bCanExecuteRemotely =
-					CompileEnvironment.PrecompiledHeaderAction != PrecompiledHeaderAction.Create ||
-					CompileEnvironment.bAllowRemotelyCompiledPCHs;
-
-				if (bPreprocessDepends && CompileEnvironment.bGenerateDependenciesFile)
-				{
-					Action PrepassAction = Graph.CreateAction(ActionType.Compile);
-					PrepassAction.PrerequisiteItems.AddRange(CompileAction.PrerequisiteItems);
-					PrepassAction.PrerequisiteItems.Remove(CompilerResponseFileItem);
-					PrepassAction.CommandDescription = "Preprocess Depends";
-					PrepassAction.StatusDescription = CompileAction.StatusDescription;
-					PrepassAction.bIsGCCCompiler = true;
-					PrepassAction.bCanExecuteRemotely = false;
-					PrepassAction.bShouldOutputStatusDescription = true;
-					PrepassAction.CommandPath = CompileAction.CommandPath;
-					PrepassAction.CommandVersion = CompileAction.CommandVersion;
-					PrepassAction.WorkingDirectory = CompileAction.WorkingDirectory;
-
-					List<string> PreprocessGlobalArguments = new(GlobalArguments);
-					List<string> PreprocessFileArguments = new(FileArguments);
-					PreprocessGlobalArguments.Remove("-c");
-
-					FileItem DependencyListFile = FileItem.GetItemByFileReference(FileReference.Combine(OutputDir, Path.GetFileName(SourceFile.AbsolutePath) + ".d"));
-					PreprocessFileArguments.Add(GetPreprocessDepencenciesListFileArgument(DependencyListFile));
-					PrepassAction.DependencyListFile = DependencyListFile;
-					PrepassAction.ProducedItems.Add(DependencyListFile);
-
-					PreprocessFileArguments.Remove("-ftime-trace");
-					PreprocessFileArguments.Remove(GetOutputFileArgument(CompileAction.ProducedItems.First()));
-
-					PrepassAction.DeleteItems.AddRange(PrepassAction.ProducedItems);
-
-					// Gets the target file so we can get the correct output path.
-					FileItem PreprocessTargetFile = PrepassAction.ProducedItems[0];
-
-					// Creates the path to the response file using the name of the output file and creates its contents.
-					FileReference PreprocessResponseFileName = new FileReference(PreprocessTargetFile.AbsolutePath + ".rsp");
-					List<string> PreprocessResponseFileContents = new();
-					PreprocessResponseFileContents.AddRange(PreprocessGlobalArguments);
-					PreprocessResponseFileContents.AddRange(PreprocessFileArguments);
-
-					// Adds the response file to the compiler input.
-					FileItem PreprocessResponseFileItem = Graph.CreateIntermediateTextFile(PreprocessResponseFileName, PreprocessResponseFileContents);
-					PrepassAction.PrerequisiteItems.Add(PreprocessResponseFileItem);
-
-					PrepassAction.CommandArguments = GetResponseFileArgument(PreprocessResponseFileItem);
-					CompileAction.DependencyListFile = DependencyListFile;
-					CompileAction.PrerequisiteItems.Add(DependencyListFile);
-				}
+				CompileCPPFile(CompileEnvironment, SourceFile, OutputDir, ModuleName, Graph, GlobalArguments, ClangPath!, ClangVersionString!, Result);
 			}
 
 			return Result;
diff --git a/Engine/Source/Programs/UnrealBuildTool/Platform/Mac/MacToolChain.cs b/Engine/Source/Programs/UnrealBuildTool/Platform/Mac/MacToolChain.cs
--- a/Engine/Source/Programs/UnrealBuildTool/Platform/Mac/MacToolChain.cs
+++ b/Engine/Source/Programs/UnrealBuildTool/Platform/Mac/MacToolChain.cs
@@ -315,101 +315,24 @@ namespace UnrealBuildTool
 
 			GetCompileArguments_Global(CompileEnvironment, GlobalArguments);
 
+			FileReference CompilerPath = new FileReference(Settings.ToolchainDir + MacCompiler);
+			string CompilerVersion = GetFullClangVersion();
+
 			CPPOutput Result = new CPPOutput();
 			// Create a compile action for each source file.
 			foreach (FileItem SourceFile in InputFiles)
 			{
-				Action CompileAction = Graph.CreateAction(ActionType.Compile);
-
-				List<string> FileArguments = new();
-
-				// Add C or C++ specific compiler arguments.
-				GetCompileArguments_FileType(CompileEnvironment, SourceFile, OutputDir, FileArguments, CompileAction, Result);
-
-				// Gets the target file so we can get the correct output path.
-				FileItem TargetFile = CompileAction.ProducedItems.First();
-
-				// Creates the path to the response file using the name of the output file and creates its contents.
-				FileReference ResponseFileName = new FileReference(TargetFile.AbsolutePath + ".response");
-				List<string> ResponseFileContents = new();
-				ResponseFileContents.AddRange(GlobalArguments);
-				ResponseFileContents.AddRange(FileArguments);
+				Action CompileAction = CompileCPPFile(CompileEnvironment, SourceFile, OutputDir, ModuleName, Graph, GlobalArguments, CompilerPath, CompilerVersion, Result);
 
-				// Adds the response file to the compiler input.
-				FileItem ResponseFileItem = Graph.CreateIntermediateTextFile(ResponseFileName, ResponseFileContents);
-				CompileAction.CommandArguments += string.Format("@\"{0}\"", NormalizeCommandLinePath(ResponseFileItem));
-				CompileAction.PrerequisiteItems.Add(ResponseFileItem);
-
-				string CompilerPath = Settings.ToolchainDir + MacCompiler;
-				
 				// Analyze and then compile using the shell to perform the indirection
 				string? StaticAnalysisMode = Environment.GetEnvironmentVariable("CLANG_STATIC_ANALYZER_MODE");
 				if (StaticAnalysisMode != null && StaticAnalysisMode != "")
 				{
 					FileReference ReportFile = new FileReference(SourceFile.AbsolutePath + ".html");
-					CompileAction.CommandArguments = "-c \"" + CompilerPath + " " + string.Format("@\"{0}\"", ResponseFileName) + " --analyze -Wno-unused-command-line-argument -Xclang -analyzer-output=html -Xclang -analyzer-config -Xclang path-diagnostics-alternate=true -Xclang -analyzer-config -Xclang report-in-main-source-file=true -Xclang -analyzer-disable-checker -Xclang deadcode.DeadStores -o " + ReportFile + "; " + CompilerPath + " " + string.Format("@\"{0}\"", ResponseFileName) + "\"";
-					CompilerPath = "/bin/sh";
+					string Arguments = CompileAction.CommandArguments;
+					CompileAction.CommandArguments = "-c \"" + CompilerPath + " " + Arguments + " --analyze -Wno-unused-command-line-argument -Xclang -analyzer-output=html -Xclang -analyzer-config -Xclang path-diagnostics-alternate=true -Xclang -analyzer-config -Xclang report-in-main-source-file=true -Xclang -analyzer-disable-checker -Xclang deadcode.DeadStores -o " + ReportFile + "; " + CompilerPath + " " + Arguments + "\"";
+					CompileAction.CommandPath = new FileReference("/bin/sh");
 				}
-
-				// RPC utility parameters are in terms of the Mac side
-				CompileAction.WorkingDirectory = GetMacDevSrcRoot();
-				CompileAction.CommandPath = new FileReference(CompilerPath);
-
-				// For compilation we delete everything we produce
-				CompileAction.DeleteItems.AddRange(CompileAction.ProducedItems);
-				CompileAction.CommandDescription = "Compile";
-				CompileAction.StatusDescription = Path.GetFileName(SourceFile.AbsolutePath);
-				CompileAction.bIsGCCCompiler = true;
-				// We're already distributing the command by execution on Mac.
-				CompileAction.bCanExecuteRemotely = CompileEnvironment.PrecompiledHeaderAction != PrecompiledHeaderAction.Create;
-				CompileAction.bShouldOutputStatusDescription = true;
-				CompileAction.CommandVersion = GetFullClangVersion();
-
-				if (bPreprocessDepends && CompileEnvironment.bGenerateDependenciesFile)
-				{
-					Action PrepassAction = Graph.CreateAction(ActionType.Compile);
-					PrepassAction.PrerequisiteItems.AddRange(CompileAction.PrerequisiteItems);
-					PrepassAction.CommandDescription = "Preprocess Depends";
-					PrepassAction.StatusDescription = CompileAction.StatusDescription;
-					PrepassAction.bIsGCCCompiler = true;
-					PrepassAction.bCanExecuteRemotely = false;
-					PrepassAction.bShouldOutputStatusDescription = true;
-					PrepassAction.CommandVersion = CompileAction.CommandVersion;
-					PrepassAction.WorkingDirectory = CompileAction.WorkingDirectory;
-
-					List<string> PreprocessGlobalArguments = new(GlobalArguments);
-					List<string> PreprocessFileArguments = new(FileArguments);
-					PreprocessGlobalArguments.Remove("-c");
-
-					FileItem DependencyListFile = FileItem.GetItemByFileReference(FileReference.Combine(OutputDir, Path.GetFileName(SourceFile.AbsolutePath) + ".d"));
-					PreprocessFileArguments.Add(string.Format("-M -MF\"{0}\"", NormalizeCommandLinePath(DependencyListFile)));
-					PrepassAction.DependencyListFile = DependencyListFile;
-					PrepassAction.ProducedItems.Add(DependencyListFile);
-
-					PreprocessFileArguments.Remove("-ftime-trace");
-					PreprocessFileArguments.Remove(string.Format(" -o \"{0}\"", NormalizeCommandLinePath(CompileAction.ProducedItems.First())));
-
-					PrepassAction.DeleteItems.AddRange(PrepassAction.ProducedItems);
-
-					// Gets the target file so we can get the correct output path.
-					FileItem PreprocessTargetFile = PrepassAction.ProducedItems[0];
-
-					// Creates the path to the response file using the name of the output file and creates its contents.
-					FileReference PreprocessResponseFileName = new FileReference(PreprocessTargetFile.AbsolutePath + ".response");
-					List<string> PreprocessResponseFileContents = new();
-					PreprocessResponseFileContents.AddRange(PreprocessGlobalArguments);
-					PreprocessResponseFileContents.AddRange(PreprocessFileArguments);
-
-					// Adds the response file to the compiler input.
-					FileItem PreprocessResponseFileItem = Graph.CreateIntermediateTextFile(PreprocessResponseFileName, PreprocessResponseFileContents);
-					PrepassAction.PrerequisiteItems.Add(PreprocessResponseFileItem);
-
-					PrepassAction.CommandPath = new FileReference(CompilerPath);
-					PrepassAction.CommandArguments = string.Format("@\"{0}\"", PreprocessResponseFileItem);
-					CompileAction.DependencyListFile = DependencyListFile;
-					CompileAction.PrerequisiteItems.Add(DependencyListFile);
-				}
-
 			}
 			return Result;
 		}
diff --git a/Engine/Source/Programs/UnrealBuildTool/ToolChain/AppleToolChain.cs b/Engine/Source/Programs/UnrealBuildTool/ToolChain/AppleToolChain.cs
--- a/Engine/Source/Programs/UnrealBuildTool/ToolChain/AppleToolChain.cs
+++ b/Engine/Source/Programs/UnrealBuildTool/ToolChain/AppleToolChain.cs
@@ -180,7 +180,7 @@ namespace UnrealBuildTool
 		static string? FullClangVersion = null;
 
 		protected Version GetClangVersion()
-		{	
+		{
 			if (ClangVersion == null)
 			{
 				FileReference ClangLocation = new FileReference("/usr/bin/clang");
@@ -327,6 +327,30 @@ namespace UnrealBuildTool
 			}
 		}
 
+		/// <inheritdoc/>
+		protected override void GetCompileArguments_Analyze(CppCompileEnvironment CompileEnvironment, List<string> Arguments)
+		{
+			// TODO: Replace logic in MacToolChain.CompileCPPFiles and IOSToolChain.CompileCPPFiles
+			string? StaticAnalysisMode = Environment.GetEnvironmentVariable("CLANG_STATIC_ANALYZER_MODE");
+			if (string.IsNullOrEmpty(StaticAnalysisMode))
+			{
+				return;
+			}
+			
+			base.GetCompileArguments_Analyze(CompileEnvironment, Arguments);
+
+			Arguments.Add("-Xclang -analyzer-output=html");
+			Arguments.Add("-Xclang -analyzer-config -Xclang stable-report-filename=true");
+			Arguments.Add("-Xclang -analyzer-config -Xclang report-in-main-source-file=true");
+			Arguments.Add("-Xclang -analyzer-config -Xclang path-diagnostics-alternate=true");
+			Arguments.Add("-Xclang -analyzer-disable-checker -Xclang deadcode.DeadStores");
+
+			if (string.Equals(StaticAnalysisMode, "shallow"))
+			{
+				Arguments.Add("-Xclang -analyzer-config -Xclang mode=shallow");
+			}
+		}
+
 		/// <inheritdoc/>
 		protected override void GetCompileArguments_Global(CppCompileEnvironment CompileEnvironment, List<string> Arguments)
 		{
@@ -339,6 +363,11 @@ namespace UnrealBuildTool
 
 			Arguments.Add("-fmessage-length=0");
 			Arguments.Add("-fpascal-strings");
+
+			if (!string.IsNullOrEmpty(Environment.GetEnvironmentVariable("CLANG_STATIC_ANALYZER_MODE")))
+			{
+				GetCompileArguments_Analyze(CompileEnvironment, Arguments);
+			}
 		}
 
 		protected string GetDsymutilPath(ILogger Logger, out string ExtraOptions, bool bIsForLTOBuild=false)
diff --git a/Engine/Source/Programs/UnrealBuildTool/ToolChain/ClangToolChain.cs b/Engine/Source/Programs/UnrealBuildTool/ToolChain/ClangToolChain.cs
--- a/Engine/Source/Programs/UnrealBuildTool/ToolChain/ClangToolChain.cs
+++ b/Engine/Source/Programs/UnrealBuildTool/ToolChain/ClangToolChain.cs
@@ -121,6 +121,18 @@ namespace UnrealBuildTool
 
 		protected bool bPreprocessDepends = false;
 
+		// Dummy define to work around clang compilation related to the windows maximum path length limitation
+		protected static string ClangDummyDefine; 
+		protected const int ClangCmdLineMaxSize = 32 * 1024;
+		protected const int ClangCmdlineDangerZone = 30 * 1024;
+
+		static ClangToolChain()
+		{
+			const string DummyStart = "-D \"DUMMY_DEFINE";
+			const string DummyEnd = "\"";
+			ClangDummyDefine = DummyStart + String.Format("_").PadRight(ClangCmdLineMaxSize - ClangCmdlineDangerZone - DummyStart.Length - DummyEnd.Length, 'X') + DummyEnd;
+		}
+
 		public ClangToolChain(ClangToolChainOptions InOptions, ILogger InLogger)
 			: base(InLogger)
 		{
@@ -597,6 +609,16 @@ namespace UnrealBuildTool
 			}
 		}
 
+		/// <summary>
+		/// Compile arguments for running clang-analyze. Currently unsupported.
+		/// </summary>
+		/// <param name="CompileEnvironment"></param>
+		/// <param name="Arguments"></param>
+		protected virtual void GetCompileArguments_Analyze(CppCompileEnvironment CompileEnvironment, List<string> Arguments)
+		{
+			Arguments.Add("--analyze");
+		}
+
 		/// <summary>
 		/// Common compile arguments for all files in a module.
 		/// Override and call base.GetCompileArguments_Global() in derived classes.
@@ -744,5 +766,117 @@ namespace UnrealBuildTool
 			// Add the parameters needed to compile the output file to the command-line.
 			Arguments.Add(GetOutputFileArgument(OutputFile));
 		}
+
+		protected virtual Action CompileCPPFile(CppCompileEnvironment CompileEnvironment, FileItem SourceFile, DirectoryReference OutputDir, string ModuleName, IActionGraphBuilder Graph, IReadOnlyCollection<string> GlobalArguments, FileReference CommandPath, string CommandVersion, CPPOutput Result)
+		{
+			Action CompileAction = Graph.CreateAction(ActionType.Compile);
+
+			List<string> FileArguments = new();
+
+			// Add C or C++ specific compiler arguments.
+			GetCompileArguments_FileType(CompileEnvironment, SourceFile, OutputDir, FileArguments, CompileAction, Result);
+
+			// Gets the target file so we can get the correct output path.
+			FileItem TargetFile = CompileAction.ProducedItems.First();
+
+			// Creates the path to the response file using the name of the output file and creates its contents.
+			FileReference ResponseFileName = new FileReference(TargetFile.AbsolutePath + ".response");
+			List<string> ResponseFileContents = new();
+			ResponseFileContents.AddRange(GlobalArguments);
+			ResponseFileContents.AddRange(FileArguments);
+			ResponseFileContents = ResponseFileContents.Select(x => Utils.ExpandVariables(x)).ToList();
+
+			if (RuntimePlatform.IsWindows)
+			{
+				// Enables Clang Cmd Line Work Around
+				// Clang reads the response file and computes the total cmd line length.
+				// Depending on the length it will then use either cmd line or response file mode.
+				// Clang currently underestimates the total size of the cmd line, which can trigger the following clang error in cmd line mode:
+				// >>>> xxx-clang.exe': The filename or extension is too long.  (0xCE)
+				// Clang processes and modifies the response file contents and this makes the final cmd line size hard for us to predict.
+				// To be conservative we add a dummy define to inflate the response file size and force clang to use the response file mode when we are close to the limit.
+				int CmdLineLength = CommandPath.ToString().Length + string.Join(' ', ResponseFileContents).Length;
+				bool bIsInDangerZone = CmdLineLength >= ClangCmdlineDangerZone && CmdLineLength <= ClangCmdLineMaxSize;
+				if (bIsInDangerZone)
+				{
+					ResponseFileContents.Add(ClangDummyDefine);
+				}
+			}
+
+			// Adds the response file to the compiler input.
+			FileItem CompilerResponseFileItem = Graph.CreateIntermediateTextFile(ResponseFileName, ResponseFileContents);
+			CompileAction.CommandArguments = GetResponseFileArgument(CompilerResponseFileItem);
+			CompileAction.PrerequisiteItems.Add(CompilerResponseFileItem);
+
+			CompileAction.WorkingDirectory = Unreal.EngineSourceDirectory;
+			CompileAction.CommandPath = CommandPath;
+			CompileAction.CommandVersion = CommandVersion;
+			CompileAction.CommandDescription = "Compile";
+			CompileAction.StatusDescription = Path.GetFileName(SourceFile.AbsolutePath);
+			CompileAction.bIsGCCCompiler = true;
+
+			// Don't farm out creation of pre-compiled headers as it is the critical path task.
+			CompileAction.bCanExecuteRemotely =
+				CompileEnvironment.PrecompiledHeaderAction != PrecompiledHeaderAction.Create ||
+				CompileEnvironment.bAllowRemotelyCompiledPCHs;
+
+			if (bPreprocessDepends && CompileEnvironment.bGenerateDependenciesFile)
+			{
+				Action PrepassAction = Graph.CreateAction(ActionType.Compile);
+				PrepassAction.PrerequisiteItems.AddRange(CompileAction.PrerequisiteItems);
+				PrepassAction.PrerequisiteItems.Remove(CompilerResponseFileItem);
+				PrepassAction.CommandDescription = "Preprocess Depends";
+				PrepassAction.StatusDescription = CompileAction.StatusDescription;
+				PrepassAction.bIsGCCCompiler = true;
+				PrepassAction.bCanExecuteRemotely = false;
+				PrepassAction.bShouldOutputStatusDescription = true;
+				PrepassAction.CommandPath = CompileAction.CommandPath;
+				PrepassAction.CommandVersion = CompileAction.CommandVersion;
+				PrepassAction.WorkingDirectory = CompileAction.WorkingDirectory;
+
+				List<string> PreprocessGlobalArguments = new(GlobalArguments);
+				List<string> PreprocessFileArguments = new(FileArguments);
+				PreprocessGlobalArguments.Remove("-c");
+
+				FileItem DependencyListFile = FileItem.GetItemByFileReference(FileReference.Combine(OutputDir, Path.GetFileName(SourceFile.AbsolutePath) + ".d"));
+				PreprocessFileArguments.Add(GetPreprocessDepencenciesListFileArgument(DependencyListFile));
+				PrepassAction.DependencyListFile = DependencyListFile;
+				PrepassAction.ProducedItems.Add(DependencyListFile);
+
+				PreprocessFileArguments.Remove("-ftime-trace");
+				PreprocessFileArguments.Remove(GetOutputFileArgument(CompileAction.ProducedItems.First()));
+
+				PrepassAction.DeleteItems.AddRange(PrepassAction.ProducedItems);
+
+				// Gets the target file so we can get the correct output path.
+				FileItem PreprocessTargetFile = PrepassAction.ProducedItems[0];
+
+				// Creates the path to the response file using the name of the output file and creates its contents.
+				FileReference PreprocessResponseFileName = new FileReference(PreprocessTargetFile.AbsolutePath + ".response");
+				List<string> PreprocessResponseFileContents = new();
+				PreprocessResponseFileContents.AddRange(PreprocessGlobalArguments);
+				PreprocessResponseFileContents.AddRange(PreprocessFileArguments);
+
+				if (RuntimePlatform.IsWindows)
+				{
+					int CmdLineLength = CommandPath.ToString().Length + string.Join(' ', ResponseFileContents).Length;
+					bool bIsInDangerZone = CmdLineLength >= ClangCmdlineDangerZone && CmdLineLength <= ClangCmdLineMaxSize;
+					if (bIsInDangerZone)
+					{
+						ResponseFileContents.Add(ClangDummyDefine);
+					}
+				}
+
+				// Adds the response file to the compiler input.
+				FileItem PreprocessResponseFileItem = Graph.CreateIntermediateTextFile(PreprocessResponseFileName, PreprocessResponseFileContents);
+				PrepassAction.PrerequisiteItems.Add(PreprocessResponseFileItem);
+
+				PrepassAction.CommandArguments = GetResponseFileArgument(PreprocessResponseFileItem);
+				CompileAction.DependencyListFile = DependencyListFile;
+				CompileAction.PrerequisiteItems.Add(DependencyListFile);
+			}
+
+			return CompileAction;
+		}
 	}
 }
\ No newline at end of file
