UnrealBuildTool: Support two-pass CLANG_STATIC_ANALYZER_MODE for all clang platforms

#rnx

diff --git a/Engine/Source/Programs/UnrealBuildTool/Platform/IOS/IOSToolChain.cs b/Engine/Source/Programs/UnrealBuildTool/Platform/IOS/IOSToolChain.cs
--- a/Engine/Source/Programs/UnrealBuildTool/Platform/IOS/IOSToolChain.cs
+++ b/Engine/Source/Programs/UnrealBuildTool/Platform/IOS/IOSToolChain.cs
@@ -519,17 +519,6 @@ namespace UnrealBuildTool
 			foreach (FileItem SourceFile in InputFiles)
 			{
 				Action CompileAction = CompileCPPFile(CompileEnvironment, SourceFile, OutputDir, ModuleName, Graph, GlobalArguments, CompilerPath, CompilerVersion, Result);
-
-				// Analyze and then compile using the shell to perform the indirection
-				string? StaticAnalysisMode = Environment.GetEnvironmentVariable("CLANG_STATIC_ANALYZER_MODE");
-				if (StaticAnalysisMode != null && StaticAnalysisMode != "")
-				{
-					FileReference ReportFile = new FileReference(SourceFile.AbsolutePath + ".html");
-					string Arguments = CompileAction.CommandArguments;
-					CompileAction.CommandArguments = "-c \"" + CompilerPath + " " + Arguments + " --analyze -Wno-unused-command-line-argument -Xclang -analyzer-output=html -Xclang -analyzer-config -Xclang path-diagnostics-alternate=true -Xclang -analyzer-config -Xclang report-in-main-source-file=true -Xclang -analyzer-disable-checker -Xclang deadcode.DeadStores -o " + ReportFile + "; " + CompilerPath + " " + Arguments + "\"";
-					CompileAction.CommandPath = new FileReference("/bin/sh");
-				}
-
 				CompileAction.PrerequisiteItems.AddRange(FrameworkTokenFiles);
 			}
 			return Result;
diff --git a/Engine/Source/Programs/UnrealBuildTool/Platform/Mac/MacToolChain.cs b/Engine/Source/Programs/UnrealBuildTool/Platform/Mac/MacToolChain.cs
--- a/Engine/Source/Programs/UnrealBuildTool/Platform/Mac/MacToolChain.cs
+++ b/Engine/Source/Programs/UnrealBuildTool/Platform/Mac/MacToolChain.cs
@@ -322,17 +322,7 @@ namespace UnrealBuildTool
 			// Create a compile action for each source file.
 			foreach (FileItem SourceFile in InputFiles)
 			{
-				Action CompileAction = CompileCPPFile(CompileEnvironment, SourceFile, OutputDir, ModuleName, Graph, GlobalArguments, CompilerPath, CompilerVersion, Result);
-
-				// Analyze and then compile using the shell to perform the indirection
-				string? StaticAnalysisMode = Environment.GetEnvironmentVariable("CLANG_STATIC_ANALYZER_MODE");
-				if (StaticAnalysisMode != null && StaticAnalysisMode != "")
-				{
-					FileReference ReportFile = new FileReference(SourceFile.AbsolutePath + ".html");
-					string Arguments = CompileAction.CommandArguments;
-					CompileAction.CommandArguments = "-c \"" + CompilerPath + " " + Arguments + " --analyze -Wno-unused-command-line-argument -Xclang -analyzer-output=html -Xclang -analyzer-config -Xclang path-diagnostics-alternate=true -Xclang -analyzer-config -Xclang report-in-main-source-file=true -Xclang -analyzer-disable-checker -Xclang deadcode.DeadStores -o " + ReportFile + "; " + CompilerPath + " " + Arguments + "\"";
-					CompileAction.CommandPath = new FileReference("/bin/sh");
-				}
+				CompileCPPFile(CompileEnvironment, SourceFile, OutputDir, ModuleName, Graph, GlobalArguments, CompilerPath, CompilerVersion, Result);
 			}
 			return Result;
 		}
diff --git a/Engine/Source/Programs/UnrealBuildTool/ToolChain/AppleToolChain.cs b/Engine/Source/Programs/UnrealBuildTool/ToolChain/AppleToolChain.cs
--- a/Engine/Source/Programs/UnrealBuildTool/ToolChain/AppleToolChain.cs
+++ b/Engine/Source/Programs/UnrealBuildTool/ToolChain/AppleToolChain.cs
@@ -327,30 +327,6 @@ namespace UnrealBuildTool
 			}
 		}
 
-		/// <inheritdoc/>
-		protected override void GetCompileArguments_Analyze(CppCompileEnvironment CompileEnvironment, List<string> Arguments)
-		{
-			// TODO: Replace logic in MacToolChain.CompileCPPFiles and IOSToolChain.CompileCPPFiles
-			string? StaticAnalysisMode = Environment.GetEnvironmentVariable("CLANG_STATIC_ANALYZER_MODE");
-			if (string.IsNullOrEmpty(StaticAnalysisMode))
-			{
-				return;
-			}
-			
-			base.GetCompileArguments_Analyze(CompileEnvironment, Arguments);
-
-			Arguments.Add("-Xclang -analyzer-output=html");
-			Arguments.Add("-Xclang -analyzer-config -Xclang stable-report-filename=true");
-			Arguments.Add("-Xclang -analyzer-config -Xclang report-in-main-source-file=true");
-			Arguments.Add("-Xclang -analyzer-config -Xclang path-diagnostics-alternate=true");
-			Arguments.Add("-Xclang -analyzer-disable-checker -Xclang deadcode.DeadStores");
-
-			if (string.Equals(StaticAnalysisMode, "shallow"))
-			{
-				Arguments.Add("-Xclang -analyzer-config -Xclang mode=shallow");
-			}
-		}
-
 		/// <inheritdoc/>
 		protected override void GetCompileArguments_Global(CppCompileEnvironment CompileEnvironment, List<string> Arguments)
 		{
@@ -363,11 +339,6 @@ namespace UnrealBuildTool
 
 			Arguments.Add("-fmessage-length=0");
 			Arguments.Add("-fpascal-strings");
-
-			if (!string.IsNullOrEmpty(Environment.GetEnvironmentVariable("CLANG_STATIC_ANALYZER_MODE")))
-			{
-				GetCompileArguments_Analyze(CompileEnvironment, Arguments);
-			}
 		}
 
 		protected string GetDsymutilPath(ILogger Logger, out string ExtraOptions, bool bIsForLTOBuild=false)
diff --git a/Engine/Source/Programs/UnrealBuildTool/ToolChain/ClangToolChain.cs b/Engine/Source/Programs/UnrealBuildTool/ToolChain/ClangToolChain.cs
--- a/Engine/Source/Programs/UnrealBuildTool/ToolChain/ClangToolChain.cs
+++ b/Engine/Source/Programs/UnrealBuildTool/ToolChain/ClangToolChain.cs
@@ -130,7 +130,7 @@ namespace UnrealBuildTool
 		{
 			const string DummyStart = "-D \"DUMMY_DEFINE";
 			const string DummyEnd = "\"";
-			ClangDummyDefine = DummyStart + String.Format("_").PadRight(ClangCmdLineMaxSize - ClangCmdlineDangerZone - DummyStart.Length - DummyEnd.Length, 'X') + DummyEnd;
+			ClangDummyDefine = DummyStart + "_".PadRight(ClangCmdLineMaxSize - ClangCmdlineDangerZone - DummyStart.Length - DummyEnd.Length, 'X') + DummyEnd;
 		}
 
 		public ClangToolChain(ClangToolChainOptions InOptions, ILogger InLogger)
@@ -355,9 +355,9 @@ namespace UnrealBuildTool
 			return $"\"{NormalizeCommandLinePath(SourceFile)}\"";
 		}
 
-		protected virtual string GetOutputFileArgument(FileItem ForceIncludeFile)
+		protected virtual string GetOutputFileArgument(FileItem OutputFile)
 		{
-			return $"-o \"{NormalizeCommandLinePath(ForceIncludeFile)}\"";
+			return $"-o \"{NormalizeCommandLinePath(OutputFile)}\"";
 		}
 
 		protected virtual string GetDepencenciesListFileArgument(FileItem DependencyListFile)
@@ -529,7 +529,7 @@ namespace UnrealBuildTool
 			if (CompileEnvironment.bPGOOptimize)
 			{
 				Log.TraceInformationOnce("Enabling Profile Guided Optimization (PGO). Linking will take a while.");
-				Arguments.Add(string.Format("-fprofile-instr-use=\"{0}\"", Path.Combine(CompileEnvironment.PGODirectory!, CompileEnvironment.PGOFilenamePrefix!)));
+				Arguments.Add($"-fprofile-instr-use=\"{Path.Combine(CompileEnvironment.PGODirectory!, CompileEnvironment.PGOFilenamePrefix!)}\"");
 			}
 			else if (CompileEnvironment.bPGOProfile)
 			{
@@ -820,6 +820,7 @@ namespace UnrealBuildTool
 				CompileEnvironment.PrecompiledHeaderAction != PrecompiledHeaderAction.Create ||
 				CompileEnvironment.bAllowRemotelyCompiledPCHs;
 
+			// Two-pass compile where the preprocessor is run first to output the dependency list
 			if (bPreprocessDepends && CompileEnvironment.bGenerateDependenciesFile)
 			{
 				Action PrepassAction = Graph.CreateAction(ActionType.Compile);
@@ -863,7 +864,7 @@ namespace UnrealBuildTool
 					bool bIsInDangerZone = CmdLineLength >= ClangCmdlineDangerZone && CmdLineLength <= ClangCmdLineMaxSize;
 					if (bIsInDangerZone)
 					{
-						ResponseFileContents.Add(ClangDummyDefine);
+						PreprocessResponseFileContents.Add(ClangDummyDefine);
 					}
 				}
 
@@ -876,6 +877,70 @@ namespace UnrealBuildTool
 				CompileAction.PrerequisiteItems.Add(DependencyListFile);
 			}
 
+			// Special analysis mode which outputs to files. Requires running the compile twice as clang will not emit an object file in this mode.
+			// Generally not recommended as it takes twice as long.
+			string? StaticAnalysisMode = Environment.GetEnvironmentVariable("CLANG_STATIC_ANALYZER_MODE");
+			if (!string.IsNullOrEmpty(StaticAnalysisMode))
+			{
+				string AnalyzerOutput = Environment.GetEnvironmentVariable("CLANG_ANALYZER_OUTPUT") ?? "html";
+
+				List<string> AnalyzeGlobalArguments = new(GlobalArguments);
+				List<string> AnalyzeFileArguments = new(FileArguments);
+
+				AnalyzeFileArguments.Remove("-ftime-trace");
+				AnalyzeFileArguments.Remove(GetOutputFileArgument(CompileAction.ProducedItems.First()));
+				if (CompileAction.DependencyListFile != null)
+				{
+					AnalyzeFileArguments.Remove(GetDepencenciesListFileArgument(CompileAction.DependencyListFile));
+				}
+
+				AnalyzeGlobalArguments.Add("-Wno-unused-command-line-argument");
+				AnalyzeGlobalArguments.Add("--analyze");
+				AnalyzeGlobalArguments.Add($"-Xclang -analyzer-output={AnalyzerOutput}");
+				AnalyzeGlobalArguments.Add("-Xclang -analyzer-config -Xclang stable-report-filename=true");
+				AnalyzeGlobalArguments.Add("-Xclang -analyzer-config -Xclang report-in-main-source-file=true");
+				AnalyzeGlobalArguments.Add("-Xclang -analyzer-config -Xclang path-diagnostics-alternate=true");
+				AnalyzeGlobalArguments.Add("-Xclang -analyzer-disable-checker -Xclang deadcode.DeadStores");
+				if (string.Equals(StaticAnalysisMode, "shallow"))
+				{
+					AnalyzeGlobalArguments.Add("-Xclang -analyzer-config -Xclang mode=shallow");
+				}
+				FileItem AnalyzerOutputDir = FileItem.GetItemByFileReference(FileReference.Combine(OutputDir, Path.GetFileName(SourceFile.AbsolutePath) + $".{AnalyzerOutput}"));
+				AnalyzeFileArguments.Add(GetOutputFileArgument(AnalyzerOutputDir));
+
+				// Creates the path to the response file using the name of the output file and creates its contents.
+				FileReference AnalyzeResponseFileName = new FileReference(AnalyzerOutputDir + ".response");
+				List<string> AnalyzeResponseFileContents = new();
+				AnalyzeResponseFileContents.AddRange(AnalyzeGlobalArguments);
+				AnalyzeResponseFileContents.AddRange(AnalyzeFileArguments);
+
+				if (RuntimePlatform.IsWindows)
+				{
+					int CmdLineLength = CommandPath.ToString().Length + string.Join(' ', ResponseFileContents).Length;
+					bool bIsInDangerZone = CmdLineLength >= ClangCmdlineDangerZone && CmdLineLength <= ClangCmdLineMaxSize;
+					if (bIsInDangerZone)
+					{
+						AnalyzeResponseFileContents.Add(ClangDummyDefine);
+					}
+				}
+
+				// Adds the response file to the compiler input.
+				FileItem AnalyzeResponseFileItem = Graph.CreateIntermediateTextFile(AnalyzeResponseFileName, AnalyzeResponseFileContents);
+				CompileAction.PrerequisiteItems.Add(AnalyzeResponseFileItem);
+
+				if (RuntimePlatform.IsWindows)
+				{
+					CompileAction.CommandArguments = $"/c \" \"{CompileAction.CommandPath}\" {GetResponseFileArgument(AnalyzeResponseFileItem)} && \"{CompileAction.CommandPath}\" {CompileAction.CommandArguments} \"";
+					CompileAction.CommandPath = FileReference.Combine(new DirectoryReference(Environment.GetFolderPath(Environment.SpecialFolder.System)), "cmd.exe");
+				}
+				else
+				{
+					CompileAction.CommandArguments = $"-c ' \"{CompileAction.CommandPath}\" {GetResponseFileArgument(AnalyzeResponseFileItem)}; \"{CompileAction.CommandPath}\" {CompileAction.CommandArguments} '";
+					CompileAction.CommandPath = new FileReference("/bin/sh");
+				}
+				CompileAction.bCanExecuteRemotely = false;
+			}
+
 			return CompileAction;
 		}
 	}
