UnrealBuildTool: Support preprocess depends for all clang not just Mac and Linux

#rnx

diff --git a/Engine/Source/Programs/UnrealBuildTool/Configuration/TargetRules.cs b/Engine/Source/Programs/UnrealBuildTool/Configuration/TargetRules.cs
--- a/Engine/Source/Programs/UnrealBuildTool/Configuration/TargetRules.cs
+++ b/Engine/Source/Programs/UnrealBuildTool/Configuration/TargetRules.cs
@@ -1344,6 +1344,13 @@ namespace UnrealBuildTool
 		[CommandLine("-Preprocess")]
 		public bool bPreprocessOnly = false;
 
+		/// <summary>
+		/// Generate dependency files by preprocessing. This is only recommended when distributing builds as it adds additional overhead.
+		/// </summary>
+		[CommandLine("-PreprocessDepends")]
+		[XmlConfigFile(Category = "BuildConfiguration")]
+		public bool bPreprocessDepends = false;
+
 		/// <summary>
 		/// The minimum number of files that must use a pre-compiled header before it will be created and used.
 		/// </summary>
@@ -2923,6 +2930,11 @@ namespace UnrealBuildTool
 			get { return Inner.bPreprocessOnly; }
 		}
 
+		public bool bPreprocessDepends
+		{
+			get { return Inner.bPreprocessDepends; }
+		}
+
 		public int MinFilesUsingPrecompiledHeader
 		{
 			get { return Inner.MinFilesUsingPrecompiledHeader; }
diff --git a/Engine/Source/Programs/UnrealBuildTool/Executors/Experimental/HordeExecutor.cs b/Engine/Source/Programs/UnrealBuildTool/Executors/Experimental/HordeExecutor.cs
--- a/Engine/Source/Programs/UnrealBuildTool/Executors/Experimental/HordeExecutor.cs
+++ b/Engine/Source/Programs/UnrealBuildTool/Executors/Experimental/HordeExecutor.cs
@@ -23,7 +23,7 @@
 	/// * Currently only supported for Linux and Mac toolchains.
 	/// 
 	/// Current Requirements:
-	/// * -PreprocessDepends: For Linux and Mac, run a prepass of clang to generate the dependency list. Needed to get all input files required to compile.
+	/// * -PreprocessDepends: Run a prepass of clang to generate the dependency list. Needed to get all input files required to compile.
 	///						  This will add additional overhead to the build and is not recommended unless using Horde compute remote distribution.
 	/// * -NoPCH:			  Cannot build with pch, as absolute paths are embedded into the .gch file for included headers, and these paths must exist exactly.
 	/// * -HordeCompute:	  Flag to enable this executor
diff --git a/Engine/Source/Programs/UnrealBuildTool/Platform/Linux/LinuxToolChain.cs b/Engine/Source/Programs/UnrealBuildTool/Platform/Linux/LinuxToolChain.cs
--- a/Engine/Source/Programs/UnrealBuildTool/Platform/Linux/LinuxToolChain.cs
+++ b/Engine/Source/Programs/UnrealBuildTool/Platform/Linux/LinuxToolChain.cs
@@ -126,13 +126,6 @@ namespace UnrealBuildTool
 			BreakpadEncoderPath = FileReference.Combine(Unreal.EngineDirectory, "Binaries", "Linux", $"BreakpadSymbolEncoder{BuildHostPlatform.Current.BinarySuffix}");
 		}
 
-		public override void SetUpGlobalEnvironment(ReadOnlyTargetRules Target)
-		{
-			base.SetUpGlobalEnvironment(Target);
-
-			bPreprocessDepends = Target.LinuxPlatform.bPreprocessDepends;
-		}
-
 		protected virtual bool CrossCompiling()
 		{
 			return bIsCrossCompiling;
diff --git a/Engine/Source/Programs/UnrealBuildTool/Platform/Linux/UEBuildLinux.cs b/Engine/Source/Programs/UnrealBuildTool/Platform/Linux/UEBuildLinux.cs
--- a/Engine/Source/Programs/UnrealBuildTool/Platform/Linux/UEBuildLinux.cs
+++ b/Engine/Source/Programs/UnrealBuildTool/Platform/Linux/UEBuildLinux.cs
@@ -79,13 +79,6 @@ namespace UnrealBuildTool
 		[XmlConfigFile(Category = "BuildConfiguration", Name = "bTuneDebugInfoForLLDB")]
 		public bool bTuneDebugInfoForLLDB = false;
 
-		/// <summary>
-		/// Generate dependency files by preprocessing. This is only recommended when distributing builds as it adds additional overhead.
-		/// </summary>
-		[CommandLine("-PreprocessDepends")]
-		[XmlConfigFile(Category = "BuildConfiguration", Name = "bPreprocessDepends")]
-		public bool bPreprocessDepends = false;
-
 		/// <summary>
 		/// Enables runtime ray tracing support.
 		/// </summary>
@@ -145,11 +138,6 @@ namespace UnrealBuildTool
 			get { return Inner.bTuneDebugInfoForLLDB; }
 		}
 
-		public bool bPreprocessDepends
-		{
-			get { return Inner.bPreprocessDepends; }
-		}
-
 		public bool bEnableRayTracing
 		{
 			get { return Inner.bEnableRayTracing; }
diff --git a/Engine/Source/Programs/UnrealBuildTool/Platform/Mac/MacToolChain.cs b/Engine/Source/Programs/UnrealBuildTool/Platform/Mac/MacToolChain.cs
--- a/Engine/Source/Programs/UnrealBuildTool/Platform/Mac/MacToolChain.cs
+++ b/Engine/Source/Programs/UnrealBuildTool/Platform/Mac/MacToolChain.cs
@@ -121,8 +121,6 @@ namespace UnrealBuildTool
 		{
 			base.SetUpGlobalEnvironment(Target);
 
-			bPreprocessDepends = Target.MacPlatform.bPreprocessDepends;
-
 			// validation, because sometimes this is called from a shell script and quoting messes up		
 			if (!Target.Architecture.All(C => char.IsLetterOrDigit(C) || C == '_' || C == '+'))
 			{
diff --git a/Engine/Source/Programs/UnrealBuildTool/Platform/Mac/UEBuildMac.cs b/Engine/Source/Programs/UnrealBuildTool/Platform/Mac/UEBuildMac.cs
--- a/Engine/Source/Programs/UnrealBuildTool/Platform/Mac/UEBuildMac.cs
+++ b/Engine/Source/Programs/UnrealBuildTool/Platform/Mac/UEBuildMac.cs
@@ -1,4 +1,4 @@
-﻿// Copyright Epic Games, Inc. All Rights Reserved.
+// Copyright Epic Games, Inc. All Rights Reserved.
 
 
 using System;
@@ -41,13 +41,6 @@ namespace UnrealBuildTool
 		[CommandLine("-NoDSYM", Value = "false")]
 		[XmlConfigFile(Category = "BuildConfiguration", Name = "bUseDSYMFiles")]
 		public bool bUseDSYMFiles = false;
-
-		/// <summary>
-		/// Generate dependency files by preprocessing. This is only recommended when distributing builds as it adds additional overhead.
-		/// </summary>
-		[CommandLine("-PreprocessDepends")]
-		[XmlConfigFile(Category = "BuildConfiguration", Name = "bPreprocessDepends")]
-		public bool bPreprocessDepends = false;
 	}
 
 	/// <summary>
@@ -90,11 +83,6 @@ namespace UnrealBuildTool
 			get { return Inner.bEnableUndefinedBehaviorSanitizer; }
 		}
 
-		public bool bPreprocessDepends
-		{
-			get { return Inner.bPreprocessDepends; }
-		}
-
 #pragma warning restore CS1591
 		#endregion
 	}
diff --git a/Engine/Source/Programs/UnrealBuildTool/ToolChain/ClangToolChain.cs b/Engine/Source/Programs/UnrealBuildTool/ToolChain/ClangToolChain.cs
--- a/Engine/Source/Programs/UnrealBuildTool/ToolChain/ClangToolChain.cs
+++ b/Engine/Source/Programs/UnrealBuildTool/ToolChain/ClangToolChain.cs
@@ -106,7 +106,12 @@ namespace UnrealBuildTool
 		/// <summary>
 		/// Indicates that the target is a moduler build i.e. Target.LinkType == TargetLinkType.Modular
 		/// </summary>
-		ModularBuild = 1 << 16
+		ModularBuild = 1 << 16,
+
+		/// <summary>
+		/// Generate dependency files by preprocessing. This is only recommended when distributing builds as it adds additional overhead.
+		/// </summary>
+		PreprocessDepends = 1 << 17,
 	}
 
 	abstract class ClangToolChain : ISPCToolChain
@@ -119,8 +124,6 @@ namespace UnrealBuildTool
 
 		protected ClangToolChainOptions Options;
 
-		protected bool bPreprocessDepends = false;
-
 		// Dummy define to work around clang compilation related to the windows maximum path length limitation
 		protected static string ClangDummyDefine; 
 		protected const int ClangCmdLineMaxSize = 32 * 1024;
@@ -139,6 +142,16 @@ namespace UnrealBuildTool
 			Options = InOptions;
 		}
 
+		public override void SetUpGlobalEnvironment(ReadOnlyTargetRules Target)
+		{
+			base.SetUpGlobalEnvironment(Target);
+
+			if (Target.bPreprocessDepends)
+			{
+				Options |= ClangToolChainOptions.PreprocessDepends;
+			}
+		}
+
 		public override void FinalizeOutput(ReadOnlyTargetRules Target, TargetMakefileBuilder MakefileBuilder)
 		{
 			if (Target.bPrintToolChainTimingInfo && Target.bParseTimingInfoForTracing)
@@ -755,7 +768,7 @@ namespace UnrealBuildTool
 			}
 
 			// Generate the included header dependency list
-			if (!bPreprocessDepends && CompileEnvironment.bGenerateDependenciesFile)
+			if (!Options.HasFlag(ClangToolChainOptions.PreprocessDepends) && CompileEnvironment.bGenerateDependenciesFile)
 			{
 				FileItem DependencyListFile = FileItem.GetItemByFileReference(FileReference.Combine(OutputDir, Path.GetFileName(SourceFile.AbsolutePath) + ".d"));
 				Arguments.Add(GetDepencenciesListFileArgument(DependencyListFile));
@@ -821,7 +834,7 @@ namespace UnrealBuildTool
 				CompileEnvironment.bAllowRemotelyCompiledPCHs;
 
 			// Two-pass compile where the preprocessor is run first to output the dependency list
-			if (bPreprocessDepends && CompileEnvironment.bGenerateDependenciesFile)
+			if (Options.HasFlag(ClangToolChainOptions.PreprocessDepends) && CompileEnvironment.bGenerateDependenciesFile)
 			{
 				Action PrepassAction = Graph.CreateAction(ActionType.Compile);
 				PrepassAction.PrerequisiteItems.AddRange(CompileAction.PrerequisiteItems);
