Update baseline tool template to use "clang" rather than "clang-cl".  Licensee was encountering problems with builds using SN-DBS.  The updated template should be a better baseline for compilers, particularly clang.

diff --git a/Engine/Source/Programs/UnrealBuildTool/Executors/SNDBS.cs b/Engine/Source/Programs/UnrealBuildTool/Executors/SNDBS.cs
--- a/Engine/Source/Programs/UnrealBuildTool/Executors/SNDBS.cs
+++ b/Engine/Source/Programs/UnrealBuildTool/Executors/SNDBS.cs
@@ -513,34 +513,37 @@ omp=true
 omp_min_threads=1",
 			["clang++.exe"] = @"
 [tool]
-family=clang-cl
+family=clang
+extensions=.c;.cc;.cpp;.cxx;.c++;.h;.hpp;.s;.asm
+adjacent_metadata_extensions=.pch;.gch;.pth;.o;.obj
 include_path01=..\include
 include_path02=..\include\c++\v1
 include_path03=..\lib\clang\*\include
-include_path04=..\usr\include
+include_path04=%CPATH%;%C_INCLUDE_PATH%;%CPLUS_INCLUDE_PATH%
 
 [files]
 main=clang++.exe
-file01=clang-shared.dll
-file02=libclang.dll
-file03=NXMangledNamePrinter.dll
-file04=..\lib\*
+
+[include-path-patterns]
+include01=-[IF][ \t]*(\""[^\""]+\""|[^ ]+)
+
+[additional-include-file-patterns]
+includefile01=--?include[ \t]*(\""[^\""]+\""|[^ ]+)
+
+[additional-input-file-patterns]
+inputfile01=--?(?:include-pch|fprofile-instr-use=|fprofile-sample-use=|fprofile-use=|fsanitize-blacklist=)[ \t]*(\""[^\""]+\""|[^ ]+)
 
 [output-file-patterns]
-outputfile01=\s*""([^ "",]+\.cpp\.txt)\""
+outputfile01=-o[ \t]*(\""[^\""]+\""|[^ ]+)
 
-[output-file-rules]
-rule01=*.log|discard=true
-rule02=*.dat|discard=true
-rule03=*.tmp|discard=true
+[definition-patterns]
+definition01=-D[ \t]*""?([a-zA-Z_0-9]+)(?:[\s""]|$)()
+definition02=-D[ \t]*""?([a-zA-Z_0-9]+)=(?:\\""|<)(.+?)(?:\\""|>)
+definition03=-D[ \t]*""?([a-zA-Z_0-9]+)=((?!(?:\\""|<))[^ ""]*)
 
-[system-file-filters]
-filter01=msvcr*.dll
-filter02=msvcp*.dll
-filter03=vcruntime140*.dll
-filter04=appcrt140*.dll
-filter05=desktopcrt140*.dll
-filter06=concrt140*.dll",
+[input-scanners]
+scanner01=c .*
+scanner02=assembler .s;.asm",
 		};
 	}
 }
