UnrealBuildTool: Support running static code analysis -StaticAnalyzer=Clang for clang toolchains.

Static code analysis automatically disables linking and pch. Dependency checking does not work as no output file is written.

#rnx

diff --git a/Engine/Source/Programs/UnrealBuildTool/Configuration/TargetRules.cs b/Engine/Source/Programs/UnrealBuildTool/Configuration/TargetRules.cs
--- a/Engine/Source/Programs/UnrealBuildTool/Configuration/TargetRules.cs
+++ b/Engine/Source/Programs/UnrealBuildTool/Configuration/TargetRules.cs
@@ -197,7 +197,7 @@ namespace UnrealBuildTool
 	/// Output type for the static analyzer. This currently only works for the Clang static analyzer.
 	/// The Clang static analyzer can do either Text, which prints the analysis to stdout, or
 	/// html, where it writes out a navigable HTML page for each issue that it finds, per file.
-	/// The HTML is output in the same directory as the object fil that would otherwise have
+	/// The HTML is output in the same directory as the object file that would otherwise have
 	/// been generated. 
 	/// All other analyzers default automatically to Text. 
 	/// </summary>
@@ -214,6 +214,23 @@ namespace UnrealBuildTool
 		Html,
 	}
 
+	/// <summary>
+	/// Output type for the static analyzer. This currently only works for the Clang static analyzer.
+	/// The Clang static analyzer can do a shallow quick analysis. However the default deep is recommended.
+	/// </summary>
+	public enum StaticAnalyzerMode
+	{
+		/// <summary>
+		/// Default deep analysis.
+		/// </summary>
+		Deep,
+
+		/// <summary>
+		/// Quick analysis. Not recommended.
+		/// </summary>
+		Shallow,
+	}
+
 	/// <summary>
 	/// Utility class for EngineIncludeOrderVersion defines
 	/// </summary>
@@ -1419,6 +1436,13 @@ namespace UnrealBuildTool
 		[XmlConfigFile(Category = "BuildConfiguration")]
 		public StaticAnalyzerOutputType StaticAnalyzerOutputType = StaticAnalyzerOutputType.Text;
 
+		/// <summary>
+		/// The mode to use for the static analyzer.
+		/// </summary>
+		[CommandLine("-StaticAnalyzerMode")]
+		[XmlConfigFile(Category = "BuildConfiguration")]
+		public StaticAnalyzerMode StaticAnalyzerMode = StaticAnalyzerMode.Deep;
+
 		/// <summary>
 		/// The minimum number of files that must use a pre-compiled header before it will be created and used.
 		/// </summary>
@@ -3013,6 +3037,11 @@ namespace UnrealBuildTool
 			get { return Inner.StaticAnalyzerOutputType; }
 		}
 
+		public StaticAnalyzerMode StaticAnalyzerMode
+		{
+			get { return Inner.StaticAnalyzerMode; }
+		}
+
 		public int MinFilesUsingPrecompiledHeader
 		{
 			get { return Inner.MinFilesUsingPrecompiledHeader; }
diff --git a/Engine/Source/Programs/UnrealBuildTool/Platform/IOS/UEBuildIOS.cs b/Engine/Source/Programs/UnrealBuildTool/Platform/IOS/UEBuildIOS.cs
--- a/Engine/Source/Programs/UnrealBuildTool/Platform/IOS/UEBuildIOS.cs
+++ b/Engine/Source/Programs/UnrealBuildTool/Platform/IOS/UEBuildIOS.cs
@@ -760,6 +760,25 @@ namespace UnrealBuildTool
 
 		public override void ValidateTarget(TargetRules Target)
 		{
+			if (!string.IsNullOrWhiteSpace(Environment.GetEnvironmentVariable("CLANG_STATIC_ANALYZER_MODE")))
+			{
+				Target.StaticAnalyzer = StaticAnalyzer.Default;
+				Target.StaticAnalyzerOutputType = (Environment.GetEnvironmentVariable("CLANG_ANALYZER_OUTPUT")?.Contains("html", StringComparison.OrdinalIgnoreCase) == true) ? StaticAnalyzerOutputType.Html : StaticAnalyzerOutputType.Text;
+				Target.StaticAnalyzerMode = string.Equals(Environment.GetEnvironmentVariable("CLANG_STATIC_ANALYZER_MODE"), "shallow") ? StaticAnalyzerMode.Shallow : StaticAnalyzerMode.Deep;
+			}
+			else if (Target.StaticAnalyzer == StaticAnalyzer.Clang)
+			{
+				Target.StaticAnalyzer = StaticAnalyzer.Default;
+			}
+
+			// Disable linking if we're using a static analyzer
+			// Disable PCHs for clang static analyzer.
+			if (Target.StaticAnalyzer == StaticAnalyzer.Default)
+			{
+				Target.bDisableLinking = true;
+				Target.bUsePCHFiles = false;
+			}
+
 			// we assume now we are building with IOS8 or later
 			if (Target.bCompileAgainstEngine)
 			{
diff --git a/Engine/Source/Programs/UnrealBuildTool/Platform/Linux/UEBuildLinux.cs b/Engine/Source/Programs/UnrealBuildTool/Platform/Linux/UEBuildLinux.cs
--- a/Engine/Source/Programs/UnrealBuildTool/Platform/Linux/UEBuildLinux.cs
+++ b/Engine/Source/Programs/UnrealBuildTool/Platform/Linux/UEBuildLinux.cs
@@ -212,6 +212,25 @@ namespace UnrealBuildTool
 
 		public override void ValidateTarget(TargetRules Target)
 		{
+			if (!string.IsNullOrWhiteSpace(Environment.GetEnvironmentVariable("CLANG_STATIC_ANALYZER_MODE")))
+			{
+				Target.StaticAnalyzer = StaticAnalyzer.Default;
+				Target.StaticAnalyzerOutputType = (Environment.GetEnvironmentVariable("CLANG_ANALYZER_OUTPUT")?.Contains("html", StringComparison.OrdinalIgnoreCase) == true) ? StaticAnalyzerOutputType.Html : StaticAnalyzerOutputType.Text;
+				Target.StaticAnalyzerMode = string.Equals(Environment.GetEnvironmentVariable("CLANG_STATIC_ANALYZER_MODE"), "shallow") ? StaticAnalyzerMode.Shallow : StaticAnalyzerMode.Deep;
+			}
+			else if (Target.StaticAnalyzer == StaticAnalyzer.Clang)
+			{
+				Target.StaticAnalyzer = StaticAnalyzer.Default;
+			}
+
+			// Disable linking if we're using a static analyzer
+			// Disable PCHs for clang static analyzer.
+			if (Target.StaticAnalyzer == StaticAnalyzer.Default)
+			{
+				Target.bDisableLinking = true;
+				Target.bUsePCHFiles = false;
+			}
+
 			// Editor target types get overwritten in UEBuildTarget.cs so lets avoid adding this here. ResetTarget is called with
 			// default settings for TargetRules meanings Type == Game once then Type == Editor a 2nd time when building the Editor.
 			// BuildVersion string is not set at this point so we can avoid setting a Sanitizer suffix if this is the first ResetTarget
diff --git a/Engine/Source/Programs/UnrealBuildTool/Platform/Mac/UEBuildMac.cs b/Engine/Source/Programs/UnrealBuildTool/Platform/Mac/UEBuildMac.cs
--- a/Engine/Source/Programs/UnrealBuildTool/Platform/Mac/UEBuildMac.cs
+++ b/Engine/Source/Programs/UnrealBuildTool/Platform/Mac/UEBuildMac.cs
@@ -109,6 +109,25 @@ namespace UnrealBuildTool
 
 		public override void ValidateTarget(TargetRules Target)
 		{
+			if (!string.IsNullOrWhiteSpace(Environment.GetEnvironmentVariable("CLANG_STATIC_ANALYZER_MODE")))
+			{
+				Target.StaticAnalyzer = StaticAnalyzer.Default;
+				Target.StaticAnalyzerOutputType = (Environment.GetEnvironmentVariable("CLANG_ANALYZER_OUTPUT")?.Contains("html", StringComparison.OrdinalIgnoreCase) == true) ? StaticAnalyzerOutputType.Html : StaticAnalyzerOutputType.Text;
+				Target.StaticAnalyzerMode = string.Equals(Environment.GetEnvironmentVariable("CLANG_STATIC_ANALYZER_MODE"), "shallow") ? StaticAnalyzerMode.Shallow : StaticAnalyzerMode.Deep;
+			}
+			else if (Target.StaticAnalyzer == StaticAnalyzer.Clang)
+			{
+				Target.StaticAnalyzer = StaticAnalyzer.Default;
+			}
+
+			// Disable linking if we're using a static analyzer
+			// Disable PCHs for clang static analyzer.
+			if (Target.StaticAnalyzer == StaticAnalyzer.Default)
+			{
+				Target.bDisableLinking = true;
+				Target.bUsePCHFiles = false;
+			}
+
 			if (BuildHostPlatform.Current.Platform != UnrealTargetPlatform.Mac)
 			{
 				// @todo: Temporarily disable precompiled header files when building remotely due to errors
diff --git a/Engine/Source/Programs/UnrealBuildTool/ToolChain/ClangToolChain.cs b/Engine/Source/Programs/UnrealBuildTool/ToolChain/ClangToolChain.cs
--- a/Engine/Source/Programs/UnrealBuildTool/ToolChain/ClangToolChain.cs
+++ b/Engine/Source/Programs/UnrealBuildTool/ToolChain/ClangToolChain.cs
@@ -107,11 +107,6 @@ namespace UnrealBuildTool
 		/// Indicates that the target is a moduler build i.e. Target.LinkType == TargetLinkType.Modular
 		/// </summary>
 		ModularBuild = 1 << 16,
-
-		/// <summary>
-		/// Generate dependency files by preprocessing. This is only recommended when distributing builds as it adds additional overhead.
-		/// </summary>
-		PreprocessDepends = 1 << 17,
 	}
 
 	abstract class ClangToolChain : ISPCToolChain
@@ -129,6 +124,12 @@ namespace UnrealBuildTool
 		protected const int ClangCmdLineMaxSize = 32 * 1024;
 		protected const int ClangCmdlineDangerZone = 30 * 1024;
 
+		// Target settings
+		protected bool PreprocessDepends = false;
+		protected StaticAnalyzer StaticAnalyzer = StaticAnalyzer.None;
+		protected StaticAnalyzerMode StaticAnalyzerMode = StaticAnalyzerMode.Deep;
+		protected StaticAnalyzerOutputType StaticAnalyzerOutputType = StaticAnalyzerOutputType.Text;
+
 		static ClangToolChain()
 		{
 			const string DummyStart = "-D \"DUMMY_DEFINE";
@@ -146,10 +147,10 @@ namespace UnrealBuildTool
 		{
 			base.SetUpGlobalEnvironment(Target);
 
-			if (Target.bPreprocessDepends)
-			{
-				Options |= ClangToolChainOptions.PreprocessDepends;
-			}
+			PreprocessDepends = Target.bPreprocessDepends;
+			StaticAnalyzer = Target.StaticAnalyzer;
+			StaticAnalyzerMode = Target.StaticAnalyzerMode;
+			StaticAnalyzerOutputType = Target.StaticAnalyzerOutputType;
 		}
 
 		public override void FinalizeOutput(ReadOnlyTargetRules Target, TargetMakefileBuilder MakefileBuilder)
@@ -629,7 +630,15 @@ namespace UnrealBuildTool
 		/// <param name="Arguments"></param>
 		protected virtual void GetCompileArguments_Analyze(CppCompileEnvironment CompileEnvironment, List<string> Arguments)
 		{
+			Arguments.Add("-Wno-unused-command-line-argument");
 			Arguments.Add("--analyze");
+			Arguments.Add($"-Xclang -analyzer-output={StaticAnalyzerOutputType.ToString().ToLowerInvariant()}");
+			Arguments.Add("-Xclang -analyzer-config -Xclang stable-report-filename=true");
+			Arguments.Add("-Xclang -analyzer-config -Xclang report-in-main-source-file=true");
+			Arguments.Add("-Xclang -analyzer-config -Xclang path-diagnostics-alternate=true");
+			Arguments.Add("-Xclang -analyzer-disable-checker -Xclang deadcode.DeadStores");
+			Arguments.Add("-Xclang -analyzer-disable-checker -Xclang security.FloatLoopCounter");
+			if (StaticAnalyzerMode == StaticAnalyzerMode.Shallow) Arguments.Add("-Xclang -analyzer-config -Xclang mode=shallow");
 		}
 
 		/// <summary>
@@ -666,6 +675,12 @@ namespace UnrealBuildTool
 			// Add sanitizer flags to the argument list.
 			GetCompilerArguments_Sanitizers(CompileEnvironment, Arguments);
 
+			// Add analysis flags to the argument list.
+			if (StaticAnalyzer == StaticAnalyzer.Default)
+			{
+				GetCompileArguments_Analyze(CompileEnvironment, Arguments);
+			}
+
 			// Add additional arguments to the argument list.
 			GetCompileArguments_AdditionalArgs(CompileEnvironment, Arguments);
 		}
@@ -759,21 +774,28 @@ namespace UnrealBuildTool
 			// Add the source file path to the command-line.
 			Arguments.Add(GetSourceFileArgument(SourceFile));
 
-			// Generate the timing info
-			if (CompileEnvironment.bPrintTimingInfo)
+			if (StaticAnalyzer != StaticAnalyzer.Default)
 			{
-				FileItem TraceFile = FileItem.GetItemByFileReference(FileReference.Combine(OutputDir, Path.GetFileName(SourceFile.AbsolutePath) + ".json"));
-				Arguments.Add("-ftime-trace");
-				CompileAction.ProducedItems.Add(TraceFile);
-			}
+				// Generate the timing info
+				if (CompileEnvironment.bPrintTimingInfo)
+				{
+					FileItem TraceFile = FileItem.GetItemByFileReference(FileReference.Combine(OutputDir, Path.GetFileName(SourceFile.AbsolutePath) + ".json"));
+					Arguments.Add("-ftime-trace");
+					CompileAction.ProducedItems.Add(TraceFile);
+				}
 
-			// Generate the included header dependency list
-			if (!Options.HasFlag(ClangToolChainOptions.PreprocessDepends) && CompileEnvironment.bGenerateDependenciesFile)
+				// Generate the included header dependency list
+				if (!PreprocessDepends && CompileEnvironment.bGenerateDependenciesFile)
+				{
+					FileItem DependencyListFile = FileItem.GetItemByFileReference(FileReference.Combine(OutputDir, Path.GetFileName(SourceFile.AbsolutePath) + ".d"));
+					Arguments.Add(GetDepencenciesListFileArgument(DependencyListFile));
+					CompileAction.DependencyListFile = DependencyListFile;
+					CompileAction.ProducedItems.Add(DependencyListFile);
+				}
+			}
+			else if (StaticAnalyzer == StaticAnalyzer.Default)
 			{
-				FileItem DependencyListFile = FileItem.GetItemByFileReference(FileReference.Combine(OutputDir, Path.GetFileName(SourceFile.AbsolutePath) + ".d"));
-				Arguments.Add(GetDepencenciesListFileArgument(DependencyListFile));
-				CompileAction.DependencyListFile = DependencyListFile;
-				CompileAction.ProducedItems.Add(DependencyListFile);
+				OutputFile = FileItem.GetItemByFileReference(FileReference.Combine(OutputDir, Path.GetFileName(SourceFile.AbsolutePath) + ".analysis"));
 			}
 
 			// Add the parameters needed to compile the output file to the command-line.
@@ -824,7 +846,7 @@ namespace UnrealBuildTool
 			CompileAction.WorkingDirectory = Unreal.EngineSourceDirectory;
 			CompileAction.CommandPath = CommandPath;
 			CompileAction.CommandVersion = CommandVersion;
-			CompileAction.CommandDescription = "Compile";
+			CompileAction.CommandDescription = StaticAnalyzer == StaticAnalyzer.Default ? "Analyze" : "Compile";
 			CompileAction.StatusDescription = Path.GetFileName(SourceFile.AbsolutePath);
 			CompileAction.bIsGCCCompiler = true;
 
@@ -834,7 +856,7 @@ namespace UnrealBuildTool
 				CompileEnvironment.bAllowRemotelyCompiledPCHs;
 
 			// Two-pass compile where the preprocessor is run first to output the dependency list
-			if (Options.HasFlag(ClangToolChainOptions.PreprocessDepends) && CompileEnvironment.bGenerateDependenciesFile)
+			if (PreprocessDepends && CompileEnvironment.bGenerateDependenciesFile)
 			{
 				Action PrepassAction = Graph.CreateAction(ActionType.Compile);
 				PrepassAction.PrerequisiteItems.AddRange(CompileAction.PrerequisiteItems);
@@ -890,70 +912,6 @@ namespace UnrealBuildTool
 				CompileAction.PrerequisiteItems.Add(DependencyListFile);
 			}
 
-			// Special analysis mode which outputs to files. Requires running the compile twice as clang will not emit an object file in this mode.
-			// Generally not recommended as it takes twice as long.
-			string? StaticAnalysisMode = Environment.GetEnvironmentVariable("CLANG_STATIC_ANALYZER_MODE");
-			if (!string.IsNullOrEmpty(StaticAnalysisMode))
-			{
-				string AnalyzerOutput = Environment.GetEnvironmentVariable("CLANG_ANALYZER_OUTPUT") ?? "html";
-
-				List<string> AnalyzeGlobalArguments = new(GlobalArguments);
-				List<string> AnalyzeFileArguments = new(FileArguments);
-
-				AnalyzeFileArguments.Remove("-ftime-trace");
-				AnalyzeFileArguments.Remove(GetOutputFileArgument(CompileAction.ProducedItems.First()));
-				if (CompileAction.DependencyListFile != null)
-				{
-					AnalyzeFileArguments.Remove(GetDepencenciesListFileArgument(CompileAction.DependencyListFile));
-				}
-
-				AnalyzeGlobalArguments.Add("-Wno-unused-command-line-argument");
-				AnalyzeGlobalArguments.Add("--analyze");
-				AnalyzeGlobalArguments.Add($"-Xclang -analyzer-output={AnalyzerOutput}");
-				AnalyzeGlobalArguments.Add("-Xclang -analyzer-config -Xclang stable-report-filename=true");
-				AnalyzeGlobalArguments.Add("-Xclang -analyzer-config -Xclang report-in-main-source-file=true");
-				AnalyzeGlobalArguments.Add("-Xclang -analyzer-config -Xclang path-diagnostics-alternate=true");
-				AnalyzeGlobalArguments.Add("-Xclang -analyzer-disable-checker -Xclang deadcode.DeadStores");
-				if (string.Equals(StaticAnalysisMode, "shallow"))
-				{
-					AnalyzeGlobalArguments.Add("-Xclang -analyzer-config -Xclang mode=shallow");
-				}
-				FileItem AnalyzerOutputDir = FileItem.GetItemByFileReference(FileReference.Combine(OutputDir, Path.GetFileName(SourceFile.AbsolutePath) + $".{AnalyzerOutput}"));
-				AnalyzeFileArguments.Add(GetOutputFileArgument(AnalyzerOutputDir));
-
-				// Creates the path to the response file using the name of the output file and creates its contents.
-				FileReference AnalyzeResponseFileName = new FileReference(AnalyzerOutputDir + ".response");
-				List<string> AnalyzeResponseFileContents = new();
-				AnalyzeResponseFileContents.AddRange(AnalyzeGlobalArguments);
-				AnalyzeResponseFileContents.AddRange(AnalyzeFileArguments);
-
-				if (RuntimePlatform.IsWindows)
-				{
-					int CmdLineLength = CommandPath.ToString().Length + string.Join(' ', ResponseFileContents).Length;
-					bool bIsInDangerZone = CmdLineLength >= ClangCmdlineDangerZone && CmdLineLength <= ClangCmdLineMaxSize;
-					if (bIsInDangerZone)
-					{
-						AnalyzeResponseFileContents.Add(ClangDummyDefine);
-					}
-				}
-
-				// Adds the response file to the compiler input.
-				FileItem AnalyzeResponseFileItem = Graph.CreateIntermediateTextFile(AnalyzeResponseFileName, AnalyzeResponseFileContents);
-				CompileAction.PrerequisiteItems.Add(AnalyzeResponseFileItem);
-
-				if (RuntimePlatform.IsWindows)
-				{
-					CompileAction.CommandArguments = $"/c \" \"{CompileAction.CommandPath}\" {GetResponseFileArgument(AnalyzeResponseFileItem)} && \"{CompileAction.CommandPath}\" {CompileAction.CommandArguments} \"";
-					CompileAction.CommandPath = FileReference.Combine(new DirectoryReference(Environment.GetFolderPath(Environment.SpecialFolder.System)), "cmd.exe");
-				}
-				else
-				{
-					CompileAction.CommandArguments = $"-c ' \"{CompileAction.CommandPath}\" {GetResponseFileArgument(AnalyzeResponseFileItem)}; \"{CompileAction.CommandPath}\" {CompileAction.CommandArguments} '";
-					CompileAction.CommandPath = new FileReference("/bin/sh");
-				}
-				CompileAction.bCanExecuteRemotely = false;
-			}
-
 			return CompileAction;
 		}
 	}
