UnrealBuildTool: Prevent adding duplicate actions when using FASTBuild

#rnx

diff --git a/Engine/Source/Programs/UnrealBuildTool/Executors/Experimental/FASTBuild.cs b/Engine/Source/Programs/UnrealBuildTool/Executors/Experimental/FASTBuild.cs
--- a/Engine/Source/Programs/UnrealBuildTool/Executors/Experimental/FASTBuild.cs
+++ b/Engine/Source/Programs/UnrealBuildTool/Executors/Experimental/FASTBuild.cs
@@ -455,7 +455,7 @@
 
 			if (PreCompileActions.Any())
 			{
-				bool bResult = LocalExecutor.ExecuteActions(PreCompileActions, Logger);
+				bool bResult = LocalExecutor.ExecuteActions(PreCompileActions.Distinct().ToList(), Logger);
 
 				if (!bResult)
 					return false;
