Android build fix:
correcting clang version for 'no-bitwise-instead-of-logical' arg.

diff --git a/Engine/Source/Programs/UnrealBuildTool/Platform/Android/AndroidToolChain.cs b/Engine/Source/Programs/UnrealBuildTool/Platform/Android/AndroidToolChain.cs
--- a/Engine/Source/Programs/UnrealBuildTool/Platform/Android/AndroidToolChain.cs
+++ b/Engine/Source/Programs/UnrealBuildTool/Platform/Android/AndroidToolChain.cs
@@ -566,7 +566,7 @@ namespace UnrealBuildTool
                 Result += " -Wno-ordered-compare-function-pointers";
             }
 			
-			if (CompilerVersionGreaterOrEqual(14, 0, 0))
+			if (CompilerVersionGreaterOrEqual(14, 0, 2))
 			{
 				Result += " -Wno-bitwise-instead-of-logical";
 			}
