UnrealBuildTool: Update sn-dbs clang++ config

#rnx

diff --git a/Engine/Source/Programs/UnrealBuildTool/Executors/SNDBS.cs b/Engine/Source/Programs/UnrealBuildTool/Executors/SNDBS.cs
--- a/Engine/Source/Programs/UnrealBuildTool/Executors/SNDBS.cs
+++ b/Engine/Source/Programs/UnrealBuildTool/Executors/SNDBS.cs
@@ -525,13 +525,13 @@ include_path04=%CPATH%;%C_INCLUDE_PATH%;%CPLUS_INCLUDE_PATH%
 main=clang++.exe
 
 [include-path-patterns]
-include01=-[IF][ \t]*(\""[^\""]+\""|[^ ]+)
+include01=-(?:I|F|isystem|idirafter)[ \t]*(\""[^\""]+\""|[^ ]+)
 
 [additional-include-file-patterns]
 includefile01=--?include[ \t]*(\""[^\""]+\""|[^ ]+)
 
 [additional-input-file-patterns]
-inputfile01=--?(?:include-pch|fprofile-instr-use=|fprofile-sample-use=|fprofile-use=|fsanitize-blacklist=)[ \t]*(\""[^\""]+\""|[^ ]+)
+inputfile01=--?(?:include-pch|fprofile-instr-use=|fprofile-sample-use=|fprofile-use=|fsanitize-ignorelist=)[ \t]*(\""[^\""]+\""|[^ ]+)
 
 [output-file-patterns]
 outputfile01=-o[ \t]*(\""[^\""]+\""|[^ ]+)
