UnrealBuildTool: Fix clang include pch logic for non cpp files to match msvc

#rnx

diff --git a/Engine/Source/Programs/UnrealBuildTool/ToolChain/ClangToolChain.cs b/Engine/Source/Programs/UnrealBuildTool/ToolChain/ClangToolChain.cs
--- a/Engine/Source/Programs/UnrealBuildTool/ToolChain/ClangToolChain.cs
+++ b/Engine/Source/Programs/UnrealBuildTool/ToolChain/ClangToolChain.cs
@@ -355,7 +355,6 @@ namespace UnrealBuildTool
 		{
 			// Add the precompiled header file's path to the force include path
 			// This needs to be before the other force include paths to ensure clang uses it instead of the source header file.
-			// Will be removed if compiling non c++ files.
 			if (CompileEnvironment.PrecompiledHeaderAction == PrecompiledHeaderAction.Include)
 			{
 				Arguments.Add(GetForceIncludeFileArgument(CompileEnvironment.PrecompiledHeaderIncludeFilename!));
@@ -702,8 +701,6 @@ namespace UnrealBuildTool
 			CompileAction.PrerequisiteItems.Add(SourceFile);
 
 			string Extension = Path.GetExtension(SourceFile.AbsolutePath).ToUpperInvariant();
-			bool IsCPlusPlus = false;
-
 			if (CompileEnvironment.PrecompiledHeaderAction == PrecompiledHeaderAction.Create)
 			{
 				GetCompileArguments_PCH(CompileEnvironment, Arguments);
@@ -727,20 +724,12 @@ namespace UnrealBuildTool
 			{
 				// Compile the file as C++ code.
 				GetCompileArguments_CPP(CompileEnvironment, Arguments);
-				IsCPlusPlus = true;
 			}
 
 			if (CompileEnvironment.PrecompiledHeaderAction == PrecompiledHeaderAction.Include)
 			{
-				// Only force include PCH for c++ files
-				if (!IsCPlusPlus)
-				{
-					Arguments.Remove(GetForceIncludeFileArgument(CompileEnvironment.PrecompiledHeaderIncludeFilename!));
-				}
-				else
-				{
-					CompileAction.PrerequisiteItems.Add(CompileEnvironment.PrecompiledHeaderFile!);
-				}
+				CompileAction.PrerequisiteItems.Add(CompileEnvironment.PrecompiledHeaderFile!);
+				CompileAction.PrerequisiteItems.Add(FileItem.GetItemByFileReference(CompileEnvironment.PrecompiledHeaderIncludeFilename!));
 			}
 
 			FileItem OutputFile;
