UnrealBuildTool: Fix LogTrace and wrong preprocessor output extension in ClangToolChain

#rnx

diff --git a/Engine/Source/Programs/UnrealBuildTool/ToolChain/ClangToolChain.cs b/Engine/Source/Programs/UnrealBuildTool/ToolChain/ClangToolChain.cs
--- a/Engine/Source/Programs/UnrealBuildTool/ToolChain/ClangToolChain.cs
+++ b/Engine/Source/Programs/UnrealBuildTool/ToolChain/ClangToolChain.cs
@@ -526,7 +526,7 @@ namespace UnrealBuildTool
 				}
 				else
 				{
-					Log.TraceWarningOnce($"CrashDiagnosticDirectory has been specified but directory \"{CrashDiagnosticDirectory}\" does not exist. Compiler argument \"-fcrash-diagnostics-dir\" has been discarded.");
+					Log.TraceWarningOnce("CrashDiagnosticDirectory has been specified but directory \"{CrashDiagnosticDirectory}\" does not exist. Compiler argument \"-fcrash-diagnostics-dir\" has been discarded.", CrashDiagnosticDirectory);
 				}
 			}
 		}
@@ -740,7 +740,7 @@ namespace UnrealBuildTool
 			}
 			else if (CompileEnvironment.bPreprocessOnly)
 			{
-				OutputFile = FileItem.GetItemByFileReference(FileReference.Combine(OutputDir, Path.GetFileName(SourceFile.AbsolutePath) + ".cpp.i"));
+				OutputFile = FileItem.GetItemByFileReference(FileReference.Combine(OutputDir, Path.GetFileName(SourceFile.AbsolutePath) + ".i"));
 				CompileResult.ObjectFiles.Add(OutputFile);
 
 				// Clang does EITHER pre-process or object file.
