UnrealBuildTool: -BytesPerUnityCPP command line arg to control unity size

#rnx

diff --git a/Engine/Source/Programs/UnrealBuildTool/Configuration/TargetRules.cs b/Engine/Source/Programs/UnrealBuildTool/Configuration/TargetRules.cs
--- a/Engine/Source/Programs/UnrealBuildTool/Configuration/TargetRules.cs
+++ b/Engine/Source/Programs/UnrealBuildTool/Configuration/TargetRules.cs
@@ -1360,6 +1360,7 @@ namespace UnrealBuildTool
 		/// <summary>
 		/// An approximate number of bytes of C++ code to target for inclusion in a single unified C++ file.
 		/// </summary>
+		[CommandLine("-BytesPerUnityCPP")]
 		[XmlConfigFile(Category = "BuildConfiguration")]
 		public int NumIncludedBytesPerUnityCPP = 384 * 1024;
 
