UBT: Fix ability to do single file compiles on a generated header.

diff --git a/Engine/Source/Programs/UnrealBuildTool/Configuration/UEBuildModuleCPP.cs b/Engine/Source/Programs/UnrealBuildTool/Configuration/UEBuildModuleCPP.cs
--- a/Engine/Source/Programs/UnrealBuildTool/Configuration/UEBuildModuleCPP.cs
+++ b/Engine/Source/Programs/UnrealBuildTool/Configuration/UEBuildModuleCPP.cs
@@ -533,14 +533,27 @@ namespace UnrealBuildTool
 			List<FileItem> CPPFiles = new List<FileItem>(InputFiles.CPPFiles);
 
 			// Compile all the generated CPP files
-			if (GeneratedCppDirectories != null && !CompileEnvironment.bHackHeaderGenerator && SpecificFilesToCompile.Count == 0)
+			if (GeneratedCppDirectories != null && !CompileEnvironment.bHackHeaderGenerator)
 			{
 				List<string> GeneratedFiles = new List<string>();
-				foreach (string GeneratedDir in GeneratedCppDirectories)
+				if (SpecificFilesToCompile.Count == 0)
 				{
-					if (Directory.Exists(GeneratedDir))
+					foreach (string GeneratedDir in GeneratedCppDirectories)
 					{
-						GeneratedFiles.AddRange(Directory.GetFiles(GeneratedDir, "*.gen.cpp"));
+						if (Directory.Exists(GeneratedDir))
+						{
+							GeneratedFiles.AddRange(Directory.GetFiles(GeneratedDir, "*.gen.cpp"));
+						}
+					}
+				}
+				else
+				{
+					foreach (FileReference FileToCompile in SpecificFilesToCompile)
+					{
+						if (GeneratedCppDirectories.Any(x => FileToCompile.IsUnderDirectory(new DirectoryReference(x))))
+						{
+							GeneratedFiles.Add(FileToCompile.FullName);
+						}
 					}
 				}
 
