UBT: Fix single-file compile for monolithic targets.

diff --git a/Engine/Source/Programs/UnrealBuildTool/Configuration/UEBuildBinary.cs b/Engine/Source/Programs/UnrealBuildTool/Configuration/UEBuildBinary.cs
--- a/Engine/Source/Programs/UnrealBuildTool/Configuration/UEBuildBinary.cs
+++ b/Engine/Source/Programs/UnrealBuildTool/Configuration/UEBuildBinary.cs
@@ -477,12 +477,15 @@
 						FileReference PrecompiledManifestLocation = Module.PrecompiledManifestLocation;
 						BuildProducts.Add(PrecompiledManifestLocation, BuildProductType.BuildResource);
 
-						PrecompiledManifest ModuleManifest = PrecompiledManifest.Read(PrecompiledManifestLocation);
-						foreach(FileReference OutputFile in ModuleManifest.OutputFiles)
+						if (FileReference.Exists(PrecompiledManifestLocation)) // May be single file compile; skipped modulePrecompiledManifestLocation{
 						{
-							if(!BuildProducts.ContainsKey(OutputFile))
+							PrecompiledManifest ModuleManifest = PrecompiledManifest.Read(PrecompiledManifestLocation);
+							foreach (FileReference OutputFile in ModuleManifest.OutputFiles)
 							{
-								BuildProducts.Add(OutputFile, BuildProductType.BuildResource);
+								if (!BuildProducts.ContainsKey(OutputFile))
+								{
+									BuildProducts.Add(OutputFile, BuildProductType.BuildResource);
+								}
 							}
 						}
 					}
