Fix for -BypassGradlePackaging when Intermediate/Android cleaned

diff --git a/Engine/Source/Programs/UnrealBuildTool/Platform/Android/UEDeployAndroid.cs b/Engine/Source/Programs/UnrealBuildTool/Platform/Android/UEDeployAndroid.cs
--- a/Engine/Source/Programs/UnrealBuildTool/Platform/Android/UEDeployAndroid.cs
+++ b/Engine/Source/Programs/UnrealBuildTool/Platform/Android/UEDeployAndroid.cs
@@ -4362,7 +4362,7 @@ namespace UnrealBuildTool
 				return;
 			}
 
-			if (BypassGradlePackaging && !bEnableBundle && !bSkipGradleBuild)
+			if (BypassGradlePackaging && !bEnableBundle && !bSkipGradleBuild && !bCreateFromScratch)
 			{
 				Logger.LogInformation("Attemping BypassGradlePackaging");
 				int BuildListComboRemaining = BuildListComboTotal;
@@ -4393,7 +4393,7 @@ namespace UnrealBuildTool
 
 					if (!File.Exists(DestApkName))
 					{
-						Logger.LogInformation("Output .apk [DestApkName] does not exist, will do full packaging with Gradle", DestApkName);
+						Logger.LogInformation("Output .apk [{DestApkName}] does not exist, will do full packaging with Gradle", DestApkName);
 						continue;
 					}
 
@@ -4401,7 +4401,14 @@ namespace UnrealBuildTool
 					bool bVerifySOChanged = true;
 					if (bVerifySOChanged)
 					{
-						string FinalSOName = UnrealBuildPath + "/jni/" + NDKArch + "/libUnreal.so";
+						string JNIDirectory = Path.Combine(UnrealBuildPath, "jni", NDKArch);
+						if (!Directory.Exists(JNIDirectory))
+						{
+							Logger.LogInformation("JNI directory {JNIDirectory} does not exist, will do full packaging with Gradle", JNIDirectory);
+							continue;
+						}
+
+						string FinalSOName = Path.Combine(JNIDirectory, "libUnreal.so");
 						if (!CopyIfDifferent(SourceSOName, FinalSOName, false, false))
 						{
 							Logger.LogInformation("{FinalSOName} unchanged, skipping repackage of {DestApkName}", FinalSOName, DestApkName);
