UnrealBuildTool: Clean up clang toolchain paths and versioning.
Remove need for most of the derived classes to override CompileCPPFiles

#rnx

diff --git a/Engine/Source/Programs/UnrealBuildTool/Platform/Android/AndroidToolChain.cs b/Engine/Source/Programs/UnrealBuildTool/Platform/Android/AndroidToolChain.cs
--- a/Engine/Source/Programs/UnrealBuildTool/Platform/Android/AndroidToolChain.cs
+++ b/Engine/Source/Programs/UnrealBuildTool/Platform/Android/AndroidToolChain.cs
@@ -100,6 +100,10 @@ namespace UnrealBuildTool
 		public string? NDKToolchainVersion;
 		public UInt64 NDKVersionInt;
 
+		int ClangVersionMajor = -1;
+		int ClangVersionMinor = -1;
+		int ClangVersionPatch = -1;
+
 		protected void SetClangVersion(int Major, int Minor, int Patch)
 		{
 			ClangVersionMajor = Major;
@@ -285,6 +289,11 @@ namespace UnrealBuildTool
 			}
 		}
 
+		protected override ClangToolChainInfo GetToolChainInfo()
+		{
+			return new ClangToolChainInfo(FileReference.FromString(ClangPath)!, FileReference.FromString(ArPathArm64)!, Logger);
+		}
+
 		public virtual void ParseArchitectures()
 		{
 			// look in ini settings for what platforms to compile for
diff --git a/Engine/Source/Programs/UnrealBuildTool/Platform/IOS/IOSToolChain.cs b/Engine/Source/Programs/UnrealBuildTool/Platform/IOS/IOSToolChain.cs
--- a/Engine/Source/Programs/UnrealBuildTool/Platform/IOS/IOSToolChain.cs
+++ b/Engine/Source/Programs/UnrealBuildTool/Platform/IOS/IOSToolChain.cs
@@ -17,6 +17,7 @@ using EpicGames.Core;
 using System.Security.Cryptography.X509Certificates;
 using UnrealBuildBase;
 using Microsoft.Extensions.Logging;
+using System.Text.RegularExpressions;
 
 namespace UnrealBuildTool
 {
@@ -124,20 +125,22 @@ namespace UnrealBuildTool
 		private Lazy<IOSToolChainSettings> Settings;
 
 		/// <summary>
-		/// Which compiler frontend to use
+		/// Which compiler\linker frontend to use
 		/// </summary>
 		private const string IOSCompiler = "clang++";
 
-		/// <summary>
-		/// Which linker frontend to use
-		/// </summary>
-		private const string IOSLinker = "clang++";
-
 		/// <summary>
 		/// Which library archiver to use
 		/// </summary>
 		private const string IOSArchiver = "libtool";
 
+		protected override ClangToolChainInfo GetToolChainInfo()
+		{
+			FileReference CompilerPath = new FileReference(Settings.Value.ToolchainDir + IOSCompiler);
+			FileReference ArchiverPath = new FileReference(Settings.Value.ToolchainDir + IOSArchiver);
+			return new AppleToolChainInfo(CompilerPath, ArchiverPath, Logger);
+		}
+
 		public override string GetSDKVersion()
 		{
 			return Settings.Value.IOSSDKVersionFloat.ToString();
@@ -214,7 +217,7 @@ namespace UnrealBuildTool
 		{
 			base.GetCompileArguments_Optimizations(CompileEnvironment, Arguments);
 
-			if (GetClangVersion().Major >= 12)
+			if (CompilerVersionGreaterOrEqual(12, 0, 0))
 			{
 				// We have 'this' vs nullptr comparisons that get optimized away for newer versions of Clang, which is undesirable until we refactor these checks.
 				Arguments.Add("-fno-delete-null-pointer-checks");
@@ -501,9 +504,6 @@ namespace UnrealBuildTool
 
 			GetCompileArguments_Global(CompileEnvironment, GlobalArguments);
 
-			FileReference CompilerPath = new FileReference(Settings.Value.ToolchainDir + IOSCompiler);
-			string CompilerVersion = GetFullClangVersion();
-
 			List<FileItem> FrameworkTokenFiles = new List<FileItem>();
 			foreach (UEBuildFramework Framework in CompileEnvironment.AdditionalFrameworks)
 			{
@@ -518,7 +518,7 @@ namespace UnrealBuildTool
 			// Create a compile action for each source file.
 			foreach (FileItem SourceFile in InputFiles)
 			{
-				Action CompileAction = CompileCPPFile(CompileEnvironment, SourceFile, OutputDir, ModuleName, Graph, GlobalArguments, CompilerPath, CompilerVersion, Result);
+				Action CompileAction = CompileCPPFile(CompileEnvironment, SourceFile, OutputDir, ModuleName, Graph, GlobalArguments, Result);
 				CompileAction.PrerequisiteItems.AddRange(FrameworkTokenFiles);
 			}
 			return Result;
@@ -526,8 +526,7 @@ namespace UnrealBuildTool
 
 		public override FileItem LinkFiles(LinkEnvironment LinkEnvironment, bool bBuildImportLibraryOnly, IActionGraphBuilder Graph)
 		{
-			string LinkerPath = Settings.Value.ToolchainDir +
-				(LinkEnvironment.bIsBuildingLibrary ? IOSArchiver : IOSLinker);
+			FileReference LinkerPath = LinkEnvironment.bIsBuildingLibrary ? Info.Archiver : Info.Clang;
 
 			// Create an action that invokes the linker.
 			Action LinkAction = Graph.CreateAction(ActionType.Link);
@@ -540,7 +539,7 @@ namespace UnrealBuildTool
 			if (LinkEnvironment.bIsBuildingDLL)
 			{
 				// @todo roll this put into GetLinkArguments_Global
-				LinkerPath = IOSLinker;
+				LinkerPath = Info.Clang;
 				LinkCommandArguments += " -dynamiclib -Xlinker -export_dynamic -Xlinker -no_deduplicate";
 
 				string InstallName = LinkEnvironment.InstallName ?? String.Format("@executable_path/Frameworks/{0}", LinkEnvironment.OutputFilePath.MakeRelativeTo(LinkEnvironment.OutputFilePath.Directory.ParentDirectory!));
@@ -654,7 +653,7 @@ namespace UnrealBuildTool
 
 			LinkAction.StatusDescription = string.Format("{0}", OutputFile.AbsolutePath);
 
-			LinkAction.CommandVersion = GetFullClangVersion();
+			LinkAction.CommandVersion = Info.ClangVersionString;
 
 			LinkAction.CommandPath = BuildHostPlatform.Current.Shell;
 			if (LinkEnvironment.Configuration == CppConfiguration.Shipping && Path.GetExtension(OutputFile.AbsolutePath) != ".a")
diff --git a/Engine/Source/Programs/UnrealBuildTool/Platform/Linux/LinuxToolChain.cs b/Engine/Source/Programs/UnrealBuildTool/Platform/Linux/LinuxToolChain.cs
--- a/Engine/Source/Programs/UnrealBuildTool/Platform/Linux/LinuxToolChain.cs
+++ b/Engine/Source/Programs/UnrealBuildTool/Platform/Linux/LinuxToolChain.cs
@@ -15,6 +15,29 @@ namespace UnrealBuildTool
 {
 	class LinuxToolChain : ClangToolChain
 	{
+		protected class LinuxToolChainInfo : ClangToolChainInfo
+		{
+			// cache the location of NDK tools
+			public bool bIsCrossCompiling { get; init; }
+			public DirectoryReference? BaseLinuxPath { get; init; }
+			public DirectoryReference? MultiArchRoot { get; init; }
+
+			public FileReference Objcopy { get; init; }
+			public FileReference DumpSyms { get; init; }
+			public FileReference BreakpadEncoder { get; init; }
+
+			public LinuxToolChainInfo(DirectoryReference? BaseLinuxPath, DirectoryReference? MultiArchRoot, FileReference Clang, FileReference Archiver, FileReference Objcopy, ILogger Logger)
+				: base(Clang, Archiver, Logger)
+			{
+				this.BaseLinuxPath = BaseLinuxPath;
+				this.MultiArchRoot = MultiArchRoot;
+				this.Objcopy = Objcopy;
+				// these are supplied by the engine and do not change depending on the circumstances
+				DumpSyms = FileReference.Combine(Unreal.EngineDirectory, "Binaries", "Linux", $"dump_syms{BuildHostPlatform.Current.BinarySuffix}");
+				BreakpadEncoder = FileReference.Combine(Unreal.EngineDirectory, "Binaries", "Linux", $"BreakpadSymbolEncoder{BuildHostPlatform.Current.BinarySuffix}");
+			}
+		}
+
 		/** Flavor of the current build (target triplet)*/
 		string Architecture;
 
@@ -31,30 +54,56 @@ namespace UnrealBuildTool
 		/** Platform SDK to use */
 		protected LinuxPlatformSDK PlatformSDK;
 
-		/** Toolchain information to print during the build. */
-		protected string? ToolchainInfo;
+		protected LinuxToolChainInfo LinuxInfo => (Info as LinuxToolChainInfo)!;
 
 		public LinuxToolChain(string InArchitecture, LinuxPlatformSDK InSDK, ClangToolChainOptions InOptions, ILogger InLogger)
 			: this(UnrealTargetPlatform.Linux, InArchitecture, InSDK, InOptions, InLogger)
 		{
-			MultiArchRoot = PlatformSDK.GetSDKLocation();
-			BaseLinuxPath = PlatformSDK.GetBaseLinuxPathForArchitecture(InArchitecture);
+			CheckDefaultCompilerSettings();
+
+			// prevent unknown clangs since the build is likely to fail on too old or too new compilers
+			if (CompilerVersionLessThan(13, 0, 0) || CompilerVersionGreaterOrEqual(14, 0, 0))
+			{
+				throw new BuildException(
+					string.Format("This version of the Unreal Engine can only be compiled with clang 13.0. clang {0} may not build it - please use a different version.",
+						Info.ClangVersion)
+					);
+			}
+		}
+
+		public LinuxToolChain(UnrealTargetPlatform InPlatform, string InArchitecture, LinuxPlatformSDK InSDK, ClangToolChainOptions InOptions, ILogger InLogger)
+			: base(InOptions, InLogger)
+		{
+			Architecture = InArchitecture;
+			PlatformSDK = InSDK;
+		}
+
+		protected override ClangToolChainInfo GetToolChainInfo()
+		{
+			DirectoryReference? MultiArchRoot = PlatformSDK.GetSDKLocation();
+			DirectoryReference? BaseLinuxPath = PlatformSDK.GetBaseLinuxPathForArchitecture(Architecture);
 
 			bool bForceUseSystemCompiler = PlatformSDK.ForceUseSystemCompiler();
-			bool bHasValidCompiler;
 
 			if (bForceUseSystemCompiler)
 			{
-				// Validate the system toolchain.
-				BaseLinuxPath = null;
-				MultiArchRoot = null;
-
-				ToolchainInfo = "system toolchain";
-
 				// use native linux toolchain
-				ClangPath = FileReference.FromString(LinuxCommon.WhichClang(Logger));
-				LlvmArPath = FileReference.FromString(LinuxCommon.Which("llvm-ar", Logger));
-				ObjcopyPath = FileReference.FromString(LinuxCommon.Which("llvm-objcopy", Logger));
+				FileReference? ClangPath = FileReference.FromString(LinuxCommon.WhichClang(Logger));
+				FileReference? LlvmArPath = FileReference.FromString(LinuxCommon.Which("llvm-ar", Logger));
+				FileReference? ObjcopyPath = FileReference.FromString(LinuxCommon.Which("llvm-objcopy", Logger));
+
+				if (ClangPath == null)
+				{
+					throw new BuildException("Unable to find system clang; cannot instantiate Linux toolchain");
+				}
+				else if (LlvmArPath == null)
+				{
+					throw new BuildException("Unable to find system llvm-ar; cannot instantiate Linux toolchain");
+				}
+				else if (ObjcopyPath == null)
+				{
+					throw new BuildException("Unable to find system llvm-objcopy; cannot instantiate Linux toolchain");
+				}
 
 				// When compiling on Linux, use a faster way to relink circularly dependent libraries.
 				// Race condition between actions linking to the .so and action overwriting it is avoided thanks to inodes
@@ -62,7 +111,7 @@ namespace UnrealBuildTool
 
 				bIsCrossCompiling = false;
 
-				bHasValidCompiler = DetermineCompilerVersion();
+				return new LinuxToolChainInfo(null, null, ClangPath, LlvmArPath, ObjcopyPath, Logger);
 			}
 			else
 			{
@@ -76,12 +125,10 @@ namespace UnrealBuildTool
 					Logger.LogInformation("Using LINUX_ROOT (deprecated, consider LINUX_MULTIARCH_ROOT)");
 				}
 
-				ToolchainInfo = $"toolchain located at '{BaseLinuxPath}'";
-
 				// set up the path to our toolchain
-				ClangPath = FileReference.Combine(BaseLinuxPath, "bin", $"clang++{BuildHostPlatform.Current.BinarySuffix}");
-				LlvmArPath = FileReference.Combine(BaseLinuxPath, "bin", $"llvm-ar{BuildHostPlatform.Current.BinarySuffix}" );
-				ObjcopyPath = FileReference.Combine(BaseLinuxPath, "bin", $"llvm-objcopy{BuildHostPlatform.Current.BinarySuffix}");
+				FileReference ClangPath = FileReference.Combine(BaseLinuxPath, "bin", $"clang++{BuildHostPlatform.Current.BinarySuffix}");
+				FileReference LlvmArPath = FileReference.Combine(BaseLinuxPath, "bin", $"llvm-ar{BuildHostPlatform.Current.BinarySuffix}");
+				FileReference ObjcopyPath = FileReference.Combine(BaseLinuxPath, "bin", $"llvm-objcopy{BuildHostPlatform.Current.BinarySuffix}");
 
 				// When cross-compiling on Windows, use old FixDeps. It is slow, but it does not have timing issues
 				bUseFixdeps = BuildHostPlatform.Current.Platform == UnrealTargetPlatform.Win64;
@@ -93,65 +140,15 @@ namespace UnrealBuildTool
 
 				bIsCrossCompiling = true;
 
-				bHasValidCompiler = DetermineCompilerVersion();
-			}
-
-			if (!bHasValidCompiler)
-			{
-				throw new BuildException("Could not determine version of the compiler, not registering Linux toolchain.");
-			}
-
-			CheckDefaultCompilerSettings();
-
-			// refuse to use compilers that we know won't work
-
-			// prevent unknown clangs since the build is likely to fail on too old or too new compilers
-			if ((ClangVersionMajor * 10 + ClangVersionMinor) > 130 || (ClangVersionMajor * 10 + ClangVersionMinor) < 130)
-			{
-				throw new BuildException(
-					string.Format("This version of the Unreal Engine can only be compiled with clang 13.0. clang {0} may not build it - please use a different version.",
-						ClangVersionString)
-					);
+				return new LinuxToolChainInfo(BaseLinuxPath, MultiArchRoot, ClangPath, LlvmArPath, ObjcopyPath, Logger);
 			}
 		}
 
-		public LinuxToolChain(UnrealTargetPlatform InPlatform, string InArchitecture, LinuxPlatformSDK InSDK, ClangToolChainOptions InOptions, ILogger InLogger)
-			: base(InOptions, InLogger)
-		{
-			Architecture = InArchitecture;
-			PlatformSDK = InSDK;
-
-			// these are supplied by the engine and do not change depending on the circumstances
-			DumpSymsPath = FileReference.Combine(Unreal.EngineDirectory, "Binaries", "Linux", $"dump_syms{BuildHostPlatform.Current.BinarySuffix}");
-			BreakpadEncoderPath = FileReference.Combine(Unreal.EngineDirectory, "Binaries", "Linux", $"BreakpadSymbolEncoder{BuildHostPlatform.Current.BinarySuffix}");
-		}
-
 		protected virtual bool CrossCompiling()
 		{
 			return bIsCrossCompiling;
 		}
 
-
-		/// <summary>
-		/// Splits compiler version string into numerical components, leaving unchanged if not known
-		/// </summary>
-		private void DetermineCompilerMajMinPatchFromVersionString()
-		{
-			string[] Parts = ClangVersionString!.Split('.');
-			if (Parts.Length >= 1)
-			{
-				ClangVersionMajor = Convert.ToInt32(Parts[0]);
-			}
-			if (Parts.Length >= 2)
-			{
-				ClangVersionMinor = Convert.ToInt32(Parts[1]);
-			}
-			if (Parts.Length >= 3)
-			{
-				ClangVersionPatch = Convert.ToInt32(Parts[2]);
-			}
-		}
-
 		protected internal virtual string GetDumpEncodeDebugCommand(LinkEnvironment LinkEnvironment, FileItem OutputFile)
 		{
 			bool bUseCmdExe = BuildHostPlatform.Current.Platform == UnrealTargetPlatform.Win64;
@@ -171,14 +168,14 @@ namespace UnrealBuildTool
 
 			// dump_syms
 			Out.WriteLine(DumpCommand,
-				DumpSymsPath,
+				LinuxInfo.DumpSyms,
 				OutputFile.AbsolutePath,
 				SymbolsFile.AbsolutePath
 			);
 
 			// encode breakpad symbols
 			Out.WriteLine("\"{0}\" \"{1}\" \"{2}\"",
-				BreakpadEncoderPath,
+				LinuxInfo.BreakpadEncoder,
 				SymbolsFile.AbsolutePath,
 				EncodedBinarySymbolsFile.AbsolutePath
 			);
@@ -196,21 +193,21 @@ namespace UnrealBuildTool
 
 				// objcopy stripped file
 				Out.WriteLine("\"{0}\" --strip-all \"{1}\" \"{2}\"",
-					ObjcopyPath!,
+					LinuxInfo.Objcopy,
 					OutputFile.AbsolutePath,
 					StrippedFile.AbsolutePath
 				);
 
 				// objcopy debug file
 				Out.WriteLine("\"{0}\" --only-keep-debug \"{1}\" \"{2}\"",
-					ObjcopyPath!,
+					LinuxInfo.Objcopy,
 					OutputFile.AbsolutePath,
 					DebugFile.AbsolutePath
 				);
 
 				// objcopy link debug file to final so
 				Out.WriteLine("\"{0}\" --add-gnu-debuglink=\"{1}\" \"{2}\" \"{3}.temp\"",
-					ObjcopyPath!,
+					LinuxInfo.Objcopy,
 					DebugFile.AbsolutePath,
 					StrippedFile.AbsolutePath,
 					OutputFile.AbsolutePath
@@ -257,60 +254,6 @@ namespace UnrealBuildTool
 			return Out.ToString();
 		}
 
-		/// <summary>
-		/// Queries compiler for the version
-		/// </summary>
-		protected bool DetermineCompilerVersion()
-		{
-			ClangVersionString = null;
-			ClangVersionMajor = -1;
-			ClangVersionMinor = -1;
-			ClangVersionPatch = -1;
-
-			using (Process Proc = new Process())
-			{
-				Proc.StartInfo.UseShellExecute = false;
-				Proc.StartInfo.CreateNoWindow = true;
-				Proc.StartInfo.RedirectStandardOutput = true;
-				Proc.StartInfo.RedirectStandardError = true;
-
-				if (ClangPath != null && FileReference.Exists(ClangPath))
-				{
-					Proc.StartInfo.FileName = ClangPath.FullName;
-					Proc.StartInfo.Arguments = " --version";
-
-					Proc.Start();
-					Proc.WaitForExit();
-
-					if (Proc.ExitCode == 0)
-					{
-						// read just the first string
-						string? VersionString = Proc.StandardOutput.ReadLine();
-
-						if (VersionString != null)
-						{
-							Regex VersionPattern = new Regex("version \\d+(\\.\\d+)+");
-							Match VersionMatch = VersionPattern.Match(VersionString);
-
-							// version match will be like "version 3.3", so remove the "version"
-							if (VersionMatch.Value.StartsWith("version "))
-							{
-								ClangVersionString = VersionMatch.Value.Replace("version ", "");
-
-								DetermineCompilerMajMinPatchFromVersionString();
-							}
-						}
-					}
-				}
-				else
-				{
-					// icl?
-				}
-			}
-
-			return !String.IsNullOrEmpty(ClangVersionString);
-		}
-
 		/// <summary>
 		/// Checks default compiler settings
 		/// </summary>
@@ -324,9 +267,9 @@ namespace UnrealBuildTool
 				Proc.StartInfo.RedirectStandardError = true;
 				Proc.StartInfo.RedirectStandardInput = true;
 
-				if (ClangPath != null && FileReference.Exists(ClangPath))
+				if (FileReference.Exists(Info.Clang))
 				{
-					Proc.StartInfo.FileName = ClangPath.FullName;
+					Proc.StartInfo.FileName = Info.Clang.FullName;
 					Proc.StartInfo.Arguments = " -E -dM -";
 
 					Proc.Start();
@@ -605,7 +548,7 @@ namespace UnrealBuildTool
 				{
 					Arguments.Add($"-target {CompileEnvironment.Architecture}");        // Set target triple
 				}
-				Arguments.Add($"--sysroot=\"{NormalizeCommandLinePath(BaseLinuxPath!)}\"");
+				Arguments.Add($"--sysroot=\"{NormalizeCommandLinePath(LinuxInfo.BaseLinuxPath!)}\"");
 			}
 		}
 
@@ -686,7 +629,7 @@ namespace UnrealBuildTool
 				if (CrossCompiling())
 				{
 					Arguments.Add(string.Format("-Wl,-rpath=\"{0}/lib/clang/{1}.{2}.{3}/lib/linux\"",
-							BaseLinuxPath, ClangVersionMajor, ClangVersionMinor, ClangVersionPatch));
+							LinuxInfo.BaseLinuxPath, Info.ClangVersion.Major, Info.ClangVersion.Minor, Info.ClangVersion.Build));
 				}
 			}
 
@@ -776,7 +719,7 @@ namespace UnrealBuildTool
 			if (CrossCompiling())
 			{
 				Arguments.Add($"-target {LinkEnvironment.Architecture}");        // Set target triple
-				DirectoryReference SysRootPath = BaseLinuxPath!;
+				DirectoryReference SysRootPath = LinuxInfo.BaseLinuxPath!;
 				Arguments.Add($"--sysroot=\"{NormalizeCommandLinePath(SysRootPath)}\"");
 
 				// Linking with the toolchain on linux appears to not search usr/
@@ -797,14 +740,6 @@ namespace UnrealBuildTool
 
 		// cache the location of NDK tools
 		protected bool bIsCrossCompiling;
-		protected DirectoryReference? BaseLinuxPath;
-		protected FileReference? ClangPath;
-		protected FileReference? LlvmArPath;
-		protected FileReference? ObjcopyPath;
-		protected FileReference? DumpSymsPath;
-		protected FileReference? BreakpadEncoderPath;
-		protected DirectoryReference? MultiArchRoot;
-
 		/// <summary>
 		/// Whether to use old, slower way to relink circularly dependent libraries.
 		/// It makes sense to use it when cross-compiling on Windows due to race conditions between actions reading and modifying the libs.
@@ -828,14 +763,14 @@ namespace UnrealBuildTool
 		protected void PrintBuildDetails(CppCompileEnvironment CompileEnvironment, ILogger Logger)
 		{
 			Logger.LogInformation("------- Build details --------");
-			Logger.LogInformation("Using {ToolchainInfo}.", ToolchainInfo);
+			Logger.LogInformation("Using {ToolchainInfo}.", LinuxInfo.BaseLinuxPath == null ? "system toolchain" : $"toolchain located at '{LinuxInfo.BaseLinuxPath}'");
 			Logger.LogInformation("Using clang ({ClangPath}) version '{ClangVersionString}' (string), {ClangVersionMajor} (major), {ClangVersionMinor} (minor), {ClangVersionPatch} (patch)",
-				ClangPath, ClangVersionString, ClangVersionMajor, ClangVersionMinor, ClangVersionPatch);
+				Info.Clang, Info.ClangVersionString, Info.ClangVersion.Major, Info.ClangVersion.Minor, Info.ClangVersion.Build);
 
 			// inform the user which C++ library the engine is going to be compiled against - important for compatibility with third party code that uses STL
 			Logger.LogInformation("Using {Lib} standard C++ library.", ShouldUseLibcxx(CompileEnvironment.Architecture) ? "bundled libc++" : "compiler default (most likely libstdc++)");
 			Logger.LogInformation("Using lld linker");
-			Logger.LogInformation("Using llvm-ar ({LlvmAr})", LlvmArPath);
+			Logger.LogInformation("Using llvm-ar ({LlvmAr}) version '{LlvmArVersionString} (string)'", Info.Archiver, Info.ArchiverVersionString);
 
 			if (Options.HasFlag(ClangToolChainOptions.EnableAddressSanitizer) ||
 				Options.HasFlag(ClangToolChainOptions.EnableThreadSanitizer) ||
@@ -905,7 +840,7 @@ namespace UnrealBuildTool
 			CPPOutput Result = new CPPOutput();
 			foreach (FileItem SourceFile in InputFiles)
 			{
-				CompileCPPFile(CompileEnvironment, SourceFile, OutputDir, ModuleName, Graph, GlobalArguments, ClangPath!, ClangVersionString!, Result);
+				CompileCPPFile(CompileEnvironment, SourceFile, OutputDir, ModuleName, Graph, GlobalArguments, Result);
 			}
 
 			return Result;
@@ -919,7 +854,7 @@ namespace UnrealBuildTool
 			// Create an archive action
 			Action ArchiveAction = Graph.CreateAction(ActionType.Link);
 			ArchiveAction.WorkingDirectory = Unreal.EngineSourceDirectory;
-			ArchiveAction.CommandPath = LlvmArPath!;
+			ArchiveAction.CommandPath = Info.Archiver;
 
 			// this will produce a final library
 			ArchiveAction.bProducesImportLibrary = true;
@@ -1049,7 +984,7 @@ namespace UnrealBuildTool
 			LinkAction.WorkingDirectory = Unreal.EngineSourceDirectory;
 
 			string LinkCommandString;
-			LinkCommandString = "\"" + ClangPath + "\"";
+			LinkCommandString = "\"" + Info.Clang + "\"";
 
 			// Get link arguments.
 			List<string> LinkArguments = new List<string>();
@@ -1082,7 +1017,7 @@ namespace UnrealBuildTool
 			}
 			// because the logic choosing between lld and ld is somewhat messy atm (lld fails to link .DSO due to bugs), make the name of the linker clear
 			LinkAction.CommandDescription += (LinkCommandString.Contains("-fuse-ld=lld")) ? " (lld)" : " (ld)";
-			LinkAction.CommandVersion = ClangVersionString!;
+			LinkAction.CommandVersion = Info.ClangVersionString;
 			LinkAction.StatusDescription = Path.GetFileName(OutputFile.AbsolutePath);
 
 			// Add the output file to the command-line.
@@ -1624,7 +1559,7 @@ namespace UnrealBuildTool
 			}
 
 			ProcessStartInfo StartInfo = new ProcessStartInfo();
-			StartInfo.FileName = ObjcopyPath!.FullName;
+			StartInfo.FileName = LinuxInfo.Objcopy.FullName;
 			StartInfo.Arguments = "--strip-debug \"" + TargetFile.FullName + "\"";
 			StartInfo.UseShellExecute = false;
 			StartInfo.CreateNoWindow = true;
diff --git a/Engine/Source/Programs/UnrealBuildTool/Platform/Mac/MacToolChain.cs b/Engine/Source/Programs/UnrealBuildTool/Platform/Mac/MacToolChain.cs
--- a/Engine/Source/Programs/UnrealBuildTool/Platform/Mac/MacToolChain.cs
+++ b/Engine/Source/Programs/UnrealBuildTool/Platform/Mac/MacToolChain.cs
@@ -92,15 +92,10 @@ namespace UnrealBuildTool
 		}
 
 		/// <summary>
-		/// Which compiler frontend to use
+		/// Which compiler\linker frontend to use
 		/// </summary>
 		private const string MacCompiler = "clang++";
 
-		/// <summary>
-		/// Which linker frontend to use
-		/// </summary>
-		private const string MacLinker = "clang++";
-
 		/// <summary>
 		/// Which archiver to use
 		/// </summary>
@@ -113,6 +108,13 @@ namespace UnrealBuildTool
 
 		private static List<FileItem> BundleDependencies = new List<FileItem>();
 
+		protected override ClangToolChainInfo GetToolChainInfo()
+		{
+			FileReference CompilerPath = new FileReference(Settings.ToolchainDir + MacCompiler);
+			FileReference ArchiverPath = new FileReference(Settings.ToolchainDir + MacArchiver);
+			return new AppleToolChainInfo(CompilerPath, ArchiverPath, Logger);
+		}
+
 		private static void SetupXcodePaths(bool bVerbose)
 		{
 		}
@@ -259,7 +261,7 @@ namespace UnrealBuildTool
 
 			// Temporary workaround for linker warning with Xcode 14:
 			//		'ld: warning: could not create compact unwind for _inflate_fast: registers 27 not saved contiguously in frame'
-			if (GetClangVersion().Major < 14)
+			if (CompilerVersionLessThan(14, 0, 0))
 			{
 				Result += " -Wl,-fatal_warnings";
 			}
@@ -307,24 +309,6 @@ namespace UnrealBuildTool
 			return Result;
 		}
 
-		public override CPPOutput CompileCPPFiles(CppCompileEnvironment CompileEnvironment, List<FileItem> InputFiles, DirectoryReference OutputDir, string ModuleName, IActionGraphBuilder Graph)
-		{
-			List<string> GlobalArguments = new();
-
-			GetCompileArguments_Global(CompileEnvironment, GlobalArguments);
-
-			FileReference CompilerPath = new FileReference(Settings.ToolchainDir + MacCompiler);
-			string CompilerVersion = GetFullClangVersion();
-
-			CPPOutput Result = new CPPOutput();
-			// Create a compile action for each source file.
-			foreach (FileItem SourceFile in InputFiles)
-			{
-				CompileCPPFile(CompileEnvironment, SourceFile, OutputDir, ModuleName, Graph, GlobalArguments, CompilerPath, CompilerVersion, Result);
-			}
-			return Result;
-		}
-
 		private void AppendMacLine(StreamWriter Writer, string Format, params object[] Arg)
 		{
 			string PreLine = String.Format(Format, Arg);
@@ -478,13 +462,12 @@ namespace UnrealBuildTool
 			// Create an action that invokes the linker.
 			Action LinkAction = Graph.CreateAction(ActionType.Link);
 
-			string Linker = bIsBuildingLibrary ? MacArchiver : MacLinker;
-			string LinkerPath = Settings.ToolchainDir + Linker;
+			FileReference LinkerPath = bIsBuildingLibrary ? Info.Archiver : Info.Clang;
 
 			LinkAction.WorkingDirectory = GetMacDevSrcRoot();
-			LinkAction.CommandPath = new FileReference(LinkerPath);
+			LinkAction.CommandPath = LinkerPath;
 			LinkAction.CommandDescription = "Link";
-			LinkAction.CommandVersion = GetFullClangVersion();
+			LinkAction.CommandVersion = Info.ClangVersionString;
 
 			string EngineAPIVersion = LoadEngineAPIVersion();
 			string EngineDisplayVersion = LoadEngineDisplayVersion(true);
diff --git a/Engine/Source/Programs/UnrealBuildTool/ToolChain/AppleToolChain.cs b/Engine/Source/Programs/UnrealBuildTool/ToolChain/AppleToolChain.cs
--- a/Engine/Source/Programs/UnrealBuildTool/ToolChain/AppleToolChain.cs
+++ b/Engine/Source/Programs/UnrealBuildTool/ToolChain/AppleToolChain.cs
@@ -148,6 +148,18 @@ namespace UnrealBuildTool
 
 	abstract class AppleToolChain : ClangToolChain
 	{
+		protected class AppleToolChainInfo : ClangToolChainInfo
+		{
+			public AppleToolChainInfo(FileReference Clang, FileReference Archiver, ILogger Logger)
+				: base(Clang, Archiver, Logger)
+			{
+			}
+
+			// libtool doesn't provide version, just use clang's version string
+			/// <inheritdoc/>
+			protected override string QueryArchiverVersionString() => ClangVersionString;
+		}
+
 		protected FileReference? ProjectFile;
 
 		public AppleToolChain(FileReference? InProjectFile, ClangToolChainOptions InOptions, ILogger InLogger) : base(InOptions, InLogger)
@@ -176,37 +188,6 @@ namespace UnrealBuildTool
 			Utils.RunLocalProcessAndLogOutput(StartInfo, Logger);
 		}
 
-		static Version? ClangVersion = null;
-		static string? FullClangVersion = null;
-
-		protected Version GetClangVersion()
-		{
-			if (ClangVersion == null)
-			{
-				FileReference ClangLocation = new FileReference("/usr/bin/clang");
-				ClangVersion = RunToolAndCaptureVersion(ClangLocation, "--version");
-			}
-			if (ClangVersionMajor == -1 || ClangVersionMinor == -1 || ClangVersionPatch == -1)
-			{
-				ClangVersionMajor = ClangVersion.Major;
-				ClangVersionMinor = ClangVersion.Minor;
-				ClangVersionPatch = ClangVersion.Build;
-			}
-			return ClangVersion;
-		}
-
-		protected string GetFullClangVersion()
-		{
-			if (FullClangVersion == null)
-			{
-				// get the first line that has the full clang and build number
-				FileReference ClangLocation = new FileReference("/usr/bin/clang");
-				FullClangVersion = RunToolAndCaptureOutput(ClangLocation, "--version", "(.*)")!;
-			}
-
-			return FullClangVersion;
-		}
-
 		/// <inheritdoc/>
 		protected override string EscapePreprocessorDefinition(string Definition)
 		{
@@ -240,11 +221,11 @@ namespace UnrealBuildTool
 			GetCppStandardCompileArgument(CompileEnvironment, Arguments);
 			Arguments.Add("-stdlib=libc++");
 
-			if (GetClangVersion().Major >= 11)
+			if (CompilerVersionGreaterOrEqual(11, 0, 0))
 			{
 				Arguments.Add("-fpch-validate-input-files-content");
 			}
-			if (GetClangVersion().Major >= 13) // Note this is supported for >=11 on other clang platforms
+			if (CompilerVersionGreaterOrEqual(13, 0, 0)) // Note this is supported for >=11 on other clang platforms
 			{
 				Arguments.Add("-fpch-instantiate-templates");
 			}
@@ -265,7 +246,7 @@ namespace UnrealBuildTool
 
 			// clang 12.00 has a new warning for copies in ranged loops. Instances have all been fixed up (2020/6/26) but
 			// are likely to be reintroduced due to no equivalent on other platforms at this time so disable the warning
-			if (GetClangVersion().Major >= 12)
+			if (CompilerVersionGreaterOrEqual(12, 0, 0))
 			{
 				Arguments.Add("-Wno-range-loop-analysis");
 			}
@@ -330,9 +311,6 @@ namespace UnrealBuildTool
 		/// <inheritdoc/>
 		protected override void GetCompileArguments_Global(CppCompileEnvironment CompileEnvironment, List<string> Arguments)
 		{
-			// Ensure Clang version is set for base.GetCompileArguments_Global()
-			GetClangVersion();
-
 			base.GetCompileArguments_Global(CompileEnvironment, Arguments);
 
 			Arguments.Add(GetRTTIFlag(CompileEnvironment));
diff --git a/Engine/Source/Programs/UnrealBuildTool/ToolChain/ClangToolChain.cs b/Engine/Source/Programs/UnrealBuildTool/ToolChain/ClangToolChain.cs
--- a/Engine/Source/Programs/UnrealBuildTool/ToolChain/ClangToolChain.cs
+++ b/Engine/Source/Programs/UnrealBuildTool/ToolChain/ClangToolChain.cs
@@ -7,6 +7,7 @@ using System.Linq;
 using UnrealBuildBase;
 using Microsoft.Extensions.Logging;
 using System.IO;
+using System.Text.RegularExpressions;
 
 namespace UnrealBuildTool
 {
@@ -111,11 +112,68 @@ namespace UnrealBuildTool
 
 	abstract class ClangToolChain : ISPCToolChain
 	{
+		protected class ClangToolChainInfo
+		{
+			protected ILogger Logger { get; init; }
+			public FileReference Clang { get; init; }
+			public FileReference Archiver { get; init; }
+
+			public Version ClangVersion => LazyClangVersion.Value;
+			public string ClangVersionString => LazyClangVersionString.Value;
+			public string ArchiverVersionString => LazyArchiverVersionString.Value;
+
+			Lazy<Version> LazyClangVersion;
+			Lazy<string> LazyClangVersionString;
+			Lazy<string> LazyArchiverVersionString;
+
+			/// <summary>
+			/// Constructor for ClangToolChainInfo
+			/// </summary>
+			/// <param name="Clang">The path to the compiler</param>
+			/// <param name="Archiver">The path to the archiver</param>
+			/// <param name="Logger">Logging interface</param>
+			public ClangToolChainInfo(FileReference Clang, FileReference Archiver, ILogger Logger)
+			{
+				this.Logger = Logger;
+				this.Clang = Clang;
+				this.Archiver = Archiver;
+
+				LazyClangVersion = new Lazy<Version>(() => QueryClangVersion());
+				LazyClangVersionString = new Lazy<string>(() => QueryClangVersionString());
+				LazyArchiverVersionString = new Lazy<string>(() => QueryArchiverVersionString());
+			}
+
+			/// <summary>
+			/// Lazily query the clang version. Will only be executed once.
+			/// </summary>
+			/// <returns>The version of clang</returns>
+			/// <exception cref="BuildException"></exception>
+			protected virtual Version QueryClangVersion()
+			{
+				Match MatchClangVersion = Regex.Match(ClangVersionString, @"clang version (?<full>(?<major>\d+)\.(?<minor>\d+)\.(?<patch>\d+))");
+				if (MatchClangVersion.Success && Version.TryParse(MatchClangVersion.Groups["full"].Value, out Version? ClangVersion))
+				{
+					return ClangVersion;
+				}
+				throw new BuildException("Failed to query the Clang version number!");
+			}
+
+			/// <summary>
+			/// Lazily query the clang version output. Will only be executed once.
+			/// </summary>
+			/// <returns>The standard output when running Clang --version</returns>
+			protected virtual string QueryClangVersionString() => Utils.RunLocalProcessAndReturnStdOut(Clang.FullName, "--version", null);
+
+			/// <summary>
+			/// Lazily query the archiver version output. Will only be executed once.
+			/// </summary>
+			/// <returns>The standard output when running Archiver --version</returns>
+			protected virtual string QueryArchiverVersionString() => Utils.RunLocalProcessAndReturnStdOut(Archiver.FullName, "--version", null);
+		}
+
 		// The Clang version being used to compile
-		protected string? ClangVersionString = null;
-		protected int ClangVersionMajor = -1;
-		protected int ClangVersionMinor = -1;
-		protected int ClangVersionPatch = -1;
+		Lazy<ClangToolChainInfo> LazyInfo;
+		protected ClangToolChainInfo Info => LazyInfo.Value;
 
 		protected ClangToolChainOptions Options;
 
@@ -141,8 +199,12 @@ namespace UnrealBuildTool
 			: base(InLogger)
 		{
 			Options = InOptions;
+
+			LazyInfo = new Lazy<ClangToolChainInfo>(() => GetToolChainInfo());
 		}
 
+		protected abstract ClangToolChainInfo GetToolChainInfo();
+
 		public override void SetUpGlobalEnvironment(ReadOnlyTargetRules Target)
 		{
 			base.SetUpGlobalEnvironment(Target);
@@ -210,30 +272,12 @@ namespace UnrealBuildTool
 		/// <summary>
 		/// Checks if compiler version matches the requirements
 		/// </summary>
-		protected bool CompilerVersionGreaterOrEqual(int Major, int Minor, int Patch)
-		{
-			// TODO: Temporary verification check to ensure a clang version has been set until ClangVersion is standarized across all clang-based toolchains.
-			if (ClangVersionMajor == -1 || ClangVersionMinor == -1 || ClangVersionPatch == -1)
-			{
-				throw new BuildException($"ClangVersion not valid ({ClangVersionMajor}.{ClangVersionMinor}.{ClangVersionPatch}), unable to check compiler version requirements");
-			}
-
-			return new Version(ClangVersionMajor, ClangVersionMinor, ClangVersionPatch) >= new Version(Major, Minor, Patch);
-		}
+		protected bool CompilerVersionGreaterOrEqual(int Major, int Minor, int Patch) => Info.ClangVersion >= new Version(Major, Minor, Patch);
 
 		/// <summary>
 		/// Checks if compiler version matches the requirements
 		/// </summary>
-		protected bool CompilerVersionLessThan(int Major, int Minor, int Patch)
-		{
-			// TODO: Temporary verification check until ClangVersion is standarized across all clang-based toolchains to ensure a version has been set.
-			if (ClangVersionMajor == -1 || ClangVersionMinor == -1 || ClangVersionPatch == -1)
-			{
-				throw new BuildException($"ClangVersion not valid ({ClangVersionMajor}.{ClangVersionMinor}.{ClangVersionPatch}), unable to check compiler version requirements");
-			}
-
-			return new Version(ClangVersionMajor, ClangVersionMinor, ClangVersionPatch) < new Version(Major, Minor, Patch);
-		}
+		protected bool CompilerVersionLessThan(int Major, int Minor, int Patch) => Info.ClangVersion < new Version(Major, Minor, Patch);
 
 		protected virtual void GetCppStandardCompileArgument(CppCompileEnvironment CompileEnvironment, List<string> Arguments)
 		{
@@ -791,7 +835,7 @@ namespace UnrealBuildTool
 			Arguments.Add(GetOutputFileArgument(OutputFile));
 		}
 
-		protected virtual Action CompileCPPFile(CppCompileEnvironment CompileEnvironment, FileItem SourceFile, DirectoryReference OutputDir, string ModuleName, IActionGraphBuilder Graph, IReadOnlyCollection<string> GlobalArguments, FileReference CommandPath, string CommandVersion, CPPOutput Result)
+		protected virtual Action CompileCPPFile(CppCompileEnvironment CompileEnvironment, FileItem SourceFile, DirectoryReference OutputDir, string ModuleName, IActionGraphBuilder Graph, IReadOnlyCollection<string> GlobalArguments, CPPOutput Result)
 		{
 			Action CompileAction = Graph.CreateAction(ActionType.Compile);
 
@@ -819,7 +863,7 @@ namespace UnrealBuildTool
 				// >>>> xxx-clang.exe': The filename or extension is too long.  (0xCE)
 				// Clang processes and modifies the response file contents and this makes the final cmd line size hard for us to predict.
 				// To be conservative we add a dummy define to inflate the response file size and force clang to use the response file mode when we are close to the limit.
-				int CmdLineLength = CommandPath.ToString().Length + string.Join(' ', ResponseFileContents).Length;
+				int CmdLineLength = Info.Clang.ToString().Length + string.Join(' ', ResponseFileContents).Length;
 				bool bIsInDangerZone = CmdLineLength >= ClangCmdlineDangerZone && CmdLineLength <= ClangCmdLineMaxSize;
 				if (bIsInDangerZone)
 				{
@@ -833,8 +877,8 @@ namespace UnrealBuildTool
 			CompileAction.PrerequisiteItems.Add(CompilerResponseFileItem);
 
 			CompileAction.WorkingDirectory = Unreal.EngineSourceDirectory;
-			CompileAction.CommandPath = CommandPath;
-			CompileAction.CommandVersion = CommandVersion;
+			CompileAction.CommandPath = Info.Clang;
+			CompileAction.CommandVersion = Info.ClangVersionString;
 			CompileAction.CommandDescription = StaticAnalyzer == StaticAnalyzer.Default ? "Analyze" : "Compile";
 			CompileAction.StatusDescription = Path.GetFileName(SourceFile.AbsolutePath);
 			CompileAction.bIsGCCCompiler = true;
@@ -884,7 +928,7 @@ namespace UnrealBuildTool
 
 				if (RuntimePlatform.IsWindows)
 				{
-					int CmdLineLength = CommandPath.ToString().Length + string.Join(' ', ResponseFileContents).Length;
+					int CmdLineLength = Info.Clang.ToString().Length + string.Join(' ', ResponseFileContents).Length;
 					bool bIsInDangerZone = CmdLineLength >= ClangCmdlineDangerZone && CmdLineLength <= ClangCmdLineMaxSize;
 					if (bIsInDangerZone)
 					{
@@ -903,5 +947,20 @@ namespace UnrealBuildTool
 
 			return CompileAction;
 		}
+
+		public override CPPOutput CompileCPPFiles(CppCompileEnvironment CompileEnvironment, List<FileItem> InputFiles, DirectoryReference OutputDir, string ModuleName, IActionGraphBuilder Graph)
+		{
+			List<string> GlobalArguments = new();
+			GetCompileArguments_Global(CompileEnvironment, GlobalArguments);
+
+			// Create a compile action for each source file.
+			CPPOutput Result = new CPPOutput();
+			foreach (FileItem SourceFile in InputFiles)
+			{
+				CompileCPPFile(CompileEnvironment, SourceFile, OutputDir, ModuleName, Graph, GlobalArguments, Result);
+			}
+
+			return Result;
+		}
 	}
 }
\ No newline at end of file
