Add support for catching Xcode fatal errors when PCH needs to be recompiled.

#rnx

diff --git a/Engine/Source/Programs/UnrealBuildTool/Matchers/CompileEventMatcher.cs b/Engine/Source/Programs/UnrealBuildTool/Matchers/CompileEventMatcher.cs
--- a/Engine/Source/Programs/UnrealBuildTool/Matchers/CompileEventMatcher.cs
+++ b/Engine/Source/Programs/UnrealBuildTool/Matchers/CompileEventMatcher.cs
@@ -157,10 +157,33 @@ namespace UnrealBuildTool.Matchers
 
 					return builder.ToMatch(LogEventPriority.High, level, KnownLogEvents.Compiler);
 				}
+
+				// Try to match an Xcode diagnostic.
+				if (TryMatchXcodeCppEvent(builder, out eventMatch))
+				{
+					return eventMatch;
+				}
 			}
 			return null;
 		}
 
+		static readonly Regex s_xcodeRebuildPCHPattern = new Regex(@"(?<severity>(?:[Ff]atal )?[Ee]rror): file '(?<file>.*)' has been modified since the precompiled header '(?<pch>.*)' was built");
+
+		bool TryMatchXcodeCppEvent(LogEventBuilder builder, [NotNullWhen(true)] out LogEventMatch? outEvent)
+		{
+			Match? match;
+			if (builder.Current.TryMatch(s_xcodeRebuildPCHPattern, out match))
+			{
+				builder.Annotate(match.Groups["severity"], LogEventMarkup.Severity);
+				builder.AnnotateSourceFile(match.Groups["pch"], null);
+				outEvent = builder.ToMatch(LogEventPriority.Highest, LogLevel.Error, KnownLogEvents.Compiler);
+				return true;
+			}
+
+			outEvent = null;
+			return false;
+		}
+
 		static readonly Regex s_msvcPattern = new Regex($"^\\s*(?:ERROR: |WARNING: )?{FilePattern}(?:{VisualCppLocationPattern})? ?:\\s+{VisualCppSeverity}:");
 		static readonly Regex s_msvcNotePattern = new Regex($"^\\s*{FilePattern}(?:{VisualCppLocationPattern})?\\s*: note:");
 		static readonly Regex s_projectPattern = new Regex(@"\[(?<project>[^[\]]+)]\s*$");
