UnrealBuildTool: link arguments as list for ios & mac

#rnx

diff --git a/Engine/Source/Programs/UnrealBuildTool/Platform/IOS/IOSToolChain.cs b/Engine/Source/Programs/UnrealBuildTool/Platform/IOS/IOSToolChain.cs
--- a/Engine/Source/Programs/UnrealBuildTool/Platform/IOS/IOSToolChain.cs
+++ b/Engine/Source/Programs/UnrealBuildTool/Platform/IOS/IOSToolChain.cs
@@ -395,72 +395,70 @@ namespace UnrealBuildTool
 			}
 		}
 
-		string GetLinkArguments_Global(LinkEnvironment LinkEnvironment)
+		void GetLinkArguments_Global(LinkEnvironment LinkEnvironment, List<string> Arguments)
 		{
-			string Result = "";
-
-			Result += GetArchitectureArgument(LinkEnvironment.Configuration, LinkEnvironment.Architecture);
+			Arguments.Add(GetArchitectureArgument(LinkEnvironment.Configuration, LinkEnvironment.Architecture));
 
 			bool bIsDevice = (LinkEnvironment.Architecture != "-simulator");
-			Result += String.Format(" -isysroot \\\"{0}Platforms/{1}.platform/Developer/SDKs/{1}{2}.sdk\\\"",
-				Settings.Value.XcodeDeveloperDir, bIsDevice ? Settings.Value.DevicePlatformName : Settings.Value.SimulatorPlatformName, Settings.Value.IOSSDKVersion);
+			Arguments.Add(String.Format(" -isysroot \\\"{0}Platforms/{1}.platform/Developer/SDKs/{1}{2}.sdk\\\"",
+				Settings.Value.XcodeDeveloperDir, bIsDevice ? Settings.Value.DevicePlatformName : Settings.Value.SimulatorPlatformName, Settings.Value.IOSSDKVersion));
 
 			if (IsBitcodeCompilingEnabled(LinkEnvironment.Configuration))
 			{
 				FileItem OutputFile = FileItem.GetItemByFileReference(LinkEnvironment.OutputFilePath);
 
-				Result += " -fembed-bitcode -Xlinker -bitcode_verify -Xlinker -bitcode_hide_symbols -Xlinker -bitcode_symbol_map ";
-				Result += " -Xlinker \\\"" + Path.GetDirectoryName(OutputFile.AbsolutePath) + "\\\"";
+				Arguments.Add("-fembed-bitcode -Xlinker -bitcode_verify -Xlinker -bitcode_hide_symbols -Xlinker -bitcode_symbol_map ");
+				Arguments.Add("-Xlinker \\\"" + Path.GetDirectoryName(OutputFile.AbsolutePath) + "\\\"");
 			}
 
-			Result += " -dead_strip";
-			Result += " -m" + GetXcodeMinVersionParam() + "=" + ProjectSettings.RuntimeVersion;
-			Result += " -Wl";
+			Arguments.Add("-dead_strip");
+			Arguments.Add("-m" + GetXcodeMinVersionParam() + "=" + ProjectSettings.RuntimeVersion);
+			Arguments.Add("-Wl");
 			if (!IsBitcodeCompilingEnabled(LinkEnvironment.Configuration))
 			{
-				Result += "-no-pie";
+				Arguments.Add("-no-pie");
 			}
-			Result += " -stdlib=libc++";
-			Result += " -ObjC";
-			//			Result += " -v";
+			Arguments.Add("-stdlib=libc++");
+			Arguments.Add("-ObjC");
+			// Arguments.Add("-v");
 
 			// use LTO if desired (like VCToolchain does)
 			if (LinkEnvironment.bAllowLTCG)
 			{
-				Result += " -flto";
+				Arguments.Add("-flto");
 			}
 
 			string? SanitizerMode = Environment.GetEnvironmentVariable("ENABLE_ADDRESS_SANITIZER");
 			if ((SanitizerMode != null && SanitizerMode == "YES") || (Options.HasFlag(ClangToolChainOptions.EnableAddressSanitizer)))
 			{
-				Result += " -rpath \"@executable_path/Frameworks\"";
-				Result += " -fsanitize=address";
+				Arguments.Add("-rpath \"@executable_path/Frameworks\"");
+				Arguments.Add("-fsanitize=address");
 			}
 
 			string? UndefSanitizerMode = Environment.GetEnvironmentVariable("ENABLE_UNDEFINED_BEHAVIOR_SANITIZER");
 			if ((UndefSanitizerMode != null && UndefSanitizerMode == "YES") || (Options.HasFlag(ClangToolChainOptions.EnableUndefinedBehaviorSanitizer)))
 			{
-				Result += " -rpath \"@executable_path/libclang_rt.ubsan_ios_dynamic.dylib\"";
-				Result += " -fsanitize=undefined";
+				Arguments.Add("-rpath \"@executable_path/libclang_rt.ubsan_ios_dynamic.dylib\"");
+				Arguments.Add("-fsanitize=undefined");
 			}
 
 			if (Options.HasFlag(ClangToolChainOptions.EnableThreadSanitizer))
 			{
-				Result += " -rpath \"@executable_path/libclang_rt.tsan_ios_dynamic.dylib\"";
-				Result += " -fsanitize=thread";
+				Arguments.Add("-rpath \"@executable_path/libclang_rt.tsan_ios_dynamic.dylib\"");
+				Arguments.Add("-fsanitize=thread");
 			}
 
 			// need to tell where to load Framework dylibs
-			Result += " -rpath @executable_path/Frameworks";
+			Arguments.Add("-rpath @executable_path/Frameworks");
 
-			Result += " " + GetAdditionalLinkerFlags(LinkEnvironment.Configuration);
+			Arguments.Add(GetAdditionalLinkerFlags(LinkEnvironment.Configuration));
 
 			// link in the frameworks
 			foreach (string Framework in LinkEnvironment.Frameworks)
 			{
 				if (Framework != "ARKit" || Settings.Value.IOSSDKVersionFloat >= 11.0f)
 				{
-					Result += " -framework " + Framework;
+					Arguments.Add("-framework " + Framework);
 				}
 			}
 			foreach (UEBuildFramework Framework in LinkEnvironment.AdditionalFrameworks)
@@ -476,26 +474,20 @@ namespace UnrealBuildTool
 					{
 						FrameworkDir = Path.GetDirectoryName(FrameworkDir)!;
 					}
-					Result += String.Format(" -F\\\"{0}\\\"", FrameworkDir);
+					Arguments.Add(String.Format("-F\\\"{0}\\\"", FrameworkDir));
 				}
 
-				Result += " -framework " + Framework.Name;
+				Arguments.Add("-framework " + Framework.Name);
 			}
 			foreach (string Framework in LinkEnvironment.WeakFrameworks)
 			{
-				Result += " -weak_framework " + Framework;
+				Arguments.Add("-weak_framework " + Framework);
 			}
-
-			return Result;
 		}
 
-		static string GetArchiveArguments_Global(LinkEnvironment LinkEnvironment)
+		void GetArchiveArguments_Global(LinkEnvironment LinkEnvironment, List<string> Arguments)
 		{
-			string Result = "";
-
-			Result += " -static";
-
-			return Result;
+			Arguments.Add("-static");
 		}
 
 		public override CPPOutput CompileCPPFiles(CppCompileEnvironment CompileEnvironment, List<FileItem> InputFiles, DirectoryReference OutputDir, string ModuleName, IActionGraphBuilder Graph)
@@ -535,7 +527,16 @@ namespace UnrealBuildTool
 			LinkAction.WorkingDirectory = GetMacDevSrcRoot();
 
 			// build this up over the rest of the function
-			string LinkCommandArguments = LinkEnvironment.bIsBuildingLibrary ? GetArchiveArguments_Global(LinkEnvironment) : GetLinkArguments_Global(LinkEnvironment);
+			List<string> LinkArguments = new();
+			if (LinkEnvironment.bIsBuildingLibrary)
+			{
+				GetArchiveArguments_Global(LinkEnvironment, LinkArguments);
+			}
+			else
+			{
+				GetLinkArguments_Global(LinkEnvironment, LinkArguments);
+			}
+			string LinkCommandArguments = string.Join(' ', LinkArguments);
 			if (LinkEnvironment.bIsBuildingDLL)
 			{
 				// @todo roll this put into GetLinkArguments_Global
diff --git a/Engine/Source/Programs/UnrealBuildTool/Platform/Mac/MacToolChain.cs b/Engine/Source/Programs/UnrealBuildTool/Platform/Mac/MacToolChain.cs
--- a/Engine/Source/Programs/UnrealBuildTool/Platform/Mac/MacToolChain.cs
+++ b/Engine/Source/Programs/UnrealBuildTool/Platform/Mac/MacToolChain.cs
@@ -249,64 +249,58 @@ namespace UnrealBuildTool
 			return Result;
 		}
 
-		string GetLinkArguments_Global(LinkEnvironment LinkEnvironment)
+		void GetLinkArguments_Global(LinkEnvironment LinkEnvironment, List<string> Arguments)
 		{
-			string Result = "";
-
 			// Pass through architecture and OS info		
-			Result += " " + FormatArchitectureArg(LinkEnvironment.Architecture);
-			Result += string.Format(" -isysroot \"{0}\"", SDKPath);
-			Result += " -mmacosx-version-min=" + Settings.MacOSVersion;
-			Result += " -dead_strip";
+			Arguments.Add(FormatArchitectureArg(LinkEnvironment.Architecture));
+			Arguments.Add(string.Format("-isysroot \"{0}\"", SDKPath));
+			Arguments.Add("-mmacosx-version-min=" + Settings.MacOSVersion);
+			Arguments.Add("-dead_strip");
 
 			// Temporary workaround for linker warning with Xcode 14:
 			//		'ld: warning: could not create compact unwind for _inflate_fast: registers 27 not saved contiguously in frame'
 			if (CompilerVersionLessThan(14, 0, 0))
 			{
-				Result += " -Wl,-fatal_warnings";
+				Arguments.Add("-Wl,-fatal_warnings");
 			}
 
 			if (Options.HasFlag(ClangToolChainOptions.EnableAddressSanitizer) || Options.HasFlag(ClangToolChainOptions.EnableThreadSanitizer) || Options.HasFlag(ClangToolChainOptions.EnableUndefinedBehaviorSanitizer))
 			{
-				Result += " -g";
+				Arguments.Add("-g");
 				if (Options.HasFlag(ClangToolChainOptions.EnableAddressSanitizer))
 				{
-					Result += " -fsanitize=address";
+					Arguments.Add("-fsanitize=address");
 				}
 				else if (Options.HasFlag(ClangToolChainOptions.EnableThreadSanitizer))
 				{
-					Result += " -fsanitize=thread";
+					Arguments.Add("-fsanitize=thread");
 				}
 				else if (Options.HasFlag(ClangToolChainOptions.EnableUndefinedBehaviorSanitizer))
 				{
-					Result += " -fsanitize=undefined";
+					Arguments.Add("-fsanitize=undefined");
 				}
 			}
 
 			if (LinkEnvironment.bIsBuildingDLL)
 			{
-				Result += " -dynamiclib";
+				Arguments.Add("-dynamiclib");
 			}
 
 			if (LinkEnvironment.Configuration == CppConfiguration.Debug)
 			{
 				// Apple's Clang is not supposed to run the de-duplication pass when linking in debug configs. Xcode adds this flag automatically, we need it as well, otherwise linking would take very long
-				Result += " -Wl,-no_deduplicate";
+				Arguments.Add("-Wl,-no_deduplicate");
 			}
 
 			// Needed to make sure install_name_tool will be able to update paths in Mach-O headers
-			Result += " -headerpad_max_install_names";
-
-			Result += " -lc++";
+			Arguments.Add("-headerpad_max_install_names");
 
-			return Result;
+			Arguments.Add("-lc++");
 		}
 
-		static string GetArchiveArguments_Global(LinkEnvironment LinkEnvironment)
+		void GetArchiveArguments_Global(LinkEnvironment LinkEnvironment, List<string> Arguments)
 		{
-			string Result = "";
-			Result += " -static";
-			return Result;
+			Arguments.Add("-static");
 		}
 
 		private void AppendMacLine(StreamWriter Writer, string Format, params object[] Arg)
@@ -473,7 +467,17 @@ namespace UnrealBuildTool
 			string EngineDisplayVersion = LoadEngineDisplayVersion(true);
 			string VersionArg = LinkEnvironment.bIsBuildingDLL ? " -current_version " + EngineAPIVersion + " -compatibility_version " + EngineDisplayVersion : "";
 
-			string LinkCommand = VersionArg + " " + (bIsBuildingLibrary ? GetArchiveArguments_Global(LinkEnvironment) : GetLinkArguments_Global(LinkEnvironment));
+			List<string> LinkArguments = new();
+			if (bIsBuildingLibrary)
+			{
+				GetArchiveArguments_Global(LinkEnvironment, LinkArguments);
+			}
+			else
+			{
+				GetLinkArguments_Global(LinkEnvironment, LinkArguments);
+			}
+
+			string LinkCommand = VersionArg + " " + string.Join(' ', LinkArguments);
 
 			// Tell the action that we're building an import library here and it should conditionally be
 			// ignored as a prerequisite for other actions
