PR #9330: Add build information to Rider Project Model (Contributed by [REDACTED_PERSON])

diff --git a/Engine/Source/Programs/UnrealBuildTool/ProjectFiles/Rider/RiderProjectFile.cs b/Engine/Source/Programs/UnrealBuildTool/ProjectFiles/Rider/RiderProjectFile.cs
--- a/Engine/Source/Programs/UnrealBuildTool/ProjectFiles/Rider/RiderProjectFile.cs
+++ b/Engine/Source/Programs/UnrealBuildTool/ProjectFiles/Rider/RiderProjectFile.cs
@@ -5,6 +5,7 @@
 using System.Diagnostics;
 using System.IO;
 using System.Linq;
+using System.Text;
 using System.Text.RegularExpressions;
 using EpicGames.Core;
 using UnrealBuildBase;
@@ -12,6 +13,20 @@
 
 namespace UnrealBuildTool
 {
+	internal class TargetEntry
+	{
+		public TargetEntry(FileReference OutputFile, UEBuildTarget BuildTarget, bool bBuildByDefault)
+		{
+			this.OutputFile = OutputFile;
+			this.BuildTarget = BuildTarget;
+			this.bBuildByDefault = bBuildByDefault;
+		}
+		
+		public readonly FileReference OutputFile;
+		public readonly UEBuildTarget BuildTarget;
+		public readonly bool bBuildByDefault;
+	}
+
 	internal class RiderProjectFile : ProjectFile
 	{
 		private static readonly XcrunRunner AppleHelper = new XcrunRunner();
@@ -56,7 +71,7 @@
 		{
 			string ProjectName = ProjectFilePath.GetFileNameWithoutAnyExtensions();
 			DirectoryReference ProjectRootFolder = RootPath;
-			List<Tuple<FileReference, UEBuildTarget>> FileToTarget = new List<Tuple<FileReference, UEBuildTarget>>();
+			List<TargetEntry> FileToTarget = new List<TargetEntry>();
 			foreach (UnrealTargetPlatform Platform in InPlatforms)
 			{
 				if (!IsPlatformInHostGroup(Platform))
@@ -84,6 +99,8 @@
 							continue;
 						}
 						
+						bool bBuildByDefault = ShouldBuildByDefaultForSolutionTargets && ProjectTarget.SupportedPlatforms.Contains(Platform);
+						
 						DirectoryReference ConfigurationFolder = DirectoryReference.Combine(ProjectRootFolder, Platform.ToString(), Configuration.ToString());
 
 						DirectoryReference TargetFolder =
@@ -99,7 +116,7 @@
 							UEBuildTarget BuildTarget = UEBuildTarget.Create(TargetDesc, false, false, false, Logger);
 						
 							FileReference OutputFile = FileReference.Combine(TargetFolder, $"{ProjectName}.json");
-							FileToTarget.Add(Tuple.Create(OutputFile, BuildTarget));
+							FileToTarget.Add(new TargetEntry(OutputFile, BuildTarget, bBuildByDefault));
 						}
 						catch(Exception Ex)
 						{
@@ -109,17 +126,18 @@
 					}
 				}
 			}
-			foreach (Tuple<FileReference,UEBuildTarget> tuple in FileToTarget)
+			foreach (TargetEntry TargetEntry in FileToTarget)
 			{
 				try
 				{
-					CurrentTarget = tuple.Item2;
+					CurrentTarget = TargetEntry.BuildTarget;
 					CurrentTarget.PreBuildSetup(Logger);
-					SerializeTarget(tuple.Item1, CurrentTarget, Minimize, Logger);
+					SerializeTarget(TargetEntry.OutputFile, CurrentTarget, PlatformProjectGenerators, Minimize, TargetEntry.bBuildByDefault, Logger);
 				}
 				catch (Exception Ex)
 				{
-					Logger.LogWarning("Exception while generating include data for Target:{Target}, Platform: {Platform}, Configuration: {Configuration}", tuple.Item2.AppName, tuple.Item2.Platform.ToString(), tuple.Item2.Configuration.ToString());
+					Logger.LogWarning("Exception while generating include data for Target:{Target}, Platform: {Platform}, Configuration: {Configuration}",
+						TargetEntry.BuildTarget.AppName, TargetEntry.BuildTarget.Platform.ToString(), TargetEntry.BuildTarget.Configuration.ToString());
 					Logger.LogWarning("{Ex}", Ex.ToString());
 				}
 			}
@@ -142,12 +160,12 @@
 			return false;
 		}
 
-		private void SerializeTarget(FileReference OutputFile, UEBuildTarget BuildTarget, JsonWriterStyle Minimize, ILogger Logger)
+		private void SerializeTarget(FileReference OutputFile, UEBuildTarget BuildTarget, PlatformProjectGeneratorCollection PlatformProjectGenerators, JsonWriterStyle Minimize, bool bBuildByDefault, ILogger Logger)
 		{
 			DirectoryReference.CreateDirectory(OutputFile.Directory);
 			using (JsonWriter Writer = new JsonWriter(OutputFile, Minimize))
 			{
-				ExportTarget(BuildTarget, Writer, Logger);
+				ExportTarget(BuildTarget, Writer, PlatformProjectGenerators, bBuildByDefault, Logger);
 			}
 		}
 
@@ -156,8 +174,10 @@
 		/// </summary>
 		/// <param name="Target"></param>
 		/// <param name="Writer">Writer for the array data</param>
+		/// <param name="PlatformProjectGenerators"></param>
+		/// <param name="bBuildByDefault"></param>
 		/// <param name="Logger">Logger for output</param>
-		private void ExportTarget(UEBuildTarget Target, JsonWriter Writer, ILogger Logger)
+		private void ExportTarget(UEBuildTarget Target, JsonWriter Writer, PlatformProjectGeneratorCollection PlatformProjectGenerators, bool bBuildByDefault, ILogger Logger)
 		{
 			Writer.WriteObjectStart();
 
@@ -170,7 +190,7 @@
 				Writer.WriteValue("ProjectFile", Target.ProjectFile.FullName);
 			}
 			
-			ExportEnvironmentToJson(Target, Writer, Logger);
+			ExportEnvironmentToJson(Target, Writer, PlatformProjectGenerators, bBuildByDefault, Logger);
 			
 			if(Target.Binaries.Any())
 			{
@@ -460,8 +480,10 @@
 		/// </summary>
 		/// <param name="Target"></param>
 		/// <param name="Writer"></param>
+		/// <param name="PlatformProjectGenerators"></param>
+		/// <param name="bBuildByDefault"></param>
 		/// <param name="Logger"></param>
-		private void ExportEnvironmentToJson(UEBuildTarget Target, JsonWriter Writer, ILogger Logger)
+		private void ExportEnvironmentToJson(UEBuildTarget Target, JsonWriter Writer, PlatformProjectGeneratorCollection PlatformProjectGenerators, bool bBuildByDefault, ILogger Logger)
 		{
 			CppCompileEnvironment GlobalCompileEnvironment = Target.CreateCompileEnvironmentForProjectFiles(Logger);
 			
@@ -474,6 +496,8 @@
 			}
 			Writer.WriteObjectEnd();
 			
+			ExportBuildInfo(Writer, Target, PlatformProjectGenerators, bBuildByDefault);
+			
 			Writer.WriteArrayStart("EnvironmentIncludePaths");
 			foreach (DirectoryReference Path in GlobalCompileEnvironment.UserIncludePaths)
 			{
@@ -555,6 +579,109 @@
 			Writer.WriteArrayEnd();
 		}
 
+		private void ExportBuildInfo(JsonWriter Writer, UEBuildTarget Target, PlatformProjectGeneratorCollection PlatformProjectGenerators, bool bBuildByDefault)
+		{
+			if (IsStubProject) return;
+			
+			Writer.WriteObjectStart("BuildInfo");
+			UnrealTargetPlatform HostPlatform = BuildHostPlatform.Current.Platform;
+			
+			ProjectTarget ProjectTarget = ProjectTargets.OfType<ProjectTarget>().Single(It => Target.TargetRulesFile == It.TargetFilePath);
+			UnrealTargetPlatform Platform = Target.Platform;
+			UnrealTargetConfiguration Configuration = Target.Configuration;
+
+			string UProjectPath = "";
+			if (IsForeignProject)
+			{
+				UProjectPath = String.Format("\"{0}\"", ProjectTarget.UnrealProjectFilePath!.FullName);
+			}
+			
+			Writer.WriteValue("bBuildByDefault", bBuildByDefault);
+			
+			if (HostPlatform.IsInGroup(UnrealPlatformGroup.Windows))
+			{
+				PlatformProjectGenerator? ProjGenerator = PlatformProjectGenerators.GetPlatformProjectGenerator(Platform, true);
+				VCProjectFile.BuildCommandBuilder BuildCommandBuilder = new VCProjectFile.BuildCommandBuilder(Configuration, Platform, ProjectTarget, UProjectPath)
+				{
+					ProjectGenerator = ProjGenerator,
+					bIsForeignProject = IsForeignProject
+				};
+				
+				string BuildArguments = BuildCommandBuilder.GetBuildArguments();
+				WriteCommand(Writer, "BuildCmd", EscapePath(BuildCommandBuilder.BuildScript.FullName), BuildArguments);
+				WriteCommand(Writer, "RebuildCmd", EscapePath(BuildCommandBuilder.RebuildScript.FullName), BuildArguments);
+				WriteCommand(Writer, "CleanCmd", EscapePath(BuildCommandBuilder.CleanScript.FullName), BuildArguments);
+			}
+			else
+			{
+				string BuildScript = GetBuildScript(HostPlatform);
+				string BuildArguments = GetBuildArguments(HostPlatform, ProjectTarget, Target, UProjectPath, false);
+				string CleanArguments = GetBuildArguments(HostPlatform, ProjectTarget, Target, UProjectPath, true);
+				WriteCommand(Writer, "BuildCmd", BuildScript, BuildArguments);
+				WriteCommand(Writer, "RebuildCmd", "", "");
+				WriteCommand(Writer, "CleanCmd", BuildScript, CleanArguments);
+			}
+			
+			UEBuildBinary MainBinary = Target.Binaries[0];
+			Writer.WriteValue("Output", MainBinary.OutputFilePath.FullName);
+			Writer.WriteObjectEnd();
+		}
+
+		private string GetBuildScript(UnrealTargetPlatform HostPlatform)
+		{
+			DirectoryReference BatchFilesDirectory = DirectoryReference.Combine(Unreal.EngineDirectory, "Build", "BatchFiles");
+			string ScriptExtension = HostPlatform.IsInGroup(UnrealPlatformGroup.Windows) ? ".bat" : ".sh";
+			if (HostPlatform.IsInGroup(UnrealPlatformGroup.Linux))
+			{
+				BatchFilesDirectory = DirectoryReference.Combine(BatchFilesDirectory, "Linux");
+			}
+			else if (HostPlatform == UnrealTargetPlatform.Mac)
+			{
+				BatchFilesDirectory = DirectoryReference.Combine(BatchFilesDirectory, "Mac");
+			}
+			
+			return EscapePath(FileReference.Combine(BatchFilesDirectory, "Build" + ScriptExtension).FullName);
+		}
+
+		private string GetBuildArguments(UnrealTargetPlatform HostPlatform, ProjectTarget ProjectTarget, UEBuildTarget Target, string UProjectPath, bool bIsClean)
+		{
+			UnrealTargetConfiguration Configuration = Target.Configuration;
+			UnrealTargetPlatform Platform = Target.Platform;
+			string TargetName = ProjectTarget.TargetFilePath.GetFileNameWithoutAnyExtensions();
+			StringBuilder BuildArguments = new StringBuilder();
+			BuildArguments.AppendFormat("{0} {1} {2}", TargetName, Platform.ToString(), Configuration.ToString());
+
+			if (IsForeignProject)
+			{
+				BuildArguments.AppendFormat(" -Project={0}", UProjectPath);
+			}
+
+			if (Target.TargetType == TargetType.Editor)
+			{
+				BuildArguments.Append(" -buildscw");
+			}
+
+			if (bIsClean)
+			{
+				BuildArguments.Append(" -clean");
+			}
+			else if (HostPlatform.IsInGroup(UnrealPlatformGroup.Apple) &&
+			         (Platform == UnrealTargetPlatform.TVOS || Platform == UnrealTargetPlatform.IOS))
+			{
+				BuildArguments.Append(" -deploy");
+			}
+
+			return BuildArguments.ToString();
+		}
+
+		private static void WriteCommand(JsonWriter Writer, string CommandName, string Command, string Arguments)
+		{
+			Writer.WriteObjectStart(CommandName);
+			Writer.WriteValue("Command", Command);
+			Writer.WriteValue("Args", Arguments);
+			Writer.WriteObjectEnd();
+		}
+
 		private static void WriteField(string ModuleOrTargetName, JsonWriter Writer, Tuple<string, object?> Field, ILogger Logger)
 		{
 			if (Field.Item2 == null) return;
diff --git a/Engine/Source/Programs/UnrealBuildTool/ProjectFiles/VisualStudio/VCProject.cs b/Engine/Source/Programs/UnrealBuildTool/ProjectFiles/VisualStudio/VCProject.cs
--- a/Engine/Source/Programs/UnrealBuildTool/ProjectFiles/VisualStudio/VCProject.cs
+++ b/Engine/Source/Programs/UnrealBuildTool/ProjectFiles/VisualStudio/VCProject.cs
@@ -1591,6 +1591,141 @@ namespace UnrealBuildTool
 			VCProjectFileContent.AppendLine("  </PropertyGroup>");
 		}
 
+		// Helper class to generate NMake build commands and arguments
+		public class BuildCommandBuilder
+		{
+			public bool bEditorDependsOnShaderCompileWorker = true;
+			public bool bBuildLiveCodingConsole;
+			public bool bAddFastPDBToProjects;
+
+			public bool bIsForeignProject;
+			public bool bUsePrecompiled;
+			public bool bIsTestMode;
+			public bool bIsFromMSBuild;
+
+			public PlatformProjectGenerator? ProjectGenerator;
+			
+			public FileReference BuildScript { get; }
+			public FileReference RebuildScript { get; }
+			public FileReference CleanScript { get; }
+			
+			private readonly string? BuildToolOverride;
+
+			private readonly string UProjectPath;
+
+			private readonly UnrealTargetConfiguration Configuration;
+			private readonly UnrealTargetPlatform Platform;
+			private readonly ProjectTarget ProjectTarget;
+
+			public BuildCommandBuilder(UnrealTargetConfiguration InConfiguration, UnrealTargetPlatform InPlatform,
+				ProjectTarget InProjectTarget, string InUProjectPath, string? InBuildToolOverride = null)
+			{
+				Configuration = InConfiguration;
+				Platform = InPlatform;
+				ProjectTarget = InProjectTarget;
+				UProjectPath = InUProjectPath;
+				BuildToolOverride = InBuildToolOverride;
+				
+				DirectoryReference BatchFilesDirectory = DirectoryReference.Combine(Unreal.EngineDirectory, "Build", "BatchFiles");
+				BuildScript = FileReference.Combine(BatchFilesDirectory, "Build.bat");
+				RebuildScript = FileReference.Combine(BatchFilesDirectory, "Rebuild.bat");
+				CleanScript = FileReference.Combine(BatchFilesDirectory, "Clean.bat");
+			}
+			
+			public string GetBuildArguments()
+			{
+				TargetRules TargetRulesObject = ProjectTarget.TargetRules!;
+				string TargetName = ProjectTarget.TargetFilePath.GetFileNameWithoutAnyExtensions();
+
+				StringBuilder BuildArguments = new StringBuilder();
+				
+				BuildArguments.AppendFormat("{0} {1} {2}", TargetName, Platform.ToString(), Configuration.ToString());
+				if (bIsForeignProject)
+				{
+					BuildArguments.AppendFormat(" -Project={0}", UProjectPath);
+				}
+
+				if (bIsTestMode)
+				{
+					BuildArguments.Append(" -Mode=Test");
+				}
+
+				List<string> ExtraTargets = new List<string>();
+				if (!bUsePrecompiled)
+				{
+					if (TargetRulesObject.Type == TargetType.Editor && bEditorDependsOnShaderCompileWorker && !Unreal.IsEngineInstalled())
+					{
+						ExtraTargets.Add("ShaderCompileWorker Win64 Development");
+					}
+					if (TargetRulesObject.bWithLiveCoding && bBuildLiveCodingConsole && !Unreal.IsEngineInstalled() && TargetRulesObject.Name != "LiveCodingConsole")
+					{
+						ExtraTargets.Add(TargetRulesObject.bUseDebugLiveCodingConsole ? "LiveCodingConsole Win64 Debug" : "LiveCodingConsole Win64 Development");
+					}
+				}
+
+				if (ExtraTargets.Count > 0)
+				{
+					BuildArguments.Replace("\"", "\\\"");
+					BuildArguments.Insert(0, "-Target=\"");
+					BuildArguments.Append("\"");
+					foreach (string ExtraTarget in ExtraTargets)
+					{
+						BuildArguments.AppendFormat(" -Target=\"{0} -Quiet\"", ExtraTarget);
+					}
+				}
+
+				if (bUsePrecompiled)
+				{
+					BuildArguments.Append(" -UsePrecompiled");
+				}
+
+				// Always wait for the mutex between UBT invocations, so that building the whole solution doesn't fail.
+				BuildArguments.Append(" -WaitMutex");
+
+				if (bIsFromMSBuild)
+				{
+					BuildArguments.Append(" -FromMsBuild");
+				}
+
+				if (bAddFastPDBToProjects)
+				{
+					// Pass Fast PDB option to make use of Visual Studio's /DEBUG:FASTLINK option
+					BuildArguments.Append(" -FastPDB");
+				}
+
+				if (BuildToolOverride != null)
+				{
+					BuildArguments.AppendFormat(" {0}", BuildToolOverride);
+				}
+
+				if (ProjectGenerator != null)
+				{
+					BuildArguments.Append(ProjectGenerator.GetExtraBuildArguments(Platform, Configuration));
+				}
+
+				return BuildArguments.ToString();
+			}
+		}
+
+		private BuildCommandBuilder CreateArgumentsBuilder(ProjectConfigAndTargetCombination Combination, string UProjectPath, PlatformProjectGenerator? ProjGenerator)
+		{
+			BuildCommandBuilder Builder = new BuildCommandBuilder(Combination.Configuration, Combination.Platform!.Value,
+				Combination.ProjectTarget!, UProjectPath, BuildToolOverride)
+			{
+				ProjectGenerator = ProjGenerator,
+				bIsTestMode = Combination.TestMode,
+				bEditorDependsOnShaderCompileWorker = Settings.bEditorDependsOnShaderCompileWorker,
+				bBuildLiveCodingConsole = Settings.bBuildLiveCodingConsole,
+				bAddFastPDBToProjects = Settings.bAddFastPDBToProjects,
+				bIsForeignProject = IsForeignProject,
+				bUsePrecompiled = bUsePrecompiled,
+				// Always include a flag to format log messages for MSBuild
+				bIsFromMSBuild = true
+			};
+
+			return Builder;
+		}
+		
 		// Anonymous function that writes project configuration data
 		private void WriteConfiguration(string ProjectName, ProjectConfigAndTargetCombination Combination, StringBuilder VCProjectFileContent, PlatformProjectGeneratorCollection PlatformProjectGenerators, StringBuilder? VCUserFileContent, ILogger Logger)
 		{
@@ -1730,76 +1865,13 @@ namespace UnrealBuildTool
 					//	..\..\Build\BatchFiles\Build.bat <TARGETNAME> <PLATFORM> <CONFIGURATION>
 					//	ie ..\..\Build\BatchFiles\Build.bat BlankProgram Win64 Debug
 
-					StringBuilder BuildArguments = new StringBuilder();
-
-					BuildArguments.AppendFormat("{0} {1} {2}", TargetName, UBTPlatformName, UBTConfigurationName);
-					if (IsForeignProject)
-					{
-						BuildArguments.AppendFormat(" -Project={0}", UProjectPath);
-					}
-
-					if (Combination.TestMode)
-					{
-						BuildArguments.Append(" -Mode=Test");
-					}
-
-					List<string> ExtraTargets = new List<string>();
-					if (!bUsePrecompiled)
-					{
-						if (TargetRulesObject.Type == TargetType.Editor && Settings.bEditorDependsOnShaderCompileWorker && !Unreal.IsEngineInstalled())
-						{
-							ExtraTargets.Add("ShaderCompileWorker Win64 Development");
-						}
-						if (TargetRulesObject.bWithLiveCoding && Settings.bBuildLiveCodingConsole && !Unreal.IsEngineInstalled() && TargetRulesObject.Name != "LiveCodingConsole")
-						{
-							ExtraTargets.Add(TargetRulesObject.bUseDebugLiveCodingConsole ? "LiveCodingConsole Win64 Debug" : "LiveCodingConsole Win64 Development");
-						}
-					}
-
-					if (ExtraTargets.Count > 0)
-					{
-						BuildArguments.Replace("\"", "\\\"");
-						BuildArguments.Insert(0, "-Target=\"");
-						BuildArguments.Append("\"");
-						foreach (string ExtraTarget in ExtraTargets)
-						{
-							BuildArguments.AppendFormat(" -Target=\"{0} -Quiet\"", ExtraTarget);
-						}
-					}
-
-					if (bUsePrecompiled)
-					{
-						BuildArguments.Append(" -UsePrecompiled");
-					}
-
-					// Always wait for the mutex between UBT invocations, so that building the whole solution doesn't fail.
-					BuildArguments.Append(" -WaitMutex");
-
-					// Always include a flag to format log messages for MSBuild
-					BuildArguments.Append(" -FromMsBuild");
-
-					if (Settings.bAddFastPDBToProjects)
-					{
-						// Pass Fast PDB option to make use of Visual Studio's /DEBUG:FASTLINK option
-						BuildArguments.Append(" -FastPDB");
-					}
-
-					DirectoryReference BatchFilesDirectory = DirectoryReference.Combine(Unreal.EngineDirectory, "Build", "BatchFiles");
-
-					if (BuildToolOverride != null)
-					{
-						BuildArguments.AppendFormat(" {0}", BuildToolOverride);
-					}
-
-					if (ProjGenerator != null)
-					{
-						BuildArguments.Append(ProjGenerator.GetExtraBuildArguments(Platform, Configuration));
-					}
+					BuildCommandBuilder Builder = CreateArgumentsBuilder(Combination, UProjectPath, ProjGenerator);
+					string BuildArguments = Builder.GetBuildArguments();
 
 					// NMake Build command line
-					VCProjectFileContent.AppendLine("    <NMakeBuildCommandLine>{0} {1}</NMakeBuildCommandLine>", EscapePath(NormalizeProjectPath(FileReference.Combine(BatchFilesDirectory, "Build.bat"))), BuildArguments.ToString());
-					VCProjectFileContent.AppendLine("    <NMakeReBuildCommandLine>{0} {1}</NMakeReBuildCommandLine>", EscapePath(NormalizeProjectPath(FileReference.Combine(BatchFilesDirectory, "Rebuild.bat"))), BuildArguments.ToString());
-					VCProjectFileContent.AppendLine("    <NMakeCleanCommandLine>{0} {1}</NMakeCleanCommandLine>", EscapePath(NormalizeProjectPath(FileReference.Combine(BatchFilesDirectory, "Clean.bat"))), BuildArguments.ToString());
+					VCProjectFileContent.AppendLine("    <NMakeBuildCommandLine>{0} {1}</NMakeBuildCommandLine>", EscapePath(NormalizeProjectPath(Builder.BuildScript)), BuildArguments);
+					VCProjectFileContent.AppendLine("    <NMakeReBuildCommandLine>{0} {1}</NMakeReBuildCommandLine>", EscapePath(NormalizeProjectPath(Builder.RebuildScript)), BuildArguments);
+					VCProjectFileContent.AppendLine("    <NMakeCleanCommandLine>{0} {1}</NMakeCleanCommandLine>", EscapePath(NormalizeProjectPath(Builder.CleanScript)), BuildArguments);
 					VCProjectFileContent.AppendLine("    <NMakeOutput>{0}</NMakeOutput>", NormalizeProjectPath(NMakePath.FullName));
 					VCProjectFileContent.AppendLine("    <AdditionalOptions>{0}</AdditionalOptions>", GetCppStandardCompileArgument(TargetRulesObject.CppStandard));
 
