- Added support for .xcconfig files, currently for project and targets, but we can make specific config ones (Debug vs Development, etc)

diff --git a/Engine/Source/Programs/UnrealBuildTool/ProjectFiles/Xcode/XcodeProject.cs b/Engine/Source/Programs/UnrealBuildTool/ProjectFiles/Xcode/XcodeProject.cs
--- a/Engine/Source/Programs/UnrealBuildTool/ProjectFiles/Xcode/XcodeProject.cs
+++ b/Engine/Source/Programs/UnrealBuildTool/ProjectFiles/Xcode/XcodeProject.cs
@@ -124,6 +124,33 @@ namespace UnrealBuildTool
 		public string? ConfigurationContents;
 	}
 
+	class XcconfigFile
+	{
+		public string Filename;
+		public string Guid;
+		public StreamWriter Stream;
+
+		public XcconfigFile(FileReference ProjectFile, string NameSuffix)
+		{
+			FileReference File = FileReference.Combine(ProjectFile.Directory, ProjectFile.GetFileNameWithoutAnyExtensions() + NameSuffix + ".xcconfig");
+			Filename = File.GetFileName();
+			Guid = XcodeProjectFileGenerator.MakeXcodeGuid();
+			Stream = new StreamWriter(new FileStream(File.FullName, FileMode.Create, FileAccess.Write), Encoding.ASCII);
+		}
+
+		public void AppendLine(string Line)
+		{
+			Stream.Write(Line + ProjectFileGenerator.NewLine);
+			Stream.Flush();
+		}
+
+		public void Close()
+		{
+			Stream.Close();
+		}
+	}
+
+
 	class XcodeProjectFile : ProjectFile
 	{
 		Dictionary<string, XcodeFileGroup> Groups = new Dictionary<string, XcodeFileGroup>();
@@ -142,6 +169,12 @@ namespace UnrealBuildTool
 			bForDistribution = IsForDistribution;
 			BundleIdentifier = BundleID;
 			AppName = InAppName;
+
+			ProjectXcconfig = new XcconfigFile(InitFilePath, "_Project");
+			TargetXcconfig = new XcconfigFile(InitFilePath, "_Target");
+
+			AllXcconfigFiles.Add(ProjectXcconfig);
+			AllXcconfigFiles.Add(TargetXcconfig);
 		}
 
 		public override string ToString()
@@ -169,6 +202,17 @@ namespace UnrealBuildTool
 		/// </summary>
 		string[] SupportedIOSArchitectures = { "arm64" };
 
+		/// <summary>
+		/// Xcconfig file information - this is global per project, so no need to pass it around everywhere
+		/// </summary>
+		private XcconfigFile ProjectXcconfig;
+		private XcconfigFile TargetXcconfig;
+
+		/// <summary>
+		/// All Xcconfigs that can be looped over (to add to project, and close, etc)
+		/// </summary>
+		private List<XcconfigFile> AllXcconfigFiles = new List<XcconfigFile>();
+
 		/// <summary>
 		/// Gets Xcode file category based on its extension
 		/// </summary>
@@ -356,6 +400,10 @@ namespace UnrealBuildTool
 				}
 			}
 
+			foreach (XcconfigFile Xcconfig in AllXcconfigFiles)
+			{
+				PBXFileReferenceSection.Append($"\t\t{Xcconfig.Guid} /* {Xcconfig.Filename} */ = {{isa = PBXFileReference; explicitFileType = test.xcconfig; path = {Xcconfig.Filename}; sourceTree = \"<group>\"; }};" + ProjectFileGenerator.NewLine);
+			}
 			PBXFileReferenceSection.Append(string.Format("\t\t{0} /* {1} */ = {{isa = PBXFileReference; explicitFileType = {2}; path = {1}; sourceTree = BUILT_PRODUCTS_DIR; }};" + ProjectFileGenerator.NewLine, TargetAppGuid, TargetName, bIsAppBundle ? "wrapper.application" : "\"compiled.mach-o.executable\""));
 		}
 
@@ -590,7 +638,10 @@ namespace UnrealBuildTool
 			Content.Append(string.Format("\t\t{0} = {{{1}", MainGroupGuid, ProjectFileGenerator.NewLine));
 			Content.Append("\t\t\tisa = PBXGroup;" + ProjectFileGenerator.NewLine);
 			Content.Append("\t\t\tchildren = (" + ProjectFileGenerator.NewLine);
-
+			foreach (XcconfigFile Xcconfig in AllXcconfigFiles)
+			{
+				Content.Append($"\t\t\t\t{Xcconfig.Guid} /* {Xcconfig.Filename} */,");
+			}
 			foreach (XcodeFileGroup Group in RootGroup.Children.Values)
 			{
 				Content.Append(string.Format("\t\t\t\t{0} /* {1} */,{2}", Group.GroupGuid, Group.GroupName, ProjectFileGenerator.NewLine));
@@ -922,6 +973,7 @@ namespace UnrealBuildTool
 		{
 			Content.Append("\t\t" + ConfigGuid + " /* \"" + ConfigName + "\" */ = {" + ProjectFileGenerator.NewLine);
 			Content.Append("\t\t\tisa = XCBuildConfiguration;" + ProjectFileGenerator.NewLine);
+			Content.Append($"\t\t\tbaseConfigurationReference = {ProjectXcconfig.Guid} /* \"{ProjectXcconfig.Filename}\" */;" + ProjectFileGenerator.NewLine);
 			Content.Append("\t\t\tbuildSettings = {" + ProjectFileGenerator.NewLine);
 
 			Content.Append("\t\t\t\tGCC_PREPROCESSOR_DEFINITIONS = (" + ProjectFileGenerator.NewLine);
@@ -961,6 +1013,7 @@ namespace UnrealBuildTool
 			Content.Append("\t\t\t};" + ProjectFileGenerator.NewLine);
 			Content.Append("\t\t\tname = \"" + ConfigName + "\";" + ProjectFileGenerator.NewLine);
 			Content.Append("\t\t};" + ProjectFileGenerator.NewLine);
+
 		}
 
 		// cache for the below function
@@ -1418,6 +1471,7 @@ namespace UnrealBuildTool
 
 			Content.Append("\t\t" + ConfigGuid + " /* \"" + Config.DisplayName + "\" */ = {" + ProjectFileGenerator.NewLine);
 			Content.Append("\t\t\tisa = XCBuildConfiguration;" + ProjectFileGenerator.NewLine);
+			Content.Append($"\t\t\tbaseConfigurationReference = {TargetXcconfig.Guid} /* \"{TargetXcconfig.Filename}\" */;" + ProjectFileGenerator.NewLine);
 			Content.Append("\t\t\tbuildSettings = {" + ProjectFileGenerator.NewLine);
 			if (bMacOnly)
 			{
@@ -1850,6 +1904,20 @@ namespace UnrealBuildTool
 			File.WriteAllText(ManagementFile.FullName, Content.ToString(), new UTF8Encoding());
 		}
 
+		private void WriteXcconfigFiles()
+		{
+			// settings we write to xcconfig now instead of xcode project file
+			ProjectXcconfig.AppendLine("GENERATE_INFOPLIST_FILE[sdk=macosx*] = YES // Xcode 14 code-signing fix for development");
+
+			TargetXcconfig.AppendLine("// Will add more here later");
+
+			foreach (XcconfigFile Xcconfig in AllXcconfigFiles)
+			{
+				Xcconfig.Close();
+			}
+		}
+
+
 		public static IEnumerable<UnrealTargetPlatform> GetSupportedPlatforms()
 		{
 			List<UnrealTargetPlatform> SupportedPlatforms = new List<UnrealTargetPlatform>();
@@ -2023,6 +2091,9 @@ namespace UnrealBuildTool
 			ProjectFileContent.Append("\trootObject = " + ProjectGuid + " /* Project object */;" + ProjectFileGenerator.NewLine);
 			ProjectFileContent.Append("}" + ProjectFileGenerator.NewLine);
 
+			// write out .xcconfig file(s) for settings we are moving out of the xcode project file
+			WriteXcconfigFiles();
+
 			if (bSuccess)
 			{
 				FileReference PBXProjFilePath = ProjectFilePath + "/project.pbxproj";
