UnrealBuildTool: Add [RequiresUniqueBuildEnvironment] attribute to Target warning\error settings that affect command line options and will cause builds to be invalidated

#rnx

diff --git a/Engine/Source/Programs/UnrealBuildTool/Configuration/TargetRules.cs b/Engine/Source/Programs/UnrealBuildTool/Configuration/TargetRules.cs
--- a/Engine/Source/Programs/UnrealBuildTool/Configuration/TargetRules.cs
+++ b/Engine/Source/Programs/UnrealBuildTool/Configuration/TargetRules.cs
@@ -1288,6 +1288,7 @@ namespace UnrealBuildTool
 		/// <summary>
 		/// Default treatment of uncategorized warnings
 		/// </summary>
+		[RequiresUniqueBuildEnvironment]
 		public WarningLevel DefaultWarningLevel
 		{
 			get => (DefaultWarningLevelPrivate == WarningLevel.Default)? (bWarningsAsErrors ? WarningLevel.Error : WarningLevel.Warning) : DefaultWarningLevelPrivate;
@@ -1315,6 +1316,7 @@ namespace UnrealBuildTool
 		/// Forces shadow variable warnings to be treated as errors on platforms that support it.
 		/// </summary>
 		[CommandLine("-ShadowVariableErrors", Value = nameof(WarningLevel.Error))]
+		[RequiresUniqueBuildEnvironment]
 		public WarningLevel ShadowVariableWarningLevel = WarningLevel.Warning;
 
 		/// <summary>
@@ -1322,18 +1324,21 @@ namespace UnrealBuildTool
 		/// </summary>
 		[XmlConfigFile(Category = "BuildConfiguration")]
 		[CommandLine("-WarningsAsErrors")]
+		[RequiresUniqueBuildEnvironment]
 		public bool bWarningsAsErrors = false;
 
 		/// <summary>
 		/// Indicates what warning/error level to treat unsafe type casts as on platforms that support it (e.g., double->float or int64->int32)
 		/// </summary>
 		[XmlConfigFile(Category = "BuildConfiguration")]
+		[RequiresUniqueBuildEnvironment]
 		public WarningLevel UnsafeTypeCastWarningLevel = WarningLevel.Off;
 
 		/// <summary>
 		/// Forces the use of undefined identifiers in conditional expressions to be treated as errors.
 		/// </summary>
 		[XmlConfigFile(Category = "BuildConfiguration")]
+		[RequiresUniqueBuildEnvironment]
 		public bool bUndefinedIdentifierErrors = true;
 
 		/// <summary>
