UnrealBuildTool: Fix -fpch-validate-input-files-content check

#rnx

diff --git a/Engine/Source/Programs/UnrealBuildTool/ToolChain/AppleToolChain.cs b/Engine/Source/Programs/UnrealBuildTool/ToolChain/AppleToolChain.cs
--- a/Engine/Source/Programs/UnrealBuildTool/ToolChain/AppleToolChain.cs
+++ b/Engine/Source/Programs/UnrealBuildTool/ToolChain/AppleToolChain.cs
@@ -221,10 +221,6 @@ namespace UnrealBuildTool
 			GetCppStandardCompileArgument(CompileEnvironment, Arguments);
 			Arguments.Add("-stdlib=libc++");
 
-			if (CompilerVersionGreaterOrEqual(11, 0, 0))
-			{
-				Arguments.Add("-fpch-validate-input-files-content");
-			}
 			if (CompilerVersionGreaterOrEqual(13, 0, 0)) // Note this is supported for >=11 on other clang platforms
 			{
 				Arguments.Add("-fpch-instantiate-templates");
diff --git a/Engine/Source/Programs/UnrealBuildTool/ToolChain/ClangToolChain.cs b/Engine/Source/Programs/UnrealBuildTool/ToolChain/ClangToolChain.cs
--- a/Engine/Source/Programs/UnrealBuildTool/ToolChain/ClangToolChain.cs
+++ b/Engine/Source/Programs/UnrealBuildTool/ToolChain/ClangToolChain.cs
@@ -281,18 +281,21 @@ namespace UnrealBuildTool
 
 		protected virtual void GetCppStandardCompileArgument(CppCompileEnvironment CompileEnvironment, List<string> Arguments)
 		{
+			// https://clang.llvm.org/cxx_status.html
 			switch (CompileEnvironment.CppStandard)
 			{
 				case CppStandardVersion.Cpp14:
 					Arguments.Add("-std=c++14");
 					break;
-				case CppStandardVersion.Latest:
 				case CppStandardVersion.Cpp17:
 					Arguments.Add("-std=c++17");
 					break;
 				case CppStandardVersion.Cpp20:
 					Arguments.Add("-std=c++20");
 					break;
+				case CppStandardVersion.Latest:
+					Arguments.Add(CompilerVersionLessThan(13, 0, 0) ? "-std=c++20" : "-std=c++2b");
+					break;
 				default:
 					throw new BuildException($"Unsupported C++ standard type set: {CompileEnvironment.CppStandard}");
 			}
@@ -305,16 +308,17 @@ namespace UnrealBuildTool
 					Arguments.Add("-Wno-coroutine-missing-unhandled-exception");
 				}
 			}
-		}
 
-		protected virtual void GetCompileArguments_CPP(CppCompileEnvironment CompileEnvironment, List<string> Arguments)
-		{
-			Arguments.Add("-x c++");
-			if (CompileEnvironment.PrecompiledHeaderAction == PrecompiledHeaderAction.Include && CompilerVersionGreaterOrEqual(11, 0, 0))
+			if (CompileEnvironment.PrecompiledHeaderAction != PrecompiledHeaderAction.None && CompilerVersionGreaterOrEqual(11, 0, 0))
 			{
 				// Validate PCH inputs by content if mtime check fails
 				Arguments.Add("-fpch-validate-input-files-content");
 			}
+		}
+
+		protected virtual void GetCompileArguments_CPP(CppCompileEnvironment CompileEnvironment, List<string> Arguments)
+		{
+			Arguments.Add("-x c++");
 			GetCppStandardCompileArgument(CompileEnvironment, Arguments);
 		}
 
@@ -344,8 +348,6 @@ namespace UnrealBuildTool
 			Arguments.Add("-x c++-header");
 			if (CompilerVersionGreaterOrEqual(11, 0, 0))
 			{
-				// Validate PCH inputs by content if mtime check fails
-				Arguments.Add("-fpch-validate-input-files-content");
 				Arguments.Add("-fpch-instantiate-templates");
 			}
 			GetCppStandardCompileArgument(CompileEnvironment, Arguments);
