UnrealBuildTool: Update CppStandardVersion.Cpp20 for msvc to pass /std:c++20 instead of /std:c++latest as it was feature-complete* with 16.11

#rnx

*https://docs.microsoft.com/en-us/cpp/overview/visual-cpp-language-conformance?view=msvc-170#note_20abi

diff --git a/Engine/Source/Programs/UnrealBuildTool/Platform/Windows/VCToolChain.cs b/Engine/Source/Programs/UnrealBuildTool/Platform/Windows/VCToolChain.cs
--- a/Engine/Source/Programs/UnrealBuildTool/Platform/Windows/VCToolChain.cs
+++ b/Engine/Source/Programs/UnrealBuildTool/Platform/Windows/VCToolChain.cs
@@ -819,6 +819,17 @@ namespace UnrealBuildTool
 					Arguments.Add("/std:c++17");
 					break;
 				case CppStandardVersion.Cpp20:
+					Arguments.Add("/std:c++20");
+					
+					// warning C5054: operator ___: deprecated between enumerations of different types
+					// re: http://www.open-std.org/jtc1/sc22/wg21/docs/papers/2018/p1120r0.html
+						
+					// It seems unclear whether the deprecation will be enacted in C++23 or not
+					// e.g. http://www.open-std.org/jtc1/sc22/wg21/docs/papers/2020/p2139r2.html
+					// Until the path forward is clearer, it seems reasonable to leave things as they are.
+					Arguments.Add("/wd5054");
+
+					break;
 				case CppStandardVersion.Latest:
 					Arguments.Add("/std:c++latest");
 						
@@ -831,11 +842,6 @@ namespace UnrealBuildTool
 					Arguments.Add("/wd5054");
 						
 					break;
-				// Will be added when MSVC is feature-complete.
-				// https://docs.microsoft.com/en-us/cpp/build/reference/std-specify-language-standard-version?view=msvc-160
-				// case CppStandardVersion.Cpp20:
-				//	Arguments.Add("/std:c++20");
-				//  break;
 				default:
 					throw new BuildException($"Unsupported C++ standard type set: {CompileEnvironment.CppStandard}");
 			}
