Update terminology. Requires UGS update

#rnx

diff --git a/Engine/Plugins/Developer/RiderSourceCodeAccess/Source/RiderSourceCodeAccess/Private/RiderSourceCodeAccessor.cpp b/Engine/Plugins/Developer/RiderSourceCodeAccess/Source/RiderSourceCodeAccess/Private/RiderSourceCodeAccessor.cpp
--- a/Engine/Plugins/Developer/RiderSourceCodeAccess/Source/RiderSourceCodeAccess/Private/RiderSourceCodeAccessor.cpp
+++ b/Engine/Plugins/Developer/RiderSourceCodeAccess/Source/RiderSourceCodeAccess/Private/RiderSourceCodeAccessor.cpp
@@ -329,12 +329,12 @@ void FRiderSourceCodeAccessor::CachePathToSln() const
 
 		if (!FUProjectDictionary(FPaths::RootDir()).IsForeignProject(CachedSolutionPath))
 		{
-			FString MasterProjectName;
-			if (!FFileHelper::LoadFileToString(MasterProjectName, *(FPaths::EngineIntermediateDir() / TEXT("ProjectFiles/MasterProjectName.txt"))))
+			FString PrimaryProjectName;
+			if (!FFileHelper::LoadFileToString(PrimaryProjectName, *(FPaths::EngineIntermediateDir() / TEXT("ProjectFiles/PrimaryProjectName.txt"))))
 			{
-				MasterProjectName = "UE5";
+				PrimaryProjectName = "UE5";
 			}
-			CachedSolutionPath = FPaths::Combine(FPaths::RootDir(), MasterProjectName + TEXT(".sln"));
+			CachedSolutionPath = FPaths::Combine(FPaths::RootDir(), PrimaryProjectName + TEXT(".sln"));
 		}
 		else
 		{
diff --git a/Engine/Plugins/Developer/VisualStudioSourceCodeAccess/Source/VisualStudioSourceCodeAccess/Private/VisualStudioSourceCodeAccessor.cpp b/Engine/Plugins/Developer/VisualStudioSourceCodeAccess/Source/VisualStudioSourceCodeAccess/Private/VisualStudioSourceCodeAccessor.cpp
--- a/Engine/Plugins/Developer/VisualStudioSourceCodeAccess/Source/VisualStudioSourceCodeAccess/Private/VisualStudioSourceCodeAccessor.cpp
+++ b/Engine/Plugins/Developer/VisualStudioSourceCodeAccess/Source/VisualStudioSourceCodeAccess/Private/VisualStudioSourceCodeAccessor.cpp
@@ -1410,12 +1410,12 @@ FString FVisualStudioSourceCodeAccessor::GetSolutionPath() const
 
 			if (!FUProjectDictionary::GetDefault().IsForeignProject(CachedSolutionPath))
 			{
-				FString MasterProjectName;
-				if (!FFileHelper::LoadFileToString(MasterProjectName, *(FPaths::EngineIntermediateDir() / TEXT("ProjectFiles/MasterProjectName.txt"))))
+				FString PrimaryProjectName;
+				if (!FFileHelper::LoadFileToString(PrimaryProjectName, *(FPaths::EngineIntermediateDir() / TEXT("ProjectFiles/PrimaryProjectName.txt"))))
 				{
-					MasterProjectName = "UE5";
+					PrimaryProjectName = "UE5";
 				}
-				CachedSolutionPath = FPaths::Combine(FPaths::RootDir(), MasterProjectName + TEXT(".sln"));
+				CachedSolutionPath = FPaths::Combine(FPaths::RootDir(), PrimaryProjectName + TEXT(".sln"));
 			}
 			else
 			{
diff --git a/Engine/Source/Programs/UnrealBuildTool/ProjectFiles/ProjectFileGenerator.cs b/Engine/Source/Programs/UnrealBuildTool/ProjectFiles/ProjectFileGenerator.cs
--- a/Engine/Source/Programs/UnrealBuildTool/ProjectFiles/ProjectFileGenerator.cs
+++ b/Engine/Source/Programs/UnrealBuildTool/ProjectFiles/ProjectFileGenerator.cs
@@ -818,7 +818,7 @@ namespace UnrealBuildTool
 				}
 
 				// Write out the name of the primary project file, so the runtime knows to use it
-				FileReference PrimaryProjectNameLocation = FileReference.Combine(Unreal.EngineDirectory, "Intermediate", "ProjectFiles", "MasterProjectName.txt");
+				FileReference PrimaryProjectNameLocation = FileReference.Combine(Unreal.EngineDirectory, "Intermediate", "ProjectFiles", "PrimaryProjectName.txt");
 				DirectoryReference.CreateDirectory(PrimaryProjectNameLocation.Directory);
 				FileReference.WriteAllText(PrimaryProjectNameLocation, PrimaryProjectName);
 			}
@@ -842,7 +842,7 @@ namespace UnrealBuildTool
 				}
 
 				// the primary project name is always read from our intermediate directory and not the overriden one for this set of platforms
-				FileReference PrimaryProjectNameLocation = FileReference.Combine(Unreal.EngineDirectory, "Intermediate", "ProjectFiles", "MasterProjectName.txt");
+				FileReference PrimaryProjectNameLocation = FileReference.Combine(Unreal.EngineDirectory, "Intermediate", "ProjectFiles", "PrimaryProjectName.txt");
 				DirectoryReference.CreateDirectory(PrimaryProjectNameLocation.Directory);
 				FileReference.WriteAllText(PrimaryProjectNameLocation, PrimaryProjectName);
 			}
