Added Clang Analyzer support to Android toolchain

diff --git a/Engine/Source/Programs/UnrealBuildTool/Platform/Android/AndroidToolChain.cs b/Engine/Source/Programs/UnrealBuildTool/Platform/Android/AndroidToolChain.cs
--- a/Engine/Source/Programs/UnrealBuildTool/Platform/Android/AndroidToolChain.cs
+++ b/Engine/Source/Programs/UnrealBuildTool/Platform/Android/AndroidToolChain.cs
@@ -793,6 +793,17 @@ namespace UnrealBuildTool
 				Result += " -fstack-protector-strong";  // Emits extra code to check for buffer overflows
 			}
 
+			// Add analysis flags to the argument list.
+			if (StaticAnalyzer == StaticAnalyzer.Default)
+			{
+				List<string> Arguments = new List<string>();
+				GetCompileArguments_Analyze(CompileEnvironment, Arguments);
+				for (int i = 0; i < Arguments.Count; i++)
+				{
+					Result += " " + Arguments[i];
+				}
+			}
+
 			Result += " -fforce-emit-vtables";      // Helps with devirtualization
 			//Result += " -fstrict-vtable-pointers";	// Assumes that vtable pointer will never change, i.e. object of a different type would not be constracted in place of another object. Helps with devirtualization
 
@@ -1395,9 +1406,10 @@ namespace UnrealBuildTool
 					string ResponseArgument = string.Format("@\"{0}\"", ResponseFileName);
 
 					CompileAction.WorkingDirectory = Unreal.EngineSourceDirectory;
-					if(bExecuteCompilerThroughShell)
+					CompileAction.CommandDescription = StaticAnalyzer == StaticAnalyzer.Default ? "Analyze" : "Compile";
+					if (bExecuteCompilerThroughShell)
 					{
-						SetupActionToExecuteCompilerThroughShell(ref CompileAction, ClangPath!, ResponseArgument, "Compile");
+						SetupActionToExecuteCompilerThroughShell(ref CompileAction, ClangPath!, ResponseArgument, CompileAction.CommandDescription);
 					}
 					else
 					{
diff --git a/Engine/Source/Programs/UnrealBuildTool/Platform/Android/UEBuildAndroid.cs b/Engine/Source/Programs/UnrealBuildTool/Platform/Android/UEBuildAndroid.cs
--- a/Engine/Source/Programs/UnrealBuildTool/Platform/Android/UEBuildAndroid.cs
+++ b/Engine/Source/Programs/UnrealBuildTool/Platform/Android/UEBuildAndroid.cs
@@ -150,6 +150,25 @@ namespace UnrealBuildTool
 
 		public override void ValidateTarget(TargetRules Target)
 		{
+			if (!string.IsNullOrWhiteSpace(Environment.GetEnvironmentVariable("CLANG_STATIC_ANALYZER_MODE")))
+			{
+				Target.StaticAnalyzer = StaticAnalyzer.Default;
+				Target.StaticAnalyzerOutputType = (Environment.GetEnvironmentVariable("CLANG_ANALYZER_OUTPUT")?.Contains("html", StringComparison.OrdinalIgnoreCase) == true) ? StaticAnalyzerOutputType.Html : StaticAnalyzerOutputType.Text;
+				Target.StaticAnalyzerMode = string.Equals(Environment.GetEnvironmentVariable("CLANG_STATIC_ANALYZER_MODE"), "shallow") ? StaticAnalyzerMode.Shallow : StaticAnalyzerMode.Deep;
+			}
+			else if (Target.StaticAnalyzer == StaticAnalyzer.Clang)
+			{
+				Target.StaticAnalyzer = StaticAnalyzer.Default;
+			}
+
+			// Disable linking if we're using a static analyzer
+			// Disable PCHs for clang static analyzer.
+			if (Target.StaticAnalyzer == StaticAnalyzer.Default)
+			{
+				Target.bDisableLinking = true;
+				Target.bUsePCHFiles = false;
+			}
+
 			Target.bCompileRecast = true;
 			Target.bCompileISPC = false;
 		}
