UnrealBuildTool: Update project specific makefile generator to build ProjectEditor instead of UnrealEditor

#rnx

diff --git a/Engine/Source/Programs/UnrealBuildTool/ProjectFiles/Make/MakefileGenerator.cs b/Engine/Source/Programs/UnrealBuildTool/ProjectFiles/Make/MakefileGenerator.cs
--- a/Engine/Source/Programs/UnrealBuildTool/ProjectFiles/Make/MakefileGenerator.cs
+++ b/Engine/Source/Programs/UnrealBuildTool/ProjectFiles/Make/MakefileGenerator.cs
@@ -65,11 +65,14 @@ namespace UnrealBuildTool
 
 			string UnrealRootPath = Unreal.RootDirectory.FullName;
 
+			string EditorTarget = "UnrealEditor";
+
 			if (!String.IsNullOrEmpty(GameProjectName))
 			{
 				GameProjectFile = OnlyGameProject!.FullName;
 				MakeGameProjectFile = "GAMEPROJECTFILE =" + GameProjectFile + "\n";
 				ProjectBuildCommand = $"PROJECTBUILD = \"$(UNREALROOTPATH)/Engine/{Unreal.RelativeDotnetDirectory}/dotnet\" \"$(UNREALROOTPATH)/Engine/Binaries/DotNET/UnrealBuildTool/UnrealBuildTool.dll\"\n";
+				EditorTarget = GameProjectName + "Editor";
 			}
 
 			BuildCommand = "BUILD = bash \"$(UNREALROOTPATH)/Engine/Build/BatchFiles/Linux/Build.sh\"\n";
@@ -116,8 +119,8 @@ namespace UnrealBuildTool
 			MakefileContent.Append("\n\n" + BuildCommand + ProjectBuildCommand + "\n" +
 				"all: StandardSet\n\n" +
 				"RequiredTools: CrashReportClient-Linux-Shipping CrashReportClientEditor-Linux-Shipping ShaderCompileWorker UnrealLightmass EpicWebHelper-Linux-Shipping\n\n" +
-				"StandardSet: RequiredTools UnrealFrontend UnrealEditor UnrealInsights\n\n" +
-				"DebugSet: RequiredTools UnrealFrontend-Linux-Debug UnrealEditor-Linux-Debug\n\n"
+				$"StandardSet: RequiredTools UnrealFrontend {EditorTarget} UnrealInsights\n\n" +
+				$"DebugSet: RequiredTools UnrealFrontend-Linux-Debug {EditorTarget}-Linux-Debug\n\n"
 			);
 
 			foreach (ProjectFile Project in GeneratedProjectFiles)
