Modular builds: UBT: Added an option to generate .loadorder manifest file that includes a list of all dynamic modules in the order in which they are meant to be loaded (based on their dependencies).

It's needed on some platforms that don't handle DLLs dependencies out of the box. Generation of this file is integrated into the build metadata generation but is only triggered for platforms that explicitly indicate that it's needed (UEBuildPlatform.RequiresLoadOrderManifest()).

diff --git a/Engine/Source/Programs/UnrealBuildTool/Configuration/UEBuildPlatform.cs b/Engine/Source/Programs/UnrealBuildTool/Configuration/UEBuildPlatform.cs
--- a/Engine/Source/Programs/UnrealBuildTool/Configuration/UEBuildPlatform.cs
+++ b/Engine/Source/Programs/UnrealBuildTool/Configuration/UEBuildPlatform.cs
@@ -850,6 +850,20 @@
 		{
 		}
 
+		/// <summary>
+		/// Indicates whether this platform requires a .loadorder file to be generated for the build.
+		/// .loadorder files contain a list of dynamic modules and the exact order in which they should be loaded
+		/// to ensure that all dependencies are satisfied i.e. we don't attempt to load a module without loading
+		/// all its dependencies first.
+		/// As such, this file is only needed in modular builds on some platforms (depending on the way they implement dynamic modules such as DLLs).
+		/// </summary>
+		/// <param name="Target">The target rules in use</param>
+		/// <returns>True if .loadorder file should be generated</returns>
+		public virtual bool RequiresLoadOrderManifest(ReadOnlyTargetRules Target)
+		{
+			return false;
+		}
+
 
 		/// <summary>
 		/// Checks if platform is part of a given platform group
diff --git a/Engine/Source/Programs/UnrealBuildTool/Configuration/UEBuildTarget.cs b/Engine/Source/Programs/UnrealBuildTool/Configuration/UEBuildTarget.cs
--- a/Engine/Source/Programs/UnrealBuildTool/Configuration/UEBuildTarget.cs
+++ b/Engine/Source/Programs/UnrealBuildTool/Configuration/UEBuildTarget.cs
@@ -1493,6 +1493,74 @@ namespace UnrealBuildTool
 			return FileNameToModuleManifest;
 		}
 
+		Dictionary<FileReference, LoadOrderManifest> PrepareLoadOrderManifests(ILogger Logger)
+		{
+			Dictionary<FileReference, LoadOrderManifest> FileNameToManifest = new Dictionary<FileReference, LoadOrderManifest>();
+			
+			UEBuildPlatform BuildPlatform = UEBuildPlatform.GetBuildPlatform(Platform);
+
+			if (!bCompileMonolithic && BuildPlatform.RequiresLoadOrderManifest(Rules))
+			{
+				string ManifestFileName = LoadOrderManifest.GetStandardFileName(AppName, Platform, Configuration, Architecture);
+
+				// Create a map module->binary for all primary modules.
+				Dictionary<UEBuildModule, UEBuildBinary> PrimaryModules =
+					Binaries
+						.Where(x => x.Type == UEBuildBinaryType.DynamicLinkLibrary)
+						.ToDictionary(x => x.PrimaryModule as UEBuildModule, x => x);
+
+				List<Tuple<UEBuildModule, UEBuildModule>> Dependencies = new List<Tuple<UEBuildModule, UEBuildModule>>();
+
+				// Gather information about all dependencies between primary modules (binaries).
+				foreach (KeyValuePair<UEBuildModule, UEBuildBinary> ModuleInfo in PrimaryModules)
+				{
+					// Combine Private and Public dependencies and add a graph edge for each that is a primary module of some binary
+					// i.e. we care only about dependencies between primary modules.
+					// We remove Circular dependencies because we need to have a directed acyclic graph to be able to sort it.
+					// The convention seems to be that if we have
+					//     A -> B and B -> A (Circular: B)
+					// B -> A is considered secondary i.e. less important and we choose to remove this one.
+					(ModuleInfo.Key.PrivateDependencyModules ?? new List<UEBuildModule>())
+						.Union(ModuleInfo.Key.PublicDependencyModules ?? new List<UEBuildModule>())
+						.Where(x => PrimaryModules.ContainsKey(x) && !ModuleInfo.Key.Rules.CircularlyReferencedDependentModules.Contains(x.Name))
+						.ToList()
+						.ForEach(x => Dependencies.Add(new Tuple<UEBuildModule, UEBuildModule>(x, ModuleInfo.Key)));
+				}
+
+				// Sort the dependencies to generate a sequence of libraries to load respecting all the given dependencies.
+				TopologicalSorter<UEBuildModule> Sorter = new TopologicalSorter<UEBuildModule>(Dependencies);
+				Sorter.CycleHandling = TopologicalSorter<UEBuildModule>.CycleMode.BreakWithInfo;
+				Sorter.Logger = Logger;
+				Sorter.NodeToString = Module => Module.Name;
+				
+				if (!Sorter.Sort())
+				{
+					Logger.LogError("Failed to generate {ManifestFileName}: Couldn't sort dynamic modules in a way that would respect all dependencies (probably circular dependencies)", ManifestFileName);
+
+					return FileNameToManifest;
+				}
+
+				// Create the load order manifest and fill the list of all binaries.
+				LoadOrderManifest Manifest = new LoadOrderManifest();
+
+				foreach (UEBuildModule Module in Sorter.GetResult())
+				{
+					UEBuildBinary? Binary;
+
+					PrimaryModules.TryGetValue(Module, out Binary);
+
+					Manifest.Libraries.Add(
+						Path.GetRelativePath(GetExecutableDir().ToString(), Binary!.OutputFilePath.ToString()).Replace(Path.DirectorySeparatorChar, '/')
+						);
+				}
+
+				FileReference ManifestPath = FileReference.Combine(GetExecutableDir(), ManifestFileName);
+
+				FileNameToManifest.Add(ManifestPath, Manifest);
+			}
+			return FileNameToManifest;
+		}
+
 		/// <summary>
 		/// Prepare all the receipts this target (all the .target and .modules files). See the VersionManifest class for an explanation of what these files are.
 		/// </summary>
@@ -2028,13 +2096,17 @@ namespace UnrealBuildTool
 				Dictionary<FileReference, ModuleManifest> FileNameToModuleManifest = PrepareModuleManifests();
 				BuildProducts.AddRange(FileNameToModuleManifest.Select(x => new KeyValuePair<FileReference, BuildProductType>(x.Key, BuildProductType.RequiredResource)));
 
+				// Prepare load order manifests (if needed), and add them to the list of build products.
+				Dictionary<FileReference, LoadOrderManifest> FileNameToLoadOrderManifest = PrepareLoadOrderManifests(Logger);
+				BuildProducts.AddRange(FileNameToLoadOrderManifest.Select(x => new KeyValuePair<FileReference, BuildProductType>(x.Key, BuildProductType.RequiredResource)));
+
 				// Prepare the receipt
 				TargetReceipt Receipt = PrepareReceipt(TargetToolChain, BuildProducts, RuntimeDependencies);
 
 				// Create an action which to generate the module receipts
 				if (VersionFile == null)
 				{
-					WriteMetadataTargetInfo TargetInfo = new WriteMetadataTargetInfo(ProjectFile, null, null, ReceiptFileName, Receipt, FileNameToModuleManifest);
+					WriteMetadataTargetInfo TargetInfo = new WriteMetadataTargetInfo(ProjectFile, null, null, ReceiptFileName, Receipt, FileNameToModuleManifest, FileNameToLoadOrderManifest);
 					FileReference TargetInfoFile = FileReference.Combine(ProjectIntermediateDirectory, "TargetMetadata.dat");
 					CreateWriteMetadataAction(MakefileBuilder, ReceiptFileName.GetFileName(), TargetInfoFile, TargetInfo, Makefile.OutputItems);
 				}
@@ -2070,11 +2142,11 @@ namespace UnrealBuildTool
 						}
 					}
 
-					WriteMetadataTargetInfo EngineInfo = new WriteMetadataTargetInfo(null, VersionFile, Receipt.Version, null, null, EngineFileToManifest);
+					WriteMetadataTargetInfo EngineInfo = new WriteMetadataTargetInfo(null, VersionFile, Receipt.Version, null, null, EngineFileToManifest, null);
 					FileReference EngineInfoFile = FileReference.Combine(ProjectIntermediateDirectory, "EngineMetadata.dat");
 					CreateWriteMetadataAction(MakefileBuilder, VersionFile.GetFileName(), EngineInfoFile, EngineInfo, EnginePrereqItems);
 
-					WriteMetadataTargetInfo TargetInfo = new WriteMetadataTargetInfo(ProjectFile, VersionFile, null, ReceiptFileName, Receipt, TargetFileToManifest);
+					WriteMetadataTargetInfo TargetInfo = new WriteMetadataTargetInfo(ProjectFile, VersionFile, null, ReceiptFileName, Receipt, TargetFileToManifest, FileNameToLoadOrderManifest);
 					FileReference TargetInfoFile = FileReference.Combine(ProjectIntermediateDirectory, "TargetMetadata.dat");
 					CreateWriteMetadataAction(MakefileBuilder, ReceiptFileName.GetFileName(), TargetInfoFile, TargetInfo, TargetPrereqItems);
 				}
@@ -2209,6 +2281,8 @@ namespace UnrealBuildTool
 		void CreateWriteMetadataAction(TargetMakefileBuilder MakefileBuilder, string StatusDescription, FileReference InfoFile, WriteMetadataTargetInfo Info, IEnumerable<FileItem> PrerequisiteItems)
 		{
 			List<FileReference> ProducedItems = new List<FileReference>(Info.FileToManifest.Keys);
+			ProducedItems.AddRange(Info.FileToLoadOrderManifest.Keys);
+
 			if (Info.Version != null && Info.VersionFile != null)
 			{
 				ProducedItems.Add(Info.VersionFile);
@@ -4139,6 +4213,13 @@ namespace UnrealBuildTool
 			GlobalCompileEnvironment.Definitions.Add(String.Format("UBT_MODULE_MANIFEST=\"{0}\"", ModuleManifest.GetStandardFileName(AppName, Platform, Configuration, Architecture, false)));
 			GlobalCompileEnvironment.Definitions.Add(String.Format("UBT_MODULE_MANIFEST_DEBUGGAME=\"{0}\"", ModuleManifest.GetStandardFileName(AppName, Platform, UnrealTargetConfiguration.DebugGame, Architecture, true)));
 
+			// Compile in the names of the loadorder manifests.
+			if (!bCompileMonolithic && BuildPlatform.RequiresLoadOrderManifest(Rules))
+			{
+				GlobalCompileEnvironment.Definitions.Add(String.Format("UBT_LOADORDER_MANIFEST=\"{0}\"", LoadOrderManifest.GetStandardFileName(AppName, Platform, Configuration, Architecture)));
+				GlobalCompileEnvironment.Definitions.Add(String.Format("UBT_LOADORDER_MANIFEST_DEBUGGAME=\"{0}\"", LoadOrderManifest.GetStandardFileName(AppName, Platform, UnrealTargetConfiguration.DebugGame, Architecture)));
+			}
+
 			// tell the compiled code the name of the UBT platform (this affects folder on disk, etc that the game may need to know)
 			GlobalCompileEnvironment.Definitions.Add(String.Format("UBT_COMPILED_PLATFORM=" + Platform.ToString()));
 			GlobalCompileEnvironment.Definitions.Add(String.Format("UBT_COMPILED_TARGET=" + TargetType.ToString()));
diff --git a/Engine/Source/Programs/UnrealBuildTool/Modes/WriteMetadataMode.cs b/Engine/Source/Programs/UnrealBuildTool/Modes/WriteMetadataMode.cs
--- a/Engine/Source/Programs/UnrealBuildTool/Modes/WriteMetadataMode.cs
+++ b/Engine/Source/Programs/UnrealBuildTool/Modes/WriteMetadataMode.cs
@@ -51,6 +51,11 @@
 		public Dictionary<FileReference, ModuleManifest> FileToManifest;
 
 		/// <summary>
+		/// Map of load order manifest filenames to their locations on disk (generally, at most one load oder manifest is expected).
+		/// </summary>
+		public Dictionary<FileReference, LoadOrderManifest> FileToLoadOrderManifest;
+
+		/// <summary>
 		/// Constructor
 		/// </summary>
 		/// <param name="ProjectFile"></param>
@@ -59,7 +64,8 @@
 		/// <param name="ReceiptFile"></param>
 		/// <param name="Receipt"></param>
 		/// <param name="FileToManifest"></param>
-		public WriteMetadataTargetInfo(FileReference? ProjectFile, FileReference? VersionFile, BuildVersion? Version, FileReference? ReceiptFile, TargetReceipt? Receipt, Dictionary<FileReference, ModuleManifest> FileToManifest)
+		/// <param name="FileToLoadOrderManifest"></param>
+		public WriteMetadataTargetInfo(FileReference? ProjectFile, FileReference? VersionFile, BuildVersion? Version, FileReference? ReceiptFile, TargetReceipt? Receipt, Dictionary<FileReference, ModuleManifest> FileToManifest, Dictionary<FileReference, LoadOrderManifest>? FileToLoadOrderManifest)
 		{
 			this.ProjectFile = ProjectFile;
 			this.VersionFile = VersionFile;
@@ -67,6 +73,7 @@
 			this.ReceiptFile = ReceiptFile;
 			this.Receipt = Receipt;
 			this.FileToManifest = FileToManifest;
+			this.FileToLoadOrderManifest = FileToLoadOrderManifest ?? new Dictionary<FileReference, LoadOrderManifest>();
 		}
 	}
 
@@ -189,6 +196,46 @@
 				}
 			}
 
+			// Write load order manifests out.
+			foreach (KeyValuePair<FileReference, LoadOrderManifest> Pair in TargetInfo.FileToLoadOrderManifest)
+			{
+				FileReference ManifestFile = Pair.Key;
+				if (!UnrealBuildTool.IsFileInstalled(ManifestFile))
+				{
+					LoadOrderManifest Manifest = Pair.Value;
+
+					if (!FileReference.Exists(ManifestFile))
+					{
+						// If the file doesn't already exist, just write it out
+						DirectoryReference.CreateDirectory(ManifestFile.Directory);
+						Manifest.Write(ManifestFile);
+					}
+					else
+					{
+						// Otherwise write it to a buffer first
+						string OutputText;
+						using (StringWriter Writer = new StringWriter())
+						{
+							Manifest.Write(Writer);
+							OutputText = Writer.ToString();
+						}
+
+						// Check if the manifest has changed. Note that if a manifest is out of date, we should have generated a new build id causing the contents to differ.
+						if (bNoManifestChanges)
+						{
+							string CurrentText = FileReference.ReadAllText(ManifestFile);
+							if (CurrentText != OutputText)
+							{
+								Logger.LogError("Build modifies {File}. This is not permitted. Before:\n    {OldFile}\nAfter:\n    {NewFile}", ManifestFile, CurrentText.Replace("\n", "\n    "), OutputText.Replace("\n", "\n    "));
+							}
+						}
+
+						// Write it to disk
+						FileReference.WriteAllText(ManifestFile, OutputText);
+					}
+				}
+			}
+
 			// Write out the version file
 			if (TargetInfo.Version != null && TargetInfo.VersionFile != null)
 			{
diff --git a/Engine/Source/Programs/UnrealBuildTool/System/LoadOrderManifest.cs b/Engine/Source/Programs/UnrealBuildTool/System/LoadOrderManifest.cs
new file mode 100644
--- /dev/null
+++ b/Engine/Source/Programs/UnrealBuildTool/System/LoadOrderManifest.cs
@@ -0,0 +1,105 @@
+// Copyright Epic Games, Inc. All Rights Reserved.
+
+using System;
+using System.Collections.Generic;
+using System.IO;
+using EpicGames.Core;
+
+namespace UnrealBuildTool
+{
+	/// <summary>
+	/// Manifest file (.loadorder) describing the order in which all dynamic modules should be loaded.
+	/// For this reason, it's only relevant to modular builds on some platforms.
+	/// 
+	/// This file exists because on some non-Windows platforms dynamic modules (such as DLLs) don't support
+	/// dependencies and it's our responsibility to load them in the right order (Windows DLLs know which other
+	/// DLLs they rely on and these dependencies are guaranteed to be loaded first).
+	/// If we fail respect this order, we're likely to get crashes e.g. when global statics are being initialized
+	/// because some dereferenced globals may not be available yet (null).
+	/// </summary>
+	[Serializable]
+	class LoadOrderManifest
+	{
+		/// <summary>
+		/// List of all project's dynamic modules in the order in which they should be loaded.
+		/// All paths are relative to the main binary's directory.
+		/// </summary>
+		public List<string> Libraries = new List<string>();
+
+		public LoadOrderManifest()
+		{
+		}
+
+		/// <summary>
+		/// Gets a filename for a load order manifest for a build with the given platform/configuration.
+		/// </summary>
+		/// <param name="AppName">The modular app name being built</param>
+		/// <param name="Configuration">The target configuration</param>
+		/// <param name="Platform">The target platform</param>
+		/// <param name="BuildArchitecture">The architecture of the target platform</param>
+		/// <returns>Filename for the app receipt</returns>
+		public static string GetStandardFileName(string AppName, UnrealTargetPlatform Platform, UnrealTargetConfiguration Configuration, string BuildArchitecture)
+		{
+			string BaseName = AppName;
+			if (Configuration != UnrealTargetConfiguration.Development)
+			{
+				BaseName += String.Format("-{0}-{1}", Platform.ToString(), Configuration.ToString());
+			}
+			if (!String.IsNullOrEmpty(BuildArchitecture) && UEBuildPlatform.GetBuildPlatform(Platform).RequiresArchitectureSuffix())
+			{
+				BaseName += BuildArchitecture;
+			}
+			return String.Format("{0}.loadorder", BaseName);
+		}
+
+		/// <summary>
+		/// Read the load order manifest from a file.
+		/// </summary>
+		/// <param name="FileName">Filename to read from</param>
+		/// <returns>The receipt that was read</returns>
+		public static LoadOrderManifest Read(FileReference FileName)
+		{
+			using (StreamReader Reader = new StreamReader(FileName.FullName))
+			{
+				LoadOrderManifest Manifest = new LoadOrderManifest();
+
+				while (!Reader.EndOfStream)
+				{
+					string? Library = Reader.ReadLine();
+
+					if (Library != null)
+					{
+						Manifest.Libraries.Add(Library.Trim());
+					}
+				}
+
+				return Manifest;
+			}
+		}
+
+		/// <summary>
+		/// Write the load order manifest to a file.
+		/// </summary>
+		/// <param name="FileName">Filename to write to</param>
+		public void Write(FileReference FileName)
+		{
+			DirectoryReference.CreateDirectory(FileName.Directory);
+			using (StreamWriter Writer = new StreamWriter(FileName.FullName))
+			{
+				Write(Writer);
+			}
+		}
+
+		/// <summary>
+		/// Write the load order manifest to a text writer.
+		/// </summary>
+		/// <param name="Writer">The writer to output to</param>
+		public void Write(TextWriter Writer)
+		{
+			foreach (string Library in Libraries)
+			{
+				Writer.WriteLine(Library);
+			}
+		}
+	}
+}
diff --git a/Engine/Source/Programs/UnrealBuildTool/System/TopologicalSorter.cs b/Engine/Source/Programs/UnrealBuildTool/System/TopologicalSorter.cs
new file mode 100644
--- /dev/null
+++ b/Engine/Source/Programs/UnrealBuildTool/System/TopologicalSorter.cs
@@ -0,0 +1,331 @@
+// Copyright Epic Games, Inc. All Rights Reserved.
+
+using System;
+using System.Collections.Generic;
+using System.Linq;
+using Microsoft.Extensions.Logging;
+
+namespace UnrealBuildTool
+{
+	/// <summary>
+	/// Helper class implementing the standard topological sorting algorithm for directed graphs.
+	/// It generates a flattened sequence of nodes in the order that respects all the given dependencies
+	/// between nodes (edges). So, if we have edges A -> B and B -> C, A will be always before B and C, and
+	/// B will be before C.
+	/// 
+	/// In many cases there are more than one possible solutions. This helper just generates one of them.
+	/// The algorithm only works for directed acyclic graphs, as otherwise it's impossible to order the nodes.
+	/// In this implementation we can choose how cycles should be handled (CycleHandling member). By default
+	/// sorting fails, but we may set the class to arbitrarily break cycles and complete the task.
+	/// </summary>
+	/// <typeparam name="T">Type of node objects</typeparam>
+	class TopologicalSorter<T> where T : notnull
+	{
+		/// <summary>
+		/// Enumeration indicating how graph cycles should be handled.
+		/// </summary>
+		public enum CycleMode
+		{
+			/// <summary>
+			/// Fail sorting on the first cycle encountered (the algorithmically correct solution).
+			/// </summary>
+			Fail,
+
+			/// <summary>
+			/// Break cycles at the point they are detected (allows the algorithm to always complete but may skip some edges).
+			/// </summary>
+			Break,
+
+			/// <summary>
+			/// Break cycles at the point they are detected (same as Break but outputs warning logs with information about the cycles).
+			/// </summary>
+			BreakWithInfo
+		}
+
+		/// <summary>
+		/// Determines what should be done when a cycle in the graph is encountered.
+		/// </summary>
+		public CycleMode CycleHandling = CycleMode.Fail;
+
+		/// <summary>
+		/// UBT's logger object. Necessary if we want to get diagnostics, especially when BreakWithInfo is used.
+		/// </summary>
+		public ILogger? Logger = null;
+
+		/// <summary>
+		/// Functor returning a node's name. Necessary for some logging diagnostics, especially when BreakWithInfo is used.
+		/// </summary>
+		public Func<T, string>? NodeToString = null;
+
+		/// <summary>
+		/// List of all graph nodes (each user-provided node T is wrapped in an internal GraphNode).
+		/// </summary>
+		private readonly List<GraphNode> Nodes = new List<GraphNode>();
+
+		/// <summary>
+		/// When traversing the graph, it represents all the nodes visited in the current sequence.
+		/// For instance, if we go from node A to K and then X, Callstack will store A, K, X.
+		/// It's only needed for diagnostic purposes i.e. when we find a cycle we use Callstack
+		/// to output all the cycle's nodes.
+		/// </summary>
+		private Stack<GraphNode> Callstack = new Stack<GraphNode>();
+
+		/// <summary>
+		/// The resulting sequence of nodes (T, provided by the user) in the topological order.
+		/// </summary>
+		private List<T> Result = new List<T>();
+
+		/// <summary>
+		/// Initialize the sorter object and create the internal graph representation.
+		/// The list of all the graph nodes doesn't need to be provided because it's internally generated
+		/// from Edges (union of all edges' Item1 and Item2.
+		/// </summary>
+		/// <param name="Edges">List of connections between nodes</param>
+		public TopologicalSorter(List<Tuple<T, T>> Edges)
+		{
+			CreateGraph(Edges);
+		}
+
+		/// <summary>
+		/// Sort the graph to generate a flat sequence of nodes respecting all graph dependencies.
+		/// </summary>
+		/// <returns>True on success, false if failed (most likely a cycle encountered)</returns>
+		public bool Sort()
+		{
+			Result.Clear();
+			Callstack.Clear();
+
+			ClearNodesStates();
+
+			foreach (GraphNode Node in Nodes)
+			{
+				Callstack.Push(Node);
+
+				if (!Visit(Node))
+				{
+					return false;
+				}
+
+				Callstack.Pop();
+			}
+
+			Result.Reverse();
+
+			return true;
+		}
+
+		/// <summary>
+		/// Get the list of sorted nodes (valid after successfully calling Sort).
+		/// </summary>
+		/// <returns></returns>
+		public List<T> GetResult()
+		{
+			return Result;
+		}
+
+		/// <summary>
+		/// Process a single node.
+		/// </summary>
+		/// <param name="Node"></param>
+		/// <returns>True on success, false if failed (most likely a cycle encountered)</returns>
+		private bool Visit(GraphNode Node)
+		{
+			if (Node.Mark == Mark.Done)
+			{
+				// This node has already been processed.
+				return true;
+			}
+
+			if (Node.Mark == Mark.InProgress)
+			{
+				// This node is being processed i.e. we've detected a cycle.
+				return HandleCycle(Node);
+			}
+
+			Node.Mark = Mark.InProgress;
+
+			foreach (GraphNode Next in Node.Links)
+			{
+				Callstack.Push(Next);
+
+				if (!Visit(Next))
+				{
+					return false;
+				}
+
+				Callstack.Pop();
+			}
+
+			Node.Mark = Mark.Done;
+
+			Result.Add(Node.Data);
+
+			return true;
+		}
+
+		/// <summary>
+		/// Create the internal graph structure for the provided list of edges.
+		/// </summary>
+		/// <param name="Edges">List of tuples representing connections between nodes</param>
+		private void CreateGraph(List<Tuple<T, T>> Edges)
+		{
+			HashSet<T> AllNodes = new HashSet<T>();
+
+			// Create a collection of all nodes based on the ones used in Edges.
+			foreach (var Edge in Edges)
+			{
+				AllNodes.Add(Edge.Item1);
+				AllNodes.Add(Edge.Item2);
+			}
+
+			Dictionary<T, GraphNode> NodeToGraphNode = new Dictionary<T, GraphNode>();
+
+			// Create graph node objects.
+			foreach (T Node in AllNodes)
+			{
+				GraphNode GraphNode = new GraphNode(Node);
+
+				Nodes.Add(GraphNode);
+				NodeToGraphNode.Add(Node, GraphNode);
+			}
+
+			// Add edges to the graph.
+			foreach (var Edge in Edges)
+			{
+				GraphNode? NodeSrc;
+				GraphNode? NodeDst;
+
+				if (!NodeToGraphNode.TryGetValue(Edge.Item1, out NodeSrc))
+				{
+					throw new Exception($"TopologicalSorter: Failed to build graph (source node {Edge.Item1} from an edge unknown)!");
+				}
+
+				if (!NodeToGraphNode.TryGetValue(Edge.Item2, out NodeDst))
+				{
+					throw new Exception($"TopologicalSorter: Failed to build graph (source node {Edge.Item2} from an edge unknown)!");
+				}
+
+				NodeSrc.Links.Add(NodeDst);
+			}
+		}
+
+		/// <summary>
+		/// Clear the state (in-progress, done) of all graph nodes.
+		/// </summary>
+		private void ClearNodesStates()
+		{
+			foreach (GraphNode Node in Nodes)
+			{
+				Node.Mark = Mark.None;
+			}
+		}
+
+		/// <summary>
+		/// Member executed when a graph cycle is encountered.
+		/// </summary>
+		/// <param name="Node"></param>
+		/// <returns>True if operation should continue, false if a cycle is considered an error</returns>
+		private bool HandleCycle(GraphNode Node)
+		{
+			switch (CycleHandling)
+			{
+				case CycleMode.Fail:
+					// We stop further processing and fail the whole sort.
+					if (Logger != null)
+					{
+						Logger.LogError("TopologicalSorter: Cycle found in graph ({NodesInCycle})", GetNodesInCycleString());
+					}
+					return false;
+
+				case CycleMode.Break:
+					// Treat the 'in-progress' node as if it's been 'done' i.e. break the cycle at this point.
+					return true;
+
+				case CycleMode.BreakWithInfo:
+					// As in Break but output log information about the cycle.
+					if (Logger != null)
+					{
+						Logger.LogWarning("TopologicalSorter: Cycle found in graph ({NodesInCycle})", GetNodesInCycleString());
+					}
+					return true;
+			}
+
+			return false;
+		}
+
+		/// <summary>
+		/// Returns all nodes being part of a cycle (call valid only if we've really detected a cycle).
+		/// </summary>
+		/// <returns></returns>
+		private IEnumerable<GraphNode> GetNodesInCycle()
+		{
+			GraphNode[] CallstackNodes = Callstack.ToArray();
+
+			// First, find where in the callstack the cycle starts.
+			int CycleStart = 1;
+
+			for (; CycleStart < CallstackNodes.Length; ++CycleStart)
+			{
+				if (CallstackNodes[CycleStart] == CallstackNodes[0])
+				{
+					break;
+				}
+			}
+
+			return CallstackNodes.Take(CycleStart + 1).Reverse();
+		}
+
+		/// <summary>
+		/// Returns a string with a list of names of nodes forming a graph cycle (call valid only if we've really detected a cycle).
+		/// </summary>
+		/// <returns></returns>
+		private string GetNodesInCycleString()
+		{
+			if (NodeToString == null)
+			{
+				return "";
+			}
+
+			IEnumerable<GraphNode> CycleNodes = GetNodesInCycle();
+
+			return String.Join(", ", CycleNodes.Select(Node => NodeToString(Node.Data)));
+		}
+
+		/// <summary>
+		/// Graph node's state wrt the traversal.
+		/// </summary>
+		private enum Mark
+		{
+			None = 0,
+			InProgress = 1,
+			Done = 2
+		}
+
+		/// <summary>
+		/// Internal helper representing a graph node (wrapping the user-provided node and storing internal state).
+		/// </summary>
+		private class GraphNode
+		{
+			/// <summary>
+			/// User-provided node object.
+			/// </summary>
+			public T Data;
+
+			/// <summary>
+			/// List of graph nodes we can visit starting from this one.
+			/// </summary>
+			public List<GraphNode> Links = new List<GraphNode>();
+
+			/// <summary>
+			/// Traversal state (in-progress, done).
+			/// </summary>
+			public Mark Mark;
+
+			public GraphNode(T Data)
+			{
+				this.Data = Data;
+				this.Mark = Mark.None;
+			}
+		}
+	}
+}
