PR #9021: Improve support for PreBuildSteps generating files (Contributed by [REDACTED_PERSON])

diff --git a/Engine/Source/Programs/UnrealBuildTool/Modes/BuildMode.cs b/Engine/Source/Programs/UnrealBuildTool/Modes/BuildMode.cs
--- a/Engine/Source/Programs/UnrealBuildTool/Modes/BuildMode.cs
+++ b/Engine/Source/Programs/UnrealBuildTool/Modes/BuildMode.cs
@@ -715,12 +715,8 @@
 			bool bHasRunPreBuildScripts = false;
 			if(Makefile != null)
 			{
-				// Execute the scripts. We have to invalidate all cached file info after doing so, because we don't know what may have changed.
-				if(Makefile.PreBuildScripts.Length > 0)
-				{
-					Utils.ExecuteCustomBuildSteps(Makefile.PreBuildScripts, Logger);
-					DirectoryItem.ResetAllCachedInfo_SLOW();
-				}
+				// Execute the scripts.
+				Utils.ExecuteCustomBuildSteps(Makefile.PreBuildScripts, Logger);
 
 				// Don't run the pre-build steps again, even if we invalidate the makefile.
 				bHasRunPreBuildScripts = true;
diff --git a/Engine/Source/Programs/UnrealBuildTool/System/Utils.cs b/Engine/Source/Programs/UnrealBuildTool/System/Utils.cs
--- a/Engine/Source/Programs/UnrealBuildTool/System/Utils.cs
+++ b/Engine/Source/Programs/UnrealBuildTool/System/Utils.cs
@@ -1336,7 +1336,6 @@
 		/// <returns>True if the steps succeeded, false otherwise</returns>
 		public static void ExecuteCustomBuildSteps(FileReference[] ScriptFiles, ILogger Logger)
 		{
-			UnrealTargetPlatform HostPlatform = BuildHostPlatform.Current.Platform;
 			foreach(FileReference ScriptFile in ScriptFiles)
 			{
 				ProcessStartInfo StartInfo = new ProcessStartInfo();
@@ -1357,6 +1356,12 @@
 					throw new BuildException("Custom build step {0} {1} terminated with exit code {2}", StartInfo.FileName, StartInfo.Arguments, ReturnCode);
 				}
 			}
+
+			if (ScriptFiles.Length > 0)
+			{
+				// We have to invalidate all cached file info after running the scripts, because we don't know what may have changed.
+				DirectoryItem.ResetAllCachedInfo_SLOW();
+			}
 		}
 
 		/// <summary>
