- Fixed Mac arm64 builds of UnrealEditor (another nni reference had to be disabled)

diff --git a/Engine/Source/Programs/UnrealBuildTool/Platform/Mac/UEBuildMac.cs b/Engine/Source/Programs/UnrealBuildTool/Platform/Mac/UEBuildMac.cs
--- a/Engine/Source/Programs/UnrealBuildTool/Platform/Mac/UEBuildMac.cs
+++ b/Engine/Source/Programs/UnrealBuildTool/Platform/Mac/UEBuildMac.cs
@@ -150,6 +150,7 @@ namespace UnrealBuildTool
 					"VertexDeltaModel",
 					"MetaHumanAnimator",
 					"MetaHumans",
+					"NeuralMorphModel",
 
 					// MDL
 					"MDLImporter",
