UnealBuildTool: Include clang pch with -include-pch rather than -include

#rnx

diff --git a/Engine/Source/Programs/UnrealBuildTool/ToolChain/ClangToolChain.cs b/Engine/Source/Programs/UnrealBuildTool/ToolChain/ClangToolChain.cs
--- a/Engine/Source/Programs/UnrealBuildTool/ToolChain/ClangToolChain.cs
+++ b/Engine/Source/Programs/UnrealBuildTool/ToolChain/ClangToolChain.cs
@@ -397,13 +397,23 @@ namespace UnrealBuildTool
 			return GetForceIncludeFileArgument(ForceIncludeFile.Location);
 		}
 
+		protected virtual string GetIncludePCHFileArgument(FileReference IncludePCHFile)
+		{
+			return $"-include-pch \"{NormalizeCommandLinePath(IncludePCHFile)}\"";
+		}
+
+		protected virtual string GetIncludePCHFileArgument(FileItem IncludePCHFile)
+		{
+			return GetIncludePCHFileArgument(IncludePCHFile.Location);
+		}
+
 		protected virtual void GetCompileArguments_ForceInclude(CppCompileEnvironment CompileEnvironment, List<string> Arguments)
 		{
 			// Add the precompiled header file's path to the force include path
 			// This needs to be before the other force include paths to ensure clang uses it instead of the source header file.
 			if (CompileEnvironment.PrecompiledHeaderAction == PrecompiledHeaderAction.Include)
 			{
-				Arguments.Add(GetForceIncludeFileArgument(CompileEnvironment.PrecompiledHeaderIncludeFilename!));
+				Arguments.Add(GetIncludePCHFileArgument(CompileEnvironment.PrecompiledHeaderFile!));
 			}
 
 			Arguments.AddRange(CompileEnvironment.ForceIncludeFiles.Select(ForceIncludeFile => GetForceIncludeFileArgument(ForceIncludeFile)));
