- Moved everything over to .xcconfig files, with the exception of Extensions (a little-used IOS feature). It's possible, just not done yet, and is currently removed and will throw an exception
- Requires -xcconfig on the GenerateProjectFiles commandline
- Resurreceted the old non-xcconfig code in another namespace so that we can have both code paths side-by-side
- Had to change a little bit of XcodeProjectFileGenerator to manage the side-by-side classes

diff --git a/Engine/Source/Programs/UnrealBuildTool/ProjectFiles/Xcode/XcodeProject.cs b/Engine/Source/Programs/UnrealBuildTool/ProjectFiles/Xcode/XcodeProject.cs
--- a/Engine/Source/Programs/UnrealBuildTool/ProjectFiles/Xcode/XcodeProject.cs
+++ b/Engine/Source/Programs/UnrealBuildTool/ProjectFiles/Xcode/XcodeProject.cs
@@ -12,7 +12,7 @@ using EpicGames.Core;
 using UnrealBuildBase;
 using Microsoft.Extensions.Logging;
 
-namespace UnrealBuildTool
+namespace UnrealBuildTool.XcodeProjectXcconfig
 {
 	/// <summary>
 	/// Info needed to make a file a member of specific group
@@ -96,6 +96,19 @@ namespace UnrealBuildTool
 		public string BuildTarget;
 		public ProjectTarget? ProjectTarget;
 		public UnrealTargetConfiguration BuildConfig;
+
+		public bool bSupportsMac { get =>
+			XcodeProjectFileGenerator.ProjectFilePlatform.HasFlag(XcodeProjectFileGenerator.XcodeProjectFilePlatform.Mac) &&
+			(ProjectTarget == null || ProjectTarget.SupportedPlatforms.Contains(UnrealTargetPlatform.Mac));
+		}
+		public bool bSupportsIOS { get =>
+			XcodeProjectFileGenerator.ProjectFilePlatform.HasFlag(XcodeProjectFileGenerator.XcodeProjectFilePlatform.iOS) &&
+			(ProjectTarget == null || ProjectTarget.SupportedPlatforms.Contains(UnrealTargetPlatform.IOS));
+		}
+		public bool bSupportsTVOS { get =>
+			XcodeProjectFileGenerator.ProjectFilePlatform.HasFlag(XcodeProjectFileGenerator.XcodeProjectFilePlatform.tvOS) &&
+			(ProjectTarget == null || ProjectTarget.SupportedPlatforms.Contains(UnrealTargetPlatform.TVOS));
+		}
 	};
 
 	class XcodeExtensionInfo
@@ -126,28 +139,31 @@ namespace UnrealBuildTool
 
 	class XcconfigFile
 	{
+		public string Name;
 		public string Filename;
 		public string Guid;
-		public StreamWriter Stream;
+		public FileReference FileRef;
+		private StringBuilder Text;
 
-		public XcconfigFile(FileReference ProjectFile, string NameSuffix)
+		public XcconfigFile(FileReference ProjectFile, string ConfigName)
 		{
-			FileReference File = FileReference.Combine(ProjectFile.Directory, ProjectFile.GetFileNameWithoutAnyExtensions() + NameSuffix + ".xcconfig");
-			Filename = File.GetFileName();
+			Name = ConfigName;
+			FileRef = FileReference.Combine(ProjectFile.Directory, "Xcconfigs", $"{ProjectFile.GetFileNameWithoutAnyExtensions()}_{Name.Replace(" ", "")}.xcconfig");
+			Filename = FileRef.GetFileName();
 			Guid = XcodeProjectFileGenerator.MakeXcodeGuid();
-			DirectoryReference.CreateDirectory(ProjectFile.Directory);
-			Stream = new StreamWriter(new FileStream(File.FullName, FileMode.Create, FileAccess.Write), Encoding.ASCII);
+			DirectoryReference.CreateDirectory(FileRef.Directory);
+			Text = new StringBuilder();
 		}
 
 		public void AppendLine(string Line)
 		{
-			Stream.Write(Line + ProjectFileGenerator.NewLine);
-			Stream.Flush();
+			Text.Append(Line);
+			Text.Append(ProjectFileGenerator.NewLine);
 		}
 
-		public void Close()
+		public void Write()
 		{
-			Stream.Close();
+			FileReference.WriteAllTextIfDifferent(FileRef, Text.ToString());
 		}
 	}
 
@@ -171,11 +187,10 @@ namespace UnrealBuildTool
 			BundleIdentifier = BundleID;
 			AppName = InAppName;
 
-			ProjectXcconfig = new XcconfigFile(InitFilePath, "_Project");
-			TargetXcconfig = new XcconfigFile(InitFilePath, "_Target");
-
-			AllXcconfigFiles.Add(ProjectXcconfig);
-			AllXcconfigFiles.Add(TargetXcconfig);
+			AllXcconfigFiles.Add(ProjectXcconfig = new XcconfigFile(InitFilePath, "Project"));
+			AllXcconfigFiles.Add(RunXcconfig = new XcconfigFile(InitFilePath, "Run"));
+			AllXcconfigFiles.Add(IndexXcconfig = new XcconfigFile(InitFilePath, "Index"));
+			AllXcconfigFiles.Add(BuildXcconfig = new XcconfigFile(InitFilePath, "Build"));
 		}
 
 		public override string ToString()
@@ -207,12 +222,10 @@ namespace UnrealBuildTool
 		/// Xcconfig file information - this is global per project, so no need to pass it around everywhere
 		/// </summary>
 		private XcconfigFile ProjectXcconfig;
-		private XcconfigFile TargetXcconfig;
-
-		/// <summary>
-		/// All Xcconfigs that can be looped over (to add to project, and close, etc)
-		/// </summary>
-		private List<XcconfigFile> AllXcconfigFiles = new List<XcconfigFile>();
+		private XcconfigFile RunXcconfig;
+		private XcconfigFile IndexXcconfig;
+		private XcconfigFile BuildXcconfig;
+		internal List<XcconfigFile> AllXcconfigFiles = new List<XcconfigFile>();
 
 		/// <summary>
 		/// Gets Xcode file category based on its extension
@@ -403,8 +416,9 @@ namespace UnrealBuildTool
 
 			foreach (XcconfigFile Xcconfig in AllXcconfigFiles)
 			{
-				PBXFileReferenceSection.Append($"\t\t{Xcconfig.Guid} /* {Xcconfig.Filename} */ = {{isa = PBXFileReference; explicitFileType = test.xcconfig; path = {Xcconfig.Filename}; sourceTree = \"<group>\"; }};" + ProjectFileGenerator.NewLine);
+				PBXFileReferenceSection.Append($"\t\t{Xcconfig.Guid} /* {Xcconfig.Filename} */ = {{isa = PBXFileReference; explicitFileType = test.xcconfig; name = {Xcconfig.Filename}; path = Xcconfigs/{Xcconfig.Filename}; sourceTree = \"<group>\"; }};" + ProjectFileGenerator.NewLine);
 			}
+
 			PBXFileReferenceSection.Append(string.Format("\t\t{0} /* {1} */ = {{isa = PBXFileReference; explicitFileType = {2}; path = {1}; sourceTree = BUILT_PRODUCTS_DIR; }};" + ProjectFileGenerator.NewLine, TargetAppGuid, TargetName, bIsAppBundle ? "wrapper.application" : "\"compiled.mach-o.executable\""));
 		}
 
@@ -416,8 +430,11 @@ namespace UnrealBuildTool
 				string ProjectExtensionsDir = Path.Combine(Path.GetDirectoryName(UProjectPath.FullName)!, "Build/IOS/Extensions");
 				//string ProjectIntermediateDir = Path.Combine(Path.GetDirectoryName(UProjectPath.FullName), "Intermediate/IOS/Extensions");
 
+
 				if (Directory.Exists(ProjectExtensionsDir))
 				{
+					throw new BuildException($"YOU NEED TO HANDLE EXTENSIONS NOW {UProjectPath}");
+#if false
 					foreach (DirectoryInfo DI in new System.IO.DirectoryInfo(ProjectExtensionsDir).EnumerateDirectories())
 					{
 						Console.WriteLine("  Project {0} has Extension {1}!", UProjectPath, DI);
@@ -525,6 +542,7 @@ namespace UnrealBuildTool
 
 						ExtensionInfo.ConfigurationContents = ConfigSection.ToString();
 					}
+#endif
 				}
 			}
 		}
@@ -625,7 +643,7 @@ namespace UnrealBuildTool
 			Content.Append("/* End PBXCopyFilesBuildPhase section */" + ProjectFileGenerator.NewLine);
 		}
 
-		private void AppendGroupSection(StringBuilder Content, string MainGroupGuid, string ProductRefGroupGuid, string TargetAppGuid, string TargetName, List<XcodeExtensionInfo> AllExtensions)
+		private void AppendGroupSection(StringBuilder Content, string MainGroupGuid, string ProductRefGroupGuid, string XcconfigsRefGroupGuid, string TargetAppGuid, string TargetName, List<XcodeExtensionInfo> AllExtensions)
 		{
 			XcodeFileGroup? RootGroup = FindRootFileGroup(Groups);
 			if (RootGroup == null)
@@ -639,10 +657,6 @@ namespace UnrealBuildTool
 			Content.Append(string.Format("\t\t{0} = {{{1}", MainGroupGuid, ProjectFileGenerator.NewLine));
 			Content.Append("\t\t\tisa = PBXGroup;" + ProjectFileGenerator.NewLine);
 			Content.Append("\t\t\tchildren = (" + ProjectFileGenerator.NewLine);
-			foreach (XcconfigFile Xcconfig in AllXcconfigFiles)
-			{
-				Content.Append($"\t\t\t\t{Xcconfig.Guid} /* {Xcconfig.Filename} */,");
-			}
 			foreach (XcodeFileGroup Group in RootGroup.Children.Values)
 			{
 				Content.Append(string.Format("\t\t\t\t{0} /* {1} */,{2}", Group.GroupGuid, Group.GroupName, ProjectFileGenerator.NewLine));
@@ -654,6 +668,7 @@ namespace UnrealBuildTool
 			}
 
 			Content.Append(string.Format("\t\t\t\t{0} /* Products */,{1}", ProductRefGroupGuid, ProjectFileGenerator.NewLine));
+			Content.Append(string.Format("\t\t\t\t{0} /* Xcconfigs */,{1}", XcconfigsRefGroupGuid, ProjectFileGenerator.NewLine));
 			Content.Append("\t\t\t);" + ProjectFileGenerator.NewLine);
 			Content.Append("\t\t\tsourceTree = \"<group>\";" + ProjectFileGenerator.NewLine);
 			Content.Append("\t\t};" + ProjectFileGenerator.NewLine);
@@ -678,6 +693,19 @@ namespace UnrealBuildTool
 			Content.Append("\t\t\tsourceTree = \"<group>\";" + ProjectFileGenerator.NewLine);
 			Content.Append("\t\t};" + ProjectFileGenerator.NewLine);
 
+			// Xcconfigs group
+			Content.Append(string.Format("\t\t{0} /* XcconfigsRefGroupGuid */ = {{{1}", XcconfigsRefGroupGuid, ProjectFileGenerator.NewLine));
+			Content.Append("\t\t\tisa = PBXGroup;" + ProjectFileGenerator.NewLine);
+			Content.Append("\t\t\tchildren = (" + ProjectFileGenerator.NewLine);
+			foreach (XcconfigFile Xcconfig in AllXcconfigFiles)
+			{
+				Content.Append($"\t\t\t\t{Xcconfig.Guid} /* {Xcconfig.Filename} */,{ProjectFileGenerator.NewLine}");
+			}
+			Content.Append("\t\t\t);" + ProjectFileGenerator.NewLine);
+			Content.Append("\t\t\tname = Xcconfigs;" + ProjectFileGenerator.NewLine);
+			Content.Append("\t\t\tsourceTree = \"<group>\";" + ProjectFileGenerator.NewLine);
+			Content.Append("\t\t};" + ProjectFileGenerator.NewLine);
+
 			Content.Append("/* End PBXGroup section */" + ProjectFileGenerator.NewLine + ProjectFileGenerator.NewLine);
 		}
 
@@ -970,53 +998,6 @@ namespace UnrealBuildTool
 			Content.Append("/* End PBXTargetDependency section */" + ProjectFileGenerator.NewLine + ProjectFileGenerator.NewLine);
 		}
 
-		private void AppendProjectBuildConfiguration(StringBuilder Content, string ConfigName, string ConfigGuid)
-		{
-			Content.Append("\t\t" + ConfigGuid + " /* \"" + ConfigName + "\" */ = {" + ProjectFileGenerator.NewLine);
-			Content.Append("\t\t\tisa = XCBuildConfiguration;" + ProjectFileGenerator.NewLine);
-			Content.Append($"\t\t\tbaseConfigurationReference = {ProjectXcconfig.Guid} /* \"{ProjectXcconfig.Filename}\" */;" + ProjectFileGenerator.NewLine);
-			Content.Append("\t\t\tbuildSettings = {" + ProjectFileGenerator.NewLine);
-
-			Content.Append("\t\t\t\tGCC_PREPROCESSOR_DEFINITIONS = (" + ProjectFileGenerator.NewLine);
-			foreach (string Definition in IntelliSensePreprocessorDefinitions)
-			{
-				Content.Append("\t\t\t\t\t\"" + Definition.Replace("\"", "").Replace("\\", "") + "\"," + ProjectFileGenerator.NewLine);
-			}
-			Content.Append("\t\t\t\t\t\"__INTELLISENSE__\"," + ProjectFileGenerator.NewLine);
-			Content.Append("\t\t\t\t\t\"MONOLITHIC_BUILD=1\"," + ProjectFileGenerator.NewLine);
-			Content.Append("\t\t\t\t);" + ProjectFileGenerator.NewLine);
-
-			Content.Append("\t\t\t\tHEADER_SEARCH_PATHS = (" + ProjectFileGenerator.NewLine);
-			foreach (string SearchPath in IntelliSenseSystemIncludeSearchPaths)
-			{
-				string Path = SearchPath.Contains(" ") ? "\\\"" + SearchPath + "\\\"" : SearchPath;
-				Content.Append("\t\t\t\t\t\"" + Path + "\"," + ProjectFileGenerator.NewLine);
-			}
-			Content.Append("\t\t\t\t);" + ProjectFileGenerator.NewLine);
-
-			Content.Append("\t\t\t\tUSER_HEADER_SEARCH_PATHS = (" + ProjectFileGenerator.NewLine);
-			foreach (string SearchPath in IntelliSenseIncludeSearchPaths)
-			{
-				string Path = SearchPath.Contains(" ") ? "\\\"" + SearchPath + "\\\"" : SearchPath;
-				Content.Append("\t\t\t\t\t\"" + Path + "\"," + ProjectFileGenerator.NewLine);
-			}
-			Content.Append("\t\t\t\t);" + ProjectFileGenerator.NewLine);
-			Content.Append("\t\t\t\tONLY_ACTIVE_ARCH = YES;" + ProjectFileGenerator.NewLine);
-			if (ConfigName == "Debug")
-			{
-				Content.Append("\t\t\t\tENABLE_TESTABILITY = YES;" + ProjectFileGenerator.NewLine);
-			}
-			Content.Append("\t\t\t\tALWAYS_SEARCH_USER_PATHS = NO;" + ProjectFileGenerator.NewLine);
-			Content.Append("\t\t\t\tCLANG_CXX_LANGUAGE_STANDARD = \"c++14\";" + ProjectFileGenerator.NewLine);
-			Content.Append("\t\t\t\tGCC_ENABLE_CPP_RTTI = NO;" + ProjectFileGenerator.NewLine);
-			Content.Append("\t\t\t\tGCC_WARN_CHECK_SWITCH_STATEMENTS = NO;" + ProjectFileGenerator.NewLine);
-			Content.Append("\t\t\t\tUSE_HEADERMAP = NO;" + ProjectFileGenerator.NewLine);
-			Content.Append("\t\t\t};" + ProjectFileGenerator.NewLine);
-			Content.Append("\t\t\tname = \"" + ConfigName + "\";" + ProjectFileGenerator.NewLine);
-			Content.Append("\t\t};" + ProjectFileGenerator.NewLine);
-
-		}
-
 		// cache for the below function
 		Dictionary<string, IEnumerable<string>> CachedMacProjectArcitectures = new Dictionary<string, IEnumerable<string>>();
 
@@ -1102,275 +1083,220 @@ namespace UnrealBuildTool
 			return CachedMacProjectArcitectures[TargetName];
 		}
 
-		private void AppendPlatformConfiguration(StringBuilder Content, XcodeBuildConfig Config, string TargetName, FileReference? ProjectFile, bool bSupportMac, bool bSupportIOS, bool bSupportTVOS, ILogger Logger, out string? IOSRunTimeVersion, out string? TVOSRunTimeVersion, string BinariesSubDir = "/Payload")
+		private void WriteProjectXcconfigFiles(List<XcodeBuildConfig> ProjectConfigs, FileReference? ProjectFile, ILogger Logger)
 		{
-			FileReference MacExecutablePath = Config.MacExecutablePath!;
+			// Setup all the settings that are shared across all configs
 
 			string UEDir = ConvertPath(Path.GetFullPath(Directory.GetCurrentDirectory() + "../../.."));
-			string MacExecutableDir = bSupportMac ? ConvertPath(MacExecutablePath.Directory.FullName) : "";
-			string MacExecutableFileName = bSupportMac ? MacExecutablePath.GetFileName() : "";
-
-			// Get Mac architectures supported by this project
-			IEnumerable<string> SupportedMacArchitectures = GetSupportedMacArchitectures(Config, ProjectFile);
-
-			IOSRunTimeVersion = null;
-			TVOSRunTimeVersion = null;
+			string SupportedIOSArchitectures = string.Join(" ", this.SupportedIOSArchitectures);
+			string? IOSRunTimeVersion = null;
+			string? TVOSRunTimeVersion = null;
 
 			// shortcut for mac only
-			if (bSupportMac && !bSupportIOS && !bSupportTVOS)
-			{
-				Content.Append("\t\t\t\tVALID_ARCHS = \"" + string.Join(" ", SupportedMacArchitectures) + "\";" + ProjectFileGenerator.NewLine);
-				Content.Append("\t\t\t\tSUPPORTED_PLATFORMS = \"macosx\";" + ProjectFileGenerator.NewLine);
-				Content.Append("\t\t\t\tPRODUCT_NAME = \"" + MacExecutableFileName + "\";" + ProjectFileGenerator.NewLine);
-				Content.Append("\t\t\t\tCONFIGURATION_BUILD_DIR = \"" + MacExecutableDir + "\";" + ProjectFileGenerator.NewLine);
-			}
-			else
-			{
-				bool bIsUnrealGame = TargetName.Equals("UnrealGame", StringComparison.InvariantCultureIgnoreCase);
-				bool bIsUnrealClient = TargetName.Equals("UnrealClient", StringComparison.InvariantCultureIgnoreCase);
-				DirectoryReference? GameDir = ProjectFile?.Directory;
-				string? GamePath = GameDir != null ? ConvertPath(GameDir.FullName) : null;
-
-				string? IOSRunTimeDevices = null;
-				string? TVOSRunTimeDevices = null;
-				string SupportedPlatforms = bSupportMac ? "macosx" : "";
-
-				bool bAutomaticSigning = false;
-                string? UUID_IOS = "";
-                string? UUID_TVOS = "";
-                string? TEAM_IOS = "";
-                string? TEAM_TVOS = "";
-                string? IOS_CERT = "iPhone Developer";
-                string? TVOS_CERT = "iPhone Developer";
-				string IOS_BUNDLE = "";
-				string TVOS_BUNDLE = "";
-				if (bSupportIOS && InstalledPlatformInfo.IsValidPlatform(UnrealTargetPlatform.IOS, EProjectType.Code))
-                {
- 					IOSPlatform IOSPlatform = ((IOSPlatform)UEBuildPlatform.GetBuildPlatform(UnrealTargetPlatform.IOS));
-					IOSProjectSettings ProjectSettings = IOSPlatform.ReadProjectSettings(ProjectFile);
-					IOSProvisioningData ProvisioningData = IOSPlatform.ReadProvisioningData(ProjectSettings, bForDistribution);
-					IOSRunTimeVersion = ProjectSettings.RuntimeVersion;
-					IOSRunTimeDevices = ProjectSettings.RuntimeDevices;
-					SupportedPlatforms += " iphoneos";
-					bAutomaticSigning = ProjectSettings.bAutomaticSigning;
-					if (!bAutomaticSigning)
-					{
-						UUID_IOS = ProvisioningData.MobileProvisionUUID;
-						IOS_CERT = ProvisioningData.SigningCertificate;
-					}
-                    TEAM_IOS = ProvisioningData.TeamUUID;
-					IOS_BUNDLE = ProjectSettings.BundleIdentifier;
-                }
+			string? IOSRunTimeDevices = null;
+			string? TVOSRunTimeDevices = null;
 
-				if (bSupportTVOS && InstalledPlatformInfo.IsValidPlatform(UnrealTargetPlatform.TVOS, EProjectType.Code))
+			bool bAutomaticSigning = false;
+			string? UUID_IOS = "";
+			string? UUID_TVOS = "";
+			string? TEAM_IOS = "";
+			string? TEAM_TVOS = "";
+			string? IOS_CERT = "iPhone Developer";
+			string? TVOS_CERT = "iPhone Developer";
+			string IOS_BUNDLE = "";
+			string TVOS_BUNDLE = "";
+				
+			if (InstalledPlatformInfo.IsValidPlatform(UnrealTargetPlatform.IOS, EProjectType.Code))
+			{
+				IOSPlatform IOSPlatform = ((IOSPlatform)UEBuildPlatform.GetBuildPlatform(UnrealTargetPlatform.IOS));
+				IOSProjectSettings ProjectSettings = IOSPlatform.ReadProjectSettings(ProjectFile);
+				IOSProvisioningData ProvisioningData = IOSPlatform.ReadProvisioningData(ProjectSettings, bForDistribution);
+				IOSRunTimeVersion = ProjectSettings.RuntimeVersion;
+				IOSRunTimeDevices = ProjectSettings.RuntimeDevices;
+				bAutomaticSigning = ProjectSettings.bAutomaticSigning;
+				if (!bAutomaticSigning)
 				{
-					TVOSPlatform TVOSPlatform = ((TVOSPlatform)UEBuildPlatform.GetBuildPlatform(UnrealTargetPlatform.TVOS));
-					TVOSProjectSettings ProjectSettings = TVOSPlatform.ReadProjectSettings(ProjectFile);
-					TVOSProvisioningData ProvisioningData = TVOSPlatform.ReadProvisioningData(ProjectSettings, bForDistribution);
-					TVOSRunTimeVersion = ProjectSettings.RuntimeVersion;
-					TVOSRunTimeDevices = ProjectSettings.RuntimeDevices;
-					SupportedPlatforms += " appletvos";
-					if (!bAutomaticSigning)
-					{
-						UUID_TVOS = ProvisioningData.MobileProvisionUUID;
-						TVOS_CERT = ProvisioningData.SigningCertificate;
-					}
-                    TEAM_TVOS = ProvisioningData.TeamUUID;
-					TVOS_BUNDLE = ProjectSettings.BundleIdentifier;
-                }
+					UUID_IOS = ProvisioningData.MobileProvisionUUID;
+					IOS_CERT = ProvisioningData.SigningCertificate;
+				}
+				TEAM_IOS = ProvisioningData.TeamUUID;
+				IOS_BUNDLE = ProjectSettings.BundleIdentifier;
+			}
 
-				Content.Append("\t\t\t\tSUPPORTED_PLATFORMS = \"" + SupportedPlatforms.Trim() + "\";" + ProjectFileGenerator.NewLine);
-				if (bAutomaticSigning)
+			if (InstalledPlatformInfo.IsValidPlatform(UnrealTargetPlatform.TVOS, EProjectType.Code))
+			{
+				TVOSPlatform TVOSPlatform = ((TVOSPlatform)UEBuildPlatform.GetBuildPlatform(UnrealTargetPlatform.TVOS));
+				TVOSProjectSettings ProjectSettings = TVOSPlatform.ReadProjectSettings(ProjectFile);
+				TVOSProvisioningData ProvisioningData = TVOSPlatform.ReadProvisioningData(ProjectSettings, bForDistribution);
+				TVOSRunTimeVersion = ProjectSettings.RuntimeVersion;
+				TVOSRunTimeDevices = ProjectSettings.RuntimeDevices;
+				if (!bAutomaticSigning)
 				{
-					Content.Append("\t\t\t\tCODE_SIGN_STYLE = Automatic;" + ProjectFileGenerator.NewLine);
+					UUID_TVOS = ProvisioningData.MobileProvisionUUID;
+					TVOS_CERT = ProvisioningData.SigningCertificate;
 				}
-				if (IOSRunTimeVersion != null)
+				TEAM_TVOS = ProvisioningData.TeamUUID;
+				TVOS_BUNDLE = ProjectSettings.BundleIdentifier;
+			}
+
+			//#jira UE-50382 Xcode Address Sanitizer feature does not work on iOS
+			// address sanitizer dylib loader depends on the SDKROOT parameter. For macosx or default (missing, translated as macosx), the path is incorrect for iphone/appletv
+			ProjectXcconfig.AppendLine($"SDKROOT[sdk=macosx*] = macosx");
+			ProjectXcconfig.AppendLine($"SDKROOT[sdk=iphoneos] = iphoneos");
+			ProjectXcconfig.AppendLine($"SDKROOT[sdk=appletvos*] = appletvos");
+			ProjectXcconfig.AppendLine("USE_HEADERMAP = NO");
+			ProjectXcconfig.AppendLine("ONLY_ACTIVE_ARCH = YES");
+
+			if (bAutomaticSigning)
+			{
+				ProjectXcconfig.AppendLine("CODE_SIGN_STYLE = Automatic");
+			}
+
+			if (IOSRunTimeVersion != null)
+			{
+				ProjectXcconfig.AppendLine($"VALID_ARCHS[sdk=iphoneos*] = {SupportedIOSArchitectures}");
+				ProjectXcconfig.AppendLine($"IPHONEOS_DEPLOYMENT_TARGET = {IOSRunTimeVersion}");
+				ProjectXcconfig.AppendLine($"TARGETED_DEVICE_FAMILY[sdk=iphoneos*] = {IOSRunTimeDevices}");
+				if (!string.IsNullOrEmpty(TEAM_IOS))
 				{
-					Content.Append("\t\t\t\t\"VALID_ARCHS[sdk=iphoneos*]\" = \"" + string.Join(" ", this.SupportedIOSArchitectures) + "\";" + ProjectFileGenerator.NewLine);
-					Content.Append("\t\t\t\tIPHONEOS_DEPLOYMENT_TARGET = " + IOSRunTimeVersion + ";" + ProjectFileGenerator.NewLine);
-					Content.Append("\t\t\t\t\"PRODUCT_NAME[sdk=iphoneos*]\" = \"" + TargetName + "\";" + ProjectFileGenerator.NewLine); // @todo: change to Path.GetFileName(Config.IOSExecutablePath) when we stop using payload
-					Content.Append("\t\t\t\t\"TARGETED_DEVICE_FAMILY[sdk=iphoneos*]\" = \"" + IOSRunTimeDevices + "\";" + ProjectFileGenerator.NewLine);
-                    Content.Append("\t\t\t\t\"SDKROOT[sdk=iphoneos]\" = iphoneos;" + ProjectFileGenerator.NewLine);
-					if (!string.IsNullOrEmpty(TEAM_IOS))
-					{
-						Content.Append("\t\t\t\t\"DEVELOPMENT_TEAM[sdk=iphoneos*]\" = " + TEAM_IOS + ";" + ProjectFileGenerator.NewLine);
-					}
-					Content.Append("\t\t\t\t\"CODE_SIGN_IDENTITY[sdk=iphoneos*]\" = \"" + IOS_CERT + "\";" + ProjectFileGenerator.NewLine);
-					if (!bAutomaticSigning && !string.IsNullOrEmpty(UUID_IOS))
-					{
-						Content.Append("\t\t\t\t\"PROVISIONING_PROFILE_SPECIFIER[sdk=iphoneos*]\" = \"" + UUID_IOS + "\";" + ProjectFileGenerator.NewLine);
-					}
-					if (ProjectFile != null)
-					{
-						Content.Append("\t\t\t\t\"PRODUCT_BUNDLE_IDENTIFIER[sdk=iphoneos*]\" = " + IOS_BUNDLE + ";" + ProjectFileGenerator.NewLine);
-					}
+					ProjectXcconfig.AppendLine($"DEVELOPMENT_TEAM[sdk=iphoneos*] = {TEAM_IOS}");
 				}
-                if (TVOSRunTimeVersion != null)
+				ProjectXcconfig.AppendLine($"CODE_SIGN_IDENTITY[sdk=iphoneos*] = {IOS_CERT}");
+				if (!bAutomaticSigning && !string.IsNullOrEmpty(UUID_IOS))
 				{
-					Content.Append("\t\t\t\t\"VALID_ARCHS[sdk=appletvos*]\" = \"" + string.Join(" ", this.SupportedIOSArchitectures) + "\";" + ProjectFileGenerator.NewLine);
-					Content.Append("\t\t\t\tTVOS_DEPLOYMENT_TARGET = " + TVOSRunTimeVersion + ";" + ProjectFileGenerator.NewLine);
-					Content.Append("\t\t\t\t\"PRODUCT_NAME[sdk=appletvos*]\" = \"" + TargetName + "\";" + ProjectFileGenerator.NewLine); // @todo: change to Path.GetFileName(Config.TVOSExecutablePath) when we stop using payload
-					Content.Append("\t\t\t\t\"TARGETED_DEVICE_FAMILY[sdk=appletvos*]\" = \"" + TVOSRunTimeDevices + "\";" + ProjectFileGenerator.NewLine);
-                    Content.Append("\t\t\t\t\"SDKROOT[sdk=appletvos]\" = appletvos;" + ProjectFileGenerator.NewLine);
-					if (!string.IsNullOrEmpty(TEAM_TVOS))
-					{
-						Content.Append("\t\t\t\t\"DEVELOPMENT_TEAM[sdk=appletvos*]\" = " + TEAM_TVOS + ";" + ProjectFileGenerator.NewLine);
-					}
-					Content.Append("\t\t\t\t\"CODE_SIGN_IDENTITY[sdk=appletvos*]\" = \"" + TVOS_CERT + "\";" + ProjectFileGenerator.NewLine);
-					if (!bAutomaticSigning && !string.IsNullOrEmpty(UUID_TVOS))
-					{
-						Content.Append("\t\t\t\t\"PROVISIONING_PROFILE_SPECIFIER[sdk=appletvos*]\" = \"" + UUID_TVOS + "\";" + ProjectFileGenerator.NewLine);
-					}
-					if (ProjectFile != null)
-					{
-						Content.Append("\t\t\t\t\"PRODUCT_BUNDLE_IDENTIFIER[sdk=appletvos*]\" = " + TVOS_BUNDLE + ";" + ProjectFileGenerator.NewLine);
-					}
+					ProjectXcconfig.AppendLine($"PROVISIONING_PROFILE_SPECIFIER[sdk=iphoneos*] = {UUID_IOS}");
 				}
-				if (bSupportMac)
+				if (ProjectFile != null)
 				{
-					Content.Append("\t\t\t\t\"VALID_ARCHS[sdk=macosx*]\" = \"" + string.Join(" ", SupportedMacArchitectures) + "\";" + ProjectFileGenerator.NewLine);
-					Content.Append("\t\t\t\t\"PRODUCT_NAME[sdk=macosx*]\" = \"" + MacExecutableFileName + "\";" + ProjectFileGenerator.NewLine);
-					Content.Append("\t\t\t\t\"CONFIGURATION_BUILD_DIR[sdk=macosx*]\" = \"" + MacExecutableDir + "\";" + ProjectFileGenerator.NewLine);
-					Content.Append("\t\t\t\t\"SDKROOT[sdk=macosx]\" = macosx;" + ProjectFileGenerator.NewLine);
+					ProjectXcconfig.AppendLine($"PRODUCT_BUNDLE_IDENTIFIER[sdk=iphoneos*] = {IOS_BUNDLE}");
 				}
-
-				if (bIsUnrealGame || bIsUnrealClient)
+			}
+			if (TVOSRunTimeVersion != null)
+			{
+				ProjectXcconfig.AppendLine($"VALID_ARCHS[sdk=appletvos*] = {SupportedIOSArchitectures}");
+				ProjectXcconfig.AppendLine($"IPHONEOS_DEPLOYMENT_TARGET = {TVOSRunTimeVersion}");
+				ProjectXcconfig.AppendLine($"TARGETED_DEVICE_FAMILY[sdk=appletvos*] = {TVOSRunTimeDevices}");
+				if (!string.IsNullOrEmpty(TEAM_TVOS))
 				{
-					if (IOSRunTimeVersion != null)
-					{
-						Content.Append("\t\t\t\t\"CONFIGURATION_BUILD_DIR[sdk=iphoneos*]\" = \"" + UEDir + "/Engine/Binaries/IOS" + BinariesSubDir + "\";" + ProjectFileGenerator.NewLine);
-					}
-					if (TVOSRunTimeVersion != null)
-					{
-						Content.Append("\t\t\t\t\"CONFIGURATION_BUILD_DIR[sdk=appletvos*]\" = \"" + UEDir + "/Engine/Binaries/TVOS" + BinariesSubDir + "\";" + ProjectFileGenerator.NewLine);
-					}
+					ProjectXcconfig.AppendLine($"DEVELOPMENT_TEAM[sdk=iphoneos*] = {TEAM_TVOS}");
 				}
-				else if (ProjectFile != null)
+				ProjectXcconfig.AppendLine($"CODE_SIGN_IDENTITY[sdk=appletvos*] = {TVOS_CERT}");
+				if (!bAutomaticSigning && !string.IsNullOrEmpty(UUID_TVOS))
 				{
-					if (IOSRunTimeVersion != null)
-					{
-						Content.Append("\t\t\t\t\"CONFIGURATION_BUILD_DIR[sdk=iphoneos*]\" = \"" + GamePath + "/Binaries/IOS" + BinariesSubDir + "\";" + ProjectFileGenerator.NewLine);
-					}
-					if (TVOSRunTimeVersion != null)
-					{
-						Content.Append("\t\t\t\t\"CONFIGURATION_BUILD_DIR[sdk=appletvos*]\" = \"" + GamePath + "/Binaries/TVOS" + BinariesSubDir + "\";" + ProjectFileGenerator.NewLine);
-					}
+					ProjectXcconfig.AppendLine($"PROVISIONING_PROFILE_SPECIFIER[sdk=appletvos*] = {UUID_TVOS}");
 				}
-				else
+				if (ProjectFile != null)
 				{
-					if (IOSRunTimeVersion != null)
-					{
-						Content.Append("\t\t\t\t\"CONFIGURATION_BUILD_DIR[sdk=iphoneos*]\" = \"" + UEDir + "/Engine/Binaries/IOS" + BinariesSubDir + "\";" + ProjectFileGenerator.NewLine);
-					}
-					if (TVOSRunTimeVersion != null)
-					{
-						Content.Append("\t\t\t\t\"CONFIGURATION_BUILD_DIR[sdk=appletvos*]\" = \"" + UEDir + "/Engine/Binaries/TVOS" + BinariesSubDir + "\";" + ProjectFileGenerator.NewLine);
-					}
+					ProjectXcconfig.AppendLine($"PRODUCT_BUNDLE_IDENTIFIER[sdk=appletvos*] = {TVOS_BUNDLE}");
 				}
-
 			}
-		}
 
-		private void AppendNativeTargetBuildConfiguration(StringBuilder Content, XcodeBuildConfig Config, string ConfigGuid, FileReference? ProjectFile, ILogger Logger)
-		{
-			bool bMacOnly = true;
-			if (Config.ProjectTarget!.TargetRules != null && XcodeProjectFileGenerator.ProjectFilePlatform.HasFlag(XcodeProjectFileGenerator.XcodeProjectFilePlatform.iOS))
+			// Now for each config write out the specific settings
+
+			DirectoryReference? GameDir = ProjectFile?.Directory;
+			string? GamePath = GameDir != null ? ConvertPath(GameDir.FullName) : null;
+
+			// Get Mac architectures supported by this project
+			foreach (XcodeBuildConfig Config in ProjectConfigs)
 			{
-				if (Config.ProjectTarget.SupportedPlatforms.Contains(UnrealTargetPlatform.IOS))
-				{
-					bMacOnly = false;
-				}
-			}
+				// find the config file we made earlier - we know it exists
+				XcconfigFile Xcconfig = AllXcconfigFiles.First(x => x.Name == Config.DisplayName);
 
-			Content.Append("\t\t" + ConfigGuid + " /* \"" + Config.DisplayName + "\" */ = {" + ProjectFileGenerator.NewLine);
-			Content.Append("\t\t\tisa = XCBuildConfiguration;" + ProjectFileGenerator.NewLine);
-			Content.Append("\t\t\tbuildSettings = {" + ProjectFileGenerator.NewLine);
+				string TargetName = Config.BuildTarget;
+				bool bIsUnrealGame = TargetName.Equals("UnrealGame", StringComparison.InvariantCultureIgnoreCase);
+				bool bIsUnrealClient = TargetName.Equals("UnrealClient", StringComparison.InvariantCultureIgnoreCase);
 
-			string UEDir = ConvertPath(Path.GetFullPath(Directory.GetCurrentDirectory() + "../../.."));
-			//string MacExecutableDir = ConvertPath(Config.MacExecutablePath.Directory.FullName);
-			string MacExecutableFileName = Config.MacExecutablePath!.GetFileName();
+				// Setup Mac stuff
+				FileReference MacExecutablePath = Config.MacExecutablePath!;
+				string MacExecutableDir = ConvertPath(MacExecutablePath.Directory.FullName);
+				string MacExecutableFileName = MacExecutablePath.GetFileName();
+				string SupportedMacArchitectures = string.Join(" ", GetSupportedMacArchitectures(Config, ProjectFile));
 
-			string? IOSRunTimeVersion, TVOSRunTimeVersion;
-			AppendPlatformConfiguration(Content, Config, Config.BuildTarget, ProjectFile, true, !bMacOnly, !bMacOnly, Logger, out IOSRunTimeVersion, out TVOSRunTimeVersion);
+				// figure out the directory that contains the Binaries and Intermediate directories
+				string ProjectRootDir = (bIsUnrealGame || bIsUnrealClient || GamePath == null) ? $"{UEDir}/Engine" : GamePath;
 
-			if (!bMacOnly)
-			{
-				bool bIsUnrealGame = Config.BuildTarget.Equals("UnrealGame", StringComparison.InvariantCultureIgnoreCase);
-				bool bIsUnrealClient = Config.BuildTarget.Equals("UnrealClient", StringComparison.InvariantCultureIgnoreCase);
+				// UnrealClient re-uses UnrealGame
+				string PlistTargetName = bIsUnrealClient ? "UnrealGame" : Config.BuildTarget;
+				string MacInfoPlistPath = $"{ProjectRootDir}/Intermediate/Mac/{MacExecutableFileName}-Info.plist";
+				string IOSInfoPlistPath = $"{ProjectRootDir}/Intermediate/IOS/{PlistTargetName}-Info.plist";
+				string TVOSInfoPlistPath = $"{ProjectRootDir}/Intermediate/TVOS/{PlistTargetName}-Info.plist";
 
-				DirectoryReference? GameDir = ProjectFile?.Directory;
-				string? GamePath = GameDir != null ? ConvertPath(GameDir.FullName) : null;
+				// @todo is really necessary to only set entitlements in the non-UnrealGame, project case? what about mac entitlements?
+				bool bSetEntitlements = !bIsUnrealGame && !bIsUnrealClient && ProjectFile != null;
 
-				string IOSInfoPlistPath;
-				string TVOSInfoPlistPath;
-				string MacInfoPlistPath;
-				string? IOSEntitlementPath = null;
-				string? TVOSEntitlementPath = null;
-				if (bIsUnrealGame)
-				{
-					IOSInfoPlistPath = UEDir + "/Engine/Intermediate/IOS/" + Config.BuildTarget + "-Info.plist";
-					TVOSInfoPlistPath = UEDir + "/Engine/Intermediate/TVOS/" + Config.BuildTarget + "-Info.plist";
-					MacInfoPlistPath = UEDir + "/Engine/Intermediate/Mac/" + MacExecutableFileName + "-Info.plist";
-					IOSEntitlementPath = "";
-					TVOSEntitlementPath = "";
-				}
-				else if (bIsUnrealClient)
+				string SupportedPlatforms = "";
+				if (Config.ProjectTarget!.TargetRules != null)
 				{
-					IOSInfoPlistPath = UEDir + "/Engine/Intermediate/IOS/UnrealGame-Info.plist";
-					TVOSInfoPlistPath = UEDir + "/Engine/Intermediate/TVOS/UnrealGame-Info.plist";
-					MacInfoPlistPath = UEDir + "/Engine/Intermediate/Mac/" + MacExecutableFileName + "-Info.plist";
-					IOSEntitlementPath = "";
-					TVOSEntitlementPath = "";
+					if (Config.bSupportsIOS)
+					{
+						SupportedPlatforms += "iphoneos ";
 					}
-				else if (ProjectFile != null)
-				{
-					IOSInfoPlistPath = GamePath + "/Intermediate/IOS/" + Config.BuildTarget + "-Info.plist";
-					TVOSInfoPlistPath = GamePath + "/Intermediate/TVOS/" + Config.BuildTarget + "-Info.plist";
-					MacInfoPlistPath = GamePath + "/Intermediate/Mac/" + MacExecutableFileName + "-Info.plist";
-					IOSEntitlementPath = GamePath + "/Intermediate/IOS/" + Config.BuildTarget + ".entitlements";
-					TVOSEntitlementPath = GamePath + "/Intermediate/TVOS/" + Config.BuildTarget + ".entitlements";
-				}
-				else
-				{
-					if (GamePath == null)
+					if (Config.bSupportsTVOS)
 					{
-						IOSInfoPlistPath = UEDir + "/Engine/Intermediate/IOS/" + Config.BuildTarget + "-Info.plist";
-						TVOSInfoPlistPath = UEDir + "/Engine/Intermediate/TVOS/" + Config.BuildTarget + "-Info.plist";
-						MacInfoPlistPath = UEDir + "/Engine/Intermediate/Mac/" + MacExecutableFileName + "-Info.plist";
+						SupportedPlatforms += "appletvos ";
 					}
-					else
+					if (Config.bSupportsMac)
 					{
-						IOSInfoPlistPath = GamePath + "/Intermediate/IOS/" + Config.BuildTarget + "-Info.plist";
-						TVOSInfoPlistPath = GamePath + "/Intermediate/TVOS/" + Config.BuildTarget + "-Info.plist";
-						MacInfoPlistPath = GamePath + "/Intermediate/Mac/" + MacExecutableFileName + "-Info.plist";
+						SupportedPlatforms += "macosx ";
 					}
 				}
+				else
+				{
+					// @todo when does this case happen?
+					SupportedPlatforms = "macosx";
+				}
 
-				if (XcodeProjectFileGenerator.bGeneratingRunIOSProject)
+				// debug settings
+				if (Config.BuildConfig == UnrealTargetConfiguration.Debug)
 				{
-					Content.Append("\t\t\t\tINFOPLIST_FILE = \"" + IOSInfoPlistPath + "\";" + ProjectFileGenerator.NewLine);
-					Content.Append("\t\t\t\tCODE_SIGN_ENTITLEMENTS = \"" + IOSEntitlementPath + "\";" + ProjectFileGenerator.NewLine);
+					Xcconfig.AppendLine("ENABLE_TESTABILITY = YES");
 				}
-				else if (XcodeProjectFileGenerator.bGeneratingRunTVOSProject)
+
+				Xcconfig.AppendLine($"UE_BUILD_TARGET_NAME = {TargetName}");
+				Xcconfig.AppendLine($"UE_BUILD_TARGET_CONFIG = {Config.BuildConfig}");
+
+				// list the supported platforms, which will narrow down what other settings are active
+				Xcconfig.AppendLine($"SUPPORTED_PLATFORMS = {SupportedPlatforms}");
+
+				// @otodo move VALID_ARCHS up to Project config once we are always universal
+				if (Config.bSupportsMac)
 				{
-					Content.Append("\t\t\t\tINFOPLIST_FILE = \"" + TVOSInfoPlistPath + "\";" + ProjectFileGenerator.NewLine);
-					Content.Append("\t\t\t\tCODE_SIGN_ENTITLEMENTS = \"" + TVOSEntitlementPath + "\";" + ProjectFileGenerator.NewLine);
+					Xcconfig.AppendLine("");
+					Xcconfig.AppendLine("// Mac setup");
+					Xcconfig.AppendLine($"VALID_ARCHS[sdk=macosx*] = {SupportedMacArchitectures}");
+					Xcconfig.AppendLine($"PRODUCT_NAME[sdk=macosx*] = {MacExecutableFileName}");
+					Xcconfig.AppendLine($"CONFIGURATION_BUILD_DIR[sdk=macosx*] = {MacExecutableDir}");
+					Xcconfig.AppendLine($"INFOPLIST_FILE[sdk=macosx*] = {MacInfoPlistPath}");
 				}
-				else
+
+				if (Config.bSupportsIOS)
 				{
-					Content.Append("\t\t\t\t\"INFOPLIST_FILE[sdk=macosx*]\" = \"" + MacInfoPlistPath + "\";" + ProjectFileGenerator.NewLine);
-					if (IOSRunTimeVersion != null)
+					Xcconfig.AppendLine("");
+					Xcconfig.AppendLine("// IOS setup");
+					Xcconfig.AppendLine($"PRODUCT_NAME[sdk=iphoneos*] = {TargetName}"); // @todo: change to Path.GetFileName(Config.IOSExecutablePath) when we stop using payload
+					Xcconfig.AppendLine($"CONFIGURATION_BUILD_DIR[sdk=iphoneos*] = {ProjectRootDir}/Binaries/IOS/Payload");
+					Xcconfig.AppendLine($"INFOPLIST_FILE[sdk=iphoneos*] = {IOSInfoPlistPath}");
+					if (bSetEntitlements)
 					{
-						Content.Append("\t\t\t\t\"INFOPLIST_FILE[sdk=iphoneos*]\" = \"" + IOSInfoPlistPath + "\";" + ProjectFileGenerator.NewLine);
-						Content.Append("\t\t\t\t\"CODE_SIGN_ENTITLEMENTS[sdk=iphoneos*]\" = \"" + IOSEntitlementPath + "\";" + ProjectFileGenerator.NewLine);
+						Xcconfig.AppendLine($"CODE_SIGN_ENTITLEMENTS[sdk=iphoneos*] = {ProjectRootDir}/Intermediate/IOS/{PlistTargetName}.entitlements");
 					}
-					if (TVOSRunTimeVersion != null)
+				}
+				if (Config.bSupportsTVOS)
+				{
+					Xcconfig.AppendLine("");
+					Xcconfig.AppendLine("// TVOS setup");
+					Xcconfig.AppendLine($"PRODUCT_NAME[sdk=appletvos*] = {TargetName}"); // @todo: change to Path.GetFileName(Config.IOSExecutablePath) when we stop using payload
+					Xcconfig.AppendLine($"CONFIGURATION_BUILD_DIR[sdk=appletvos*] = {ProjectRootDir}/Binaries/TVOS/Payload");
+					Xcconfig.AppendLine($"INFOPLIST_FILE[sdk=appletvos*] = {TVOSInfoPlistPath}");
+					if (bSetEntitlements)
 					{
-						Content.Append("\t\t\t\t\"INFOPLIST_FILE[sdk=appletvos*]\" = \"" + TVOSInfoPlistPath + "\";" + ProjectFileGenerator.NewLine);
-						Content.Append("\t\t\t\t\"CODE_SIGN_ENTITLEMENTS[sdk=appletvos*]\" = \"" + TVOSEntitlementPath + "\";" + ProjectFileGenerator.NewLine);
+						Xcconfig.AppendLine($"CODE_SIGN_ENTITLEMENTS[sdk=appletvos*] = {ProjectRootDir}/Intermediate/TVOS/{PlistTargetName}.entitlements");
 					}
 				}
 
+
+				// @todo can we remove this and use GENERATE_INFOPLIST_FILE when plist doesn't exist?
+
 				// Prepare a temp Info.plist file so Xcode has some basic info about the target immediately after opening the project.
 				// This is needed for the target to pass the settings validation before code signing. UBT will overwrite this plist file later, with proper contents.
 				if (BuildHostPlatform.Current.Platform == UnrealTargetPlatform.Mac)
@@ -1413,7 +1339,7 @@ namespace UnrealBuildTool
 								ReceiptFilename = TargetReceipt.GetDefaultPath(ProjectPath, GameName, UnrealTargetPlatform.IOS, Config.BuildConfig, "");
 							}
 							Directory.CreateDirectory(Path.GetDirectoryName(IOSInfoPlistPath)!);
-                            bool bSupportPortrait, bSupportLandscape;
+							bool bSupportPortrait, bSupportLandscape;
 							TargetReceipt? Receipt;
 							TargetReceipt.TryRead(ReceiptFilename, out Receipt);
 							bool bBuildAsFramework = UEDeployIOS.GetCompileAsDll(Receipt);
@@ -1427,80 +1353,14 @@ namespace UnrealBuildTool
 					}
 				}
 			}
-
-			// #jira UE-143619: Pre Monterey macOS requires this option for a packaged app to run on iOS15 due to new code signature format. Could be removed once Monterey is minimum OS.
-			Content.Append("\t\t\t\tOTHER_CODE_SIGN_FLAGS = \"--generate-entitlement-der\";" + ProjectFileGenerator.NewLine);
-
-			Content.Append("\t\t\t\tMACOSX_DEPLOYMENT_TARGET = " + MacToolChain.Settings.MacOSVersion + ";" + ProjectFileGenerator.NewLine);
-			Content.Append("\t\t\t\tINFOPLIST_OUTPUT_FORMAT = xml;" + ProjectFileGenerator.NewLine);
-			Content.Append("\t\t\t\tCOMBINE_HIDPI_IMAGES = YES;" + ProjectFileGenerator.NewLine);
-			
-            //#jira UE-50382 Xcode Address Sanitizer feature does not work on iOS
-            // address sanitizer dylib loader depends on the SDKROOT parameter. For macosx or default (missing, translated as macosx), the path is incorrect for iphone/appletv
-            if (XcodeProjectFileGenerator.bGeneratingRunIOSProject)
-            {
-                Content.Append("\t\t\t\tSDKROOT = iphoneos;" + ProjectFileGenerator.NewLine);
-            }
-            else if (XcodeProjectFileGenerator.bGeneratingRunTVOSProject)
-            {
-                Content.Append("\t\t\t\tSDKROOT = appletvos;" + ProjectFileGenerator.NewLine);
-            }
-            else
-            {
-                Content.Append("\t\t\t\tSDKROOT = macosx;" + ProjectFileGenerator.NewLine);
-            }
-            Content.Append("\t\t\t\tGCC_PRECOMPILE_PREFIX_HEADER = YES;" + ProjectFileGenerator.NewLine);
-			Content.Append("\t\t\t\tGCC_PREFIX_HEADER = \"" + UEDir + "/Engine/Source/Editor/UnrealEd/Public/UnrealEd.h\";" + ProjectFileGenerator.NewLine);
-			Content.Append("\t\t\t};" + ProjectFileGenerator.NewLine);
-			Content.Append("\t\t\tname = \"" + Config.DisplayName + "\";" + ProjectFileGenerator.NewLine);
-			Content.Append("\t\t};" + ProjectFileGenerator.NewLine);
 		}
 
-		private void AppendLegacyTargetBuildConfiguration(StringBuilder Content, XcodeBuildConfig Config, string ConfigGuid, FileReference? ProjectFile)
+		private void AppendBuildConfiguration(StringBuilder Content, XcodeBuildConfig Config, string ConfigGuid,  XcconfigFile Xcconfig)
 		{
-			bool bMacOnly = true;
-			if (Config.ProjectTarget!.TargetRules != null && XcodeProjectFileGenerator.ProjectFilePlatform.HasFlag(XcodeProjectFileGenerator.XcodeProjectFilePlatform.iOS))
-			{
-				if (Config.ProjectTarget.SupportedPlatforms.Contains(UnrealTargetPlatform.IOS))
-				{
-					bMacOnly = false;
-				}
-			}
-
-			// Get Mac architectures supported by this project
-			IEnumerable<string> SupportedMacArchitectures = GetSupportedMacArchitectures(Config, ProjectFilePath);
-
 			Content.Append("\t\t" + ConfigGuid + " /* \"" + Config.DisplayName + "\" */ = {" + ProjectFileGenerator.NewLine);
 			Content.Append("\t\t\tisa = XCBuildConfiguration;" + ProjectFileGenerator.NewLine);
-			Content.Append($"\t\t\tbaseConfigurationReference = {TargetXcconfig.Guid} /* \"{TargetXcconfig.Filename}\" */;" + ProjectFileGenerator.NewLine);
+			Content.Append($"\t\t\tbaseConfigurationReference = {Xcconfig.Guid} /* \"{Xcconfig.Filename}\" */;" + ProjectFileGenerator.NewLine);
 			Content.Append("\t\t\tbuildSettings = {" + ProjectFileGenerator.NewLine);
-			if (bMacOnly)
-			{
-				Content.Append("\t\t\t\tVALID_ARCHS = \"" + string.Join(" ", SupportedMacArchitectures) + "\";" + ProjectFileGenerator.NewLine);
-				Content.Append("\t\t\t\tSUPPORTED_PLATFORMS = \"macosx\";" + ProjectFileGenerator.NewLine);
-			}
-			else
-			{
-				IEnumerable<string> ValidArchs = SupportedMacArchitectures;
-				string SupportedPlatforms = "macosx";
-				if (InstalledPlatformInfo.IsValidPlatform(UnrealTargetPlatform.IOS, EProjectType.Code))
-				{
-					ValidArchs = ValidArchs.Union(this.SupportedIOSArchitectures);
-					SupportedPlatforms += " iphoneos";
-				}
-				if (InstalledPlatformInfo.IsValidPlatform(UnrealTargetPlatform.TVOS, EProjectType.Code))
-				{
-					ValidArchs = ValidArchs.Union(this.SupportedIOSArchitectures);
-					SupportedPlatforms += " appletvos";
-				}
-				Content.Append("\t\t\t\tVALID_ARCHS = \"" + string.Join(" ", ValidArchs) + "\";" + ProjectFileGenerator.NewLine);
-				Content.Append("\t\t\t\tSUPPORTED_PLATFORMS = \"" + SupportedPlatforms + "\";" + ProjectFileGenerator.NewLine);
-			}
-			Content.Append("\t\t\t\tGCC_PREPROCESSOR_DEFINITIONS = ();" + ProjectFileGenerator.NewLine);
-			Content.Append("\t\t\t\tHEADER_SEARCH_PATHS = ();" + ProjectFileGenerator.NewLine);
-			Content.Append("\t\t\t\tUSER_HEADER_SEARCH_PATHS = ();" + ProjectFileGenerator.NewLine);
-			Content.Append("\t\t\t\tUE_BUILD_TARGET_NAME = \"" + Config.BuildTarget + "\";" + ProjectFileGenerator.NewLine);
-			Content.Append("\t\t\t\tUE_BUILD_TARGET_CONFIG = \"" + Config.BuildConfig + "\";" + ProjectFileGenerator.NewLine);
 			Content.Append("\t\t\t};" + ProjectFileGenerator.NewLine);
 			Content.Append("\t\t\tname = \"" + Config.DisplayName + "\";" + ProjectFileGenerator.NewLine);
 			Content.Append("\t\t};" + ProjectFileGenerator.NewLine);
@@ -1515,22 +1375,24 @@ namespace UnrealBuildTool
 
 			foreach (KeyValuePair<string, XcodeBuildConfig> Config in ProjectBuildConfigs)
 			{
-				AppendProjectBuildConfiguration(Content, Config.Value.DisplayName, Config.Key);
+				// find the premade file that matches the config name (like "Debug Editor")
+				XcconfigFile Xcconfig = AllXcconfigFiles.First(x => x.Name == Config.Value.DisplayName);
+				AppendBuildConfiguration(Content, Config.Value, Config.Key, Xcconfig);
 			}
 
 			foreach (KeyValuePair<string, XcodeBuildConfig> Config in TargetBuildConfigs)
 			{
-				AppendNativeTargetBuildConfiguration(Content, Config.Value, Config.Key, GameProjectPath, Logger);
+				AppendBuildConfiguration(Content, Config.Value, Config.Key, RunXcconfig);
 			}
 
 			foreach (KeyValuePair<string, XcodeBuildConfig> Config in BuildTargetBuildConfigs)
 			{
-				AppendLegacyTargetBuildConfiguration(Content, Config.Value, Config.Key, GameProjectPath);
+				AppendBuildConfiguration(Content, Config.Value, Config.Key, BuildXcconfig);
 			}
 
 			foreach (KeyValuePair<string, XcodeBuildConfig> Config in IndexTargetBuildConfigs)
 			{
-				AppendNativeTargetBuildConfiguration(Content, Config.Value, Config.Key, GameProjectPath, Logger);
+				AppendBuildConfiguration(Content, Config.Value, Config.Key, IndexXcconfig);
 			}
 
 			foreach (XcodeExtensionInfo EI in AllExtensions)
@@ -1905,16 +1767,57 @@ namespace UnrealBuildTool
 			File.WriteAllText(ManagementFile.FullName, Content.ToString(), new UTF8Encoding());
 		}
 
-		private void WriteXcconfigFiles()
+
+		private void WriteTargetXcconfigFiles()
 		{
-			// settings we write to xcconfig now instead of xcode project file
-			ProjectXcconfig.AppendLine("GENERATE_INFOPLIST_FILE[sdk=macosx*] = YES // Xcode 14 code-signing fix for development");
+			string UEDir = ConvertPath(Path.GetFullPath(Directory.GetCurrentDirectory() + "../../.."));
+
+			//////////////////// RUN ////////////////////////
+			RunXcconfig.AppendLine("GENERATE_INFOPLIST_FILE[sdk=macosx*] = YES // Xcode 14 code-signing fix for development");
+	
+			// #jira UE-143619: Pre Monterey macOS requires this option for a packaged app to run on iOS15 due to new code signature format. Could be removed once Monterey is miniuS.
+			RunXcconfig.AppendLine("OTHER_CODE_SIGN_FLAGS = --generate-entitlement-der");
+
+			RunXcconfig.AppendLine($"MACOSX_DEPLOYMENT_TARGET = {MacToolChain.Settings.MacOSVersion};");
+			RunXcconfig.AppendLine("INFOPLIST_OUTPUT_FORMAT = xml");
+			RunXcconfig.AppendLine("COMBINE_HIDPI_IMAGES = YES");
+
+			RunXcconfig.AppendLine("SDKROOT[sdk=iphoneos*] = iphoneos");
+			RunXcconfig.AppendLine("SDKROOT[sdk=appletvos*] = appletvos");
+			RunXcconfig.AppendLine("SDKROOT[sdk=macosx*] = macosx");
 
-			TargetXcconfig.AppendLine("// Will add more here later");
 
+			//////////////////// INDEX ////////////////////////
+			IndexXcconfig.AppendLine("ALWAYS_SEARCH_USER_PATHS = NO");
+			IndexXcconfig.AppendLine("CLANG_CXX_LANGUAGE_STANDARD = c++14");
+			IndexXcconfig.AppendLine("GCC_ENABLE_CPP_RTTI = NO");
+			IndexXcconfig.AppendLine("GCC_WARN_CHECK_SWITCH_STATEMENTS = NO");
+			IndexXcconfig.AppendLine("GCC_PRECOMPILE_PREFIX_HEADER = YES");
+			IndexXcconfig.AppendLine($"GCC_PREFIX_HEADER = {UEDir}/Engine/Source/Editor/UnrealEd/Public/UnrealEd.h");
+
+			if (IntelliSensePreprocessorDefinitions.Count > 0)
+			{
+				throw new BuildException("project has defines! " + ProjectFilePath.FullName);
+			}
+			string Definitions = string.Join(" ", IntelliSensePreprocessorDefinitions);
+			IndexXcconfig.AppendLine($"GCC_PREPROCESSOR_DEFINITIONS = __INTELLISENSE__ MONOLITHIC_BUILD=1 {Definitions}");
+
+			string SearchPaths = string.Join(" ", IntelliSenseSystemIncludeSearchPaths.Select(x => x.Contains(" ") ? $"\"{x}\"" : x));
+			IndexXcconfig.AppendLine($"HEADER_SEARCH_PATHS = {SearchPaths}");
+
+			SearchPaths = string.Join(" ", IntelliSenseIncludeSearchPaths.Select(x => x.Contains(" ") ? $"\"{x}\"" : x));
+			IndexXcconfig.AppendLine($"USER_HEADER_SEARCH_PATHS = {SearchPaths}");
+
+
+			//////////////////// BUILD ////////////////////////
+			BuildXcconfig.AppendLine("// Will add more here later");
+
+
+			//////////////////// CLEANUP ////////////////////////
+			// make sure they all written out
 			foreach (XcconfigFile Xcconfig in AllXcconfigFiles)
 			{
-				Xcconfig.Close();
+				Xcconfig.Write();
 			}
 		}
 
@@ -1952,14 +1855,14 @@ namespace UnrealBuildTool
 			};
 		}
 
-		public bool ShouldIncludeProjectInWorkspace(ILogger Logger)
+		public static bool ShouldIncludeProjectInWorkspace(ProjectFile Proj, ILogger Logger)
 		{
-			return CanBuildProjectLocally(Logger);
+			return CanBuildProjectLocally(Proj, Logger);
 		}
 
-		public bool CanBuildProjectLocally(ILogger Logger)
+		public static bool CanBuildProjectLocally(ProjectFile Proj, ILogger Logger)
 		{
-			foreach (Project ProjectTarget in ProjectTargets)
+			foreach (Project ProjectTarget in Proj.ProjectTargets)
 			{
 				foreach (UnrealTargetPlatform Platform in GetSupportedPlatforms())
 				{
@@ -1997,6 +1900,7 @@ namespace UnrealBuildTool
 			string ProjectConfigListGuid = XcodeProjectFileGenerator.MakeXcodeGuid();
 			string MainGroupGuid = XcodeProjectFileGenerator.MakeXcodeGuid();
 			string ProductRefGroupGuid = XcodeProjectFileGenerator.MakeXcodeGuid();
+			string XcconfigsRefGroupGuid = XcodeProjectFileGenerator.MakeXcodeGuid();
 			string SourcesBuildPhaseGuid = XcodeProjectFileGenerator.MakeXcodeGuid();
 			string CopyExtensionsBuildPhaseGuid = XcodeProjectFileGenerator.MakeXcodeGuid();
 			string ShellScriptSectionGuid = XcodeProjectFileGenerator.MakeXcodeGuid();
@@ -2041,8 +1945,17 @@ namespace UnrealBuildTool
 				{
 					bIsAppBundle = false;
 				}
+
+				// make a project-level xcconfig for this configuration
+				XcconfigFile Xcconfig = new XcconfigFile(ProjectFilePath, Config.DisplayName);
+				AllXcconfigFiles.Add(Xcconfig);
+				Xcconfig.AppendLine("// pull in the shared settings for all configs");
+				Xcconfig.AppendLine($"#include \"{ProjectFilePath.GetFileNameWithoutAnyExtensions()}_Project.xcconfig\"");
+				Xcconfig.AppendLine("");
 			}
 
+			WriteProjectXcconfigFiles(ProjectBuildConfigs.Values.ToList(), GameProjectPath, Logger);
+
 			StringBuilder PBXBuildFileSection = new StringBuilder();
 			StringBuilder PBXFileReferenceSection = new StringBuilder();
 			StringBuilder PBXSourcesBuildPhaseSection = new StringBuilder();
@@ -2077,7 +1990,7 @@ namespace UnrealBuildTool
 				AppendContainerItemProxySection(ProjectFileContent, EI.Name, EI.TargetGuid, EI.TargetProxyGuid, ProjectGuid);
 				AppendTargetDependencySection(ProjectFileContent, EI.Name, EI.TargetGuid, EI.TargetDependencyGuid, EI.TargetProxyGuid);
 			}
-			AppendGroupSection(ProjectFileContent, MainGroupGuid, ProductRefGroupGuid, TargetAppGuid, TargetName, AllExtensions);
+			AppendGroupSection(ProjectFileContent, MainGroupGuid, ProductRefGroupGuid, XcconfigsRefGroupGuid, TargetAppGuid, TargetName, AllExtensions);
 			AppendLegacyTargetSection(ProjectFileContent, BuildTargetName, BuildTargetGuid, BuildTargetConfigListGuid, GameProjectPath);
 			AppendRunTargetSection(ProjectFileContent, TargetName, TargetGuid, TargetConfigListGuid, TargetDependencyGuid, TargetAppGuid, CopyExtensionsBuildPhaseGuid, ShellScriptSectionGuid, AllExtensions, bIsAppBundle);
 			AppendIndexTargetSection(ProjectFileContent, IndexTargetName, IndexTargetGuid, IndexTargetConfigListGuid, SourcesBuildPhaseGuid);
@@ -2093,7 +2006,7 @@ namespace UnrealBuildTool
 			ProjectFileContent.Append("}" + ProjectFileGenerator.NewLine);
 
 			// write out .xcconfig file(s) for settings we are moving out of the xcode project file
-			WriteXcconfigFiles();
+			WriteTargetXcconfigFiles();
 
 			if (bSuccess)
 			{
@@ -2101,7 +2014,7 @@ namespace UnrealBuildTool
 				bSuccess = ProjectFileGenerator.WriteFileIfChanged(PBXProjFilePath.FullName, ProjectFileContent.ToString(), Logger, new UTF8Encoding());
 			}
 
-			bool bNeedScheme = CanBuildProjectLocally(Logger);
+			bool bNeedScheme = CanBuildProjectLocally(this, Logger);
 
 			if (bNeedScheme)
 			{
diff --git a/Engine/Source/Programs/UnrealBuildTool/ProjectFiles/Xcode/XcodeProjectFileGenerator.cs b/Engine/Source/Programs/UnrealBuildTool/ProjectFiles/Xcode/XcodeProjectFileGenerator.cs
--- a/Engine/Source/Programs/UnrealBuildTool/ProjectFiles/Xcode/XcodeProjectFileGenerator.cs
+++ b/Engine/Source/Programs/UnrealBuildTool/ProjectFiles/Xcode/XcodeProjectFileGenerator.cs
@@ -1,4 +1,4 @@
-// Copyright Epic Games, Inc. All Rights Reserved.
+﻿// Copyright Epic Games, Inc. All Rights Reserved.
 
 using System;
 using System.Collections.Generic;
@@ -33,6 +33,11 @@
 		/// </summary>
 		string AppName = "";
 
+		/// <summary>
+		/// Whether or not to use the new .xcconfig file 
+		/// </summary>
+		private bool bUseXcconfig = false;
+
 		public XcodeProjectFileGenerator(FileReference? InOnlyGameProject, CommandLineArguments CommandLine)
 			: base(InOnlyGameProject)
 		{
@@ -49,6 +54,11 @@
 			{
 				AppName = CommandLine.GetString("-appname=");
 			}
+
+			if (CommandLine.HasOption("-xcconfig"))
+			{
+				bUseXcconfig = true;
+			}
 		}
 
 		/// <summary>
@@ -110,7 +120,14 @@
 		/// <returns>The newly allocated project file object</returns>
 		protected override ProjectFile AllocateProjectFile(FileReference InitFilePath, DirectoryReference BaseDir)
 		{
-			return new XcodeProjectFile(InitFilePath, BaseDir, bForDistribution, BundleIdentifier, AppName);
+			if (bUseXcconfig)
+			{
+				return new XcodeProjectXcconfig.XcodeProjectFile(InitFilePath, BaseDir, bForDistribution, BundleIdentifier, AppName);
+			}
+			else
+			{
+				return new XcodeProjectLegacy.XcodeProjectFile(InitFilePath, BaseDir, bForDistribution, BundleIdentifier, AppName);
+			}
 		}
 
 		private bool WriteWorkspaceSettingsFile(string Path, ILogger Logger)
@@ -167,7 +184,7 @@
 			WorkspaceDataContent.Append("<Workspace" + ProjectFileGenerator.NewLine);
 			WorkspaceDataContent.Append("   version = \"1.0\">" + ProjectFileGenerator.NewLine);
 
-			List<XcodeProjectFile> BuildableProjects = new List<XcodeProjectFile>();
+			List<ProjectFile> BuildableProjects = new List<ProjectFile>();
 
 			System.Action< List<PrimaryProjectFolder> /* Folders */, string /* Ident */ >? AddProjectsFunction = null;
 			AddProjectsFunction = (FolderList, Ident) =>
@@ -180,12 +197,12 @@
 						AddProjectsFunction!(CurFolder.SubFolders, Ident + "   ");
 				
 						// Filter out anything that isn't an XC project, and that shouldn't be in the workspace
-						IEnumerable<XcodeProjectFile> SupportedProjects =
-								CurFolder.ChildProjects.OfType<XcodeProjectFile>()
-									.Where(P => P.ShouldIncludeProjectInWorkspace(Logger))
+						IEnumerable<ProjectFile> SupportedProjects =
+								CurFolder.ChildProjects.Where(P => P.GetType() == typeof(XcodeProjectXcconfig.XcodeProjectFile) || P.GetType() == typeof(XcodeProjectLegacy.XcodeProjectFile))
+									.Where(P => XcodeProjectXcconfig.XcodeProjectFile.ShouldIncludeProjectInWorkspace(P, Logger))
 									.OrderBy(P => P.ProjectFilePath.GetFileName());
 
-						foreach (XcodeProjectFile XcodeProject in SupportedProjects)
+						foreach (ProjectFile XcodeProject in SupportedProjects)
 						{
 							WorkspaceDataContent.Append(Ident + "      <FileRef" + ProjectFileGenerator.NewLine);
 							WorkspaceDataContent.Append(Ident + "         location = \"group:" + XcodeProject.ProjectFilePath.MakeRelativeTo(ProjectFileGenerator.PrimaryProjectPath) + "\">" + ProjectFileGenerator.NewLine);
@@ -220,7 +237,7 @@
 				return TargetA.Name.CompareTo(TargetB.Name);
 			});
 
-			foreach (XcodeProjectFile XcodeProject in BuildableProjects)
+			foreach (ProjectFile XcodeProject in BuildableProjects)
 			{
 				FileReference SchemeManagementFile = XcodeProject.ProjectFilePath + "/xcuserdata/" + Environment.UserName + ".xcuserdatad/xcschemes/xcschememanagement.plist";
 				if (FileReference.Exists(SchemeManagementFile))
diff --git a/Engine/Source/Programs/UnrealBuildTool/ProjectFiles/Xcode/XcodeProjectLegacy.cs b/Engine/Source/Programs/UnrealBuildTool/ProjectFiles/Xcode/XcodeProjectLegacy.cs
new file mode 100644
--- /dev/null
+++ b/Engine/Source/Programs/UnrealBuildTool/ProjectFiles/Xcode/XcodeProjectLegacy.cs
@@ -0,0 +1,2055 @@
+﻿// Copyright Epic Games, Inc. All Rights Reserved.
+
+using System;
+using System.Collections.Generic;
+using System.IO;
+using System.Xml;
+using System.Xml.XPath;
+using System.Xml.Linq;
+using System.Linq;
+using System.Text;
+using EpicGames.Core;
+using UnrealBuildBase;
+using Microsoft.Extensions.Logging;
+
+namespace UnrealBuildTool.XcodeProjectLegacy
+{
+	/// <summary>
+	/// Info needed to make a file a member of specific group
+	/// </summary>
+	class XcodeSourceFile : ProjectFile.SourceFile
+	{
+		/// <summary>
+		/// Constructor
+		/// </summary>
+		public XcodeSourceFile(FileReference InitFilePath, DirectoryReference? InitRelativeBaseFolder)
+			: base(InitFilePath, InitRelativeBaseFolder)
+		{
+			FileGuid = XcodeProjectFileGenerator.MakeXcodeGuid();
+			FileRefGuid = XcodeProjectFileGenerator.MakeXcodeGuid();
+		}
+
+		/// <summary>
+		/// File Guid for use in Xcode project
+		/// </summary>
+		public string FileGuid
+		{
+			get;
+			private set;
+		}
+
+		public void ReplaceGuids(string NewFileGuid, string NewFileRefGuid)
+		{
+			FileGuid = NewFileGuid;
+			FileRefGuid = NewFileRefGuid;
+		}
+
+		/// <summary>
+		/// File reference Guid for use in Xcode project
+		/// </summary>
+		public string FileRefGuid
+		{
+			get;
+			private set;
+		}
+	}
+
+	/// <summary>
+	/// Represents a group of files shown in Xcode's project navigator as a folder
+	/// </summary>
+	class XcodeFileGroup
+	{
+		public XcodeFileGroup(string InName, string InPath, bool InIsReference)
+		{
+			GroupName = InName;
+			GroupPath = InPath;
+			GroupGuid = XcodeProjectFileGenerator.MakeXcodeGuid();
+			bIsReference = InIsReference;
+		}
+
+		public string GroupGuid;
+		public string GroupName;
+		public string GroupPath;
+		public Dictionary<string, XcodeFileGroup> Children = new Dictionary<string, XcodeFileGroup>();
+		public List<XcodeSourceFile> Files = new List<XcodeSourceFile>();
+		public bool bIsReference;
+	}
+
+	class XcodeBuildConfig
+	{
+		public XcodeBuildConfig(string InDisplayName, string InBuildTarget, FileReference? InMacExecutablePath, FileReference? InIOSExecutablePath, FileReference? InTVOSExecutablePath,
+			ProjectTarget? InProjectTarget, UnrealTargetConfiguration InBuildConfig)
+		{
+			DisplayName = InDisplayName;
+			MacExecutablePath = InMacExecutablePath;
+			IOSExecutablePath = InIOSExecutablePath;
+			TVOSExecutablePath = InTVOSExecutablePath;
+			BuildTarget = InBuildTarget;
+			ProjectTarget = InProjectTarget;
+			BuildConfig = InBuildConfig;
+		}
+
+		public string DisplayName;
+		public FileReference? MacExecutablePath;
+		public FileReference? IOSExecutablePath;
+		public FileReference? TVOSExecutablePath;
+		public string BuildTarget;
+		public ProjectTarget? ProjectTarget;
+		public UnrealTargetConfiguration BuildConfig;
+	};
+
+	class XcodeExtensionInfo
+	{
+		public XcodeExtensionInfo(string InName)
+		{
+			Name = InName;
+			TargetDependencyGuid = XcodeProjectFileGenerator.MakeXcodeGuid();
+			TargetProxyGuid = XcodeProjectFileGenerator.MakeXcodeGuid();
+			TargetGuid = XcodeProjectFileGenerator.MakeXcodeGuid();
+			ProductGuid = XcodeProjectFileGenerator.MakeXcodeGuid();
+			ResourceBuildPhaseGuid = XcodeProjectFileGenerator.MakeXcodeGuid();
+			ConfigListGuid = XcodeProjectFileGenerator.MakeXcodeGuid();
+			AllConfigs = new Dictionary<string, XcodeBuildConfig>();
+		}
+
+		public string Name;
+		public string TargetDependencyGuid;
+		public string TargetProxyGuid;
+		public string TargetGuid;
+		public string ProductGuid;
+		public string ResourceBuildPhaseGuid;
+		public string ConfigListGuid;
+		public Dictionary<string, XcodeBuildConfig> AllConfigs;
+
+		public string? ConfigurationContents;
+	}
+
+	class XcodeProjectFile : ProjectFile
+	{
+		Dictionary<string, XcodeFileGroup> Groups = new Dictionary<string, XcodeFileGroup>();
+
+		/// <summary>
+		/// Constructs a new project file object
+		/// </summary>
+		/// <param name="InitFilePath">The path to the project file on disk</param>
+		/// <param name="BaseDir">The base directory for files within this project</param>
+		/// <param name="IsForDistribution">True for distribution builds</param>
+		/// <param name="BundleID">Override option for bundle identifier</param>
+		/// <param name="InAppName"></param>
+		public XcodeProjectFile(FileReference InitFilePath, DirectoryReference BaseDir, bool IsForDistribution, string BundleID, string InAppName)
+			: base(InitFilePath, BaseDir)
+		{
+			bForDistribution = IsForDistribution;
+			BundleIdentifier = BundleID;
+			AppName = InAppName;
+		}
+
+		public override string ToString()
+		{
+			return ProjectFilePath.GetFileNameWithoutExtension();
+		}
+
+		/// <summary>
+		///  Used to mark the project for distribution (some platforms require this)
+		/// </summary>
+		bool bForDistribution = false;
+
+		/// <summary>
+		/// Override for bundle identifier
+		/// </summary>
+		string BundleIdentifier = "";
+
+		/// <summary>
+		/// Override AppName
+		/// </summary>
+		string AppName = "";
+
+		/// <summary>
+		/// Architectures supported for iOS
+		/// </summary>
+		string[] SupportedIOSArchitectures = { "arm64" };
+
+		/// <summary>
+		/// Gets Xcode file category based on its extension
+		/// </summary>
+		private string GetFileCategory(string Extension)
+		{
+			// @todo Mac: Handle more categories
+			switch (Extension)
+			{
+				case ".framework":
+					return "Frameworks";
+				default:
+					return "Sources";
+			}
+		}
+
+		/// <summary>
+		/// Gets Xcode file type based on its extension
+		/// </summary>
+		private string GetFileType(string Extension)
+		{
+			// @todo Mac: Handle more file types
+			switch (Extension)
+			{
+				case ".c":
+				case ".m":
+					return "sourcecode.c.objc";
+				case ".cc":
+				case ".cpp":
+				case ".mm":
+					return "sourcecode.cpp.objcpp";
+				case ".h":
+				case ".inl":
+				case ".pch":
+					return "sourcecode.c.h";
+				case ".framework":
+					return "wrapper.framework";
+				case ".plist":
+					return "text.plist.xml";
+				case ".png":
+					return "image.png";
+				case ".icns":
+					return "image.icns";
+				default:
+					return "file.text";
+			}
+		}
+
+		/// <summary>
+		/// Returns true if Extension is a known extension for files containing source code
+		/// </summary>
+		private bool IsSourceCode(string Extension)
+		{
+			return Extension == ".c" || Extension == ".cc" || Extension == ".cpp" || Extension == ".m" || Extension == ".mm";
+		}
+
+		private bool ShouldIncludeFileInBuildPhaseSection(XcodeSourceFile SourceFile)
+		{
+			string FileExtension = SourceFile.Reference.GetExtension();
+
+			if (IsSourceCode(FileExtension))
+			{
+				foreach (string PlatformName in UnrealTargetPlatform.GetValidPlatformNames())
+				{
+					string AltName = PlatformName == "Win64" ? "windows" : PlatformName.ToLower();
+					if ((SourceFile.Reference.FullName.ToLower().Contains("/" + PlatformName.ToLower() + "/") || SourceFile.Reference.FullName.ToLower().Contains("/" + AltName + "/"))
+						&& PlatformName != "Mac" && PlatformName != "IOS" && PlatformName != "TVOS")
+					{
+						// Build phase is used for indexing only and indexing currently works only with files that can be compiled for Mac, so skip files for other platforms
+						return false;
+					}
+				}
+
+				return true;
+			}
+
+			return false;
+		}
+
+		/// <summary>
+		/// Returns a project navigator group to which the file should belong based on its path.
+		/// Creates a group tree if it doesn't exist yet.
+		/// </summary>
+		public XcodeFileGroup? FindGroupByAbsolutePath(ref Dictionary<string, XcodeFileGroup> Groups, string AbsolutePath)
+		{
+			string[] Parts = AbsolutePath.Split(Path.DirectorySeparatorChar);
+			string CurrentPath = "/";
+			Dictionary<string, XcodeFileGroup> CurrentSubGroups = Groups;
+
+			for (int Index = 1; Index < Parts.Count(); ++Index)
+			{
+				string Part = Parts[Index];
+
+				if (CurrentPath.Length > 1)
+				{
+					CurrentPath += Path.DirectorySeparatorChar;
+				}
+
+				CurrentPath += Part;
+
+				XcodeFileGroup CurrentGroup;
+				if (!CurrentSubGroups.ContainsKey(CurrentPath))
+				{
+					CurrentGroup = new XcodeFileGroup(Path.GetFileName(CurrentPath), CurrentPath, CurrentPath.EndsWith(".xcassets"));
+					CurrentSubGroups.Add(CurrentPath, CurrentGroup);
+				}
+				else
+				{
+					CurrentGroup = CurrentSubGroups[CurrentPath];
+				}
+
+				if (CurrentPath == AbsolutePath)
+				{
+					return CurrentGroup;
+				}
+
+				CurrentSubGroups = CurrentGroup.Children;
+			}
+
+			return null;
+		}
+
+		/// <summary>
+		/// Convert all paths to Apple/Unix format (with forward slashes)
+		/// </summary>
+		/// <param name="InPath">The path to convert</param>
+		/// <returns>The normalized path</returns>
+		private static string ConvertPath(string InPath)
+		{
+			return InPath.Replace("\\", "/");
+		}
+
+		/// <summary>
+		/// Allocates a generator-specific source file object
+		/// </summary>
+		/// <param name="InitFilePath">Path to the source file on disk</param>
+		/// <param name="InitProjectSubFolder">Optional sub-folder to put the file in.  If empty, this will be determined automatically from the file's path relative to the project file</param>
+		/// <returns>The newly allocated source file object</returns>
+		public override SourceFile? AllocSourceFile(FileReference InitFilePath, DirectoryReference? InitProjectSubFolder)
+		{
+			if (InitFilePath.GetFileName().StartsWith("."))
+			{
+				return null;
+			}
+			return new XcodeSourceFile(InitFilePath, InitProjectSubFolder);
+		}
+
+		/// <summary>
+		/// Generates bodies of all sections that contain a list of source files plus a dictionary of project navigator groups.
+		/// </summary>
+		private void GenerateSectionsWithSourceFiles(StringBuilder PBXBuildFileSection, StringBuilder PBXFileReferenceSection, StringBuilder PBXSourcesBuildPhaseSection, string TargetAppGuid, string TargetName, bool bIsAppBundle)
+		{
+			SourceFiles.Sort((x, y) => { return x.Reference.FullName.CompareTo(y.Reference.FullName); });
+
+			foreach (XcodeSourceFile SourceFile in SourceFiles.OfType<XcodeSourceFile>())
+			{
+				string FileName = SourceFile.Reference.GetFileName();
+				string FileExtension = Path.GetExtension(FileName);
+				string FilePath = SourceFile.Reference.MakeRelativeTo(ProjectFilePath.Directory);
+				string FilePathMac = Utils.CleanDirectorySeparators(FilePath, '/');
+
+				if (IsGeneratedProject)
+				{
+					PBXBuildFileSection.Append(string.Format("\t\t{0} /* {1} in {2} */ = {{isa = PBXBuildFile; fileRef = {3} /* {1} */; }};" + ProjectFileGenerator.NewLine,
+						SourceFile.FileGuid,
+						FileName,
+						GetFileCategory(FileExtension),
+						SourceFile.FileRefGuid));
+				}
+
+				PBXFileReferenceSection.Append(string.Format("\t\t{0} /* {1} */ = {{isa = PBXFileReference; explicitFileType = {2}; name = \"{1}\"; path = \"{3}\"; sourceTree = SOURCE_ROOT; }};" + ProjectFileGenerator.NewLine,
+					SourceFile.FileRefGuid,
+					FileName,
+					GetFileType(FileExtension),
+					FilePathMac));
+
+				if (ShouldIncludeFileInBuildPhaseSection(SourceFile))
+				{
+					PBXSourcesBuildPhaseSection.Append("\t\t\t\t" + SourceFile.FileGuid + " /* " + FileName + " in Sources */," + ProjectFileGenerator.NewLine);
+				}
+
+				XcodeFileGroup? Group = FindGroupByAbsolutePath(ref Groups, SourceFile.Reference.Directory.FullName);
+				if (Group != null)
+				{
+					Group.Files.Add(SourceFile);
+				}
+			}
+
+			PBXFileReferenceSection.Append(string.Format("\t\t{0} /* {1} */ = {{isa = PBXFileReference; explicitFileType = {2}; path = {1}; sourceTree = BUILT_PRODUCTS_DIR; }};" + ProjectFileGenerator.NewLine, TargetAppGuid, TargetName, bIsAppBundle ? "wrapper.application" : "\"compiled.mach-o.executable\""));
+		}
+
+		private void GenerateSectionsWithExtensions(StringBuilder PBXBuildFileSection, StringBuilder PBXFileReferenceSection, StringBuilder PBXCopyFilesBuildPhaseSection, StringBuilder PBXResourcesBuildPhaseSection,
+													List<XcodeExtensionInfo> AllExtensions, FileReference? UProjectPath, List<XcodeBuildConfig> BuildConfigs, ILogger Logger)
+		{
+			if (UProjectPath != null)
+			{
+				string ProjectExtensionsDir = Path.Combine(Path.GetDirectoryName(UProjectPath.FullName)!, "Build/IOS/Extensions");
+				//string ProjectIntermediateDir = Path.Combine(Path.GetDirectoryName(UProjectPath.FullName), "Intermediate/IOS/Extensions");
+
+				if (Directory.Exists(ProjectExtensionsDir))
+				{
+					foreach (DirectoryInfo DI in new System.IO.DirectoryInfo(ProjectExtensionsDir).EnumerateDirectories())
+					{
+						Console.WriteLine("  Project {0} has Extension {1}!", UProjectPath, DI);
+
+						// assume each Extension in here will create a resulting Extension.appex
+						string Extension = DI.Name + ".appex";
+
+						string ExtensionGuid = XcodeProjectFileGenerator.MakeXcodeGuid();
+
+						// make an extension info object
+						XcodeExtensionInfo ExtensionInfo = new XcodeExtensionInfo(DI.Name);
+						AllExtensions.Add(ExtensionInfo);
+
+						PBXBuildFileSection.Append(string.Format("\t\t{0} /* {1} in Embed App Extensions */ = {{isa = PBXBuildFile; fileRef = {2} /* {1} */; settings = {{ATTRIBUTES = (RemoveHeadersOnCopy, ); }}; }};" + ProjectFileGenerator.NewLine,
+							ExtensionGuid,
+							Extension,
+							ExtensionInfo.ProductGuid));
+
+						PBXFileReferenceSection.Append(string.Format("\t\t{0} /* {1} */ = {{isa = PBXFileReference; explicitFileType = wrapper.app-extension; path = \"{1}\"; sourceTree = BUILT_PRODUCTS_DIR; }};" + ProjectFileGenerator.NewLine,
+							ExtensionInfo.ProductGuid,
+							Extension));
+
+						PBXCopyFilesBuildPhaseSection.Append(string.Format("\t\t\t\t{0} /* {1} in Embed App Extensions */," + ProjectFileGenerator.NewLine,
+							ExtensionGuid,
+							Extension));
+
+						PBXResourcesBuildPhaseSection.Append("/* Begin PBXResourcesBuildPhase section */" + ProjectFileGenerator.NewLine);
+						PBXResourcesBuildPhaseSection.Append("\t\t" + ExtensionInfo.ResourceBuildPhaseGuid + " /* Resources */ = {" + ProjectFileGenerator.NewLine);
+						PBXResourcesBuildPhaseSection.Append("\t\t\tisa = PBXResourcesBuildPhase;" + ProjectFileGenerator.NewLine);
+						PBXResourcesBuildPhaseSection.Append("\t\t\tbuildActionMask = 2147483647;" + ProjectFileGenerator.NewLine);
+						PBXResourcesBuildPhaseSection.Append("\t\t\tfiles = (" + ProjectFileGenerator.NewLine);
+						if (Directory.Exists(Path.Combine(DI.FullName, "Resources")))
+						{
+							DirectoryInfo ResourceDir = new System.IO.DirectoryInfo(Path.Combine(DI.FullName, "Resources"));
+							foreach (FileSystemInfo FSI in ResourceDir.EnumerateFileSystemInfos())
+							{
+								if (FSI.Name.StartsWith("."))
+								{
+									continue;
+								}
+								// for each resource, put it into the File/FileRef section, and into the CopyResuorceBuildPhase
+								string ResourceGuid = XcodeProjectFileGenerator.MakeXcodeGuid();
+								string ResourceRefGuid = XcodeProjectFileGenerator.MakeXcodeGuid();
+								PBXBuildFileSection.Append(string.Format("\t\t{0} /* {1} in Embed App Extensions */ = {{isa = PBXBuildFile; fileRef = {2} /* {1} */; }};" + ProjectFileGenerator.NewLine,
+									ResourceGuid,
+									FSI.Name,
+									ResourceRefGuid));
+
+								// lastKnownFileType = wrapper.app-extension; 
+								PBXFileReferenceSection.Append(string.Format("\t\t{0} /* {1} */ = {{isa = PBXFileReference; lastKnownFileType = folder.assetcatalog; path = \"{2}\"; sourceTree = \"<absolute>\"; }};" + ProjectFileGenerator.NewLine,
+									ResourceRefGuid,
+									FSI.Name,
+									// @todo: make this relative path!! 
+									FSI.FullName));
+
+								PBXResourcesBuildPhaseSection.Append("\t\t\t\t" + ResourceGuid + " /* " + FSI.Name + " in " + ResourceDir.Name + " */," + ProjectFileGenerator.NewLine);
+							}
+						}
+						PBXResourcesBuildPhaseSection.Append("\t\t\t);" + ProjectFileGenerator.NewLine);
+						PBXResourcesBuildPhaseSection.Append("\t\t\trunOnlyForDeploymentPostprocessing = 0;" + ProjectFileGenerator.NewLine);
+						PBXResourcesBuildPhaseSection.Append("\t\t};" + ProjectFileGenerator.NewLine);
+						PBXResourcesBuildPhaseSection.Append("/* End PBXResourcesBuildPhase section */" + ProjectFileGenerator.NewLine + ProjectFileGenerator.NewLine);
+
+						StringBuilder ConfigSection = new StringBuilder();
+						// copy over the configs from the general project to the extension
+						foreach (XcodeBuildConfig Configuration in BuildConfigs)
+						{
+							string ConfigGuid = XcodeProjectFileGenerator.MakeXcodeGuid();
+							string ConfigName = Configuration.DisplayName;
+
+							ConfigSection.Append("\t\t" + ConfigGuid + " /* " + ConfigName + " */ = {" + ProjectFileGenerator.NewLine);
+							ConfigSection.Append("\t\t\tisa = XCBuildConfiguration;" + ProjectFileGenerator.NewLine);
+							ConfigSection.Append("\t\t\tbuildSettings = {" + ProjectFileGenerator.NewLine);
+							ConfigSection.Append("\t\t\t\tASSETCATALOG_COMPILER_APPICON_NAME = \"iMessage App Icon\";" + ProjectFileGenerator.NewLine);
+							ConfigSection.Append("\t\t\t\tINFOPLIST_FILE = \"" + Path.Combine(DI.FullName, "Info.plist") + "\";" + ProjectFileGenerator.NewLine);
+							ConfigSection.Append("\t\t\t\tSKIP_INSTALL = YES;" + ProjectFileGenerator.NewLine);
+							ConfigSection.Append("\t\t\t\tPRODUCT_NAME = \"$(TARGET_NAME)\";" + ProjectFileGenerator.NewLine);
+
+							bool bSupportIOS = true;
+							bool bSupportTVOS = true;
+							if (bSupportIOS && InstalledPlatformInfo.IsValidPlatform(UnrealTargetPlatform.IOS, EProjectType.Code))
+							{
+								IOSPlatform IOSPlatform = ((IOSPlatform)UEBuildPlatform.GetBuildPlatform(UnrealTargetPlatform.IOS));
+								IOSProjectSettings ProjectSettings = IOSPlatform.ReadProjectSettings(UProjectPath);
+								ConfigSection.Append("\t\t\t\t\"PRODUCT_BUNDLE_IDENTIFIER[sdk=iphoneos*]\" = " + ProjectSettings.BundleIdentifier + "." + ExtensionInfo.Name + ";" + ProjectFileGenerator.NewLine);
+							}
+
+							if (bSupportTVOS && InstalledPlatformInfo.IsValidPlatform(UnrealTargetPlatform.TVOS, EProjectType.Code))
+							{
+								TVOSPlatform TVOSPlatform = ((TVOSPlatform)UEBuildPlatform.GetBuildPlatform(UnrealTargetPlatform.TVOS));
+								TVOSProjectSettings ProjectSettings = TVOSPlatform.ReadProjectSettings(UProjectPath);
+								ConfigSection.Append("\t\t\t\t\"PRODUCT_BUNDLE_IDENTIFIER[sdk=appletvos*]\" = " + ProjectSettings.BundleIdentifier + "." + ExtensionInfo.Name + ";" + ProjectFileGenerator.NewLine);
+							}
+
+							string? IOSRuntimeVersion, TVOSRuntimeVersion;
+							AppendPlatformConfiguration(ConfigSection, Configuration, ExtensionInfo.Name, UProjectPath, false, bSupportIOS, bSupportTVOS, Logger, out IOSRuntimeVersion, out TVOSRuntimeVersion);
+
+							ConfigSection.Append("\t\t\t};" + ProjectFileGenerator.NewLine);
+							ConfigSection.Append("\t\t\tname = \"" + ConfigName + "\";" + ProjectFileGenerator.NewLine);
+							ConfigSection.Append("\t\t};" + ProjectFileGenerator.NewLine);
+
+							XcodeBuildConfig Config = new XcodeBuildConfig(ConfigName, ExtensionInfo.Name, null, null, null, null, Configuration.BuildConfig);
+							ExtensionInfo.AllConfigs.Add(ConfigGuid, Config);
+						}
+
+						ExtensionInfo.ConfigurationContents = ConfigSection.ToString();
+					}
+				}
+			}
+		}
+
+		private void AppendGroup(XcodeFileGroup Group, StringBuilder Content)
+		{
+			if (!Group.bIsReference)
+			{
+				Content.Append(string.Format("\t\t{0} = {{{1}", Group.GroupGuid, ProjectFileGenerator.NewLine));
+				Content.Append("\t\t\tisa = PBXGroup;" + ProjectFileGenerator.NewLine);
+				Content.Append("\t\t\tchildren = (" + ProjectFileGenerator.NewLine);
+
+				foreach (XcodeFileGroup ChildGroup in Group.Children.Values)
+				{
+					Content.Append(string.Format("\t\t\t\t{0} /* {1} */,{2}", ChildGroup.GroupGuid, ChildGroup.GroupName, ProjectFileGenerator.NewLine));
+				}
+
+				foreach (XcodeSourceFile File in Group.Files)
+				{
+					Content.Append(string.Format("\t\t\t\t{0} /* {1} */,{2}", File.FileRefGuid, File.Reference.GetFileName(), ProjectFileGenerator.NewLine));
+				}
+
+				Content.Append("\t\t\t);" + ProjectFileGenerator.NewLine);
+				Content.Append("\t\t\tname = \"" + Group.GroupName + "\";" + ProjectFileGenerator.NewLine);
+				Content.Append("\t\t\tpath = \"" + Group.GroupPath + "\";" + ProjectFileGenerator.NewLine);
+				Content.Append("\t\t\tsourceTree = \"<absolute>\";" + ProjectFileGenerator.NewLine);
+				Content.Append("\t\t};" + ProjectFileGenerator.NewLine);
+
+				foreach (XcodeFileGroup ChildGroup in Group.Children.Values)
+				{
+					AppendGroup(ChildGroup, Content);
+				}
+			}
+		}
+
+		private void AppendBuildFileSection(StringBuilder Content, StringBuilder SectionContent)
+		{
+			Content.Append("/* Begin PBXBuildFile section */" + ProjectFileGenerator.NewLine);
+			Content.Append(SectionContent);
+			Content.Append("/* End PBXBuildFile section */" + ProjectFileGenerator.NewLine + ProjectFileGenerator.NewLine);
+		}
+
+		private void AppendFileReferenceSection(StringBuilder Content, StringBuilder SectionContent)
+		{
+			Content.Append("/* Begin PBXFileReference section */" + ProjectFileGenerator.NewLine);
+			Content.Append(SectionContent);
+			Content.Append("/* End PBXFileReference section */" + ProjectFileGenerator.NewLine + ProjectFileGenerator.NewLine);
+		}
+
+		private void AppendSourcesBuildPhaseSection(StringBuilder Content, StringBuilder SectionContent, string SourcesBuildPhaseGuid)
+		{
+			Content.Append("/* Begin PBXSourcesBuildPhase section */" + ProjectFileGenerator.NewLine);
+			Content.Append(string.Format("\t\t{0} = {{{1}", SourcesBuildPhaseGuid, ProjectFileGenerator.NewLine));
+			Content.Append("\t\t\tisa = PBXSourcesBuildPhase;" + ProjectFileGenerator.NewLine);
+			Content.Append("\t\t\tbuildActionMask = 2147483647;" + ProjectFileGenerator.NewLine);
+			Content.Append("\t\t\tfiles = (" + ProjectFileGenerator.NewLine);
+			Content.Append(SectionContent);
+			Content.Append("\t\t\t);" + ProjectFileGenerator.NewLine);
+			Content.Append("\t\t\trunOnlyForDeploymentPostprocessing = 0;" + ProjectFileGenerator.NewLine);
+			Content.Append("\t\t};" + ProjectFileGenerator.NewLine);
+			Content.Append("/* End PBXSourcesBuildPhase section */" + ProjectFileGenerator.NewLine + ProjectFileGenerator.NewLine);
+		}
+
+		private XcodeFileGroup? FindRootFileGroup(Dictionary<string, XcodeFileGroup> GroupsDict)
+		{
+			foreach (XcodeFileGroup Group in GroupsDict.Values)
+			{
+				if (Group.Children.Count > 1 || Group.Files.Count > 0)
+				{
+					return Group;
+				}
+				else
+				{
+					XcodeFileGroup? Found = FindRootFileGroup(Group.Children);
+					if (Found != null)
+					{
+						return Found;
+					}
+				}
+			}
+			return null;
+		}
+
+		private void AppendCopyExtensionsBuildPhaseSection(StringBuilder Content, StringBuilder SectionContent, string CopyFilesBuildPhaseGuid)
+		{
+			Content.Append("/* Begin PBXCopyFilesBuildPhase section */" + ProjectFileGenerator.NewLine);
+			Content.Append(string.Format("\t{0} /* Embed App Extensions */ = {{{1}", CopyFilesBuildPhaseGuid, ProjectFileGenerator.NewLine));
+			Content.Append("\t\tisa = PBXCopyFilesBuildPhase;" + ProjectFileGenerator.NewLine);
+			Content.Append("\t\tbuildActionMask = 2147483647;" + ProjectFileGenerator.NewLine);
+			Content.Append("\t\tdstPath = \"\";" + ProjectFileGenerator.NewLine);
+			Content.Append("\t\tdstSubfolderSpec = 13;" + ProjectFileGenerator.NewLine);
+			Content.Append("\t\tfiles = (" + ProjectFileGenerator.NewLine);
+			Content.Append(SectionContent);
+			Content.Append("\t\t);" + ProjectFileGenerator.NewLine);
+			Content.Append("\t\tname = \"Embed App Extensions\";" + ProjectFileGenerator.NewLine);
+			Content.Append("\t\trunOnlyForDeploymentPostprocessing = 0;" + ProjectFileGenerator.NewLine);
+			Content.Append("\t};" + ProjectFileGenerator.NewLine);
+			Content.Append("/* End PBXCopyFilesBuildPhase section */" + ProjectFileGenerator.NewLine);
+		}
+
+		private void AppendGroupSection(StringBuilder Content, string MainGroupGuid, string ProductRefGroupGuid, string TargetAppGuid, string TargetName, List<XcodeExtensionInfo> AllExtensions)
+		{
+			XcodeFileGroup? RootGroup = FindRootFileGroup(Groups);
+			if (RootGroup == null)
+			{
+				return;
+			}
+
+			Content.Append("/* Begin PBXGroup section */" + ProjectFileGenerator.NewLine);
+
+			// Main group
+			Content.Append(string.Format("\t\t{0} = {{{1}", MainGroupGuid, ProjectFileGenerator.NewLine));
+			Content.Append("\t\t\tisa = PBXGroup;" + ProjectFileGenerator.NewLine);
+			Content.Append("\t\t\tchildren = (" + ProjectFileGenerator.NewLine);
+
+			foreach (XcodeFileGroup Group in RootGroup.Children.Values)
+			{
+				Content.Append(string.Format("\t\t\t\t{0} /* {1} */,{2}", Group.GroupGuid, Group.GroupName, ProjectFileGenerator.NewLine));
+			}
+
+			foreach (XcodeSourceFile File in RootGroup.Files)
+			{
+				Content.Append(string.Format("\t\t\t\t{0} /* {1} */,{2}", File.FileRefGuid, File.Reference.GetFileName(), ProjectFileGenerator.NewLine));
+			}
+
+			Content.Append(string.Format("\t\t\t\t{0} /* Products */,{1}", ProductRefGroupGuid, ProjectFileGenerator.NewLine));
+			Content.Append("\t\t\t);" + ProjectFileGenerator.NewLine);
+			Content.Append("\t\t\tsourceTree = \"<group>\";" + ProjectFileGenerator.NewLine);
+			Content.Append("\t\t};" + ProjectFileGenerator.NewLine);
+
+			// Sources groups
+			foreach (XcodeFileGroup Group in RootGroup.Children.Values)
+			{
+				AppendGroup(Group, Content);
+			}
+
+			// Products group
+			Content.Append(string.Format("\t\t{0} /* Products */ = {{{1}", ProductRefGroupGuid, ProjectFileGenerator.NewLine));
+			Content.Append("\t\t\tisa = PBXGroup;" + ProjectFileGenerator.NewLine);
+			Content.Append("\t\t\tchildren = (" + ProjectFileGenerator.NewLine);
+			Content.Append(string.Format("\t\t\t\t{0} /* {1} */,{2}", TargetAppGuid, TargetName, ProjectFileGenerator.NewLine));
+			foreach (XcodeExtensionInfo EI in AllExtensions)
+			{
+				Content.Append(string.Format("\t\t\t\t{0} /* {1} */,{2}", EI.ProductGuid, EI.Name, ProjectFileGenerator.NewLine));
+			}
+			Content.Append("\t\t\t);" + ProjectFileGenerator.NewLine);
+			Content.Append("\t\t\tname = Products;" + ProjectFileGenerator.NewLine);
+			Content.Append("\t\t\tsourceTree = \"<group>\";" + ProjectFileGenerator.NewLine);
+			Content.Append("\t\t};" + ProjectFileGenerator.NewLine);
+
+			Content.Append("/* End PBXGroup section */" + ProjectFileGenerator.NewLine + ProjectFileGenerator.NewLine);
+		}
+
+		private void AppendLegacyTargetSection(StringBuilder Content, string TargetName, string TargetGuid, string TargetBuildConfigGuid, FileReference? UProjectPath)
+		{
+			string UEDir = ConvertPath(Path.GetFullPath(Directory.GetCurrentDirectory() + "../../.."));
+			string BuildToolPath = UEDir + "/Engine/Build/BatchFiles/Mac/XcodeBuild.sh";
+
+			Content.Append("/* Begin PBXLegacyTarget section */" + ProjectFileGenerator.NewLine);
+
+			Content.Append("\t\t" + TargetGuid + " /* " + TargetName + " */ = {" + ProjectFileGenerator.NewLine);
+			Content.Append("\t\t\tisa = PBXLegacyTarget;" + ProjectFileGenerator.NewLine);
+			Content.Append("\t\t\tbuildArgumentsString = \"$(ACTION) $(UE_BUILD_TARGET_NAME) $(PLATFORM_NAME) $(UE_BUILD_TARGET_CONFIG)"
+				+ (UProjectPath == null ? "" : " \\\"" + UProjectPath.FullName + "\\\"")
+				+ "\";" + ProjectFileGenerator.NewLine);
+			Content.Append("\t\t\tbuildConfigurationList = " + TargetBuildConfigGuid + " /* Build configuration list for PBXLegacyTarget \"" + TargetName + "\" */;" + ProjectFileGenerator.NewLine);
+			Content.Append("\t\t\tbuildPhases = (" + ProjectFileGenerator.NewLine);
+			Content.Append("\t\t\t);" + ProjectFileGenerator.NewLine);
+			Content.Append("\t\t\tbuildToolPath = \"" + BuildToolPath + "\";" + ProjectFileGenerator.NewLine);
+			Content.Append("\t\t\tbuildWorkingDirectory = \"" + UEDir + "\";" + ProjectFileGenerator.NewLine);
+			Content.Append("\t\t\tdependencies = (" + ProjectFileGenerator.NewLine);
+			Content.Append("\t\t\t);" + ProjectFileGenerator.NewLine);
+			Content.Append("\t\t\tname = \"" + TargetName + "\";" + ProjectFileGenerator.NewLine);
+			Content.Append("\t\t\tpassBuildSettingsInEnvironment = 1;" + ProjectFileGenerator.NewLine);
+			Content.Append("\t\t\tproductName = \"" + TargetName + "\";" + ProjectFileGenerator.NewLine);
+			Content.Append("\t\t};" + ProjectFileGenerator.NewLine);
+
+			Content.Append("/* End PBXLegacyTarget section */" + ProjectFileGenerator.NewLine + ProjectFileGenerator.NewLine);
+		}
+
+		private void AppendRunTargetSection(StringBuilder Content, string TargetName, string TargetGuid, string TargetBuildConfigGuid, string TargetDependencyGuid,
+				string TargetAppGuid, string CopyExtensionsBuildPhaseGuid, string ShellScriptSectionGuid, List<XcodeExtensionInfo> AllExtensions, bool bIsAppBundle)
+		{
+			List<string> DependencyGuids = new List<string>();
+			// depends on the Run target if we want one
+			if (!XcodeProjectFileGenerator.bGeneratingRunIOSProject && !XcodeProjectFileGenerator.bGeneratingRunTVOSProject)
+			{
+				DependencyGuids.Add(TargetDependencyGuid);
+			}
+			// make sure extensions get built
+			foreach (XcodeExtensionInfo EI in AllExtensions)
+			{
+				DependencyGuids.Add(EI.TargetDependencyGuid);
+			}
+
+			Dictionary<string, string> BuildPhases = new Dictionary<string, string>();
+			// add optional build phases
+			if (!string.IsNullOrEmpty(CopyExtensionsBuildPhaseGuid))
+			{
+				BuildPhases.Add(CopyExtensionsBuildPhaseGuid, "Embed App Extensions");
+			}
+			if (!string.IsNullOrEmpty(ShellScriptSectionGuid))
+			{
+				BuildPhases.Add(ShellScriptSectionGuid, "Shell Script");
+			}
+
+			// use generica target section function for an application type
+			AppendGenericTargetSection(Content, TargetName, TargetGuid, bIsAppBundle ? "com.apple.product-type.application" : "com.apple.product-type.tool", TargetBuildConfigGuid, TargetAppGuid, DependencyGuids, BuildPhases);
+		}
+
+		private void AppendGenericTargetSection(StringBuilder Content, string TargetName, string TargetGuid, string TargetType, string TargetBuildConfigGuid, string TargetAppGuid, IEnumerable<string>? TargetDependencyGuids, Dictionary<string, string> BuildPhases)
+		{
+			Content.Append("/* Begin PBXNativeTarget section */" + ProjectFileGenerator.NewLine);
+
+			Content.Append("\t\t" + TargetGuid + " /* " + TargetName + " */ = {" + ProjectFileGenerator.NewLine);
+			Content.Append("\t\t\tisa = PBXNativeTarget;" + ProjectFileGenerator.NewLine);
+			Content.Append("\t\t\tbuildConfigurationList = " + TargetBuildConfigGuid + " /* Build configuration list for PBXNativeTarget \"" + TargetName + "\" */;" + ProjectFileGenerator.NewLine);
+			Content.Append("\t\t\tbuildPhases = (" + ProjectFileGenerator.NewLine);
+			if (BuildPhases != null)
+			{
+				foreach (KeyValuePair<string, string> BuildPhasePair in BuildPhases)
+				{
+					Content.Append("\t\t\t\t" + BuildPhasePair.Key + " /* " + BuildPhasePair.Value + " */, " + ProjectFileGenerator.NewLine);
+				}
+			}
+			Content.Append("\t\t\t);" + ProjectFileGenerator.NewLine);
+			Content.Append("\t\t\tdependencies = (" + ProjectFileGenerator.NewLine);
+			if (TargetDependencyGuids != null)
+			{
+				foreach (string DependencyGuid in TargetDependencyGuids)
+				{
+					Content.Append("\t\t\t\t" + DependencyGuid + " /* PBXTargetDependency */," + ProjectFileGenerator.NewLine);
+				}
+			}
+			Content.Append("\t\t\t);" + ProjectFileGenerator.NewLine);
+			Content.Append("\t\t\tname = \"" + TargetName + "\";" + ProjectFileGenerator.NewLine);
+			Content.Append("\t\t\tpassBuildSettingsInEnvironment = 1;" + ProjectFileGenerator.NewLine);
+			Content.Append("\t\t\tproductName = \"" + TargetName + "\";" + ProjectFileGenerator.NewLine);
+			Content.Append("\t\t\tproductReference = \"" + TargetAppGuid + "\";" + ProjectFileGenerator.NewLine);
+			Content.Append("\t\t\tproductType = \"" + TargetType + "\";" + ProjectFileGenerator.NewLine);
+			Content.Append("\t\t};" + ProjectFileGenerator.NewLine);
+
+			Content.Append("/* End PBXNativeTarget section */" + ProjectFileGenerator.NewLine + ProjectFileGenerator.NewLine);
+		}
+
+		private void AppendIndexTargetSection(StringBuilder Content, string TargetName, string TargetGuid, string TargetBuildConfigGuid, string SourcesBuildPhaseGuid)
+		{
+			Content.Append("/* Begin PBXNativeTarget section */" + ProjectFileGenerator.NewLine);
+
+			Content.Append("\t\t" + TargetGuid + " /* " + TargetName + " */ = {" + ProjectFileGenerator.NewLine);
+			Content.Append("\t\t\tisa = PBXNativeTarget;" + ProjectFileGenerator.NewLine);
+			Content.Append("\t\t\tbuildConfigurationList = " + TargetBuildConfigGuid + " /* Build configuration list for PBXNativeTarget \"" + TargetName + "\" */;" + ProjectFileGenerator.NewLine);
+			Content.Append("\t\t\tbuildPhases = (" + ProjectFileGenerator.NewLine);
+			Content.Append("\t\t\t\t" + SourcesBuildPhaseGuid + " /* Sources */," + ProjectFileGenerator.NewLine);
+			Content.Append("\t\t\t);" + ProjectFileGenerator.NewLine);
+			Content.Append("\t\t\tdependencies = (" + ProjectFileGenerator.NewLine);
+			Content.Append("\t\t\t);" + ProjectFileGenerator.NewLine);
+			Content.Append("\t\t\tname = \"" + TargetName + "\";" + ProjectFileGenerator.NewLine);
+			Content.Append("\t\t\tpassBuildSettingsInEnvironment = 1;" + ProjectFileGenerator.NewLine);
+			Content.Append("\t\t\tproductName = \"" + TargetName + "\";" + ProjectFileGenerator.NewLine);
+			Content.Append("\t\t\tproductType = \"com.apple.product-type.library.static\";" + ProjectFileGenerator.NewLine);
+			Content.Append("\t\t};" + ProjectFileGenerator.NewLine);
+
+			Content.Append("/* End PBXNativeTarget section */" + ProjectFileGenerator.NewLine + ProjectFileGenerator.NewLine);
+		}
+
+		private void AppendShellScriptSection(StringBuilder Content, string ShellScriptGuid, FileReference? UProjectPath)
+		{
+			StringBuilder FrameworkScript = new StringBuilder();
+
+			// nothing to do without a project
+			if (UProjectPath == null)
+			{
+				return;
+			}
+
+			// @todo: look also in Project/Build/Frameworks directory!
+			ProjectDescriptor Project = ProjectDescriptor.FromFile(UProjectPath);
+			List<PluginInfo> AvailablePlugins = Plugins.ReadAvailablePlugins(Unreal.EngineDirectory, DirectoryReference.FromFile(UProjectPath), Project.AdditionalPluginDirectories);
+
+			// look in each plugin for frameworks
+			// @todo: Cache this kind of things since every target will re-do this work!
+			foreach (PluginInfo PI in AvailablePlugins)
+			{
+				if (!Plugins.IsPluginEnabledForTarget(PI, Project, UnrealTargetPlatform.IOS, UnrealTargetConfiguration.Development, TargetRules.TargetType.Game))
+				{
+					continue;
+				}
+
+				// for now, we copy and code sign all *.framework.zip, even if the have no code (non-code frameworks are assumed to be *.embeddedframework.zip
+				DirectoryReference FrameworkDir = DirectoryReference.Combine(PI.Directory, "Source/Frameworks");
+				if (!DirectoryReference.Exists(FrameworkDir))
+				{
+					FrameworkDir = DirectoryReference.Combine(PI.Directory, "Frameworks");
+				}
+				if (DirectoryReference.Exists(FrameworkDir))
+				{
+					// look at each zip
+					foreach (FileInfo FI in new System.IO.DirectoryInfo(FrameworkDir.FullName).EnumerateFiles("*.framework.zip"))
+					{
+						//string Guid = XcodeProjectFileGenerator.MakeXcodeGuid();
+						//string RefGuid = XcodeProjectFileGenerator.MakeXcodeGuid();
+
+						// for FI of foo.framework.zip, this will give us foo.framework
+						//string Framework = Path.GetFileNameWithoutExtension(FI.FullName);
+
+						// unzip the framework right into the .app
+						FrameworkScript.AppendFormat("\\techo Unzipping {0}...\\n", FI.FullName);
+						FrameworkScript.AppendFormat("\\tunzip -o -q {0} -d ${{FRAMEWORK_DIR}} -x \\\"__MACOSX/*\\\" \\\"*/.DS_Store\\\"\\n", FI.FullName);
+					}
+				}
+			}
+
+			string ShellScript = "set -e\\n\\nIFS=$'\\\\n'\\n\\n" +
+				"if [ $PLATFORM_NAME = iphoneos ] || [ $PLATFORM_NAME = tvos ]; then \\n" +
+				"\\tFRAMEWORK_DIR=$TARGET_BUILD_DIR/$EXECUTABLE_FOLDER_PATH/Frameworks\\n" +
+				FrameworkScript.ToString() +
+				// and now code sign anything that has been unzipped above
+				"\\tfor FRAMEWORK in ${FRAMEWORK_DIR}/*.framework; do\\n" +
+					"\\t\\t[ -d \\\"${FRAMEWORK}\\\" ] || continue\\n" +
+					"\\t\\techo Codesigning ${FRAMEWORK}\\n" +
+					"\\t\\tcodesign --force --sign ${EXPANDED_CODE_SIGN_IDENTITY} --generate-entitlement-der --verbose --preserve-metadata=identifier,entitlements,flags --timestamp=none \\\"${FRAMEWORK}\\\"\\n" +
+				"\\tdone\\n" +
+				"fi\\n";
+
+			Content.Append("/* Begin PBXShellScriptBuildPhase section */" + ProjectFileGenerator.NewLine);
+			Content.Append(string.Format("\t\t{0} /* Sign Frameworks */ = {{" + ProjectFileGenerator.NewLine, ShellScriptGuid));
+			Content.Append("\t\t\tisa = PBXShellScriptBuildPhase;" + ProjectFileGenerator.NewLine);
+			Content.Append("\t\t\tbuildActionMask = 2147483647;" + ProjectFileGenerator.NewLine);
+			Content.Append("\t\t\tfiles = (" + ProjectFileGenerator.NewLine);
+			Content.Append("\t\t\t);" + ProjectFileGenerator.NewLine);
+			Content.Append("\t\t\tinputPaths = (" + ProjectFileGenerator.NewLine);
+			Content.Append("\t\t\t);" + ProjectFileGenerator.NewLine);
+			Content.Append("\t\t\toutputPaths = (" + ProjectFileGenerator.NewLine);
+			Content.Append("\t\t\t);" + ProjectFileGenerator.NewLine);
+			Content.Append("\t\t\tname = \"Sign Manual Frameworks\";" + ProjectFileGenerator.NewLine);
+			Content.Append("\t\t\trunOnlyForDeploymentPostprocessing = 0;" + ProjectFileGenerator.NewLine);
+			Content.Append("\t\t\tshellPath = /bin/sh;" + ProjectFileGenerator.NewLine);
+			Content.Append(string.Format("\t\t\tshellScript = \"{0}\";" + ProjectFileGenerator.NewLine, ShellScript));
+			Content.Append("\t\t};" + ProjectFileGenerator.NewLine);
+			Content.Append("/* End PBXShellScriptBuildPhase section */" + ProjectFileGenerator.NewLine);
+		}
+
+		private void AppendExtensionTargetSections(StringBuilder ProjectFileContent, List<XcodeExtensionInfo> AllExtensions)
+		{
+			foreach (XcodeExtensionInfo EI in AllExtensions)
+			{
+				Dictionary<string, string> BuildPhases = new Dictionary<string, string>();
+				BuildPhases.Add(EI.ResourceBuildPhaseGuid, "Resources");
+
+				AppendGenericTargetSection(ProjectFileContent, EI.Name, EI.TargetGuid, "com.apple.product-type.app-extension.messages-sticker-pack", EI.ConfigListGuid, EI.ProductGuid, null, BuildPhases);
+			}
+		}
+
+		private void AppendProjectSection(StringBuilder Content, string TargetName, string TargetGuid, string BuildTargetName, string BuildTargetGuid, string IndexTargetName, string IndexTargetGuid, string MainGroupGuid, string ProductRefGroupGuid, string ProjectGuid, string ProjectBuildConfigGuid, FileReference? ProjectFile, List<XcodeExtensionInfo> AllExtensions)
+		{
+			Content.Append("/* Begin PBXProject section */" + ProjectFileGenerator.NewLine);
+
+			Content.Append("\t\t" + ProjectGuid + " /* Project object */ = {" + ProjectFileGenerator.NewLine);
+			Content.Append("\t\t\tisa = PBXProject;" + ProjectFileGenerator.NewLine);
+			Content.Append("\t\t\tattributes = {" + ProjectFileGenerator.NewLine);
+			Content.Append("\t\t\t\tLastUpgradeCheck = 2000;" + ProjectFileGenerator.NewLine);
+			Content.Append("\t\t\t\tORGANIZATIONNAME = \"Epic Games, Inc.\";" + ProjectFileGenerator.NewLine);
+			Content.Append("\t\t\t\tTargetAttributes = {" + ProjectFileGenerator.NewLine);
+			Content.Append("\t\t\t\t\t" + TargetGuid + " = {" + ProjectFileGenerator.NewLine);
+
+			bool bAutomaticSigning = false;
+			if (InstalledPlatformInfo.IsValidPlatform(UnrealTargetPlatform.IOS, EProjectType.Code))
+			{
+				IOSPlatform IOSPlatform = ((IOSPlatform)UEBuildPlatform.GetBuildPlatform(UnrealTargetPlatform.IOS));
+				IOSProjectSettings ProjectSettings = IOSPlatform.ReadProjectSettings(ProjectFile);
+				bAutomaticSigning = ProjectSettings.bAutomaticSigning;
+			}
+
+			if (InstalledPlatformInfo.IsValidPlatform(UnrealTargetPlatform.TVOS, EProjectType.Code))
+			{
+				TVOSPlatform TVOSPlatform = ((TVOSPlatform)UEBuildPlatform.GetBuildPlatform(UnrealTargetPlatform.TVOS));
+				TVOSProjectSettings ProjectSettings = TVOSPlatform.ReadProjectSettings(ProjectFile);
+				//TVOSProvisioningData ProvisioningData = TVOSPlatform.ReadProvisioningData(ProjectSettings, bForDistribution);
+				bAutomaticSigning = ProjectSettings.bAutomaticSigning;
+			}
+
+			if (bAutomaticSigning)
+			{
+				Content.Append("\t\t\t\t\t\tProvisioningStyle = Automatic;" + ProjectFileGenerator.NewLine);
+			}
+			else
+			{
+				Content.Append("\t\t\t\t\t\tProvisioningStyle = Manual;" + ProjectFileGenerator.NewLine);
+			}
+			Content.Append("\t\t\t\t\t};" + ProjectFileGenerator.NewLine);
+			Content.Append("\t\t\t\t};" + ProjectFileGenerator.NewLine);
+			Content.Append("\t\t\t};" + ProjectFileGenerator.NewLine);
+			Content.Append("\t\t\tbuildConfigurationList = " + ProjectBuildConfigGuid + " /* Build configuration list for PBXProject \"" + TargetName + "\" */;" + ProjectFileGenerator.NewLine);
+			Content.Append("\t\t\tcompatibilityVersion = \"Xcode 8.0\";" + ProjectFileGenerator.NewLine);
+			Content.Append("\t\t\tdevelopmentRegion = English;" + ProjectFileGenerator.NewLine);
+			Content.Append("\t\t\thasScannedForEncodings = 0;" + ProjectFileGenerator.NewLine);
+			Content.Append("\t\t\tknownRegions = (" + ProjectFileGenerator.NewLine);
+			Content.Append("\t\t\t\ten" + ProjectFileGenerator.NewLine);
+			Content.Append("\t\t\t);" + ProjectFileGenerator.NewLine);
+			Content.Append("\t\t\tmainGroup = " + MainGroupGuid + ";" + ProjectFileGenerator.NewLine);
+			Content.Append("\t\t\tproductRefGroup = " + ProductRefGroupGuid + ";" + ProjectFileGenerator.NewLine);
+			Content.Append("\t\t\tprojectDirPath = \"\";" + ProjectFileGenerator.NewLine);
+			Content.Append("\t\t\tprojectRoot = \"\";" + ProjectFileGenerator.NewLine);
+			Content.Append("\t\t\ttargets = (" + ProjectFileGenerator.NewLine);
+			Content.Append("\t\t\t" + TargetGuid + " /* " + TargetName + " */," + ProjectFileGenerator.NewLine);
+			Content.Append("\t\t\t" + BuildTargetGuid + " /* " + BuildTargetName + " */," + ProjectFileGenerator.NewLine);
+			Content.Append("\t\t\t" + IndexTargetGuid + " /* " + IndexTargetName + " */," + ProjectFileGenerator.NewLine);
+			foreach (XcodeExtensionInfo EI in AllExtensions)
+			{
+				Content.Append("\t\t\t" + EI.TargetGuid + " /* " + EI.Name + " */," + ProjectFileGenerator.NewLine);
+			}
+			Content.Append("\t\t\t);" + ProjectFileGenerator.NewLine);
+			Content.Append("\t\t};" + ProjectFileGenerator.NewLine);
+
+			Content.Append("/* End PBXProject section */" + ProjectFileGenerator.NewLine + ProjectFileGenerator.NewLine);
+		}
+
+		private void AppendContainerItemProxySection(StringBuilder Content, string TargetName, string TargetGuid, string TargetProxyGuid, string ProjectGuid)
+		{
+			Content.Append("/* Begin PBXContainerItemProxy section */" + ProjectFileGenerator.NewLine);
+			Content.Append("\t\t" + TargetProxyGuid + " /* PBXContainerItemProxy */ = {" + ProjectFileGenerator.NewLine);
+			Content.Append("\t\t\tisa = PBXContainerItemProxy;" + ProjectFileGenerator.NewLine);
+			Content.Append("\t\t\tcontainerPortal = " + ProjectGuid + " /* Project object */;" + ProjectFileGenerator.NewLine);
+			Content.Append("\t\t\tproxyType = 1;" + ProjectFileGenerator.NewLine);
+			Content.Append("\t\t\tremoteGlobalIDString = " + TargetGuid + ";" + ProjectFileGenerator.NewLine);
+			Content.Append("\t\t\tremoteInfo = \"" + TargetName + "\";" + ProjectFileGenerator.NewLine);
+			Content.Append("\t\t};" + ProjectFileGenerator.NewLine);
+			Content.Append("/* End PBXContainerItemProxy section */" + ProjectFileGenerator.NewLine + ProjectFileGenerator.NewLine);
+		}
+
+		private void AppendTargetDependencySection(StringBuilder Content, string TargetName, string TargetGuid, string TargetDependencyGuid, string TargetProxyGuid)
+		{
+			Content.Append("/* Begin PBXTargetDependency section */" + ProjectFileGenerator.NewLine);
+			Content.Append("\t\t" + TargetDependencyGuid + " /* PBXTargetDependency */ = {" + ProjectFileGenerator.NewLine);
+			Content.Append("\t\t\tisa = PBXTargetDependency;" + ProjectFileGenerator.NewLine);
+			Content.Append("\t\t\ttarget = " + TargetGuid + " /* " + TargetName + " */;" + ProjectFileGenerator.NewLine);
+			Content.Append("\t\t\ttargetProxy = " + TargetProxyGuid + " /* PBXContainerItemProxy */;" + ProjectFileGenerator.NewLine);
+			Content.Append("\t\t};" + ProjectFileGenerator.NewLine);
+			Content.Append("/* End PBXTargetDependency section */" + ProjectFileGenerator.NewLine + ProjectFileGenerator.NewLine);
+		}
+
+		private void AppendProjectBuildConfiguration(StringBuilder Content, string ConfigName, string ConfigGuid)
+		{
+			Content.Append("\t\t" + ConfigGuid + " /* \"" + ConfigName + "\" */ = {" + ProjectFileGenerator.NewLine);
+			Content.Append("\t\t\tisa = XCBuildConfiguration;" + ProjectFileGenerator.NewLine);
+			Content.Append("\t\t\tbuildSettings = {" + ProjectFileGenerator.NewLine);
+
+			Content.Append("\t\t\t\tGCC_PREPROCESSOR_DEFINITIONS = (" + ProjectFileGenerator.NewLine);
+			foreach (string Definition in IntelliSensePreprocessorDefinitions)
+			{
+				Content.Append("\t\t\t\t\t\"" + Definition.Replace("\"", "").Replace("\\", "") + "\"," + ProjectFileGenerator.NewLine);
+			}
+			Content.Append("\t\t\t\t\t\"__INTELLISENSE__\"," + ProjectFileGenerator.NewLine);
+			Content.Append("\t\t\t\t\t\"MONOLITHIC_BUILD=1\"," + ProjectFileGenerator.NewLine);
+			Content.Append("\t\t\t\t);" + ProjectFileGenerator.NewLine);
+
+			Content.Append("\t\t\t\tHEADER_SEARCH_PATHS = (" + ProjectFileGenerator.NewLine);
+			foreach (string SearchPath in IntelliSenseSystemIncludeSearchPaths)
+			{
+				string Path = SearchPath.Contains(" ") ? "\\\"" + SearchPath + "\\\"" : SearchPath;
+				Content.Append("\t\t\t\t\t\"" + Path + "\"," + ProjectFileGenerator.NewLine);
+			}
+			Content.Append("\t\t\t\t);" + ProjectFileGenerator.NewLine);
+
+			Content.Append("\t\t\t\tUSER_HEADER_SEARCH_PATHS = (" + ProjectFileGenerator.NewLine);
+			foreach (string SearchPath in IntelliSenseIncludeSearchPaths)
+			{
+				string Path = SearchPath.Contains(" ") ? "\\\"" + SearchPath + "\\\"" : SearchPath;
+				Content.Append("\t\t\t\t\t\"" + Path + "\"," + ProjectFileGenerator.NewLine);
+			}
+			Content.Append("\t\t\t\t);" + ProjectFileGenerator.NewLine);
+			Content.Append("\t\t\t\tONLY_ACTIVE_ARCH = YES;" + ProjectFileGenerator.NewLine);
+			if (ConfigName == "Debug")
+			{
+				Content.Append("\t\t\t\tENABLE_TESTABILITY = YES;" + ProjectFileGenerator.NewLine);
+			}
+			Content.Append("\t\t\t\tALWAYS_SEARCH_USER_PATHS = NO;" + ProjectFileGenerator.NewLine);
+			Content.Append("\t\t\t\tCLANG_CXX_LANGUAGE_STANDARD = \"c++14\";" + ProjectFileGenerator.NewLine);
+			Content.Append("\t\t\t\tGCC_ENABLE_CPP_RTTI = NO;" + ProjectFileGenerator.NewLine);
+			Content.Append("\t\t\t\tGCC_WARN_CHECK_SWITCH_STATEMENTS = NO;" + ProjectFileGenerator.NewLine);
+			Content.Append("\t\t\t\tUSE_HEADERMAP = NO;" + ProjectFileGenerator.NewLine);
+			Content.Append("\t\t\t};" + ProjectFileGenerator.NewLine);
+			Content.Append("\t\t\tname = \"" + ConfigName + "\";" + ProjectFileGenerator.NewLine);
+			Content.Append("\t\t};" + ProjectFileGenerator.NewLine);
+		}
+
+		// cache for the below function
+		Dictionary<string, IEnumerable<string>> CachedMacProjectArcitectures = new Dictionary<string, IEnumerable<string>>();
+
+		/// <summary>
+		/// Returns the Mac architectures that should be configured for the provided target. If the target has a project we'll adhere
+		/// to whether it's set as Intel/Universal/Apple unless the type is denied (pretty much just Editor)
+		/// 
+		/// If the target has no project we'll support allow-listed targets for installed builds and all non-editor architectures 
+		/// for source builds. Not all programs are going to compile for Apple Silicon, but being able to build and fail is useful...
+		/// </summary>
+		/// <param name="Config">Build config for the target we're generating</param>
+		/// <param name="InProjectFile">Path to the project file, or null if the target has no project</param>
+		/// <returns></returns>
+		IEnumerable<string> GetSupportedMacArchitectures(XcodeBuildConfig Config, FileReference? InProjectFile)
+		{
+			// All architectures supported
+			IEnumerable<string> AllArchitectures = new[] { MacExports.IntelArchitecture, MacExports.AppleArchitecture };
+
+			// Add a way on the command line of forcing a project file with all architectures (there isn't a good way to let this be
+			// set and checked where we can access it).
+			bool ForceAllArchitectures = Environment.GetCommandLineArgs().Contains("AllArchitectures", StringComparer.OrdinalIgnoreCase);
+
+			if (ForceAllArchitectures)
+			{
+				return AllArchitectures;
+			}
+
+			string TargetName = Config.BuildTarget;
+
+			// First time seeing this target?
+			if (!CachedMacProjectArcitectures.ContainsKey(TargetName))
+			{
+				// Default to Intel
+				IEnumerable<string> TargetArchitectures = new[] { MacExports.IntelArchitecture };
+
+				// These targets are known to work so are allow-listed
+				bool IsAllowed = MacExports.TargetsAllowedForAppleSilicon.Contains(TargetName, StringComparer.OrdinalIgnoreCase);
+
+				// determine the target architectures based on what's allowed/denied
+				if (IsAllowed)
+				{
+					TargetArchitectures = AllArchitectures;
+				}
+				else
+				{
+					// if this is an unspecified tool/program, default to Intel for installed builds because we know all of that works. 
+					if (Config.ProjectTarget!.TargetRules!.Type == TargetType.Program)
+					{
+						// For misc tools we default to Intel for installed builds because we know all of that works. 
+						TargetArchitectures = Unreal.IsEngineInstalled() ? new[] { MacExports.IntelArchitecture } : AllArchitectures;
+					}
+					else
+					{
+						// For project targets we default to Intel then check the project settings. Note the editor target will have
+						// been denied above already.
+						TargetArchitectures = new[] { MacExports.IntelArchitecture };
+
+						ConfigHierarchy EngineIni = ConfigCache.ReadHierarchy(ConfigHierarchyType.Engine, InProjectFile?.Directory, UnrealTargetPlatform.Mac);
+						string TargetArchitecture;
+						string Key = Config.ProjectTarget!.TargetRules!.Type == TargetType.Editor ? "EditorTargetArchitecture" : "TargetArchitecture";
+						if (EngineIni.GetString("/Script/MacTargetPlatform.MacTargetSettings", Key, out TargetArchitecture))
+						{
+							if (TargetArchitecture.Contains("Universal", StringComparison.OrdinalIgnoreCase))
+							{
+								TargetArchitectures = AllArchitectures;
+							}
+							else if (TargetArchitecture.Contains("Intel", StringComparison.OrdinalIgnoreCase))
+							{
+								TargetArchitectures = new[] { MacExports.IntelArchitecture };
+							}
+							else if (TargetArchitecture.Contains("Apple", StringComparison.OrdinalIgnoreCase))
+							{
+								TargetArchitectures = new[] { MacExports.AppleArchitecture };
+							}
+						}
+					}
+				}
+
+				// Cache this so we don't need to keep checking this file
+				CachedMacProjectArcitectures.Add(TargetName, TargetArchitectures);
+			}
+
+			return CachedMacProjectArcitectures[TargetName];
+		}
+
+		private void AppendPlatformConfiguration(StringBuilder Content, XcodeBuildConfig Config, string TargetName, FileReference? ProjectFile, bool bSupportMac, bool bSupportIOS, bool bSupportTVOS, ILogger Logger, out string? IOSRunTimeVersion, out string? TVOSRunTimeVersion, string BinariesSubDir = "/Payload")
+		{
+			FileReference MacExecutablePath = Config.MacExecutablePath!;
+
+			string UEDir = ConvertPath(Path.GetFullPath(Directory.GetCurrentDirectory() + "../../.."));
+			string MacExecutableDir = bSupportMac ? ConvertPath(MacExecutablePath.Directory.FullName) : "";
+			string MacExecutableFileName = bSupportMac ? MacExecutablePath.GetFileName() : "";
+
+			// Get Mac architectures supported by this project
+			IEnumerable<string> SupportedMacArchitectures = GetSupportedMacArchitectures(Config, ProjectFile);
+
+			IOSRunTimeVersion = null;
+			TVOSRunTimeVersion = null;
+
+			// shortcut for mac only
+			if (bSupportMac && !bSupportIOS && !bSupportTVOS)
+			{
+				Content.Append("\t\t\t\tVALID_ARCHS = \"" + string.Join(" ", SupportedMacArchitectures) + "\";" + ProjectFileGenerator.NewLine);
+				Content.Append("\t\t\t\tSUPPORTED_PLATFORMS = \"macosx\";" + ProjectFileGenerator.NewLine);
+				Content.Append("\t\t\t\tPRODUCT_NAME = \"" + MacExecutableFileName + "\";" + ProjectFileGenerator.NewLine);
+				Content.Append("\t\t\t\tCONFIGURATION_BUILD_DIR = \"" + MacExecutableDir + "\";" + ProjectFileGenerator.NewLine);
+			}
+			else
+			{
+				bool bIsUnrealGame = TargetName.Equals("UnrealGame", StringComparison.InvariantCultureIgnoreCase);
+				bool bIsUnrealClient = TargetName.Equals("UnrealClient", StringComparison.InvariantCultureIgnoreCase);
+				DirectoryReference? GameDir = ProjectFile?.Directory;
+				string? GamePath = GameDir != null ? ConvertPath(GameDir.FullName) : null;
+
+				string? IOSRunTimeDevices = null;
+				string? TVOSRunTimeDevices = null;
+				string SupportedPlatforms = bSupportMac ? "macosx" : "";
+
+				bool bAutomaticSigning = false;
+				string? UUID_IOS = "";
+				string? UUID_TVOS = "";
+				string? TEAM_IOS = "";
+				string? TEAM_TVOS = "";
+				string? IOS_CERT = "iPhone Developer";
+				string? TVOS_CERT = "iPhone Developer";
+				string IOS_BUNDLE = "";
+				string TVOS_BUNDLE = "";
+				if (bSupportIOS && InstalledPlatformInfo.IsValidPlatform(UnrealTargetPlatform.IOS, EProjectType.Code))
+				{
+					IOSPlatform IOSPlatform = ((IOSPlatform)UEBuildPlatform.GetBuildPlatform(UnrealTargetPlatform.IOS));
+					IOSProjectSettings ProjectSettings = IOSPlatform.ReadProjectSettings(ProjectFile);
+					IOSProvisioningData ProvisioningData = IOSPlatform.ReadProvisioningData(ProjectSettings, bForDistribution);
+					IOSRunTimeVersion = ProjectSettings.RuntimeVersion;
+					IOSRunTimeDevices = ProjectSettings.RuntimeDevices;
+					SupportedPlatforms += " iphoneos";
+					bAutomaticSigning = ProjectSettings.bAutomaticSigning;
+					if (!bAutomaticSigning)
+					{
+						UUID_IOS = ProvisioningData.MobileProvisionUUID;
+						IOS_CERT = ProvisioningData.SigningCertificate;
+					}
+					TEAM_IOS = ProvisioningData.TeamUUID;
+					IOS_BUNDLE = ProjectSettings.BundleIdentifier;
+				}
+
+				if (bSupportTVOS && InstalledPlatformInfo.IsValidPlatform(UnrealTargetPlatform.TVOS, EProjectType.Code))
+				{
+					TVOSPlatform TVOSPlatform = ((TVOSPlatform)UEBuildPlatform.GetBuildPlatform(UnrealTargetPlatform.TVOS));
+					TVOSProjectSettings ProjectSettings = TVOSPlatform.ReadProjectSettings(ProjectFile);
+					TVOSProvisioningData ProvisioningData = TVOSPlatform.ReadProvisioningData(ProjectSettings, bForDistribution);
+					TVOSRunTimeVersion = ProjectSettings.RuntimeVersion;
+					TVOSRunTimeDevices = ProjectSettings.RuntimeDevices;
+					SupportedPlatforms += " appletvos";
+					if (!bAutomaticSigning)
+					{
+						UUID_TVOS = ProvisioningData.MobileProvisionUUID;
+						TVOS_CERT = ProvisioningData.SigningCertificate;
+					}
+					TEAM_TVOS = ProvisioningData.TeamUUID;
+					TVOS_BUNDLE = ProjectSettings.BundleIdentifier;
+				}
+
+				Content.Append("\t\t\t\tSUPPORTED_PLATFORMS = \"" + SupportedPlatforms.Trim() + "\";" + ProjectFileGenerator.NewLine);
+				if (bAutomaticSigning)
+				{
+					Content.Append("\t\t\t\tCODE_SIGN_STYLE = Automatic;" + ProjectFileGenerator.NewLine);
+				}
+				if (IOSRunTimeVersion != null)
+				{
+					Content.Append("\t\t\t\t\"VALID_ARCHS[sdk=iphoneos*]\" = \"" + string.Join(" ", this.SupportedIOSArchitectures) + "\";" + ProjectFileGenerator.NewLine);
+					Content.Append("\t\t\t\tIPHONEOS_DEPLOYMENT_TARGET = " + IOSRunTimeVersion + ";" + ProjectFileGenerator.NewLine);
+					Content.Append("\t\t\t\t\"PRODUCT_NAME[sdk=iphoneos*]\" = \"" + TargetName + "\";" + ProjectFileGenerator.NewLine); // @todo: change to Path.GetFileName(Config.IOSExecutablePath) when we stop using payload
+					Content.Append("\t\t\t\t\"TARGETED_DEVICE_FAMILY[sdk=iphoneos*]\" = \"" + IOSRunTimeDevices + "\";" + ProjectFileGenerator.NewLine);
+					Content.Append("\t\t\t\t\"SDKROOT[sdk=iphoneos]\" = iphoneos;" + ProjectFileGenerator.NewLine);
+					if (!string.IsNullOrEmpty(TEAM_IOS))
+					{
+						Content.Append("\t\t\t\t\"DEVELOPMENT_TEAM[sdk=iphoneos*]\" = " + TEAM_IOS + ";" + ProjectFileGenerator.NewLine);
+					}
+					Content.Append("\t\t\t\t\"CODE_SIGN_IDENTITY[sdk=iphoneos*]\" = \"" + IOS_CERT + "\";" + ProjectFileGenerator.NewLine);
+					if (!bAutomaticSigning && !string.IsNullOrEmpty(UUID_IOS))
+					{
+						Content.Append("\t\t\t\t\"PROVISIONING_PROFILE_SPECIFIER[sdk=iphoneos*]\" = \"" + UUID_IOS + "\";" + ProjectFileGenerator.NewLine);
+					}
+					if (ProjectFile != null)
+					{
+						Content.Append("\t\t\t\t\"PRODUCT_BUNDLE_IDENTIFIER[sdk=iphoneos*]\" = " + IOS_BUNDLE + ";" + ProjectFileGenerator.NewLine);
+					}
+				}
+				if (TVOSRunTimeVersion != null)
+				{
+					Content.Append("\t\t\t\t\"VALID_ARCHS[sdk=appletvos*]\" = \"" + string.Join(" ", this.SupportedIOSArchitectures) + "\";" + ProjectFileGenerator.NewLine);
+					Content.Append("\t\t\t\tTVOS_DEPLOYMENT_TARGET = " + TVOSRunTimeVersion + ";" + ProjectFileGenerator.NewLine);
+					Content.Append("\t\t\t\t\"PRODUCT_NAME[sdk=appletvos*]\" = \"" + TargetName + "\";" + ProjectFileGenerator.NewLine); // @todo: change to Path.GetFileName(Config.TVOSExecutablePath) when we stop using payload
+					Content.Append("\t\t\t\t\"TARGETED_DEVICE_FAMILY[sdk=appletvos*]\" = \"" + TVOSRunTimeDevices + "\";" + ProjectFileGenerator.NewLine);
+					Content.Append("\t\t\t\t\"SDKROOT[sdk=appletvos]\" = appletvos;" + ProjectFileGenerator.NewLine);
+					if (!string.IsNullOrEmpty(TEAM_TVOS))
+					{
+						Content.Append("\t\t\t\t\"DEVELOPMENT_TEAM[sdk=appletvos*]\" = " + TEAM_TVOS + ";" + ProjectFileGenerator.NewLine);
+					}
+					Content.Append("\t\t\t\t\"CODE_SIGN_IDENTITY[sdk=appletvos*]\" = \"" + TVOS_CERT + "\";" + ProjectFileGenerator.NewLine);
+					if (!bAutomaticSigning && !string.IsNullOrEmpty(UUID_TVOS))
+					{
+						Content.Append("\t\t\t\t\"PROVISIONING_PROFILE_SPECIFIER[sdk=appletvos*]\" = \"" + UUID_TVOS + "\";" + ProjectFileGenerator.NewLine);
+					}
+					if (ProjectFile != null)
+					{
+						Content.Append("\t\t\t\t\"PRODUCT_BUNDLE_IDENTIFIER[sdk=appletvos*]\" = " + TVOS_BUNDLE + ";" + ProjectFileGenerator.NewLine);
+					}
+				}
+				if (bSupportMac)
+				{
+					Content.Append("\t\t\t\t\"VALID_ARCHS[sdk=macosx*]\" = \"" + string.Join(" ", SupportedMacArchitectures) + "\";" + ProjectFileGenerator.NewLine);
+					Content.Append("\t\t\t\t\"PRODUCT_NAME[sdk=macosx*]\" = \"" + MacExecutableFileName + "\";" + ProjectFileGenerator.NewLine);
+					Content.Append("\t\t\t\t\"CONFIGURATION_BUILD_DIR[sdk=macosx*]\" = \"" + MacExecutableDir + "\";" + ProjectFileGenerator.NewLine);
+					Content.Append("\t\t\t\t\"SDKROOT[sdk=macosx]\" = macosx;" + ProjectFileGenerator.NewLine);
+				}
+
+				if (bIsUnrealGame || bIsUnrealClient)
+				{
+					if (IOSRunTimeVersion != null)
+					{
+						Content.Append("\t\t\t\t\"CONFIGURATION_BUILD_DIR[sdk=iphoneos*]\" = \"" + UEDir + "/Engine/Binaries/IOS" + BinariesSubDir + "\";" + ProjectFileGenerator.NewLine);
+					}
+					if (TVOSRunTimeVersion != null)
+					{
+						Content.Append("\t\t\t\t\"CONFIGURATION_BUILD_DIR[sdk=appletvos*]\" = \"" + UEDir + "/Engine/Binaries/TVOS" + BinariesSubDir + "\";" + ProjectFileGenerator.NewLine);
+					}
+				}
+				else if (ProjectFile != null)
+				{
+					if (IOSRunTimeVersion != null)
+					{
+						Content.Append("\t\t\t\t\"CONFIGURATION_BUILD_DIR[sdk=iphoneos*]\" = \"" + GamePath + "/Binaries/IOS" + BinariesSubDir + "\";" + ProjectFileGenerator.NewLine);
+					}
+					if (TVOSRunTimeVersion != null)
+					{
+						Content.Append("\t\t\t\t\"CONFIGURATION_BUILD_DIR[sdk=appletvos*]\" = \"" + GamePath + "/Binaries/TVOS" + BinariesSubDir + "\";" + ProjectFileGenerator.NewLine);
+					}
+				}
+				else
+				{
+					if (IOSRunTimeVersion != null)
+					{
+						Content.Append("\t\t\t\t\"CONFIGURATION_BUILD_DIR[sdk=iphoneos*]\" = \"" + UEDir + "/Engine/Binaries/IOS" + BinariesSubDir + "\";" + ProjectFileGenerator.NewLine);
+					}
+					if (TVOSRunTimeVersion != null)
+					{
+						Content.Append("\t\t\t\t\"CONFIGURATION_BUILD_DIR[sdk=appletvos*]\" = \"" + UEDir + "/Engine/Binaries/TVOS" + BinariesSubDir + "\";" + ProjectFileGenerator.NewLine);
+					}
+				}
+
+			}
+		}
+
+		private void AppendNativeTargetBuildConfiguration(StringBuilder Content, XcodeBuildConfig Config, string ConfigGuid, FileReference? ProjectFile, ILogger Logger)
+		{
+			bool bMacOnly = true;
+			if (Config.ProjectTarget!.TargetRules != null && XcodeProjectFileGenerator.ProjectFilePlatform.HasFlag(XcodeProjectFileGenerator.XcodeProjectFilePlatform.iOS))
+			{
+				if (Config.ProjectTarget.SupportedPlatforms.Contains(UnrealTargetPlatform.IOS))
+				{
+					bMacOnly = false;
+				}
+			}
+
+			Content.Append("\t\t" + ConfigGuid + " /* \"" + Config.DisplayName + "\" */ = {" + ProjectFileGenerator.NewLine);
+			Content.Append("\t\t\tisa = XCBuildConfiguration;" + ProjectFileGenerator.NewLine);
+			Content.Append("\t\t\tbuildSettings = {" + ProjectFileGenerator.NewLine);
+
+			string UEDir = ConvertPath(Path.GetFullPath(Directory.GetCurrentDirectory() + "../../.."));
+			//string MacExecutableDir = ConvertPath(Config.MacExecutablePath.Directory.FullName);
+			string MacExecutableFileName = Config.MacExecutablePath!.GetFileName();
+
+			string? IOSRunTimeVersion, TVOSRunTimeVersion;
+			AppendPlatformConfiguration(Content, Config, Config.BuildTarget, ProjectFile, true, !bMacOnly, !bMacOnly, Logger, out IOSRunTimeVersion, out TVOSRunTimeVersion);
+
+			if (!bMacOnly)
+			{
+				bool bIsUnrealGame = Config.BuildTarget.Equals("UnrealGame", StringComparison.InvariantCultureIgnoreCase);
+				bool bIsUnrealClient = Config.BuildTarget.Equals("UnrealClient", StringComparison.InvariantCultureIgnoreCase);
+
+				DirectoryReference? GameDir = ProjectFile?.Directory;
+				string? GamePath = GameDir != null ? ConvertPath(GameDir.FullName) : null;
+
+				string IOSInfoPlistPath;
+				string TVOSInfoPlistPath;
+				string MacInfoPlistPath;
+				string? IOSEntitlementPath = null;
+				string? TVOSEntitlementPath = null;
+				if (bIsUnrealGame)
+				{
+					IOSInfoPlistPath = UEDir + "/Engine/Intermediate/IOS/" + Config.BuildTarget + "-Info.plist";
+					TVOSInfoPlistPath = UEDir + "/Engine/Intermediate/TVOS/" + Config.BuildTarget + "-Info.plist";
+					MacInfoPlistPath = UEDir + "/Engine/Intermediate/Mac/" + MacExecutableFileName + "-Info.plist";
+					IOSEntitlementPath = "";
+					TVOSEntitlementPath = "";
+				}
+				else if (bIsUnrealClient)
+				{
+					IOSInfoPlistPath = UEDir + "/Engine/Intermediate/IOS/UnrealGame-Info.plist";
+					TVOSInfoPlistPath = UEDir + "/Engine/Intermediate/TVOS/UnrealGame-Info.plist";
+					MacInfoPlistPath = UEDir + "/Engine/Intermediate/Mac/" + MacExecutableFileName + "-Info.plist";
+					IOSEntitlementPath = "";
+					TVOSEntitlementPath = "";
+				}
+				else if (ProjectFile != null)
+				{
+					IOSInfoPlistPath = GamePath + "/Intermediate/IOS/" + Config.BuildTarget + "-Info.plist";
+					TVOSInfoPlistPath = GamePath + "/Intermediate/TVOS/" + Config.BuildTarget + "-Info.plist";
+					MacInfoPlistPath = GamePath + "/Intermediate/Mac/" + MacExecutableFileName + "-Info.plist";
+					IOSEntitlementPath = GamePath + "/Intermediate/IOS/" + Config.BuildTarget + ".entitlements";
+					TVOSEntitlementPath = GamePath + "/Intermediate/TVOS/" + Config.BuildTarget + ".entitlements";
+				}
+				else
+				{
+					if (GamePath == null)
+					{
+						IOSInfoPlistPath = UEDir + "/Engine/Intermediate/IOS/" + Config.BuildTarget + "-Info.plist";
+						TVOSInfoPlistPath = UEDir + "/Engine/Intermediate/TVOS/" + Config.BuildTarget + "-Info.plist";
+						MacInfoPlistPath = UEDir + "/Engine/Intermediate/Mac/" + MacExecutableFileName + "-Info.plist";
+					}
+					else
+					{
+						IOSInfoPlistPath = GamePath + "/Intermediate/IOS/" + Config.BuildTarget + "-Info.plist";
+						TVOSInfoPlistPath = GamePath + "/Intermediate/TVOS/" + Config.BuildTarget + "-Info.plist";
+						MacInfoPlistPath = GamePath + "/Intermediate/Mac/" + MacExecutableFileName + "-Info.plist";
+					}
+				}
+
+				if (XcodeProjectFileGenerator.bGeneratingRunIOSProject)
+				{
+					Content.Append("\t\t\t\tINFOPLIST_FILE = \"" + IOSInfoPlistPath + "\";" + ProjectFileGenerator.NewLine);
+					Content.Append("\t\t\t\tCODE_SIGN_ENTITLEMENTS = \"" + IOSEntitlementPath + "\";" + ProjectFileGenerator.NewLine);
+				}
+				else if (XcodeProjectFileGenerator.bGeneratingRunTVOSProject)
+				{
+					Content.Append("\t\t\t\tINFOPLIST_FILE = \"" + TVOSInfoPlistPath + "\";" + ProjectFileGenerator.NewLine);
+					Content.Append("\t\t\t\tCODE_SIGN_ENTITLEMENTS = \"" + TVOSEntitlementPath + "\";" + ProjectFileGenerator.NewLine);
+				}
+				else
+				{
+					Content.Append("\t\t\t\t\"INFOPLIST_FILE[sdk=macosx*]\" = \"" + MacInfoPlistPath + "\";" + ProjectFileGenerator.NewLine);
+					if (IOSRunTimeVersion != null)
+					{
+						Content.Append("\t\t\t\t\"INFOPLIST_FILE[sdk=iphoneos*]\" = \"" + IOSInfoPlistPath + "\";" + ProjectFileGenerator.NewLine);
+						Content.Append("\t\t\t\t\"CODE_SIGN_ENTITLEMENTS[sdk=iphoneos*]\" = \"" + IOSEntitlementPath + "\";" + ProjectFileGenerator.NewLine);
+					}
+					if (TVOSRunTimeVersion != null)
+					{
+						Content.Append("\t\t\t\t\"INFOPLIST_FILE[sdk=appletvos*]\" = \"" + TVOSInfoPlistPath + "\";" + ProjectFileGenerator.NewLine);
+						Content.Append("\t\t\t\t\"CODE_SIGN_ENTITLEMENTS[sdk=appletvos*]\" = \"" + TVOSEntitlementPath + "\";" + ProjectFileGenerator.NewLine);
+					}
+				}
+
+				// Prepare a temp Info.plist file so Xcode has some basic info about the target immediately after opening the project.
+				// This is needed for the target to pass the settings validation before code signing. UBT will overwrite this plist file later, with proper contents.
+				if (BuildHostPlatform.Current.Platform == UnrealTargetPlatform.Mac)
+				{
+					bool bCreateMacInfoPlist = !File.Exists(MacInfoPlistPath);
+					bool bCreateIOSInfoPlist = !File.Exists(IOSInfoPlistPath) && IOSRunTimeVersion != null;
+					bool bCreateTVOSInfoPlist = !File.Exists(TVOSInfoPlistPath) && TVOSRunTimeVersion != null;
+					if (bCreateMacInfoPlist || bCreateIOSInfoPlist || bCreateTVOSInfoPlist)
+					{
+						DirectoryReference? ProjectPath = GameDir;
+						DirectoryReference EngineDir = DirectoryReference.Combine(new DirectoryReference(UEDir), "Engine");
+						string GameName = Config.BuildTarget;
+						bool bIsClient = false;
+						if (ProjectPath == null)
+						{
+							ProjectPath = EngineDir;
+						}
+						if (bIsUnrealGame)
+						{
+							ProjectPath = EngineDir;
+							GameName = "UnrealGame";
+							bIsClient = (AppName == "UnrealClient");
+						}
+
+						if (bCreateMacInfoPlist)
+						{
+							Directory.CreateDirectory(Path.GetDirectoryName(MacInfoPlistPath)!);
+							UEDeployMac.GeneratePList(ProjectPath.FullName, bIsUnrealGame, GameName, Config.BuildTarget, EngineDir.FullName, MacExecutableFileName);
+						}
+						if (bCreateIOSInfoPlist)
+						{
+							// get the receipt
+							FileReference ReceiptFilename;
+							if (bIsUnrealGame)
+							{
+								ReceiptFilename = TargetReceipt.GetDefaultPath(Unreal.EngineDirectory, "UnrealGame", UnrealTargetPlatform.IOS, Config.BuildConfig, "");
+							}
+							else
+							{
+								ReceiptFilename = TargetReceipt.GetDefaultPath(ProjectPath, GameName, UnrealTargetPlatform.IOS, Config.BuildConfig, "");
+							}
+							Directory.CreateDirectory(Path.GetDirectoryName(IOSInfoPlistPath)!);
+							bool bSupportPortrait, bSupportLandscape;
+							TargetReceipt? Receipt;
+							TargetReceipt.TryRead(ReceiptFilename, out Receipt);
+							bool bBuildAsFramework = UEDeployIOS.GetCompileAsDll(Receipt);
+							UEDeployIOS.GenerateIOSPList(ProjectFile, Config.BuildConfig, ProjectPath.FullName, bIsUnrealGame, GameName, bIsClient, Config.BuildTarget, EngineDir.FullName, ProjectPath + "/Binaries/IOS/Payload", null, BundleIdentifier, bBuildAsFramework, Logger, out bSupportPortrait, out bSupportLandscape);
+						}
+						if (bCreateTVOSInfoPlist)
+						{
+							Directory.CreateDirectory(Path.GetDirectoryName(TVOSInfoPlistPath)!);
+							UEDeployTVOS.GenerateTVOSPList(ProjectPath.FullName, bIsUnrealGame, GameName, bIsClient, Config.BuildTarget, EngineDir.FullName, ProjectPath + "/Binaries/TVOS/Payload", null, BundleIdentifier, Logger);
+						}
+					}
+				}
+			}
+
+			// #jira UE-143619: Pre Monterey macOS requires this option for a packaged app to run on iOS15 due to new code signature format. Could be removed once Monterey is minimum OS.
+			Content.Append("\t\t\t\tOTHER_CODE_SIGN_FLAGS = \"--generate-entitlement-der\";" + ProjectFileGenerator.NewLine);
+
+			Content.Append("\t\t\t\tMACOSX_DEPLOYMENT_TARGET = " + MacToolChain.Settings.MacOSVersion + ";" + ProjectFileGenerator.NewLine);
+			Content.Append("\t\t\t\tINFOPLIST_OUTPUT_FORMAT = xml;" + ProjectFileGenerator.NewLine);
+			Content.Append("\t\t\t\tCOMBINE_HIDPI_IMAGES = YES;" + ProjectFileGenerator.NewLine);
+
+			//#jira UE-50382 Xcode Address Sanitizer feature does not work on iOS
+			// address sanitizer dylib loader depends on the SDKROOT parameter. For macosx or default (missing, translated as macosx), the path is incorrect for iphone/appletv
+			if (XcodeProjectFileGenerator.bGeneratingRunIOSProject)
+			{
+				Content.Append("\t\t\t\tSDKROOT = iphoneos;" + ProjectFileGenerator.NewLine);
+			}
+			else if (XcodeProjectFileGenerator.bGeneratingRunTVOSProject)
+			{
+				Content.Append("\t\t\t\tSDKROOT = appletvos;" + ProjectFileGenerator.NewLine);
+			}
+			else
+			{
+				Content.Append("\t\t\t\tSDKROOT = macosx;" + ProjectFileGenerator.NewLine);
+			}
+			Content.Append("\t\t\t\tGCC_PRECOMPILE_PREFIX_HEADER = YES;" + ProjectFileGenerator.NewLine);
+			Content.Append("\t\t\t\tGCC_PREFIX_HEADER = \"" + UEDir + "/Engine/Source/Editor/UnrealEd/Public/UnrealEd.h\";" + ProjectFileGenerator.NewLine);
+			Content.Append("\t\t\t};" + ProjectFileGenerator.NewLine);
+			Content.Append("\t\t\tname = \"" + Config.DisplayName + "\";" + ProjectFileGenerator.NewLine);
+			Content.Append("\t\t};" + ProjectFileGenerator.NewLine);
+		}
+
+		private void AppendLegacyTargetBuildConfiguration(StringBuilder Content, XcodeBuildConfig Config, string ConfigGuid, FileReference? ProjectFile)
+		{
+			bool bMacOnly = true;
+			if (Config.ProjectTarget!.TargetRules != null && XcodeProjectFileGenerator.ProjectFilePlatform.HasFlag(XcodeProjectFileGenerator.XcodeProjectFilePlatform.iOS))
+			{
+				if (Config.ProjectTarget.SupportedPlatforms.Contains(UnrealTargetPlatform.IOS))
+				{
+					bMacOnly = false;
+				}
+			}
+
+			// Get Mac architectures supported by this project
+			IEnumerable<string> SupportedMacArchitectures = GetSupportedMacArchitectures(Config, ProjectFilePath);
+
+			Content.Append("\t\t" + ConfigGuid + " /* \"" + Config.DisplayName + "\" */ = {" + ProjectFileGenerator.NewLine);
+			Content.Append("\t\t\tisa = XCBuildConfiguration;" + ProjectFileGenerator.NewLine);
+			Content.Append("\t\t\tbuildSettings = {" + ProjectFileGenerator.NewLine);
+			if (bMacOnly)
+			{
+				Content.Append("\t\t\t\tVALID_ARCHS = \"" + string.Join(" ", SupportedMacArchitectures) + "\";" + ProjectFileGenerator.NewLine);
+				Content.Append("\t\t\t\tSUPPORTED_PLATFORMS = \"macosx\";" + ProjectFileGenerator.NewLine);
+			}
+			else
+			{
+				IEnumerable<string> ValidArchs = SupportedMacArchitectures;
+				string SupportedPlatforms = "macosx";
+				if (InstalledPlatformInfo.IsValidPlatform(UnrealTargetPlatform.IOS, EProjectType.Code))
+				{
+					ValidArchs = ValidArchs.Union(this.SupportedIOSArchitectures);
+					SupportedPlatforms += " iphoneos";
+				}
+				if (InstalledPlatformInfo.IsValidPlatform(UnrealTargetPlatform.TVOS, EProjectType.Code))
+				{
+					ValidArchs = ValidArchs.Union(this.SupportedIOSArchitectures);
+					SupportedPlatforms += " appletvos";
+				}
+				Content.Append("\t\t\t\tVALID_ARCHS = \"" + string.Join(" ", ValidArchs) + "\";" + ProjectFileGenerator.NewLine);
+				Content.Append("\t\t\t\tSUPPORTED_PLATFORMS = \"" + SupportedPlatforms + "\";" + ProjectFileGenerator.NewLine);
+			}
+			Content.Append("\t\t\t\tGCC_PREPROCESSOR_DEFINITIONS = ();" + ProjectFileGenerator.NewLine);
+			Content.Append("\t\t\t\tHEADER_SEARCH_PATHS = ();" + ProjectFileGenerator.NewLine);
+			Content.Append("\t\t\t\tUSER_HEADER_SEARCH_PATHS = ();" + ProjectFileGenerator.NewLine);
+			Content.Append("\t\t\t\tUE_BUILD_TARGET_NAME = \"" + Config.BuildTarget + "\";" + ProjectFileGenerator.NewLine);
+			Content.Append("\t\t\t\tUE_BUILD_TARGET_CONFIG = \"" + Config.BuildConfig + "\";" + ProjectFileGenerator.NewLine);
+			Content.Append("\t\t\t};" + ProjectFileGenerator.NewLine);
+			Content.Append("\t\t\tname = \"" + Config.DisplayName + "\";" + ProjectFileGenerator.NewLine);
+			Content.Append("\t\t};" + ProjectFileGenerator.NewLine);
+		}
+
+
+		private void AppendXCBuildConfigurationSection(StringBuilder Content, Dictionary<string, XcodeBuildConfig> ProjectBuildConfigs, Dictionary<string, XcodeBuildConfig> TargetBuildConfigs,
+				Dictionary<string, XcodeBuildConfig> BuildTargetBuildConfigs, Dictionary<string, XcodeBuildConfig> IndexTargetBuildConfigs, FileReference? GameProjectPath,
+				List<XcodeExtensionInfo> AllExtensions, ILogger Logger)
+		{
+			Content.Append("/* Begin XCBuildConfiguration section */" + ProjectFileGenerator.NewLine);
+
+			foreach (KeyValuePair<string, XcodeBuildConfig> Config in ProjectBuildConfigs)
+			{
+				AppendProjectBuildConfiguration(Content, Config.Value.DisplayName, Config.Key);
+			}
+
+			foreach (KeyValuePair<string, XcodeBuildConfig> Config in TargetBuildConfigs)
+			{
+				AppendNativeTargetBuildConfiguration(Content, Config.Value, Config.Key, GameProjectPath, Logger);
+			}
+
+			foreach (KeyValuePair<string, XcodeBuildConfig> Config in BuildTargetBuildConfigs)
+			{
+				AppendLegacyTargetBuildConfiguration(Content, Config.Value, Config.Key, GameProjectPath);
+			}
+
+			foreach (KeyValuePair<string, XcodeBuildConfig> Config in IndexTargetBuildConfigs)
+			{
+				AppendNativeTargetBuildConfiguration(Content, Config.Value, Config.Key, GameProjectPath, Logger);
+			}
+
+			foreach (XcodeExtensionInfo EI in AllExtensions)
+			{
+				Content.Append(EI.ConfigurationContents);
+			}
+
+			Content.Append("/* End XCBuildConfiguration section */" + ProjectFileGenerator.NewLine + ProjectFileGenerator.NewLine);
+		}
+
+		private void AppendXCConfigurationList(StringBuilder Content, string TypeName, string TargetName, string ConfigListGuid, Dictionary<string, XcodeBuildConfig> BuildConfigs, string Default = "Development")
+		{
+			Content.Append("\t\t" + ConfigListGuid + " /* Build configuration list for " + TypeName + " \"" + TargetName + "\" */ = {" + ProjectFileGenerator.NewLine);
+			Content.Append("\t\t\tisa = XCConfigurationList;" + ProjectFileGenerator.NewLine);
+			Content.Append("\t\t\tbuildConfigurations = (" + ProjectFileGenerator.NewLine);
+			foreach (KeyValuePair<string, XcodeBuildConfig> Config in BuildConfigs)
+			{
+				Content.Append("\t\t\t\t" + Config.Key + " /* \"" + Config.Value.DisplayName + "\" */," + ProjectFileGenerator.NewLine);
+			}
+			Content.Append("\t\t\t);" + ProjectFileGenerator.NewLine);
+			Content.Append("\t\t\tdefaultConfigurationIsVisible = 0;" + ProjectFileGenerator.NewLine);
+			Content.Append("\t\t\tdefaultConfigurationName = " + Default + ";" + ProjectFileGenerator.NewLine);
+			Content.Append("\t\t};" + ProjectFileGenerator.NewLine);
+		}
+
+		private void AppendXCConfigurationListSection(StringBuilder Content, string TargetName, string BuildTargetName, string IndexTargetName, string ProjectConfigListGuid,
+			Dictionary<string, XcodeBuildConfig> ProjectBuildConfigs, string TargetConfigListGuid, Dictionary<string, XcodeBuildConfig> TargetBuildConfigs,
+			string BuildTargetConfigListGuid, Dictionary<string, XcodeBuildConfig> BuildTargetBuildConfigs,
+			string IndexTargetConfigListGuid, Dictionary<string, XcodeBuildConfig> IndexTargetBuildConfigs,
+			List<XcodeExtensionInfo> AllExtensions)
+		{
+			Content.Append("/* Begin XCConfigurationList section */" + ProjectFileGenerator.NewLine);
+
+			AppendXCConfigurationList(Content, "PBXProject", TargetName, ProjectConfigListGuid, ProjectBuildConfigs);
+			AppendXCConfigurationList(Content, "PBXLegacyTarget", BuildTargetName, BuildTargetConfigListGuid, BuildTargetBuildConfigs);
+			AppendXCConfigurationList(Content, "PBXNativeTarget", TargetName, TargetConfigListGuid, TargetBuildConfigs);
+			AppendXCConfigurationList(Content, "PBXNativeTarget", IndexTargetName, IndexTargetConfigListGuid, IndexTargetBuildConfigs);
+
+			foreach (XcodeExtensionInfo EI in AllExtensions)
+			{
+				AppendXCConfigurationList(Content, "PBXNativeTarget", EI.Name, EI.ConfigListGuid, EI.AllConfigs);
+			}
+
+			Content.Append("/* End XCConfigurationList section */" + ProjectFileGenerator.NewLine);
+		}
+
+		private List<XcodeBuildConfig> GetSupportedBuildConfigs(List<UnrealTargetPlatform> Platforms, List<UnrealTargetConfiguration> Configurations, ILogger Logger)
+		{
+			List<XcodeBuildConfig> BuildConfigs = new List<XcodeBuildConfig>();
+
+			//string ProjectName = ProjectFilePath.GetFileNameWithoutExtension();
+
+			foreach (UnrealTargetConfiguration Configuration in Configurations)
+			{
+				if (InstalledPlatformInfo.IsValidConfiguration(Configuration, EProjectType.Code))
+				{
+					foreach (UnrealTargetPlatform Platform in Platforms)
+					{
+						if (InstalledPlatformInfo.IsValidPlatform(Platform, EProjectType.Code) && (Platform == UnrealTargetPlatform.Mac || Platform == UnrealTargetPlatform.IOS || Platform == UnrealTargetPlatform.TVOS)) // @todo support other platforms
+						{
+							UEBuildPlatform? BuildPlatform;
+							if (UEBuildPlatform.TryGetBuildPlatform(Platform, out BuildPlatform) && (BuildPlatform.HasRequiredSDKsInstalled() == SDKStatus.Valid))
+							{
+								// Check we have targets (Expected to be no Engine targets when generating for a single .uproject)
+								if (ProjectTargets.Count == 0 && BaseDir != Unreal.EngineDirectory)
+								{
+									throw new BuildException("Expecting at least one ProjectTarget to be associated with project '{0}' in the TargetProjects list ", ProjectFilePath);
+								}
+
+								// Now go through all of the target types for this project
+								foreach (ProjectTarget ProjectTarget in ProjectTargets.OfType<ProjectTarget>())
+								{
+									if (MSBuildProjectFile.IsValidProjectPlatformAndConfiguration(ProjectTarget, Platform, Configuration, Logger))
+									{
+										// Figure out if this is a monolithic build
+										bool bShouldCompileMonolithic = BuildPlatform.ShouldCompileMonolithicBinary(Platform);
+										bShouldCompileMonolithic |= (ProjectTarget.CreateRulesDelegate(Platform, Configuration).LinkType == TargetLinkType.Monolithic);
+
+										string ConfigName = Configuration.ToString();
+										if (ProjectTarget.TargetRules!.Type != TargetType.Game && ProjectTarget.TargetRules.Type != TargetType.Program)
+										{
+											ConfigName += " " + ProjectTarget.TargetRules.Type.ToString();
+										}
+
+										if (BuildConfigs.Where(Config => Config.DisplayName == ConfigName).ToList().Count == 0)
+										{
+											string TargetName = ProjectTarget.TargetFilePath.GetFileNameWithoutAnyExtensions();
+
+											// Get the output directory
+											DirectoryReference RootDirectory = Unreal.EngineDirectory;
+											// Unique and Monolithic both need to use the target directory not the engine directory
+											if (ProjectTarget.TargetRules.Type != TargetType.Program && (bShouldCompileMonolithic || ProjectTarget.TargetRules.BuildEnvironment == TargetBuildEnvironment.Unique))
+											{
+												if (ProjectTarget.UnrealProjectFilePath != null)
+												{
+													RootDirectory = ProjectTarget.UnrealProjectFilePath.Directory;
+												}
+											}
+
+											if (ProjectTarget.TargetRules.Type == TargetType.Program && ProjectTarget.UnrealProjectFilePath != null)
+											{
+												RootDirectory = ProjectTarget.UnrealProjectFilePath.Directory;
+											}
+
+											// Get the output directory
+											DirectoryReference OutputDirectory = DirectoryReference.Combine(RootDirectory, "Binaries");
+
+											string ExeName = TargetName;
+											if (!bShouldCompileMonolithic && ProjectTarget.TargetRules.Type != TargetType.Program)
+											{
+												// Figure out what the compiled binary will be called so that we can point the IDE to the correct file
+												if (ProjectTarget.TargetRules.Type != TargetType.Game)
+												{
+													// Only if shared - unique retains the Target Name
+													if (ProjectTarget.TargetRules.BuildEnvironment == TargetBuildEnvironment.Shared)
+													{
+														ExeName = "Unreal" + ProjectTarget.TargetRules.Type.ToString();
+													}
+												}
+											}
+
+											if (BuildPlatform.Platform == UnrealTargetPlatform.Mac)
+											{
+												string MacExecutableName = MakeExecutableFileName(ExeName, UnrealTargetPlatform.Mac, Configuration, ProjectTarget.TargetRules.Architecture, ProjectTarget.TargetRules.UndecoratedConfiguration);
+												string IOSExecutableName = MacExecutableName.Replace("-Mac-", "-IOS-");
+												string TVOSExecutableName = MacExecutableName.Replace("-Mac-", "-TVOS-");
+												BuildConfigs.Add(new XcodeBuildConfig(ConfigName, TargetName, FileReference.Combine(OutputDirectory, "Mac", MacExecutableName), FileReference.Combine(OutputDirectory, "IOS", IOSExecutableName), FileReference.Combine(OutputDirectory, "TVOS", TVOSExecutableName), ProjectTarget, Configuration));
+											}
+											else if (BuildPlatform.Platform == UnrealTargetPlatform.IOS || BuildPlatform.Platform == UnrealTargetPlatform.TVOS)
+											{
+												string IOSExecutableName = MakeExecutableFileName(ExeName, UnrealTargetPlatform.IOS, Configuration, ProjectTarget.TargetRules.Architecture, ProjectTarget.TargetRules.UndecoratedConfiguration);
+												string TVOSExecutableName = IOSExecutableName.Replace("-IOS-", "-TVOS-");
+												//string MacExecutableName = IOSExecutableName.Replace("-IOS-", "-Mac-");
+												BuildConfigs.Add(new XcodeBuildConfig(ConfigName, TargetName, FileReference.Combine(OutputDirectory, "Mac", IOSExecutableName), FileReference.Combine(OutputDirectory, "IOS", IOSExecutableName), FileReference.Combine(OutputDirectory, "TVOS", TVOSExecutableName), ProjectTarget, Configuration));
+											}
+										}
+									}
+								}
+							}
+						}
+					}
+				}
+			}
+
+			return BuildConfigs;
+		}
+
+		private static string MakeExecutableFileName(string BinaryName, UnrealTargetPlatform Platform, UnrealTargetConfiguration Configuration, string Architecture, UnrealTargetConfiguration UndecoratedConfiguration)
+		{
+			StringBuilder Result = new StringBuilder();
+
+			Result.Append(BinaryName);
+
+			if (Configuration != UndecoratedConfiguration)
+			{
+				Result.AppendFormat("-{0}-{1}", Platform.ToString(), Configuration.ToString());
+			}
+
+			UEBuildPlatform BuildPlatform = UEBuildPlatform.GetBuildPlatform(Platform);
+			if (BuildPlatform.RequiresArchitectureSuffix())
+			{
+				Result.Append(Architecture);
+			}
+
+			return Result.ToString();
+		}
+
+		private FileReference GetUserSchemeManagementFilePath()
+		{
+			return new FileReference(ProjectFilePath.FullName + "/xcuserdata/" + Environment.UserName + ".xcuserdatad/xcschemes/xcschememanagement.plist");
+		}
+
+		private DirectoryReference GetProjectSchemeDirectory()
+		{
+			return new DirectoryReference(ProjectFilePath.FullName + "/xcshareddata/xcschemes");
+		}
+
+		private FileReference GetProjectSchemeFilePathForTarget(string TargetName)
+		{
+			return FileReference.Combine(GetProjectSchemeDirectory(), TargetName + ".xcscheme");
+		}
+
+		private void WriteSchemeFile(string TargetName, string TargetGuid, string BuildTargetGuid, string IndexTargetGuid, bool bHasEditorConfiguration, string GameProjectPath)
+		{
+
+			FileReference SchemeFilePath = GetProjectSchemeFilePathForTarget(TargetName);
+
+			DirectoryReference.CreateDirectory(SchemeFilePath.Directory);
+
+			string? OldCommandLineArguments = null;
+			if (FileReference.Exists(SchemeFilePath))
+			{
+				string OldContents = File.ReadAllText(SchemeFilePath.FullName);
+				int OldCommandLineArgumentsStart = OldContents.IndexOf("<CommandLineArguments>") + "<CommandLineArguments>".Length;
+				int OldCommandLineArgumentsEnd = OldContents.IndexOf("</CommandLineArguments>");
+				if (OldCommandLineArgumentsStart != -1 && OldCommandLineArgumentsEnd != -1)
+				{
+					OldCommandLineArguments = OldContents.Substring(OldCommandLineArgumentsStart, OldCommandLineArgumentsEnd - OldCommandLineArgumentsStart);
+				}
+			}
+
+			string DefaultConfiguration = bHasEditorConfiguration && !XcodeProjectFileGenerator.bGeneratingRunIOSProject && !XcodeProjectFileGenerator.bGeneratingRunTVOSProject ? "Development Editor" : "Development";
+
+			StringBuilder Content = new StringBuilder();
+
+			Content.Append("<?xml version=\"1.0\" encoding=\"UTF-8\"?>" + ProjectFileGenerator.NewLine);
+			Content.Append("<Scheme" + ProjectFileGenerator.NewLine);
+			Content.Append("   LastUpgradeVersion = \"2000\"" + ProjectFileGenerator.NewLine);
+			Content.Append("   version = \"1.3\">" + ProjectFileGenerator.NewLine);
+			Content.Append("   <BuildAction" + ProjectFileGenerator.NewLine);
+			Content.Append("      parallelizeBuildables = \"YES\"" + ProjectFileGenerator.NewLine);
+			Content.Append("      buildImplicitDependencies = \"YES\">" + ProjectFileGenerator.NewLine);
+			Content.Append("      <BuildActionEntries>" + ProjectFileGenerator.NewLine);
+			Content.Append("         <BuildActionEntry" + ProjectFileGenerator.NewLine);
+			Content.Append("            buildForTesting = \"YES\"" + ProjectFileGenerator.NewLine);
+			Content.Append("            buildForRunning = \"YES\"" + ProjectFileGenerator.NewLine);
+			Content.Append("            buildForProfiling = \"YES\"" + ProjectFileGenerator.NewLine);
+			Content.Append("            buildForArchiving = \"YES\"" + ProjectFileGenerator.NewLine);
+			Content.Append("            buildForAnalyzing = \"YES\">" + ProjectFileGenerator.NewLine);
+			Content.Append("            <BuildableReference" + ProjectFileGenerator.NewLine);
+			Content.Append("               BuildableIdentifier = \"primary\"" + ProjectFileGenerator.NewLine);
+			Content.Append("               BlueprintIdentifier = \"" + TargetGuid + "\"" + ProjectFileGenerator.NewLine);
+			Content.Append("               BuildableName = \"" + TargetName + ".app\"" + ProjectFileGenerator.NewLine);
+			Content.Append("               BlueprintName = \"" + TargetName + "\"" + ProjectFileGenerator.NewLine);
+			Content.Append("               ReferencedContainer = \"container:" + TargetName + ".xcodeproj\">" + ProjectFileGenerator.NewLine);
+			Content.Append("            </BuildableReference>" + ProjectFileGenerator.NewLine);
+			Content.Append("         </BuildActionEntry>" + ProjectFileGenerator.NewLine);
+			Content.Append("      </BuildActionEntries>" + ProjectFileGenerator.NewLine);
+			Content.Append("   </BuildAction>" + ProjectFileGenerator.NewLine);
+			Content.Append("   <TestAction" + ProjectFileGenerator.NewLine);
+			Content.Append("      buildConfiguration = \"" + DefaultConfiguration + "\"" + ProjectFileGenerator.NewLine);
+			Content.Append("      selectedDebuggerIdentifier = \"Xcode.DebuggerFoundation.Debugger.LLDB\"" + ProjectFileGenerator.NewLine);
+			Content.Append("      selectedLauncherIdentifier = \"Xcode.DebuggerFoundation.Launcher.LLDB\"" + ProjectFileGenerator.NewLine);
+			Content.Append("      shouldUseLaunchSchemeArgsEnv = \"YES\">" + ProjectFileGenerator.NewLine);
+			Content.Append("      <Testables>" + ProjectFileGenerator.NewLine);
+			Content.Append("      </Testables>" + ProjectFileGenerator.NewLine);
+			Content.Append("      <MacroExpansion>" + ProjectFileGenerator.NewLine);
+			Content.Append("            <BuildableReference" + ProjectFileGenerator.NewLine);
+			Content.Append("               BuildableIdentifier = \"primary\"" + ProjectFileGenerator.NewLine);
+			Content.Append("               BlueprintIdentifier = \"" + TargetGuid + "\"" + ProjectFileGenerator.NewLine);
+			Content.Append("               BuildableName = \"" + TargetName + ".app\"" + ProjectFileGenerator.NewLine);
+			Content.Append("               BlueprintName = \"" + TargetName + "\"" + ProjectFileGenerator.NewLine);
+			Content.Append("               ReferencedContainer = \"container:" + TargetName + ".xcodeproj\">" + ProjectFileGenerator.NewLine);
+			Content.Append("            </BuildableReference>" + ProjectFileGenerator.NewLine);
+			Content.Append("      </MacroExpansion>" + ProjectFileGenerator.NewLine);
+			Content.Append("      <AdditionalOptions>" + ProjectFileGenerator.NewLine);
+			Content.Append("      </AdditionalOptions>" + ProjectFileGenerator.NewLine);
+			Content.Append("   </TestAction>" + ProjectFileGenerator.NewLine);
+			Content.Append("   <LaunchAction" + ProjectFileGenerator.NewLine);
+			Content.Append("      buildConfiguration = \"" + DefaultConfiguration + "\"" + ProjectFileGenerator.NewLine);
+			Content.Append("      selectedDebuggerIdentifier = \"Xcode.DebuggerFoundation.Debugger.LLDB\"" + ProjectFileGenerator.NewLine);
+			Content.Append("      selectedLauncherIdentifier = \"Xcode.DebuggerFoundation.Launcher.LLDB\"" + ProjectFileGenerator.NewLine);
+			Content.Append("      launchStyle = \"0\"" + ProjectFileGenerator.NewLine);
+			Content.Append("      useCustomWorkingDirectory = \"NO\"" + ProjectFileGenerator.NewLine);
+			Content.Append("      ignoresPersistentStateOnLaunch = \"NO\"" + ProjectFileGenerator.NewLine);
+			Content.Append("      debugDocumentVersioning = \"YES\"" + ProjectFileGenerator.NewLine);
+			Content.Append("      debugServiceExtension = \"internal\"" + ProjectFileGenerator.NewLine);
+			Content.Append("      allowLocationSimulation = \"YES\">" + ProjectFileGenerator.NewLine);
+			Content.Append("      <BuildableProductRunnable" + ProjectFileGenerator.NewLine);
+			Content.Append("         runnableDebuggingMode = \"0\">" + ProjectFileGenerator.NewLine);
+			Content.Append("            <BuildableReference" + ProjectFileGenerator.NewLine);
+			Content.Append("               BuildableIdentifier = \"primary\"" + ProjectFileGenerator.NewLine);
+			Content.Append("               BlueprintIdentifier = \"" + TargetGuid + "\"" + ProjectFileGenerator.NewLine);
+			Content.Append("               BuildableName = \"" + TargetName + ".app\"" + ProjectFileGenerator.NewLine);
+			Content.Append("               BlueprintName = \"" + TargetName + "\"" + ProjectFileGenerator.NewLine);
+			Content.Append("               ReferencedContainer = \"container:" + TargetName + ".xcodeproj\">" + ProjectFileGenerator.NewLine);
+			Content.Append("            </BuildableReference>" + ProjectFileGenerator.NewLine);
+			Content.Append("      </BuildableProductRunnable>" + ProjectFileGenerator.NewLine);
+			if (string.IsNullOrEmpty(OldCommandLineArguments))
+			{
+				if (bHasEditorConfiguration && TargetName != "UE5")
+				{
+					Content.Append("      <CommandLineArguments>" + ProjectFileGenerator.NewLine);
+					if (IsForeignProject)
+					{
+						Content.Append("         <CommandLineArgument" + ProjectFileGenerator.NewLine);
+						Content.Append("            argument = \"&quot;" + GameProjectPath + "&quot;\"" + ProjectFileGenerator.NewLine);
+						Content.Append("            isEnabled = \"YES\">" + ProjectFileGenerator.NewLine);
+						Content.Append("         </CommandLineArgument>" + ProjectFileGenerator.NewLine);
+					}
+					else
+					{
+						Content.Append("         <CommandLineArgument" + ProjectFileGenerator.NewLine);
+						Content.Append("            argument = \"" + TargetName + "\"" + ProjectFileGenerator.NewLine);
+						Content.Append("            isEnabled = \"YES\">" + ProjectFileGenerator.NewLine);
+						Content.Append("         </CommandLineArgument>" + ProjectFileGenerator.NewLine);
+					}
+					// Always add a configuration argument
+					Content.Append("         <CommandLineArgument" + ProjectFileGenerator.NewLine);
+					Content.Append("            argument = \"-RunConfig=$(Configuration)\"" + ProjectFileGenerator.NewLine);
+					Content.Append("            isEnabled = \"YES\">" + ProjectFileGenerator.NewLine);
+					Content.Append("         </CommandLineArgument>" + ProjectFileGenerator.NewLine);
+					Content.Append("      </CommandLineArguments>" + ProjectFileGenerator.NewLine);
+				}
+			}
+			else
+			{
+				Content.Append("      <CommandLineArguments>" + OldCommandLineArguments + "</CommandLineArguments>");
+			}
+			Content.Append("      <AdditionalOptions>" + ProjectFileGenerator.NewLine);
+			Content.Append("      </AdditionalOptions>" + ProjectFileGenerator.NewLine);
+			Content.Append("   </LaunchAction>" + ProjectFileGenerator.NewLine);
+			Content.Append("   <ProfileAction" + ProjectFileGenerator.NewLine);
+			Content.Append("      buildConfiguration = \"" + DefaultConfiguration + "\"" + ProjectFileGenerator.NewLine);
+			Content.Append("      shouldUseLaunchSchemeArgsEnv = \"YES\"" + ProjectFileGenerator.NewLine);
+			Content.Append("      savedToolIdentifier = \"\"" + ProjectFileGenerator.NewLine);
+			Content.Append("      useCustomWorkingDirectory = \"NO\"" + ProjectFileGenerator.NewLine);
+			Content.Append("      debugDocumentVersioning = \"YES\">" + ProjectFileGenerator.NewLine);
+			Content.Append("      <BuildableProductRunnable" + ProjectFileGenerator.NewLine);
+			Content.Append("         runnableDebuggingMode = \"0\">" + ProjectFileGenerator.NewLine);
+			Content.Append("            <BuildableReference" + ProjectFileGenerator.NewLine);
+			Content.Append("               BuildableIdentifier = \"primary\"" + ProjectFileGenerator.NewLine);
+			Content.Append("               BlueprintIdentifier = \"" + TargetGuid + "\"" + ProjectFileGenerator.NewLine);
+			Content.Append("               BuildableName = \"" + TargetName + ".app\"" + ProjectFileGenerator.NewLine);
+			Content.Append("               BlueprintName = \"" + TargetName + "\"" + ProjectFileGenerator.NewLine);
+			Content.Append("               ReferencedContainer = \"container:" + TargetName + ".xcodeproj\">" + ProjectFileGenerator.NewLine);
+			Content.Append("            </BuildableReference>" + ProjectFileGenerator.NewLine);
+			Content.Append("      </BuildableProductRunnable>" + ProjectFileGenerator.NewLine);
+			Content.Append("   </ProfileAction>" + ProjectFileGenerator.NewLine);
+			Content.Append("   <AnalyzeAction" + ProjectFileGenerator.NewLine);
+			Content.Append("      buildConfiguration = \"" + DefaultConfiguration + "\">" + ProjectFileGenerator.NewLine);
+			Content.Append("   </AnalyzeAction>" + ProjectFileGenerator.NewLine);
+			Content.Append("   <ArchiveAction" + ProjectFileGenerator.NewLine);
+			Content.Append("      buildConfiguration = \"" + DefaultConfiguration + "\"" + ProjectFileGenerator.NewLine);
+			Content.Append("      revealArchiveInOrganizer = \"YES\">" + ProjectFileGenerator.NewLine);
+			Content.Append("   </ArchiveAction>" + ProjectFileGenerator.NewLine);
+			Content.Append("</Scheme>" + ProjectFileGenerator.NewLine);
+
+			File.WriteAllText(SchemeFilePath.FullName, Content.ToString(), new UTF8Encoding());
+
+			Content.Clear();
+
+			Content.Append("<?xml version=\"1.0\" encoding=\"UTF-8\"?>" + ProjectFileGenerator.NewLine);
+			Content.Append("<!DOCTYPE plist PUBLIC \"-//Apple//DTD PLIST 1.0//EN\" \"http://www.apple.com/DTDs/PropertyList-1.0.dtd\">" + ProjectFileGenerator.NewLine);
+			Content.Append("<plist version=\"1.0\">" + ProjectFileGenerator.NewLine);
+			Content.Append("<dict>" + ProjectFileGenerator.NewLine);
+			Content.Append("\t<key>SchemeUserState</key>" + ProjectFileGenerator.NewLine);
+			Content.Append("\t<dict>" + ProjectFileGenerator.NewLine);
+			Content.Append("\t\t<key>" + TargetName + ".xcscheme_^#shared#^_</key>" + ProjectFileGenerator.NewLine);
+			Content.Append("\t\t<dict>" + ProjectFileGenerator.NewLine);
+			Content.Append("\t\t\t<key>orderHint</key>" + ProjectFileGenerator.NewLine);
+			Content.Append("\t\t\t<integer>1</integer>" + ProjectFileGenerator.NewLine);
+			Content.Append("\t\t</dict>" + ProjectFileGenerator.NewLine);
+			Content.Append("\t</dict>" + ProjectFileGenerator.NewLine);
+			Content.Append("\t<key>SuppressBuildableAutocreation</key>" + ProjectFileGenerator.NewLine);
+			Content.Append("\t<dict>" + ProjectFileGenerator.NewLine);
+			Content.Append("\t\t<key>" + TargetGuid + "</key>" + ProjectFileGenerator.NewLine);
+			Content.Append("\t\t<dict>" + ProjectFileGenerator.NewLine);
+			Content.Append("\t\t\t<key>primary</key>" + ProjectFileGenerator.NewLine);
+			Content.Append("\t\t\t<true/>" + ProjectFileGenerator.NewLine);
+			Content.Append("\t\t</dict>" + ProjectFileGenerator.NewLine);
+			Content.Append("\t\t<key>" + BuildTargetGuid + "</key>" + ProjectFileGenerator.NewLine);
+			Content.Append("\t\t<dict>" + ProjectFileGenerator.NewLine);
+			Content.Append("\t\t\t<key>primary</key>" + ProjectFileGenerator.NewLine);
+			Content.Append("\t\t\t<true/>" + ProjectFileGenerator.NewLine);
+			Content.Append("\t\t</dict>" + ProjectFileGenerator.NewLine);
+			Content.Append("\t\t<key>" + IndexTargetGuid + "</key>" + ProjectFileGenerator.NewLine);
+			Content.Append("\t\t<dict>" + ProjectFileGenerator.NewLine);
+			Content.Append("\t\t\t<key>primary</key>" + ProjectFileGenerator.NewLine);
+			Content.Append("\t\t\t<true/>" + ProjectFileGenerator.NewLine);
+			Content.Append("\t\t</dict>" + ProjectFileGenerator.NewLine);
+			Content.Append("\t</dict>" + ProjectFileGenerator.NewLine);
+			Content.Append("</dict>" + ProjectFileGenerator.NewLine);
+			Content.Append("</plist>" + ProjectFileGenerator.NewLine);
+
+			FileReference ManagementFile = GetUserSchemeManagementFilePath();
+			if (!DirectoryReference.Exists(ManagementFile.Directory))
+			{
+				DirectoryReference.CreateDirectory(ManagementFile.Directory);
+			}
+
+			File.WriteAllText(ManagementFile.FullName, Content.ToString(), new UTF8Encoding());
+		}
+
+		public static IEnumerable<UnrealTargetPlatform> GetSupportedPlatforms()
+		{
+			List<UnrealTargetPlatform> SupportedPlatforms = new List<UnrealTargetPlatform>();
+
+			if (XcodeProjectFileGenerator.ProjectFilePlatform.HasFlag(XcodeProjectFileGenerator.XcodeProjectFilePlatform.Mac))
+			{
+				SupportedPlatforms.Add(UnrealTargetPlatform.Mac);
+			}
+
+			if (XcodeProjectFileGenerator.ProjectFilePlatform.HasFlag(XcodeProjectFileGenerator.XcodeProjectFilePlatform.iOS))
+			{
+				SupportedPlatforms.Add(UnrealTargetPlatform.IOS);
+			}
+
+			if (XcodeProjectFileGenerator.ProjectFilePlatform.HasFlag(XcodeProjectFileGenerator.XcodeProjectFilePlatform.tvOS))
+			{
+				SupportedPlatforms.Add(UnrealTargetPlatform.TVOS);
+			}
+
+			return SupportedPlatforms;
+		}
+
+		public static IEnumerable<UnrealTargetConfiguration> GetSupportedConfigurations()
+		{
+			return new UnrealTargetConfiguration[] {
+				UnrealTargetConfiguration.Debug,
+				UnrealTargetConfiguration.DebugGame,
+				UnrealTargetConfiguration.Development,
+				UnrealTargetConfiguration.Test,
+				UnrealTargetConfiguration.Shipping
+			};
+		}
+
+		public bool ShouldIncludeProjectInWorkspace(ILogger Logger)
+		{
+			return CanBuildProjectLocally(Logger);
+		}
+
+		public bool CanBuildProjectLocally(ILogger Logger)
+		{
+			foreach (Project ProjectTarget in ProjectTargets)
+			{
+				foreach (UnrealTargetPlatform Platform in GetSupportedPlatforms())
+				{
+					foreach (UnrealTargetConfiguration Config in GetSupportedConfigurations())
+					{
+						if (MSBuildProjectFile.IsValidProjectPlatformAndConfiguration(ProjectTarget, Platform, Config, Logger))
+						{
+							return true;
+						}
+					}
+				}
+			}
+
+			return false;
+		}
+
+		/// Implements Project interface
+		public override bool WriteProjectFile(List<UnrealTargetPlatform> InPlatforms, List<UnrealTargetConfiguration> InConfigurations, PlatformProjectGeneratorCollection PlatformProjectGenerators, ILogger Logger)
+		{
+			bool bSuccess = true;
+
+			string TargetName = ProjectFilePath.GetFileNameWithoutExtension();
+			string TargetGuid = XcodeProjectFileGenerator.MakeXcodeGuid();
+			string TargetConfigListGuid = XcodeProjectFileGenerator.MakeXcodeGuid();
+			string TargetDependencyGuid = XcodeProjectFileGenerator.MakeXcodeGuid();
+			string TargetProxyGuid = XcodeProjectFileGenerator.MakeXcodeGuid();
+			string TargetAppGuid = XcodeProjectFileGenerator.MakeXcodeGuid();
+			string BuildTargetName = TargetName + "_Build";
+			string BuildTargetGuid = XcodeProjectFileGenerator.MakeXcodeGuid();
+			string BuildTargetConfigListGuid = XcodeProjectFileGenerator.MakeXcodeGuid();
+			string IndexTargetName = TargetName + "_Index";
+			string IndexTargetGuid = XcodeProjectFileGenerator.MakeXcodeGuid();
+			string IndexTargetConfigListGuid = XcodeProjectFileGenerator.MakeXcodeGuid();
+			string ProjectGuid = XcodeProjectFileGenerator.MakeXcodeGuid();
+			string ProjectConfigListGuid = XcodeProjectFileGenerator.MakeXcodeGuid();
+			string MainGroupGuid = XcodeProjectFileGenerator.MakeXcodeGuid();
+			string ProductRefGroupGuid = XcodeProjectFileGenerator.MakeXcodeGuid();
+			string SourcesBuildPhaseGuid = XcodeProjectFileGenerator.MakeXcodeGuid();
+			string CopyExtensionsBuildPhaseGuid = XcodeProjectFileGenerator.MakeXcodeGuid();
+			string ShellScriptSectionGuid = XcodeProjectFileGenerator.MakeXcodeGuid();
+
+
+			// Figure out all the desired configurations
+			List<XcodeBuildConfig> BuildConfigs = GetSupportedBuildConfigs(InPlatforms, InConfigurations, Logger);
+			if (BuildConfigs.Count == 0)
+			{
+				return true;
+			}
+
+			FileReference? GameProjectPath = null;
+			foreach (Project Target in ProjectTargets)
+			{
+				if (Target.UnrealProjectFilePath != null)
+				{
+					GameProjectPath = Target.UnrealProjectFilePath;
+					break;
+				}
+			}
+
+			bool bHasEditorConfiguration = false;
+			bool bIsAppBundle = true;
+
+			Dictionary<string, XcodeBuildConfig> ProjectBuildConfigs = new Dictionary<string, XcodeBuildConfig>();
+			Dictionary<string, XcodeBuildConfig> TargetBuildConfigs = new Dictionary<string, XcodeBuildConfig>();
+			Dictionary<string, XcodeBuildConfig> BuildTargetBuildConfigs = new Dictionary<string, XcodeBuildConfig>();
+			Dictionary<string, XcodeBuildConfig> IndexTargetBuildConfigs = new Dictionary<string, XcodeBuildConfig>();
+			foreach (XcodeBuildConfig Config in BuildConfigs)
+			{
+				ProjectBuildConfigs[XcodeProjectFileGenerator.MakeXcodeGuid()] = Config;
+				TargetBuildConfigs[XcodeProjectFileGenerator.MakeXcodeGuid()] = Config;
+				BuildTargetBuildConfigs[XcodeProjectFileGenerator.MakeXcodeGuid()] = Config;
+				IndexTargetBuildConfigs[XcodeProjectFileGenerator.MakeXcodeGuid()] = Config;
+
+				if (Config.ProjectTarget!.TargetRules!.Type == TargetType.Editor)
+				{
+					bHasEditorConfiguration = true;
+				}
+				else if (Config.ProjectTarget.TargetRules.bIsBuildingConsoleApplication || Config.ProjectTarget.TargetRules.bShouldCompileAsDLL)
+				{
+					bIsAppBundle = false;
+				}
+			}
+
+			StringBuilder PBXBuildFileSection = new StringBuilder();
+			StringBuilder PBXFileReferenceSection = new StringBuilder();
+			StringBuilder PBXSourcesBuildPhaseSection = new StringBuilder();
+			StringBuilder PBXCopyExtensionsBuildPhaseSection = new StringBuilder();
+			StringBuilder PBXResourcesBuildPhaseSection = new StringBuilder();
+			List<XcodeExtensionInfo> AllExtensions = new List<XcodeExtensionInfo>();
+			GenerateSectionsWithSourceFiles(PBXBuildFileSection, PBXFileReferenceSection, PBXSourcesBuildPhaseSection, TargetAppGuid, TargetName, bIsAppBundle);
+			GenerateSectionsWithExtensions(PBXBuildFileSection, PBXFileReferenceSection, PBXCopyExtensionsBuildPhaseSection, PBXResourcesBuildPhaseSection, AllExtensions, GameProjectPath, BuildConfigs, Logger);
+
+			StringBuilder ProjectFileContent = new StringBuilder();
+
+			ProjectFileContent.Append("// !$*UTF8*$!" + ProjectFileGenerator.NewLine);
+			ProjectFileContent.Append("{" + ProjectFileGenerator.NewLine);
+			ProjectFileContent.Append("\tarchiveVersion = 1;" + ProjectFileGenerator.NewLine);
+			ProjectFileContent.Append("\tclasses = {" + ProjectFileGenerator.NewLine);
+			ProjectFileContent.Append("\t};" + ProjectFileGenerator.NewLine);
+			ProjectFileContent.Append("\tobjectVersion = 46;" + ProjectFileGenerator.NewLine);
+			ProjectFileContent.Append("\tobjects = {" + ProjectFileGenerator.NewLine + ProjectFileGenerator.NewLine);
+
+			AppendBuildFileSection(ProjectFileContent, PBXBuildFileSection);
+			AppendFileReferenceSection(ProjectFileContent, PBXFileReferenceSection);
+			AppendSourcesBuildPhaseSection(ProjectFileContent, PBXSourcesBuildPhaseSection, SourcesBuildPhaseGuid);
+			AppendCopyExtensionsBuildPhaseSection(ProjectFileContent, PBXCopyExtensionsBuildPhaseSection, CopyExtensionsBuildPhaseGuid);
+			ProjectFileContent.Append(PBXResourcesBuildPhaseSection);
+			AppendContainerItemProxySection(ProjectFileContent, BuildTargetName, BuildTargetGuid, TargetProxyGuid, ProjectGuid);
+			if (!XcodeProjectFileGenerator.bGeneratingRunIOSProject)
+			{
+				AppendTargetDependencySection(ProjectFileContent, BuildTargetName, BuildTargetGuid, TargetDependencyGuid, TargetProxyGuid);
+			}
+			foreach (XcodeExtensionInfo EI in AllExtensions)
+			{
+				AppendContainerItemProxySection(ProjectFileContent, EI.Name, EI.TargetGuid, EI.TargetProxyGuid, ProjectGuid);
+				AppendTargetDependencySection(ProjectFileContent, EI.Name, EI.TargetGuid, EI.TargetDependencyGuid, EI.TargetProxyGuid);
+			}
+			AppendGroupSection(ProjectFileContent, MainGroupGuid, ProductRefGroupGuid, TargetAppGuid, TargetName, AllExtensions);
+			AppendLegacyTargetSection(ProjectFileContent, BuildTargetName, BuildTargetGuid, BuildTargetConfigListGuid, GameProjectPath);
+			AppendRunTargetSection(ProjectFileContent, TargetName, TargetGuid, TargetConfigListGuid, TargetDependencyGuid, TargetAppGuid, CopyExtensionsBuildPhaseGuid, ShellScriptSectionGuid, AllExtensions, bIsAppBundle);
+			AppendIndexTargetSection(ProjectFileContent, IndexTargetName, IndexTargetGuid, IndexTargetConfigListGuid, SourcesBuildPhaseGuid);
+			AppendExtensionTargetSections(ProjectFileContent, AllExtensions);
+			AppendProjectSection(ProjectFileContent, TargetName, TargetGuid, BuildTargetName, BuildTargetGuid, IndexTargetName, IndexTargetGuid, MainGroupGuid, ProductRefGroupGuid, ProjectGuid, ProjectConfigListGuid, GameProjectPath, AllExtensions);
+			AppendXCBuildConfigurationSection(ProjectFileContent, ProjectBuildConfigs, TargetBuildConfigs, BuildTargetBuildConfigs, IndexTargetBuildConfigs, GameProjectPath, AllExtensions, Logger);
+			AppendXCConfigurationListSection(ProjectFileContent, TargetName, BuildTargetName, IndexTargetName, ProjectConfigListGuid, ProjectBuildConfigs,
+				TargetConfigListGuid, TargetBuildConfigs, BuildTargetConfigListGuid, BuildTargetBuildConfigs, IndexTargetConfigListGuid, IndexTargetBuildConfigs, AllExtensions);
+			AppendShellScriptSection(ProjectFileContent, ShellScriptSectionGuid, GameProjectPath);
+
+			ProjectFileContent.Append("\t};" + ProjectFileGenerator.NewLine);
+			ProjectFileContent.Append("\trootObject = " + ProjectGuid + " /* Project object */;" + ProjectFileGenerator.NewLine);
+			ProjectFileContent.Append("}" + ProjectFileGenerator.NewLine);
+
+			if (bSuccess)
+			{
+				FileReference PBXProjFilePath = ProjectFilePath + "/project.pbxproj";
+				bSuccess = ProjectFileGenerator.WriteFileIfChanged(PBXProjFilePath.FullName, ProjectFileContent.ToString(), Logger, new UTF8Encoding());
+			}
+
+			bool bNeedScheme = CanBuildProjectLocally(Logger);
+
+			if (bNeedScheme)
+			{
+				if (bSuccess)
+				{
+					WriteSchemeFile(TargetName, TargetGuid, BuildTargetGuid, IndexTargetGuid, bHasEditorConfiguration, GameProjectPath != null ? GameProjectPath.FullName : "");
+				}
+			}
+			else
+			{
+				// clean this up because we don't want it persisting if we narrow our project list
+				DirectoryReference SchemeDir = GetProjectSchemeDirectory();
+
+				if (DirectoryReference.Exists(SchemeDir))
+				{
+					DirectoryReference.Delete(SchemeDir, true);
+				}
+			}
+
+			return bSuccess;
+		}
+	}
+}
