UnrealBuildTool: Allow PCH when running the clang analyzer

#rnx

diff --git a/Engine/Source/Programs/UnrealBuildTool/Platform/Android/AndroidToolChain.cs b/Engine/Source/Programs/UnrealBuildTool/Platform/Android/AndroidToolChain.cs
--- a/Engine/Source/Programs/UnrealBuildTool/Platform/Android/AndroidToolChain.cs
+++ b/Engine/Source/Programs/UnrealBuildTool/Platform/Android/AndroidToolChain.cs
@@ -794,7 +794,7 @@ namespace UnrealBuildTool
 			}
 
 			// Add analysis flags to the argument list.
-			if (StaticAnalyzer == StaticAnalyzer.Default)
+			if (IsAnalyzing(CompileEnvironment))
 			{
 				List<string> Arguments = new List<string>();
 				GetCompileArguments_Analyze(CompileEnvironment, Arguments);
@@ -1406,7 +1406,7 @@ namespace UnrealBuildTool
 					string ResponseArgument = string.Format("@\"{0}\"", ResponseFileName);
 
 					CompileAction.WorkingDirectory = Unreal.EngineSourceDirectory;
-					CompileAction.CommandDescription = StaticAnalyzer == StaticAnalyzer.Default ? "Analyze" : "Compile";
+					CompileAction.CommandDescription = IsAnalyzing(CompileEnvironment) ? "Analyze" : "Compile";
 					if (bExecuteCompilerThroughShell)
 					{
 						SetupActionToExecuteCompilerThroughShell(ref CompileAction, ClangPath!, ResponseArgument, CompileAction.CommandDescription);
diff --git a/Engine/Source/Programs/UnrealBuildTool/Platform/Android/UEBuildAndroid.cs b/Engine/Source/Programs/UnrealBuildTool/Platform/Android/UEBuildAndroid.cs
--- a/Engine/Source/Programs/UnrealBuildTool/Platform/Android/UEBuildAndroid.cs
+++ b/Engine/Source/Programs/UnrealBuildTool/Platform/Android/UEBuildAndroid.cs
@@ -162,11 +162,9 @@ namespace UnrealBuildTool
 			}
 
 			// Disable linking if we're using a static analyzer
-			// Disable PCHs for clang static analyzer.
 			if (Target.StaticAnalyzer == StaticAnalyzer.Default)
 			{
 				Target.bDisableLinking = true;
-				Target.bUsePCHFiles = false;
 			}
 
 			Target.bCompileRecast = true;
diff --git a/Engine/Source/Programs/UnrealBuildTool/Platform/IOS/UEBuildIOS.cs b/Engine/Source/Programs/UnrealBuildTool/Platform/IOS/UEBuildIOS.cs
--- a/Engine/Source/Programs/UnrealBuildTool/Platform/IOS/UEBuildIOS.cs
+++ b/Engine/Source/Programs/UnrealBuildTool/Platform/IOS/UEBuildIOS.cs
@@ -772,11 +772,9 @@ namespace UnrealBuildTool
 			}
 
 			// Disable linking if we're using a static analyzer
-			// Disable PCHs for clang static analyzer.
 			if (Target.StaticAnalyzer == StaticAnalyzer.Default)
 			{
 				Target.bDisableLinking = true;
-				Target.bUsePCHFiles = false;
 			}
 
 			// we assume now we are building with IOS8 or later
diff --git a/Engine/Source/Programs/UnrealBuildTool/Platform/Linux/UEBuildLinux.cs b/Engine/Source/Programs/UnrealBuildTool/Platform/Linux/UEBuildLinux.cs
--- a/Engine/Source/Programs/UnrealBuildTool/Platform/Linux/UEBuildLinux.cs
+++ b/Engine/Source/Programs/UnrealBuildTool/Platform/Linux/UEBuildLinux.cs
@@ -224,11 +224,9 @@ namespace UnrealBuildTool
 			}
 
 			// Disable linking if we're using a static analyzer
-			// Disable PCHs for clang static analyzer.
 			if (Target.StaticAnalyzer == StaticAnalyzer.Default)
 			{
 				Target.bDisableLinking = true;
-				Target.bUsePCHFiles = false;
 			}
 
 			// Editor target types get overwritten in UEBuildTarget.cs so lets avoid adding this here. ResetTarget is called with
diff --git a/Engine/Source/Programs/UnrealBuildTool/Platform/Mac/UEBuildMac.cs b/Engine/Source/Programs/UnrealBuildTool/Platform/Mac/UEBuildMac.cs
--- a/Engine/Source/Programs/UnrealBuildTool/Platform/Mac/UEBuildMac.cs
+++ b/Engine/Source/Programs/UnrealBuildTool/Platform/Mac/UEBuildMac.cs
@@ -121,11 +121,9 @@ namespace UnrealBuildTool
 			}
 
 			// Disable linking if we're using a static analyzer
-			// Disable PCHs for clang static analyzer.
 			if (Target.StaticAnalyzer == StaticAnalyzer.Default)
 			{
 				Target.bDisableLinking = true;
-				Target.bUsePCHFiles = false;
 			}
 
 			if (BuildHostPlatform.Current.Platform != UnrealTargetPlatform.Mac)
diff --git a/Engine/Source/Programs/UnrealBuildTool/Platform/Windows/UEBuildWindows.cs b/Engine/Source/Programs/UnrealBuildTool/Platform/Windows/UEBuildWindows.cs
--- a/Engine/Source/Programs/UnrealBuildTool/Platform/Windows/UEBuildWindows.cs
+++ b/Engine/Source/Programs/UnrealBuildTool/Platform/Windows/UEBuildWindows.cs
@@ -875,15 +875,11 @@ namespace UnrealBuildTool
 				Target.bDisableLinking = true;
 			}
 
-			 // Disable PCHs for PVS studio and clang static analyzer.
+			 // Disable PCHs for PVS studio analyzer.
 			if(Target.StaticAnalyzer == StaticAnalyzer.PVSStudio)
 			{
 				Target.bUsePCHFiles = false;
 			}
-			else if (Target.WindowsPlatform.Compiler.IsClang() && Target.StaticAnalyzer == StaticAnalyzer.Default)
-			{
-				Target.bUsePCHFiles = false;
-			}
 			
 			// @todo: Override PCH settings
 			if (Target.WindowsPlatform.Compiler == WindowsCompiler.Intel)
diff --git a/Engine/Source/Programs/UnrealBuildTool/Platform/Windows/VCToolChain.cs b/Engine/Source/Programs/UnrealBuildTool/Platform/Windows/VCToolChain.cs
--- a/Engine/Source/Programs/UnrealBuildTool/Platform/Windows/VCToolChain.cs
+++ b/Engine/Source/Programs/UnrealBuildTool/Platform/Windows/VCToolChain.cs
@@ -267,7 +267,7 @@ namespace UnrealBuildTool
 				// compiler version of 19.xx).
 				Arguments.Add($"-fms-compatibility-version=19.{EnvVars.ToolChainVersion.GetComponent(1)}");
 				
-				if (Target.StaticAnalyzer == StaticAnalyzer.Default)
+				if (Target.StaticAnalyzer == StaticAnalyzer.Default && CompileEnvironment.PrecompiledHeaderAction != PrecompiledHeaderAction.Create)
 				{
 					// Enable the static analyzer but only via the backend. Using the frontend
 					// flag ('--analyze') will enable a suite of default checkers, some of which
diff --git a/Engine/Source/Programs/UnrealBuildTool/ToolChain/ClangToolChain.cs b/Engine/Source/Programs/UnrealBuildTool/ToolChain/ClangToolChain.cs
--- a/Engine/Source/Programs/UnrealBuildTool/ToolChain/ClangToolChain.cs
+++ b/Engine/Source/Programs/UnrealBuildTool/ToolChain/ClangToolChain.cs
@@ -279,6 +279,8 @@ namespace UnrealBuildTool
 		/// </summary>
 		protected bool CompilerVersionLessThan(int Major, int Minor, int Patch) => Info.ClangVersion < new Version(Major, Minor, Patch);
 
+		protected bool IsAnalyzing(CppCompileEnvironment CompileEnvironment) => StaticAnalyzer == StaticAnalyzer.Default && CompileEnvironment.PrecompiledHeaderAction != PrecompiledHeaderAction.Create;
+
 		protected virtual void GetCppStandardCompileArgument(CppCompileEnvironment CompileEnvironment, List<string> Arguments)
 		{
 			// https://clang.llvm.org/cxx_status.html
@@ -731,7 +733,7 @@ namespace UnrealBuildTool
 			GetCompilerArguments_Sanitizers(CompileEnvironment, Arguments);
 
 			// Add analysis flags to the argument list.
-			if (StaticAnalyzer == StaticAnalyzer.Default)
+			if (IsAnalyzing(CompileEnvironment))
 			{
 				GetCompileArguments_Analyze(CompileEnvironment, Arguments);
 			}
@@ -819,7 +821,7 @@ namespace UnrealBuildTool
 			// Add the source file path to the command-line.
 			Arguments.Add(GetSourceFileArgument(SourceFile));
 
-			if (StaticAnalyzer != StaticAnalyzer.Default)
+			if (!IsAnalyzing(CompileEnvironment))
 			{
 				// Generate the timing info
 				if (CompileEnvironment.bPrintTimingInfo)
@@ -838,7 +840,7 @@ namespace UnrealBuildTool
 					CompileAction.ProducedItems.Add(DependencyListFile);
 				}
 			}
-			else if (StaticAnalyzer == StaticAnalyzer.Default)
+			else
 			{
 				OutputFile = FileItem.GetItemByFileReference(FileReference.Combine(OutputDir, Path.GetFileName(SourceFile.AbsolutePath) + ".analysis"));
 			}
@@ -891,7 +893,7 @@ namespace UnrealBuildTool
 			CompileAction.WorkingDirectory = Unreal.EngineSourceDirectory;
 			CompileAction.CommandPath = Info.Clang;
 			CompileAction.CommandVersion = Info.ClangVersionString;
-			CompileAction.CommandDescription = StaticAnalyzer == StaticAnalyzer.Default ? "Analyze" : "Compile";
+			CompileAction.CommandDescription = IsAnalyzing(CompileEnvironment) ? "Analyze" : "Compile";
 			CompileAction.StatusDescription = Path.GetFileName(SourceFile.AbsolutePath);
 			CompileAction.bIsGCCCompiler = true;
 
