UnrealBuildTool: Add TargetRules properties to control which clang checkers are enabled. Reconcile msvc cl-clang analyze with clang analyze

#rnx

diff --git a/Engine/Source/Programs/UnrealBuildTool/Configuration/TargetRules.cs b/Engine/Source/Programs/UnrealBuildTool/Configuration/TargetRules.cs
--- a/Engine/Source/Programs/UnrealBuildTool/Configuration/TargetRules.cs
+++ b/Engine/Source/Programs/UnrealBuildTool/Configuration/TargetRules.cs
@@ -1455,19 +1455,71 @@ namespace UnrealBuildTool
 		public StaticAnalyzer StaticAnalyzer = StaticAnalyzer.None;
 
 		/// <summary>
-		/// The output type to use for the static analyzer.
+		/// The output type to use for the static analyzer. This is only supported for Clang.
 		/// </summary>
 		[CommandLine("-StaticAnalyzerOutputType")]
 		[XmlConfigFile(Category = "BuildConfiguration")]
 		public StaticAnalyzerOutputType StaticAnalyzerOutputType = StaticAnalyzerOutputType.Text;
 
 		/// <summary>
-		/// The mode to use for the static analyzer.
+		/// The mode to use for the static analyzer. This is only supported for Clang.
+		/// Shallow mode completes quicker but is generally not recommended.
 		/// </summary>
 		[CommandLine("-StaticAnalyzerMode")]
 		[XmlConfigFile(Category = "BuildConfiguration")]
 		public StaticAnalyzerMode StaticAnalyzerMode = StaticAnalyzerMode.Deep;
 
+		/// <summary>
+		/// The static analyzer checkers that should be enabled rather than the defaults. This is only supported for Clang.
+		/// See https://clang.llvm.org/docs/analyzer/checkers.html for a full list. Or run:
+		///    'clang -Xclang -analyzer-checker-help' 
+		/// or: 
+		///    'clang -Xclang -analyzer-checker-help-alpha' 
+		/// for the list of experimental checkers.
+		/// </summary>
+		[CommandLine("-StaticAnalyzerChecker", ListSeparator = '+')]
+		public List<string> StaticAnalyzerCheckers = new List<string>();
+
+		/// <summary>
+		/// The static analyzer default checkers that should be disabled. Unused if StaticAnalyzerCheckers is set. This is only supported for Clang.
+		/// This overrides the default disabled checkers, which are deadcode.DeadStores and security.FloatLoopCounter
+		/// See https://clang.llvm.org/docs/analyzer/checkers.html for a full list. Or run:
+		///    'clang -Xclang -analyzer-checker-help' 
+		/// or: 
+		///    'clang -Xclang -analyzer-checker-help-alpha' 
+		/// for the list of experimental checkers.
+		/// </summary>
+		[CommandLine("-StaticAnalyzerDisableChecker", ListSeparator = '+')]
+		private List<string> StaticAnalyzerDisabledCheckersOverride = new List<string>();
+
+		/// <summary>
+		/// The static analyzer default checkers that should be disabled. Unused if StaticAnalyzerCheckers is set. This is only supported for Clang.
+		/// See https://clang.llvm.org/docs/analyzer/checkers.html for a full list. Or run:
+		///    'clang -Xclang -analyzer-checker-help' 
+		/// or: 
+		///    'clang -Xclang -analyzer-checker-help-alpha' 
+		/// for the list of experimental checkers.
+		/// </summary>
+		public List<string> StaticAnalyzerDisabledCheckers
+		{
+			get
+			{
+				return StaticAnalyzerDisabledCheckersOverride.Count > 0 ? StaticAnalyzerDisabledCheckersOverride : new List<string>() { "deadcode.DeadStores", "security.FloatLoopCounter" };
+			}
+			set { StaticAnalyzerDisabledCheckersOverride = value; }
+		}
+
+		/// <summary>
+		/// The static analyzer non-default checkers that should be enabled. Unused if StaticAnalyzerCheckers is set. This is only supported for Clang.
+		/// See https://clang.llvm.org/docs/analyzer/checkers.html for a full list. Or run:
+		///    'clang -Xclang -analyzer-checker-help' 
+		/// or: 
+		///    'clang -Xclang -analyzer-checker-help-alpha' 
+		/// for the list of experimental checkers.
+		/// </summary>
+		[CommandLine("-StaticAnalyzerAdditionalChecker", ListSeparator = '+')]
+		public List<string> StaticAnalyzerAdditionalCheckers = new List<string>();
+
 		/// <summary>
 		/// The minimum number of files that must use a pre-compiled header before it will be created and used.
 		/// </summary>
@@ -3077,6 +3129,21 @@ namespace UnrealBuildTool
 			get { return Inner.StaticAnalyzerMode; }
 		}
 
+		public IReadOnlyList<string> StaticAnalyzerCheckers
+		{
+			get { return Inner.StaticAnalyzerCheckers; }
+		}
+
+		public IReadOnlyList<string> StaticAnalyzerDisabledCheckers
+		{
+			get { return Inner.StaticAnalyzerDisabledCheckers; }
+		}
+
+		public IReadOnlyList<string> StaticAnalyzerAdditionalCheckers
+		{
+			get { return Inner.StaticAnalyzerAdditionalCheckers; }
+		}
+
 		public int MinFilesUsingPrecompiledHeader
 		{
 			get { return Inner.MinFilesUsingPrecompiledHeader; }
diff --git a/Engine/Source/Programs/UnrealBuildTool/Platform/Windows/VCToolChain.cs b/Engine/Source/Programs/UnrealBuildTool/Platform/Windows/VCToolChain.cs
--- a/Engine/Source/Programs/UnrealBuildTool/Platform/Windows/VCToolChain.cs
+++ b/Engine/Source/Programs/UnrealBuildTool/Platform/Windows/VCToolChain.cs
@@ -269,63 +269,59 @@ namespace UnrealBuildTool
 				
 				if (Target.StaticAnalyzer == StaticAnalyzer.Default && CompileEnvironment.PrecompiledHeaderAction != PrecompiledHeaderAction.Create)
 				{
-					// Enable the static analyzer but only via the backend. Using the frontend
-					// flag ('--analyze') will enable a suite of default checkers, some of which
-					// we don't want (like deadcode.DeadStore)
-					Arguments.Add("-Xclang -analyze");
+					Arguments.Add("-Wno-unused-command-line-argument");
 
-					// Make sure we get textual output and not XML.
-					if (Target.StaticAnalyzerOutputType == StaticAnalyzerOutputType.Html)
+					if (Target.StaticAnalyzerCheckers.Count == 0)
 					{
-						// Write out a pretty web page with navigation to understand how the analysis was derived.
-						Arguments.Add("-Xclang -analyzer-output=html");
-						
-						// If writing to HTML, use the source filename as a basis for the report filename. 
-						Arguments.Add("-Xclang -analyzer-config -Xclang stable-report-filename=true");
+						// Enable the static analyzer with default checkers.
+						Arguments.Add("--analyze");
 					}
 					else
 					{
-						Arguments.Add("-Xclang -analyzer-output=text");
+						// Enable the static analyzer but only via the backend to disable all default checkers.
+						Arguments.Add("-Xclang -analyze");
 					}
 
+					// Ensure the compiler sets the __clang_analyzer__ macro correctly.
+					Arguments.Add("-Xclang -setup-static-analyzer");
+
 					// Make sure we check inside nested blocks (e.g. 'if ((foo = getchar()) == 0) {}')
 					Arguments.Add("-Xclang -analyzer-opt-analyze-nested-blocks");
 
+					// Write out a pretty web page with navigation to understand how the analysis was derived if HTML is enabled.
+					Arguments.Add($"-Xclang -analyzer-output={Target.StaticAnalyzerOutputType.ToString().ToLowerInvariant()}");
+
 					// Needed for some of the C++ checkers.
 					Arguments.Add("-Xclang -analyzer-config -Xclang aggressive-binary-operation-simplification=true");
 
-					// Ensure the compiler sets the __clang_analyzer__ macro correctly.
-					Arguments.Add("-Xclang -setup-static-analyzer");
+					// If writing to HTML, use the source filename as a basis for the report filename. 
+					Arguments.Add("-Xclang -analyzer-config -Xclang stable-report-filename=true");
+					Arguments.Add("-Xclang -analyzer-config -Xclang report-in-main-source-file=true");
+					Arguments.Add("-Xclang -analyzer-config -Xclang path-diagnostics-alternate=true");
+
+					// Run shallow analyze if requested.
+					if (Target.StaticAnalyzerMode == StaticAnalyzerMode.Shallow) Arguments.Add("-Xclang -analyzer-config -Xclang mode=shallow");
 
-					// Enable only specific checkers of families of checkers
-					// See https://clang.llvm.org/docs/analyzer/checkers.html for a full list. Or run:
-					//    'clang -Xclang -analyzer-checker-help' 
-					// or: 
-					//    'clang -Xclang -analyzer-checker-help-alpha' 
-					// for the list of experimental checkers.
-					String[] EnabledCheckers = 
+					if (Target.StaticAnalyzerCheckers.Count > 0)
 					{
-						"core",
-						"unix.Malloc",
-						"unix.MallocSizeof",
-						"cplusplus",
-						"optin.cplusplus.UninitializedObject",
-						"optin.cplusplus.VirtualCall",
-						// "deadstore",					// Check for dead stores (noisy in UE).
-						// "security.FloatLoopCounter",	// Check if using floats for loop counters
-						// Experimental checkers (as of clang 11.0)
-						// "alpha.core.Conversion",		// Sign conversion (noisy)
-						// "alpha.core.PointerArithm",	// Sketchy pointer arithmetic
-						// "alpha.cplusplus.DeleteWithNonVirtualDtor",
-						// "alpha.cplusplus.InvalidatedIterator",
-						// "alpha.cplusplus.IteratorRange",
-						// "alpha.cplusplus.MismatchedIterator",
-						// "alpha.cplusplus.InvalidatedIterator",
-					};
-
-					foreach (String Checker in EnabledCheckers)
+						// Only enable specific checkers.
+						foreach (string Checker in Target.StaticAnalyzerCheckers)
+						{
+							Arguments.Add($"-Xclang -analyzer-checker -Xclang {Checker}");
+						}
+					}
+					else
 					{
-						Arguments.Add($"-Xclang -analyzer-checker={Checker}");
+						// Disable default checkers.
+						foreach (string Checker in Target.StaticAnalyzerDisabledCheckers)
+						{
+							Arguments.Add($"-Xclang -analyzer-disable-checker -Xclang {Checker}");
+						}
+						// Enable additional non-default checkers.
+						foreach (string Checker in Target.StaticAnalyzerAdditionalCheckers)
+						{
+							Arguments.Add($"-Xclang -analyzer-checker -Xclang {Checker}");
+						}
 					}
 				}
 			}
diff --git a/Engine/Source/Programs/UnrealBuildTool/ToolChain/ClangToolChain.cs b/Engine/Source/Programs/UnrealBuildTool/ToolChain/ClangToolChain.cs
--- a/Engine/Source/Programs/UnrealBuildTool/ToolChain/ClangToolChain.cs
+++ b/Engine/Source/Programs/UnrealBuildTool/ToolChain/ClangToolChain.cs
@@ -187,6 +187,9 @@ namespace UnrealBuildTool
 		protected StaticAnalyzer StaticAnalyzer = StaticAnalyzer.None;
 		protected StaticAnalyzerMode StaticAnalyzerMode = StaticAnalyzerMode.Deep;
 		protected StaticAnalyzerOutputType StaticAnalyzerOutputType = StaticAnalyzerOutputType.Text;
+		protected IReadOnlyList<string> StaticAnalyzerCheckers = new List<string>();
+		protected IReadOnlyList<string> StaticAnalyzerDisabledCheckers = new List<string>();
+		protected IReadOnlyList<string> StaticAnalyzerAdditionalCheckers = new List<string>();
 
 		static ClangToolChain()
 		{
@@ -213,6 +216,9 @@ namespace UnrealBuildTool
 			StaticAnalyzer = Target.StaticAnalyzer;
 			StaticAnalyzerMode = Target.StaticAnalyzerMode;
 			StaticAnalyzerOutputType = Target.StaticAnalyzerOutputType;
+			StaticAnalyzerCheckers = Target.StaticAnalyzerCheckers;
+			StaticAnalyzerDisabledCheckers = Target.StaticAnalyzerDisabledCheckers;
+			StaticAnalyzerAdditionalCheckers = Target.StaticAnalyzerAdditionalCheckers;
 		}
 
 		public override void FinalizeOutput(ReadOnlyTargetRules Target, TargetMakefileBuilder MakefileBuilder)
@@ -681,21 +687,65 @@ namespace UnrealBuildTool
 		}
 
 		/// <summary>
-		/// Compile arguments for running clang-analyze. Currently unsupported.
+		/// Compile arguments for running clang-analyze.
 		/// </summary>
 		/// <param name="CompileEnvironment"></param>
 		/// <param name="Arguments"></param>
 		protected virtual void GetCompileArguments_Analyze(CppCompileEnvironment CompileEnvironment, List<string> Arguments)
 		{
 			Arguments.Add("-Wno-unused-command-line-argument");
-			Arguments.Add("--analyze");
+			if (StaticAnalyzerCheckers.Count == 0)
+			{
+				// Enable the static analyzer with default checkers.
+				Arguments.Add("--analyze");
+			}
+			else
+			{
+				// Enable the static analyzer but only via the backend to disable all default checkers.
+				Arguments.Add("-Xclang -analyze");
+			}
+
+			// Ensure the compiler sets the __clang_analyzer__ macro correctly.
+			Arguments.Add("-Xclang -setup-static-analyzer");
+
+			// Make sure we check inside nested blocks (e.g. 'if ((foo = getchar()) == 0) {}')
+			Arguments.Add("-Xclang -analyzer-opt-analyze-nested-blocks");
+
+			// Write out a pretty web page with navigation to understand how the analysis was derived if HTML is enabled.
 			Arguments.Add($"-Xclang -analyzer-output={StaticAnalyzerOutputType.ToString().ToLowerInvariant()}");
+
+			// Needed for some of the C++ checkers.
+			Arguments.Add("-Xclang -analyzer-config -Xclang aggressive-binary-operation-simplification=true");
+
+			// If writing to HTML, use the source filename as a basis for the report filename.
 			Arguments.Add("-Xclang -analyzer-config -Xclang stable-report-filename=true");
 			Arguments.Add("-Xclang -analyzer-config -Xclang report-in-main-source-file=true");
 			Arguments.Add("-Xclang -analyzer-config -Xclang path-diagnostics-alternate=true");
-			Arguments.Add("-Xclang -analyzer-disable-checker -Xclang deadcode.DeadStores");
-			Arguments.Add("-Xclang -analyzer-disable-checker -Xclang security.FloatLoopCounter");
+
+			// Run shallow analyze if requested.
 			if (StaticAnalyzerMode == StaticAnalyzerMode.Shallow) Arguments.Add("-Xclang -analyzer-config -Xclang mode=shallow");
+
+			if (StaticAnalyzerCheckers.Count > 0)
+			{
+				// Only enable specific checkers.
+				foreach (string Checker in StaticAnalyzerCheckers)
+				{
+					Arguments.Add($"-Xclang -analyzer-checker -Xclang {Checker}");
+				}
+			}
+			else
+			{
+				// Disable default checkers.
+				foreach (string Checker in StaticAnalyzerDisabledCheckers)
+				{
+					Arguments.Add($"-Xclang -analyzer-disable-checker -Xclang {Checker}");
+				}
+				// Enable additional non-default checkers.
+				foreach (string Checker in StaticAnalyzerAdditionalCheckers)
+				{
+					Arguments.Add($"-Xclang -analyzer-checker -Xclang {Checker}");
+				}
+			}
 		}
 
 		/// <summary>
