Added the ability to turn off merging of module and generated cpp files in the module.

diff --git a/Engine/Source/Programs/UnrealBuildTool/Configuration/ModuleRules.cs b/Engine/Source/Programs/UnrealBuildTool/Configuration/ModuleRules.cs
--- a/Engine/Source/Programs/UnrealBuildTool/Configuration/ModuleRules.cs
+++ b/Engine/Source/Programs/UnrealBuildTool/Configuration/ModuleRules.cs
@@ -817,6 +817,11 @@ namespace UnrealBuildTool
 			}
 		}
 
+		/// <summary>
+		/// Whether to merge module and generated unity files for faster compilation.
+		/// </summary>
+		public bool bMergeUnityFiles = true;
+
 		/// <summary>
 		/// The number of source files in this module before unity build will be activated for that module.  If set to
 		/// anything besides -1, will override the default setting which is controlled by MinGameModuleSourceFilesForUnityBuild
diff --git a/Engine/Source/Programs/UnrealBuildTool/Configuration/UEBuildModuleCPP.cs b/Engine/Source/Programs/UnrealBuildTool/Configuration/UEBuildModuleCPP.cs
--- a/Engine/Source/Programs/UnrealBuildTool/Configuration/UEBuildModuleCPP.cs
+++ b/Engine/Source/Programs/UnrealBuildTool/Configuration/UEBuildModuleCPP.cs
@@ -559,7 +559,7 @@ namespace UnrealBuildTool
 
 				if (GeneratedFiles.Count > 0)
 				{
-					bool bMergeUnityFiles = Target.bMergeModuleAndGeneratedUnityFiles;
+					bool bMergeUnityFiles = Target.bMergeModuleAndGeneratedUnityFiles && Rules.bMergeUnityFiles;
 
 					// Create a compile environment for the generated files. We can disable creating debug info here to improve link times.
 					CppCompileEnvironment GeneratedCPPCompileEnvironment = CompileEnvironment;
