UnrealBuildTool: Use --analyzer-no-default-checks to disable default clang analyzers and remove unneeded flags

#rnx

diff --git a/Engine/Source/Programs/UnrealBuildTool/Platform/Windows/VCToolChain.cs b/Engine/Source/Programs/UnrealBuildTool/Platform/Windows/VCToolChain.cs
--- a/Engine/Source/Programs/UnrealBuildTool/Platform/Windows/VCToolChain.cs
+++ b/Engine/Source/Programs/UnrealBuildTool/Platform/Windows/VCToolChain.cs
@@ -271,19 +271,8 @@ namespace UnrealBuildTool
 				{
 					Arguments.Add("-Wno-unused-command-line-argument");
 
-					if (Target.StaticAnalyzerCheckers.Count == 0)
-					{
-						// Enable the static analyzer with default checkers.
-						Arguments.Add("--analyze");
-					}
-					else
-					{
-						// Enable the static analyzer but only via the backend to disable all default checkers.
-						Arguments.Add("-Xclang -analyze");
-					}
-
-					// Ensure the compiler sets the __clang_analyzer__ macro correctly.
-					Arguments.Add("-Xclang -setup-static-analyzer");
+					// Enable the static analyzer with default checks.
+					Arguments.Add("--analyze");
 
 					// Make sure we check inside nested blocks (e.g. 'if ((foo = getchar()) == 0) {}')
 					Arguments.Add("-Xclang -analyzer-opt-analyze-nested-blocks");
@@ -304,7 +293,10 @@ namespace UnrealBuildTool
 
 					if (Target.StaticAnalyzerCheckers.Count > 0)
 					{
-						// Only enable specific checkers.
+						// Disable all default checks
+						Arguments.Add("--analyzer-no-default-checks");
+
+						// Only enable specific checks.
 						foreach (string Checker in Target.StaticAnalyzerCheckers)
 						{
 							Arguments.Add($"-Xclang -analyzer-checker -Xclang {Checker}");
@@ -312,12 +304,12 @@ namespace UnrealBuildTool
 					}
 					else
 					{
-						// Disable default checkers.
+						// Disable default checks.
 						foreach (string Checker in Target.StaticAnalyzerDisabledCheckers)
 						{
 							Arguments.Add($"-Xclang -analyzer-disable-checker -Xclang {Checker}");
 						}
-						// Enable additional non-default checkers.
+						// Enable additional non-default checks.
 						foreach (string Checker in Target.StaticAnalyzerAdditionalCheckers)
 						{
 							Arguments.Add($"-Xclang -analyzer-checker -Xclang {Checker}");
diff --git a/Engine/Source/Programs/UnrealBuildTool/ToolChain/ClangToolChain.cs b/Engine/Source/Programs/UnrealBuildTool/ToolChain/ClangToolChain.cs
--- a/Engine/Source/Programs/UnrealBuildTool/ToolChain/ClangToolChain.cs
+++ b/Engine/Source/Programs/UnrealBuildTool/ToolChain/ClangToolChain.cs
@@ -694,19 +694,9 @@ namespace UnrealBuildTool
 		protected virtual void GetCompileArguments_Analyze(CppCompileEnvironment CompileEnvironment, List<string> Arguments)
 		{
 			Arguments.Add("-Wno-unused-command-line-argument");
-			if (StaticAnalyzerCheckers.Count == 0)
-			{
-				// Enable the static analyzer with default checkers.
-				Arguments.Add("--analyze");
-			}
-			else
-			{
-				// Enable the static analyzer but only via the backend to disable all default checkers.
-				Arguments.Add("-Xclang -analyze");
-			}
 
-			// Ensure the compiler sets the __clang_analyzer__ macro correctly.
-			Arguments.Add("-Xclang -setup-static-analyzer");
+			// Enable the static analyzer with default checks.
+			Arguments.Add("--analyze");
 
 			// Make sure we check inside nested blocks (e.g. 'if ((foo = getchar()) == 0) {}')
 			Arguments.Add("-Xclang -analyzer-opt-analyze-nested-blocks");
@@ -727,6 +717,9 @@ namespace UnrealBuildTool
 
 			if (StaticAnalyzerCheckers.Count > 0)
 			{
+				// Disable all default checks
+				Arguments.Add("--analyzer-no-default-checks");
+
 				// Only enable specific checkers.
 				foreach (string Checker in StaticAnalyzerCheckers)
 				{
@@ -735,12 +728,12 @@ namespace UnrealBuildTool
 			}
 			else
 			{
-				// Disable default checkers.
+				// Disable default checks.
 				foreach (string Checker in StaticAnalyzerDisabledCheckers)
 				{
 					Arguments.Add($"-Xclang -analyzer-disable-checker -Xclang {Checker}");
 				}
-				// Enable additional non-default checkers.
+				// Enable additional non-default checks.
 				foreach (string Checker in StaticAnalyzerAdditionalCheckers)
 				{
 					Arguments.Add($"-Xclang -analyzer-checker -Xclang {Checker}");
