UnrealBuildTool: Ignore build outputs when analyzing

#rnx

diff --git a/Engine/Source/Programs/UnrealBuildTool/Platform/Android/UEBuildAndroid.cs b/Engine/Source/Programs/UnrealBuildTool/Platform/Android/UEBuildAndroid.cs
--- a/Engine/Source/Programs/UnrealBuildTool/Platform/Android/UEBuildAndroid.cs
+++ b/Engine/Source/Programs/UnrealBuildTool/Platform/Android/UEBuildAndroid.cs
@@ -161,10 +161,11 @@ namespace UnrealBuildTool
 				Target.StaticAnalyzer = StaticAnalyzer.Default;
 			}
 
-			// Disable linking if we're using a static analyzer
+			// Disable linking and ignore build outputs if we're using a static analyzer
 			if (Target.StaticAnalyzer == StaticAnalyzer.Default)
 			{
 				Target.bDisableLinking = true;
+				Target.bIgnoreBuildOutputs = true;
 			}
 
 			Target.bCompileRecast = true;
diff --git a/Engine/Source/Programs/UnrealBuildTool/Platform/IOS/UEBuildIOS.cs b/Engine/Source/Programs/UnrealBuildTool/Platform/IOS/UEBuildIOS.cs
--- a/Engine/Source/Programs/UnrealBuildTool/Platform/IOS/UEBuildIOS.cs
+++ b/Engine/Source/Programs/UnrealBuildTool/Platform/IOS/UEBuildIOS.cs
@@ -771,10 +771,11 @@ namespace UnrealBuildTool
 				Target.StaticAnalyzer = StaticAnalyzer.Default;
 			}
 
-			// Disable linking if we're using a static analyzer
+			// Disable linking and ignore build outputs if we're using a static analyzer
 			if (Target.StaticAnalyzer == StaticAnalyzer.Default)
 			{
 				Target.bDisableLinking = true;
+				Target.bIgnoreBuildOutputs = true;
 			}
 
 			// we assume now we are building with IOS8 or later
diff --git a/Engine/Source/Programs/UnrealBuildTool/Platform/Linux/UEBuildLinux.cs b/Engine/Source/Programs/UnrealBuildTool/Platform/Linux/UEBuildLinux.cs
--- a/Engine/Source/Programs/UnrealBuildTool/Platform/Linux/UEBuildLinux.cs
+++ b/Engine/Source/Programs/UnrealBuildTool/Platform/Linux/UEBuildLinux.cs
@@ -223,10 +223,11 @@ namespace UnrealBuildTool
 				Target.StaticAnalyzer = StaticAnalyzer.Default;
 			}
 
-			// Disable linking if we're using a static analyzer
+			// Disable linking and ignore build outputs if we're using a static analyzer
 			if (Target.StaticAnalyzer == StaticAnalyzer.Default)
 			{
 				Target.bDisableLinking = true;
+				Target.bIgnoreBuildOutputs = true;
 			}
 
 			// Editor target types get overwritten in UEBuildTarget.cs so lets avoid adding this here. ResetTarget is called with
diff --git a/Engine/Source/Programs/UnrealBuildTool/Platform/Mac/UEBuildMac.cs b/Engine/Source/Programs/UnrealBuildTool/Platform/Mac/UEBuildMac.cs
--- a/Engine/Source/Programs/UnrealBuildTool/Platform/Mac/UEBuildMac.cs
+++ b/Engine/Source/Programs/UnrealBuildTool/Platform/Mac/UEBuildMac.cs
@@ -120,10 +120,11 @@ namespace UnrealBuildTool
 				Target.StaticAnalyzer = StaticAnalyzer.Default;
 			}
 
-			// Disable linking if we're using a static analyzer
+			// Disable linking and ignore build outputs if we're using a static analyzer
 			if (Target.StaticAnalyzer == StaticAnalyzer.Default)
 			{
 				Target.bDisableLinking = true;
+				Target.bIgnoreBuildOutputs = true;
 			}
 
 			if (BuildHostPlatform.Current.Platform != UnrealTargetPlatform.Mac)
diff --git a/Engine/Source/Programs/UnrealBuildTool/Platform/Windows/UEBuildWindows.cs b/Engine/Source/Programs/UnrealBuildTool/Platform/Windows/UEBuildWindows.cs
--- a/Engine/Source/Programs/UnrealBuildTool/Platform/Windows/UEBuildWindows.cs
+++ b/Engine/Source/Programs/UnrealBuildTool/Platform/Windows/UEBuildWindows.cs
@@ -869,14 +869,15 @@ namespace UnrealBuildTool
 				}
 			}
 
-			// Disable linking if we're using a static analyzer
-			if(Target.StaticAnalyzer != StaticAnalyzer.None)
+			// Disable linking and ignore build outputs if we're using a static analyzer
+			if (Target.StaticAnalyzer == StaticAnalyzer.Default)
 			{
 				Target.bDisableLinking = true;
+				Target.bIgnoreBuildOutputs = true;
 			}
 
-			 // Disable PCHs for PVS studio analyzer.
-			if(Target.StaticAnalyzer == StaticAnalyzer.PVSStudio)
+			// Disable PCHs for PVS studio analyzer.
+			if (Target.StaticAnalyzer == StaticAnalyzer.PVSStudio)
 			{
 				Target.bUsePCHFiles = false;
 			}
