UnrealBuildTool: Fix not disabling linking when running PVS-Studio analyzer caused in 21202249

#rnx

diff --git a/Engine/Source/Programs/UnrealBuildTool/Platform/Windows/UEBuildWindows.cs b/Engine/Source/Programs/UnrealBuildTool/Platform/Windows/UEBuildWindows.cs
--- a/Engine/Source/Programs/UnrealBuildTool/Platform/Windows/UEBuildWindows.cs
+++ b/Engine/Source/Programs/UnrealBuildTool/Platform/Windows/UEBuildWindows.cs
@@ -870,7 +870,7 @@ namespace UnrealBuildTool
 			}
 
 			// Disable linking and ignore build outputs if we're using a static analyzer
-			if (Target.StaticAnalyzer == StaticAnalyzer.Default)
+			if (Target.StaticAnalyzer != StaticAnalyzer.None)
 			{
 				Target.bDisableLinking = true;
 				Target.bIgnoreBuildOutputs = true;
