ISPC toolchain - Add support for an optional post-compile step to compile LLVM bytecode to a native obj file, using the target platform's toolchain. Without this step, we can end up with unoptimized output on some platforms
- Fix bug where compiler targets were duplicated in the outputs list
- Rename Platform to HostPlatform in functions which expect a host platform to be passed

diff --git a/Engine/Source/Programs/UnrealBuildTool/ToolChain/ISPCToolChain.cs b/Engine/Source/Programs/UnrealBuildTool/ToolChain/ISPCToolChain.cs
--- a/Engine/Source/Programs/UnrealBuildTool/ToolChain/ISPCToolChain.cs
+++ b/Engine/Source/Programs/UnrealBuildTool/ToolChain/ISPCToolChain.cs
@@ -141,24 +141,24 @@ namespace UnrealBuildTool
 		/// <summary>
 		/// Get host compiler path for ISPC.
 		/// </summary>
-		/// <param name="Platform">Which OS build platform is running on.</param>
+		/// <param name="HostPlatform">Which OS build platform is running on.</param>
 		/// <returns>Path to ISPC compiler</returns>
-		public virtual string GetISPCHostCompilerPath(UnrealTargetPlatform Platform)
+		public virtual string GetISPCHostCompilerPath(UnrealTargetPlatform HostPlatform)
 		{
 			string ISPCCompilerPathCommon = Path.Combine(Unreal.EngineSourceDirectory.FullName, "ThirdParty", "Intel", "ISPC", "bin");
 			string ISPCArchitecturePath = "";
 			string ExeExtension = ".exe";
 
-			if (UEBuildPlatform.IsPlatformInGroup(Platform, UnrealPlatformGroup.Windows))
+			if (UEBuildPlatform.IsPlatformInGroup(HostPlatform, UnrealPlatformGroup.Windows))
 			{
 				ISPCArchitecturePath = "Windows";
 			}
-			else if (Platform == UnrealTargetPlatform.Linux)
+			else if (HostPlatform == UnrealTargetPlatform.Linux)
 			{
 				ISPCArchitecturePath = "Linux";
 				ExeExtension = "";
 			}
-			else if (Platform == UnrealTargetPlatform.Mac)
+			else if (HostPlatform == UnrealTargetPlatform.Mac)
 			{
 				ISPCArchitecturePath = "Mac";
 				ExeExtension = "";
@@ -171,6 +171,17 @@ namespace UnrealBuildTool
 			return Path.Combine(ISPCCompilerPathCommon, ISPCArchitecturePath, "ispc" + ExeExtension);
 		}
 
+		/// <summary>
+		/// Get the host bytecode-to-obj compiler path for ISPC. Only used for platforms that support compiling ISPC to LLVM bytecode
+		/// </summary>
+		/// <param name="HostPlatform">Which OS build platform is running on.</param>
+		/// <returns>Path to bytecode to obj compiler</returns>
+		public virtual string? GetISPCHostBytecodeCompilerPath(UnrealTargetPlatform HostPlatform)
+		{
+			// Return null if the platform toolchain doesn't support separate bytecode to obj compilation
+			return null;
+		}
+
 		static Dictionary<UnrealTargetPlatform, string> ISPCCompilerVersions = new Dictionary<UnrealTargetPlatform, string>();
 
 		/// <summary>
@@ -536,17 +547,22 @@ namespace UnrealBuildTool
 			}
 
 			// Build target triplet
+			string PlatformObjectFileFormat = GetISPCObjectFileFormat(CompileEnvironment.Platform);
+
 			GlobalArguments.Add($"--target-os={GetISPCOSTarget(CompileEnvironment.Platform)}");
 			GlobalArguments.Add($"--arch={GetISPCArchTarget(CompileEnvironment.Platform, null)}");
 			GlobalArguments.Add($"--target={TargetString}");
-			GlobalArguments.Add($"--emit-{GetISPCObjectFileFormat(CompileEnvironment.Platform)}");
+			GlobalArguments.Add($"--emit-{PlatformObjectFileFormat}");
+
+			bool bByteCodeOutput = ( PlatformObjectFileFormat == "llvm" );
 
+			List<string> CommonArgs = new List<string>();
 			if (CompileEnvironment.Configuration == CppConfiguration.Debug)
 			{
 				if (CompileEnvironment.Platform == UnrealTargetPlatform.Mac)
 				{
 					// Turn off debug symbols on Mac due to dsym generation issue
-					GlobalArguments.Add("-O0");
+					CommonArgs.Add("-O0");
 					// Ideally we would be able to turn on symbols and specify the dwarf version, but that does
 					// does not seem to be working currently, ie:
 					//    GlobalArguments.Add("-g -O0 --dwarf-version=2");
@@ -554,13 +570,14 @@ namespace UnrealBuildTool
 				}
 				else
 				{
-					GlobalArguments.Add("-g -O0");
+					CommonArgs.Add("-g -O0");
 				}
 			}
 			else
 			{
-				GlobalArguments.Add("-O2");
+				CommonArgs.Add("-O2");
 			}
+			GlobalArguments.AddRange(CommonArgs);
 
 			// PIC is needed for modular builds except on Microsoft platforms
 			if ((CompileEnvironment.bIsBuildingDLL ||
@@ -610,6 +627,7 @@ namespace UnrealBuildTool
 
 				List<FileItem> CompiledISPCObjFiles = new List<FileItem>();
 
+				string CompiledISPCObjFileSuffix = bByteCodeOutput ? ".bc" : GetISPCObjectFileSuffix(CompileEnvironment.Platform);
 				foreach (string Target in CompileTargets)
 				{
 					string ObjTarget = Target;
@@ -627,7 +645,7 @@ namespace UnrealBuildTool
 						CompiledISPCObjFile = FileItem.GetItemByFileReference(
 						FileReference.Combine(
 							OutputDir,
-							Path.GetFileName(ISPCFile.AbsolutePath) + "_" + ObjTarget + GetISPCObjectFileSuffix(CompileEnvironment.Platform)
+							Path.GetFileName(ISPCFile.AbsolutePath) + "_" + ObjTarget + CompiledISPCObjFileSuffix
 							)
 						);
 					}
@@ -636,7 +654,7 @@ namespace UnrealBuildTool
 						CompiledISPCObjFile = FileItem.GetItemByFileReference(
 						FileReference.Combine(
 							OutputDir,
-							Path.GetFileName(ISPCFile.AbsolutePath) + GetISPCObjectFileSuffix(CompileEnvironment.Platform)
+							Path.GetFileName(ISPCFile.AbsolutePath) + CompiledISPCObjFileSuffix
 							)
 						);
 					}
@@ -645,15 +663,18 @@ namespace UnrealBuildTool
 					CompiledISPCObjFiles.Add(CompiledISPCObjFile);
 				}
 
-				// Add the common ISPC obj file to the produced item list.
+				// Add the common ISPC obj file to the produced item list if it's not already in it
 				FileItem CompiledISPCObjFileNoISA = FileItem.GetItemByFileReference(
 					FileReference.Combine(
 						OutputDir,
-						Path.GetFileName(ISPCFile.AbsolutePath) + GetISPCObjectFileSuffix(CompileEnvironment.Platform)
+						Path.GetFileName(ISPCFile.AbsolutePath) + CompiledISPCObjFileSuffix
 						)
 					);
 
-				CompiledISPCObjFiles.Add(CompiledISPCObjFileNoISA);
+				if (CompileTargets.Count > 1)
+				{
+					CompiledISPCObjFiles.Add(CompiledISPCObjFileNoISA);
+				}
 
 				// Add the output ISPC obj file
 				Arguments.Add($"-o \"{NormalizeCommandLinePath(CompiledISPCObjFileNoISA)}\"");
@@ -677,7 +698,6 @@ namespace UnrealBuildTool
 				}
 
 				CompileAction.ProducedItems.AddRange(CompiledISPCObjFiles);
-				Result.ObjectFiles.AddRange(CompiledISPCObjFiles);
 
 				FileReference ResponseFileName = new FileReference(CompiledISPCObjFileNoISA.AbsolutePath + ".response");
 				FileItem ResponseFileItem = Graph.CreateIntermediateTextFile(ResponseFileName, Arguments.Select(x => Utils.ExpandVariables(x)));
@@ -688,6 +708,56 @@ namespace UnrealBuildTool
 				CompileAction.PrerequisiteItems.Add(ISPCFile);
 
 				Logger.LogDebug("   ISPC Compiling {StatusDescription}: \"{CommandPath}\" {CommandArguments}", CompileAction.StatusDescription, CompileAction.CommandPath, CompileAction.CommandArguments);
+
+				if (bByteCodeOutput)
+				{
+					// If the platform toolchain supports bytecode compilation for ISPC, compile the bytecode object files to actual native object files 
+					string? ByteCodeCompilerPath = GetISPCHostBytecodeCompilerPath(BuildHostPlatform.Current.Platform);
+					if (ByteCodeCompilerPath != null)
+					{
+						List<FileItem> FinalObjectFiles = new List<FileItem>();
+						foreach (FileItem CompiledBytecodeObjFile in CompiledISPCObjFiles)
+						{
+							FileItem FinalCompiledISPCObjFile = FileItem.GetItemByFileReference(
+								FileReference.Combine(
+									OutputDir,
+									Path.GetFileNameWithoutExtension(CompiledBytecodeObjFile.AbsolutePath) + GetISPCObjectFileSuffix(CompileEnvironment.Platform)
+									)
+								);
+
+							Action PostCompileAction = Graph.CreateAction(ActionType.Compile);
+
+							List<string> PostCompileArgs = new List<string>();
+							PostCompileArgs.Add($"\"{NormalizeCommandLinePath(CompiledBytecodeObjFile)}\"");
+							PostCompileArgs.Add("-c");
+							PostCompileArgs.AddRange(CommonArgs);
+							PostCompileArgs.Add($"-o \"{NormalizeCommandLinePath(FinalCompiledISPCObjFile)}\"");
+
+							// Write the args to a response file
+							FileReference PostCompileResponseFileName = new FileReference(FinalCompiledISPCObjFile.AbsolutePath + ".response");
+							FileItem PostCompileResponseFileItem = Graph.CreateIntermediateTextFile(PostCompileResponseFileName, PostCompileArgs.Select(x => Utils.ExpandVariables(x)));
+							PostCompileAction.CommandArguments = $"@\"{PostCompileResponseFileName}\"";
+							PostCompileAction.PrerequisiteItems.Add(PostCompileResponseFileItem);
+
+							PostCompileAction.PrerequisiteItems.Add(CompiledBytecodeObjFile);
+							PostCompileAction.ProducedItems.Add(FinalCompiledISPCObjFile);
+							PostCompileAction.CommandDescription = "CompileByteCode";
+							PostCompileAction.WorkingDirectory = Unreal.EngineSourceDirectory;
+							PostCompileAction.CommandPath = new FileReference(ByteCodeCompilerPath);
+							PostCompileAction.StatusDescription = Path.GetFileName(ISPCFile.AbsolutePath);
+
+							// Disable remote execution to workaround mismatched case on XGE
+							PostCompileAction.bCanExecuteRemotely = false;
+
+							FinalObjectFiles.Add(FinalCompiledISPCObjFile);
+							Logger.LogDebug("   ISPC Compiling bytecode {StatusDescription}: \"{CommandPath}\" {CommandArguments} {ProducedItems}", PostCompileAction.StatusDescription, PostCompileAction.CommandPath, PostCompileAction.CommandArguments, PostCompileAction.ProducedItems);
+						}
+						// Override the output object files
+						CompiledISPCObjFiles = FinalObjectFiles;
+					}
+				}
+
+				Result.ObjectFiles.AddRange(CompiledISPCObjFiles);
 			}
 
 			return Result;
