PR #6739: Fixed QMake project generator defines (Contributed by [REDACTED_PERSON])

diff --git a/Engine/Source/Programs/UnrealBuildTool/ProjectFiles/QMake/QMakefileGenerator.cs b/Engine/Source/Programs/UnrealBuildTool/ProjectFiles/QMake/QMakefileGenerator.cs
--- a/Engine/Source/Programs/UnrealBuildTool/ProjectFiles/QMake/QMakefileGenerator.cs
+++ b/Engine/Source/Programs/UnrealBuildTool/ProjectFiles/QMake/QMakefileGenerator.cs
@@ -98,6 +98,7 @@
 			List<string> IncludeDirectories = new List<string>();
 			List<string> SystemIncludeDirectories = new List<string>();
 			List<string> DefinesAndValues = new List<string>();
+			HashSet<string> DefinesOnly = new HashSet<string>();
 
 			// DefineList.Add ("");
 
@@ -161,7 +162,7 @@
 
 					SplitDefinitionAndValue(CurDefine, out define, out value);
 
-					if (!DefinesAndValues.Contains(define))
+					if (!DefinesOnly.Contains(define))
 					{
 						// Logger.LogInformation (CurDefine);
 						if (string.IsNullOrEmpty(value))
@@ -172,12 +173,29 @@
 						}
 						else
 						{
+							// Escape special characters
+							// If the value contains spaces, we must put the whole value inside quotes
+							value = value.Replace("\"", "\\\\\\\"");
+							bool hasSpaces = false;
+							if (value.Contains(" "))
+							{
+								value = value.Replace(" ", "\\ ");
+								hasSpaces = true;
+							}
+							value = value.Replace("{", "\\{");
+							value = value.Replace("}", "\\}");
+
 							DefinesAndValues.Add("\t");
 							DefinesAndValues.Add(define);
 							DefinesAndValues.Add("=");
+							if (hasSpaces)
+								DefinesAndValues.Add("\"");
 							DefinesAndValues.Add(value);
+							if (hasSpaces)
+								DefinesAndValues.Add("\"");
 							DefinesAndValues.Add(" \\\n");
 						}
+						DefinesOnly.Add(define);
 					}
 				}
 			}
