- Rewrote the XcodeProjectFile class to not be spagetti of string maniuplation and guid passing [still requires opt-in via -xcconfig option to GenerateProjectFiles]
- Now each project node is its own class, and they can reference each other
- Projects with only one sub index target now have it rolled up into the main Index target
- Extension support is removed (will return later with an entirely new way to handle unique cases like that)

- Allowed an exception for UHT when disabling plugins on arm64, since Indexing needed the headers to have been generated

diff --git a/Engine/Source/Programs/UnrealBuildTool/Platform/Mac/UEBuildMac.cs b/Engine/Source/Programs/UnrealBuildTool/Platform/Mac/UEBuildMac.cs
--- a/Engine/Source/Programs/UnrealBuildTool/Platform/Mac/UEBuildMac.cs
+++ b/Engine/Source/Programs/UnrealBuildTool/Platform/Mac/UEBuildMac.cs
@@ -1,4 +1,4 @@
-// Copyright Epic Games, Inc. All Rights Reserved.
+﻿// Copyright Epic Games, Inc. All Rights Reserved.
 
 
 using System;
@@ -137,7 +137,7 @@ namespace UnrealBuildTool
 			bool bCompilingForArm = Target.Architecture.IndexOf("arm", StringComparison.OrdinalIgnoreCase) >= 0;
 			bool bCompilingMultipleArchitectures = Target.Architecture.Contains("+");
 
-			if (bCompilingForArm)
+			if (bCompilingForArm && Target.Name != "UnrealHeaderTool")
 			{
 				Target.DisablePlugins.AddRange(new string[]
 				{
diff --git a/Engine/Source/Programs/UnrealBuildTool/ProjectFiles/Xcode/XcodeProject.cs b/Engine/Source/Programs/UnrealBuildTool/ProjectFiles/Xcode/XcodeProject.cs
--- a/Engine/Source/Programs/UnrealBuildTool/ProjectFiles/Xcode/XcodeProject.cs
+++ b/Engine/Source/Programs/UnrealBuildTool/ProjectFiles/Xcode/XcodeProject.cs
@@ -11,73 +11,65 @@ using System.Text;
 using EpicGames.Core;
 using UnrealBuildBase;
 using Microsoft.Extensions.Logging;
+using System.Text.RegularExpressions;
+using System.Diagnostics.CodeAnalysis;
+
+
+/*****
+
+Here's how this works:
+
+  * An XcodeProjectFile (subclass of generic ProjectFile class) is created, along with it - UnrealData and XcodeFileCollection objects are made
+  * High level code calls AddModule() which this code will use to cache information about the Modules in the project (including build settings, etc)
+    * These are used to determine what source files can be indexed together (we use native xcode code compilation for indexing, so we make compiling succesful for best index)
+    * A few #defines are removed or modified (FOO_API, etc) which would otherwise make every module a separate target
+  * High level code then calls WriteProjectFile() which is the meat of all this
+  * This code then creates a hierarchy/reference-chain of xcode project nodes (an xcode project is basically a series of Guid/object pairs that reference each other)
+  * Then each node writes itself into the project file that is saved to disk
+
+.xcconfig files:
+  * We also now use Xcconfig files for all of the build settings, instead of jamming them into the xcode project file itself
+  * This makes it easier to see the various settings (you can also see them in the Xcode UI as before, but now with more read-only columns - they must be set via editing the file)
+  * The files are in the Xcode file browsing pane, in the an Xcconfigs folder under each project
+  * The files are:
+    * _Project - settings apply to all targets in the projectapplies to all targets), one for each configuration (Debug, Development Editor, etc) which applies to all targets
+    * _Debug/_DevelopmentEditor, etc - applies to all targets when building in that configuration
+    * _Run - applies to the native run project, which does no compiling
+    * _Index/SubIndex0, etc - applies when Indexing, which uses native Xcode compiling (these will include #defines from UBT, etc)
+  * There is currently no _Build xcconfig file since it's a makefile project and has no build setting needs
+
+
+Known issues:
+  * No PBXFrameworksBuildPhase nodes are made
+ 
+Future ideas:
+  * I have started working on a Template project that we merge Build/Index targets into, which will allow a licensee to setup codesigning, add extensions, etc. 
+  * Always make a final build xcodeproj from UBT for handling the codesigning/extensions/frameworks
+  * Allow for non-conflicting #defines to share SubIndex targets, hopefully will greatly reduce the sub targets in UE5
+
+**/
 
 namespace UnrealBuildTool.XcodeProjectXcconfig
 {
-	/// <summary>
-	/// Info needed to make a file a member of specific group
-	/// </summary>
-	class XcodeSourceFile : ProjectFile.SourceFile
+	static class StringBuilderExtensions
 	{
-		/// <summary>
-		/// Constructor
-		/// </summary>
-		public XcodeSourceFile(FileReference InitFilePath, DirectoryReference? InitRelativeBaseFolder)
-			: base(InitFilePath, InitRelativeBaseFolder)
-		{
-			FileGuid = XcodeProjectFileGenerator.MakeXcodeGuid();
-			FileRefGuid = XcodeProjectFileGenerator.MakeXcodeGuid();
-		}
-
-		/// <summary>
-		/// File Guid for use in Xcode project
-		/// </summary>
-		public string FileGuid
-		{
-			get;
-			private set;
-		}
-
-		public void ReplaceGuids(string NewFileGuid, string NewFileRefGuid)
+		public static void WriteLine(this StringBuilder SB, string Line = "")
 		{
-			FileGuid = NewFileGuid;
-			FileRefGuid = NewFileRefGuid;
+			SB.Append(Line);
+			SB.Append(ProjectFileGenerator.NewLine);
 		}
-
-		/// <summary>
-		/// File reference Guid for use in Xcode project
-		/// </summary>
-		public string FileRefGuid
+		public static void WriteLine(this StringBuilder SB, int Indent, string Line = "")
 		{
-			get;
-			private set;
+			SB.Append(new String('\t', Indent));
+			SB.Append(Line);
+			SB.Append(ProjectFileGenerator.NewLine);
 		}
 	}
 
-	/// <summary>
-	/// Represents a group of files shown in Xcode's project navigator as a folder
-	/// </summary>
-	class XcodeFileGroup
-	{
-		public XcodeFileGroup(string InName, string InPath, bool InIsReference)
-		{
-			GroupName = InName;
-			GroupPath = InPath;
-			GroupGuid = XcodeProjectFileGenerator.MakeXcodeGuid();
-			bIsReference = InIsReference;
-		}
-
-		public string GroupGuid;
-		public string GroupName;
-		public string GroupPath;
-		public Dictionary<string, XcodeFileGroup> Children = new Dictionary<string, XcodeFileGroup>();
-		public List<XcodeSourceFile> Files = new List<XcodeSourceFile>();
-		public bool bIsReference;
-	}
 
-	class XcodeBuildConfig
+	class UnrealBuildConfig
 	{
-		public XcodeBuildConfig(string InDisplayName, string InBuildTarget, FileReference? InMacExecutablePath, FileReference? InIOSExecutablePath, FileReference? InTVOSExecutablePath,
+		public UnrealBuildConfig(string InDisplayName, string InBuildTarget, FileReference? InMacExecutablePath, FileReference? InIOSExecutablePath, FileReference? InTVOSExecutablePath,
 			ProjectTarget? InProjectTarget, UnrealTargetConfiguration InBuildConfig)
 		{
 			DisplayName = InDisplayName;
@@ -111,361 +103,411 @@ namespace UnrealBuildTool.XcodeProjectXcconfig
 		}
 	};
 
-	class XcodeExtensionInfo
+	class UnrealBatchedFiles
 	{
-		public XcodeExtensionInfo(string InName)
-		{
-			Name = InName;
-			TargetDependencyGuid = XcodeProjectFileGenerator.MakeXcodeGuid();
-			TargetProxyGuid = XcodeProjectFileGenerator.MakeXcodeGuid();
-			TargetGuid = XcodeProjectFileGenerator.MakeXcodeGuid();
-			ProductGuid = XcodeProjectFileGenerator.MakeXcodeGuid();
-			ResourceBuildPhaseGuid = XcodeProjectFileGenerator.MakeXcodeGuid();
-			ConfigListGuid = XcodeProjectFileGenerator.MakeXcodeGuid();
-			AllConfigs = new Dictionary<string, XcodeBuildConfig>();
-		}
-
-		public string Name;
-		public string TargetDependencyGuid;
-		public string TargetProxyGuid;
-		public string TargetGuid;
-		public string ProductGuid;
-		public string ResourceBuildPhaseGuid;
-		public string ConfigListGuid;
-		public Dictionary<string, XcodeBuildConfig> AllConfigs;
+		// build settings that cause uniqueness
+		public string Definitions = "";
+		public FileReference? PCHFile = null;
+		public bool bEnableRTTI = false;
 
-		public string? ConfigurationContents;
-	}
+		// union of settings for all modules
+		public HashSet<string> APIDefines = new();
+		public HashSet<DirectoryReference> SystemIncludePaths = new();
+		public HashSet<DirectoryReference> UserIncludePaths = new();
 
-	class XcconfigFile
-	{
-		public string Name;
-		public string Filename;
-		public string Guid;
-		public FileReference FileRef;
-		private StringBuilder Text;
+		public List<XcodeSourceFile> Files = new();
+		public List<UEBuildModuleCPP> Modules = new List<UEBuildModuleCPP>();
 
-		public XcconfigFile(FileReference ProjectFile, string ConfigName)
+		public void WriteXcconfigSettings(StringBuilder Content)
 		{
-			Name = ConfigName;
-			FileRef = FileReference.Combine(ProjectFile.Directory, "Xcconfigs", $"{ProjectFile.GetFileNameWithoutAnyExtensions()}_{Name.Replace(" ", "")}.xcconfig");
-			Filename = FileRef.GetFileName();
-			Guid = XcodeProjectFileGenerator.MakeXcodeGuid();
-			DirectoryReference.CreateDirectory(FileRef.Directory);
-			Text = new StringBuilder();
-		}
+			Content.WriteLine($"GCC_PREPROCESSOR_DEFINITIONS = __INTELLISENSE__ MONOLITHIC_BUILD=1 {Definitions} {string.Join(" ", APIDefines)}");
+			if (PCHFile != null)
+			{
+				Content.WriteLine($"GCC_PREFIX_HEADER = {PCHFile.FullName}");
+			}
+			Content.WriteLine($"GCC_ENABLE_CPP_RTTI = " + (bEnableRTTI ? "YES" : "NO"));
 
-		public void AppendLine(string Line)
-		{
-			Text.Append(Line);
-			Text.Append(ProjectFileGenerator.NewLine);
-		}
+			string SearchPaths = string.Join(" ", SystemIncludePaths.Select(x => x.FullName.Contains(' ') ? $"\"{x.FullName}\"" : x.FullName));
+			SearchPaths = SearchPaths.Replace("/x86_64/", "/$(ARCHS)/");
+			Content.WriteLine($"HEADER_SEARCH_PATHS = {SearchPaths}");
 
-		public void Write()
-		{
-			FileReference.WriteAllTextIfDifferent(FileRef, Text.ToString());
+			SearchPaths = string.Join(" ", UserIncludePaths.Select(x => x.FullName.Contains(' ') ? $"\"{x.FullName}\"" : x.FullName));
+			SearchPaths = SearchPaths.Replace("/x86_64/", "/$(ARCHS)/");
+			Content.WriteLine($"USER_HEADER_SEARCH_PATHS = {SearchPaths}");
 		}
 	}
 
+	//class UnrealRemoteExtension
+	//{
 
-	class XcodeProjectFile : ProjectFile
+	//}
+
+	class UnrealData
 	{
-		Dictionary<string, XcodeFileGroup> Groups = new Dictionary<string, XcodeFileGroup>();
+		public bool bIsAppBundle;
+		public bool bHasEditorConfiguration;
+		public bool bUseAutomaticSigning = false;
 
-		/// <summary>
-		/// Constructs a new project file object
-		/// </summary>
-		/// <param name="InitFilePath">The path to the project file on disk</param>
-		/// <param name="BaseDir">The base directory for files within this project</param>
-		/// <param name="IsForDistribution">True for distribution builds</param>
-		/// <param name="BundleID">Override option for bundle identifier</param>
-		/// <param name="InAppName"></param>
-		public XcodeProjectFile(FileReference InitFilePath, DirectoryReference BaseDir, bool IsForDistribution, string BundleID, string InAppName)
-			: base(InitFilePath, BaseDir)
-		{
-			bForDistribution = IsForDistribution;
-			BundleIdentifier = BundleID;
-			AppName = InAppName;
+		public List<UnrealBuildConfig> AllConfigs = new();
 
-			AllXcconfigFiles.Add(ProjectXcconfig = new XcconfigFile(InitFilePath, "Project"));
-			AllXcconfigFiles.Add(RunXcconfig = new XcconfigFile(InitFilePath, "Run"));
-			AllXcconfigFiles.Add(IndexXcconfig = new XcconfigFile(InitFilePath, "Index"));
-			AllXcconfigFiles.Add(BuildXcconfig = new XcconfigFile(InitFilePath, "Build"));
-		}
+		public List<UnrealExtensionInfo> AllExtensions = new();
 
-		public override string ToString()
-		{
-			return ProjectFilePath.GetFileNameWithoutExtension();
-		}
+		public List<UnrealBatchedFiles> BatchedFiles = new();
+ 
+		public FileReference? UProjectFileLocation = null;
+		public FileReference? XcodeProjectFileLocation = null;
+
+		// settings read from project configs
+		public IOSProjectSettings? IOSProjectSettings;
+		public IOSProvisioningData? IOSProvisioningData;
+
+		public TVOSProjectSettings? TVOSProjectSettings;
+		public TVOSProvisioningData? TVOSProvisioningData;
+
+		public string ProductName="";
 
 		/// <summary>
 		///  Used to mark the project for distribution (some platforms require this)
 		/// </summary>
-		bool bForDistribution = false;
+		public bool bForDistribution = false;
 
 		/// <summary>
 		/// Override for bundle identifier
 		/// </summary>
-		string BundleIdentifier = "";
+		public string BundleIdentifier = "";
 
 		/// <summary>
 		/// Override AppName
 		/// </summary>
-		string AppName = "";
+		public string AppName = "";
 
 		/// <summary>
 		/// Architectures supported for iOS
 		/// </summary>
-		string[] SupportedIOSArchitectures = { "arm64" };
+		public string[] SupportedIOSArchitectures = { "arm64" };
 
 		/// <summary>
-		/// Xcconfig file information - this is global per project, so no need to pass it around everywhere
+		/// UBT logger object
 		/// </summary>
-		private XcconfigFile ProjectXcconfig;
-		private XcconfigFile RunXcconfig;
-		private XcconfigFile IndexXcconfig;
-		private XcconfigFile BuildXcconfig;
-		internal List<XcconfigFile> AllXcconfigFiles = new List<XcconfigFile>();
+		public ILogger? Logger;
 
-		/// <summary>
-		/// Gets Xcode file category based on its extension
-		/// </summary>
-		private string GetFileCategory(string Extension)
-		{
-			// @todo Mac: Handle more categories
-			switch (Extension)
-			{
-				case ".framework":
-					return "Frameworks";
-				default:
-					return "Sources";
-			}
-		}
+		private XcodeProjectFile? ProjectFile;
 
-		/// <summary>
-		/// Gets Xcode file type based on its extension
-		/// </summary>
-		private string GetFileType(string Extension)
+		public bool Initialize(XcodeProjectFile ProjectFile, List<UnrealTargetPlatform> Platforms, List<UnrealTargetConfiguration> Configurations, ILogger Logger)
 		{
-			// @todo Mac: Handle more file types
-			switch (Extension)
+			this.ProjectFile = ProjectFile;
+			this.Logger = Logger;
+			ProductName = ProjectFile.ProjectFilePath.GetFileNameWithoutAnyExtensions();
+			XcodeProjectFileLocation = ProjectFile.ProjectFilePath;
+
+
+			// Figure out all the desired configurations on the unreal side
+			AllConfigs = GetSupportedBuildConfigs(Platforms, Configurations, Logger);
+			// if we can't find any configs, we will fail to create a project
+			if (AllConfigs.Count == 0)
 			{
-				case ".c":
-				case ".m":
-					return "sourcecode.c.objc";
-				case ".cc":
-				case ".cpp":
-				case ".mm":
-					return "sourcecode.cpp.objcpp";
-				case ".h":
-				case ".inl":
-				case ".pch":
-					return "sourcecode.c.h";
-				case ".framework":
-					return "wrapper.framework";
-				case ".plist":
-					return "text.plist.xml";
-				case ".png":
-					return "image.png";
-				case ".icns":
-					return "image.icns";
-				default:
-					return "file.text";
+				return false;
 			}
-		}
 
-		/// <summary>
-		/// Returns true if Extension is a known extension for files containing source code
-		/// </summary>
-		private bool IsSourceCode(string Extension)
-		{
-			return Extension == ".c" || Extension == ".cc" || Extension == ".cpp" || Extension == ".m" || Extension == ".mm";
-		}
+			// cache any extension info
+			ProcessExtensions(Logger);
 
-		private bool ShouldIncludeFileInBuildPhaseSection(XcodeSourceFile SourceFile)
-		{
-			string FileExtension = SourceFile.Reference.GetExtension();
 
-			if (IsSourceCode(FileExtension))
+			// find a uproject file (UE5 target won't have one)
+			foreach (Project Target in ProjectFile.ProjectTargets)
 			{
-				foreach (string PlatformName in UnrealTargetPlatform.GetValidPlatformNames())
+				if (Target.UnrealProjectFilePath != null)
 				{
-					string AltName = PlatformName == "Win64" ? "windows" : PlatformName.ToLower();
-					if ((SourceFile.Reference.FullName.ToLower().Contains("/" + PlatformName.ToLower() + "/") || SourceFile.Reference.FullName.ToLower().Contains("/" + AltName + "/"))
-						&& PlatformName != "Mac" && PlatformName != "IOS" && PlatformName != "TVOS")
-					{
-						// Build phase is used for indexing only and indexing currently works only with files that can be compiled for Mac, so skip files for other platforms
-						return false;
-					}
+					UProjectFileLocation = Target.UnrealProjectFilePath;
+					break;
 				}
+			}
 
-				return true;
+			// this project makes an app bundle (.app directory instead of a raw executable or dylib) if none of the fings make a non-appbundle
+			bIsAppBundle = !AllConfigs.Any(x => x.ProjectTarget!.TargetRules!.bIsBuildingConsoleApplication || x.ProjectTarget.TargetRules.bShouldCompileAsDLL);
+			bHasEditorConfiguration = AllConfigs.Any(x => x.ProjectTarget!.TargetRules!.Type == TargetType.Editor);
+
+			// read config settings
+			if (InstalledPlatformInfo.IsValidPlatform(UnrealTargetPlatform.IOS, EProjectType.Code))
+			{
+				IOSPlatform IOSPlatform = ((IOSPlatform)UEBuildPlatform.GetBuildPlatform(UnrealTargetPlatform.IOS));
+				IOSProjectSettings = IOSPlatform.ReadProjectSettings(UProjectFileLocation);
+				IOSProvisioningData = IOSPlatform.ReadProvisioningData(IOSProjectSettings, bForDistribution);
+				bUseAutomaticSigning |= IOSProjectSettings.bAutomaticSigning;
 			}
 
-			return false;
+
+			if (InstalledPlatformInfo.IsValidPlatform(UnrealTargetPlatform.IOS, EProjectType.Code))
+			{
+				TVOSPlatform TVOSPlatform = ((TVOSPlatform)UEBuildPlatform.GetBuildPlatform(UnrealTargetPlatform.TVOS));
+				TVOSProjectSettings = TVOSPlatform.ReadProjectSettings(UProjectFileLocation);
+				TVOSProvisioningData = TVOSPlatform.ReadProvisioningData(TVOSProjectSettings, bForDistribution);
+				bUseAutomaticSigning |= TVOSProjectSettings.bAutomaticSigning;
+			}
+
+			return true;
 		}
 
-		/// <summary>
-		/// Returns a project navigator group to which the file should belong based on its path.
-		/// Creates a group tree if it doesn't exist yet.
-		/// </summary>
-		public XcodeFileGroup? FindGroupByAbsolutePath(ref Dictionary<string, XcodeFileGroup> Groups, string AbsolutePath)
+		public void AddModule(UEBuildModuleCPP Module, CppCompileEnvironment CompileEnvironment)
 		{
-			string[] Parts = AbsolutePath.Split(Path.DirectorySeparatorChar);
-			string CurrentPath = "/";
-			Dictionary<string, XcodeFileGroup> CurrentSubGroups = Groups;
+			// we need to keep all _API defines, but we can gather them up to append at the end, instead of using them to compute sameness
 
-			for (int Index = 1; Index < Parts.Count(); ++Index)
-			{
-				string Part = Parts[Index];
+			// remove some extraneous defines that will cause every module to be unique, but we can do without
+			List<string> Defines = new();
+			List<string> APIDefines = new();
+			Regex Regex = new Regex("#define ([A-Z0-9_]*) ?(.*)?");
 
-				if (CurrentPath.Length > 1)
+			string DefinesString = "";
+			FileReference? PCHFile = null;
+			if (CompileEnvironment.ForceIncludeFiles.Count == 0)
+			{
+				// if there are no ForceInclude files, then that means it's a module that forces the includes to come from a generated PCH file
+				// and so we will use this for definitions and uniqueness
+				PCHFile = CompileEnvironment.PrecompiledHeaderIncludeFilename;
+			}
+			else
+			{
+				foreach (string Line in File.ReadAllLines(CompileEnvironment.ForceIncludeFiles.First(x => x.FullName.Contains("Definitions")).FullName))
 				{
-					CurrentPath += Path.DirectorySeparatorChar;
-				}
+					Match Match = Regex.Match(Line);
+					if (!Match.Success)
+					{
+						continue;
+					}
 
-				CurrentPath += Part;
+					string Key = Match.Groups[1].Value;
+					string Value = Match.Groups[2].Value;
+					// if no value, then just define with no = stuff
+					if (!Match.Groups[2].Success)
+					{
+						Defines.Add(Key);
+						continue;
+					}
 
-				XcodeFileGroup CurrentGroup;
-				if (!CurrentSubGroups.ContainsKey(CurrentPath))
-				{
-					CurrentGroup = new XcodeFileGroup(Path.GetFileName(CurrentPath), CurrentPath, CurrentPath.EndsWith(".xcassets"));
-					CurrentSubGroups.Add(CurrentPath, CurrentGroup);
-				}
-				else
-				{
-					CurrentGroup = CurrentSubGroups[CurrentPath];
+					// skip some known per-module defines we don't need
+					if (Key.StartsWith("UE_MODULE_NAME") || Key.StartsWith("UE_PLUGIN_NAME") || Key.StartsWith("UBT_MODULE_MANIFEST"))
+					{
+						continue;
+					}
+					// these API ones are per module but still need to be defined, can be defined to nothing
+					else if (Key.Contains("_API"))
+					{
+						APIDefines.Add($"{Key}=");
+					}
+					else
+					{
+						// if the value is normal, just add the define as is
+						if (Value.All(x => char.IsLetterOrDigit(x) || x == '_'))
+						{
+							Defines.Add($"{Key}={Value}");
+						}
+						else
+						{
+							// escape any quotes in the define, then quote the whole thing
+							Value = Value.Replace("\"", "\\\"");
+							Defines.Add($"{Key}=\"{Value}\"");
+						}
+					}
 				}
 
-				if (CurrentPath == AbsolutePath)
+				// sort and joing them into a single string to act as a key (and happily string to put into .xcconfig)
+				Defines.Sort();
+				DefinesString = string.Join(" ", Defines);
+			}
+
+			// now find a matching SubTarget, and make a new one if needed
+			UnrealBatchedFiles? FileBatch = null;
+			foreach (UnrealBatchedFiles Search in BatchedFiles)
+			{
+				if (Search.Definitions == DefinesString && Search.PCHFile == PCHFile &&
+					Search.bEnableRTTI == CompileEnvironment.bUseRTTI)
 				{
-					return CurrentGroup;
+					FileBatch = Search;
+					break;
 				}
-
-				CurrentSubGroups = CurrentGroup.Children;
+			}
+			if (FileBatch == null)
+			{
+				//FileBatch = new UnrealBatchedFiles(UnrealData, SubTargets.Count + 1);
+				FileBatch = new UnrealBatchedFiles();
+				BatchedFiles.Add(FileBatch);
+				FileBatch.Definitions = DefinesString;
+				FileBatch.PCHFile = PCHFile;
+				FileBatch.bEnableRTTI = CompileEnvironment.bUseRTTI;
 			}
 
-			return null;
+			FileBatch.Modules.Add(Module);
+			FileBatch.APIDefines.UnionWith(APIDefines);
+			FileBatch.SystemIncludePaths.UnionWith(CompileEnvironment.SystemIncludePaths);
+			FileBatch.UserIncludePaths.UnionWith(CompileEnvironment.UserIncludePaths);
 		}
 
-		/// <summary>
-		/// Convert all paths to Apple/Unix format (with forward slashes)
-		/// </summary>
-		/// <param name="InPath">The path to convert</param>
-		/// <returns>The normalized path</returns>
-		private static string ConvertPath(string InPath)
+		public void FilterBuildableFiles(XcodeFileCollection FileCollection)
 		{
-			return InPath.Replace("\\", "/");
+			foreach (XcodeSourceFile File in FileCollection.BuildableFiles)
+			{
+				AddFileToBatch(File);
+			}
 		}
 
-		/// <summary>
-		/// Allocates a generator-specific source file object
-		/// </summary>
-		/// <param name="InitFilePath">Path to the source file on disk</param>
-		/// <param name="InitProjectSubFolder">Optional sub-folder to put the file in.  If empty, this will be determined automatically from the file's path relative to the project file</param>
-		/// <returns>The newly allocated source file object</returns>
-		public override SourceFile? AllocSourceFile(FileReference InitFilePath, DirectoryReference? InitProjectSubFolder)
+		internal void AddFileToBatch(XcodeSourceFile File)
 		{
-			if (InitFilePath.GetFileName().StartsWith("."))
+			foreach (UnrealBatchedFiles Batch in BatchedFiles)
 			{
-				return null;
+				foreach (UEBuildModuleCPP Module in Batch.Modules)
+				{
+					if (Module.ContainsFile(File.Reference))
+					{
+						Batch.Files.Add(File);
+						return;
+					}
+				}
 			}
-			return new XcodeSourceFile(InitFilePath, InitProjectSubFolder);
 		}
 
-		/// <summary>
-		/// Generates bodies of all sections that contain a list of source files plus a dictionary of project navigator groups.
-		/// </summary>
-		private void GenerateSectionsWithSourceFiles(StringBuilder PBXBuildFileSection, StringBuilder PBXFileReferenceSection, StringBuilder PBXSourcesBuildPhaseSection, string TargetAppGuid, string TargetName, bool bIsAppBundle)
+
+		private List<UnrealBuildConfig> GetSupportedBuildConfigs(List<UnrealTargetPlatform> Platforms, List<UnrealTargetConfiguration> Configurations, ILogger Logger)
 		{
-			SourceFiles.Sort((x, y) => { return x.Reference.FullName.CompareTo(y.Reference.FullName); });
+			List<UnrealBuildConfig> BuildConfigs = new List<UnrealBuildConfig>();
 
-			foreach (XcodeSourceFile SourceFile in SourceFiles.OfType<XcodeSourceFile>())
-			{
-				string FileName = SourceFile.Reference.GetFileName();
-				string FileExtension = Path.GetExtension(FileName);
-				string FilePath = SourceFile.Reference.MakeRelativeTo(ProjectFilePath.Directory);
-				string FilePathMac = Utils.CleanDirectorySeparators(FilePath, '/');
+			//string ProjectName = ProjectFilePath.GetFileNameWithoutExtension();
 
-				if (IsGeneratedProject)
+			foreach (UnrealTargetConfiguration Configuration in Configurations)
+			{
+				if (InstalledPlatformInfo.IsValidConfiguration(Configuration, EProjectType.Code))
 				{
-					PBXBuildFileSection.Append(string.Format("\t\t{0} /* {1} in {2} */ = {{isa = PBXBuildFile; fileRef = {3} /* {1} */; }};" + ProjectFileGenerator.NewLine,
-						SourceFile.FileGuid,
-						FileName,
-						GetFileCategory(FileExtension),
-						SourceFile.FileRefGuid));
-				}
+					foreach (UnrealTargetPlatform Platform in Platforms)
+					{
+						if (InstalledPlatformInfo.IsValidPlatform(Platform, EProjectType.Code) && (Platform == UnrealTargetPlatform.Mac || Platform == UnrealTargetPlatform.IOS || Platform == UnrealTargetPlatform.TVOS)) // @todo support other platforms
+						{
+							UEBuildPlatform? BuildPlatform;
+							if (UEBuildPlatform.TryGetBuildPlatform(Platform, out BuildPlatform) && (BuildPlatform.HasRequiredSDKsInstalled() == SDKStatus.Valid))
+							{
+								// Check we have targets (Expected to be no Engine targets when generating for a single .uproject)
+								if (ProjectFile!.ProjectTargets.Count == 0 && ProjectFile!.BaseDir != Unreal.EngineDirectory)
+								{
+									throw new BuildException($"Expecting at least one ProjectTarget to be associated with project '{XcodeProjectFileLocation}' in the TargetProjects list ");
+								}
 
-				PBXFileReferenceSection.Append(string.Format("\t\t{0} /* {1} */ = {{isa = PBXFileReference; explicitFileType = {2}; name = \"{1}\"; path = \"{3}\"; sourceTree = SOURCE_ROOT; }};" + ProjectFileGenerator.NewLine,
-					SourceFile.FileRefGuid,
-					FileName,
-					GetFileType(FileExtension),
-					FilePathMac));
+								// Now go through all of the target types for this project
+								foreach (ProjectTarget ProjectTarget in ProjectFile.ProjectTargets.OfType<ProjectTarget>())
+								{
+									if (MSBuildProjectFile.IsValidProjectPlatformAndConfiguration(ProjectTarget, Platform, Configuration, Logger))
+									{
+										// Figure out if this is a monolithic build
+										bool bShouldCompileMonolithic = BuildPlatform.ShouldCompileMonolithicBinary(Platform);
+										bShouldCompileMonolithic |= (ProjectTarget.CreateRulesDelegate(Platform, Configuration).LinkType == TargetLinkType.Monolithic);
 
-				if (ShouldIncludeFileInBuildPhaseSection(SourceFile))
-				{
-					PBXSourcesBuildPhaseSection.Append("\t\t\t\t" + SourceFile.FileGuid + " /* " + FileName + " in Sources */," + ProjectFileGenerator.NewLine);
-				}
+										string ConfigName = Configuration.ToString();
+										if (ProjectTarget.TargetRules!.Type != TargetType.Game && ProjectTarget.TargetRules.Type != TargetType.Program)
+										{
+											ConfigName += " " + ProjectTarget.TargetRules.Type.ToString();
+										}
 
-				XcodeFileGroup? Group = FindGroupByAbsolutePath(ref Groups, SourceFile.Reference.Directory.FullName);
-				if (Group != null)
-				{
-					Group.Files.Add(SourceFile);
-				}
-			}
+										if (BuildConfigs.Where(Config => Config.DisplayName == ConfigName).ToList().Count == 0)
+										{
+											string TargetName = ProjectTarget.TargetFilePath.GetFileNameWithoutAnyExtensions();
 
-			foreach (XcconfigFile Xcconfig in AllXcconfigFiles)
-			{
-				PBXFileReferenceSection.Append($"\t\t{Xcconfig.Guid} /* {Xcconfig.Filename} */ = {{isa = PBXFileReference; explicitFileType = test.xcconfig; name = {Xcconfig.Filename}; path = Xcconfigs/{Xcconfig.Filename}; sourceTree = \"<group>\"; }};" + ProjectFileGenerator.NewLine);
+											// Get the output directory
+											DirectoryReference RootDirectory = Unreal.EngineDirectory;
+											// Unique and Monolithic both need to use the target directory not the engine directory
+											if (ProjectTarget.TargetRules.Type != TargetType.Program && (bShouldCompileMonolithic || ProjectTarget.TargetRules.BuildEnvironment == TargetBuildEnvironment.Unique))
+											{
+												if (ProjectTarget.UnrealProjectFilePath != null)
+												{
+													RootDirectory = ProjectTarget.UnrealProjectFilePath.Directory;
+												}
+											}
+
+											if (ProjectTarget.TargetRules.Type == TargetType.Program && ProjectTarget.UnrealProjectFilePath != null)
+											{
+												RootDirectory = ProjectTarget.UnrealProjectFilePath.Directory;
+											}
+
+											// Get the output directory
+											DirectoryReference OutputDirectory = DirectoryReference.Combine(RootDirectory, "Binaries");
+
+											string ExeName = TargetName;
+											if (!bShouldCompileMonolithic && ProjectTarget.TargetRules.Type != TargetType.Program)
+											{
+												// Figure out what the compiled binary will be called so that we can point the IDE to the correct file
+												if (ProjectTarget.TargetRules.Type != TargetType.Game)
+												{
+													// Only if shared - unique retains the Target Name
+													if (ProjectTarget.TargetRules.BuildEnvironment == TargetBuildEnvironment.Shared)
+													{
+														ExeName = "Unreal" + ProjectTarget.TargetRules.Type.ToString();
+													}
+												}
+											}
+
+											if (BuildPlatform.Platform == UnrealTargetPlatform.Mac)
+											{
+												string MacExecutableName = MakeExecutableFileName(ExeName, UnrealTargetPlatform.Mac, Configuration, ProjectTarget.TargetRules.Architecture, ProjectTarget.TargetRules.UndecoratedConfiguration);
+												string IOSExecutableName = MacExecutableName.Replace("-Mac-", "-IOS-");
+												string TVOSExecutableName = MacExecutableName.Replace("-Mac-", "-TVOS-");
+												BuildConfigs.Add(new UnrealBuildConfig(ConfigName, TargetName, FileReference.Combine(OutputDirectory, "Mac", MacExecutableName), FileReference.Combine(OutputDirectory, "IOS", IOSExecutableName), FileReference.Combine(OutputDirectory, "TVOS", TVOSExecutableName), ProjectTarget, Configuration));
+											}
+											else if (BuildPlatform.Platform == UnrealTargetPlatform.IOS || BuildPlatform.Platform == UnrealTargetPlatform.TVOS)
+											{
+												string IOSExecutableName = MakeExecutableFileName(ExeName, UnrealTargetPlatform.IOS, Configuration, ProjectTarget.TargetRules.Architecture, ProjectTarget.TargetRules.UndecoratedConfiguration);
+												string TVOSExecutableName = IOSExecutableName.Replace("-IOS-", "-TVOS-");
+												//string MacExecutableName = IOSExecutableName.Replace("-IOS-", "-Mac-");
+												BuildConfigs.Add(new UnrealBuildConfig(ConfigName, TargetName, FileReference.Combine(OutputDirectory, "Mac", IOSExecutableName), FileReference.Combine(OutputDirectory, "IOS", IOSExecutableName), FileReference.Combine(OutputDirectory, "TVOS", TVOSExecutableName), ProjectTarget, Configuration));
+											}
+										}
+									}
+								}
+							}
+						}
+					}
+				}
 			}
 
-			PBXFileReferenceSection.Append(string.Format("\t\t{0} /* {1} */ = {{isa = PBXFileReference; explicitFileType = {2}; path = {1}; sourceTree = BUILT_PRODUCTS_DIR; }};" + ProjectFileGenerator.NewLine, TargetAppGuid, TargetName, bIsAppBundle ? "wrapper.application" : "\"compiled.mach-o.executable\""));
+			return BuildConfigs;
 		}
 
-		private void GenerateSectionsWithExtensions(StringBuilder PBXBuildFileSection, StringBuilder PBXFileReferenceSection, StringBuilder PBXCopyFilesBuildPhaseSection, StringBuilder PBXResourcesBuildPhaseSection,
-													List<XcodeExtensionInfo> AllExtensions, FileReference? UProjectPath, List<XcodeBuildConfig> BuildConfigs, ILogger Logger)
+		private void ProcessExtensions(ILogger Logger)
 		{
-			if (UProjectPath != null)
+			if (UProjectFileLocation != null)
 			{
-				string ProjectExtensionsDir = Path.Combine(Path.GetDirectoryName(UProjectPath.FullName)!, "Build/IOS/Extensions");
+				DirectoryReference ProjectExtensionsDir = DirectoryReference.Combine(UProjectFileLocation.Directory, "Build/IOS/Extensions");
 				//string ProjectIntermediateDir = Path.Combine(Path.GetDirectoryName(UProjectPath.FullName), "Intermediate/IOS/Extensions");
 
-
-				if (Directory.Exists(ProjectExtensionsDir))
+				if (DirectoryReference.Exists(ProjectExtensionsDir))
 				{
-					throw new BuildException($"YOU NEED TO HANDLE EXTENSIONS NOW {UProjectPath}");
-#if false
-					foreach (DirectoryInfo DI in new System.IO.DirectoryInfo(ProjectExtensionsDir).EnumerateDirectories())
+					foreach (DirectoryReference ExtensionDir in DirectoryReference.EnumerateDirectories(ProjectExtensionsDir))
 					{
-						Console.WriteLine("  Project {0} has Extension {1}!", UProjectPath, DI);
+						// assume each Extension directory in here will create a resulting Extension.appex
+						string ExtensionName = ExtensionDir.GetDirectoryName();
+						Logger.LogInformation("  Project {0} has Extension {1}!", UProjectFileLocation.GetFileName(), ExtensionName);
 
-						// assume each Extension in here will create a resulting Extension.appex
-						string Extension = DI.Name + ".appex";
+						string Extension = ExtensionName + ".appex";
 
 						string ExtensionGuid = XcodeProjectFileGenerator.MakeXcodeGuid();
 
 						// make an extension info object
-						XcodeExtensionInfo ExtensionInfo = new XcodeExtensionInfo(DI.Name);
+						UnrealExtensionInfo ExtensionInfo = new UnrealExtensionInfo(ExtensionName);
 						AllExtensions.Add(ExtensionInfo);
-
-						PBXBuildFileSection.Append(string.Format("\t\t{0} /* {1} in Embed App Extensions */ = {{isa = PBXBuildFile; fileRef = {2} /* {1} */; settings = {{ATTRIBUTES = (RemoveHeadersOnCopy, ); }}; }};" + ProjectFileGenerator.NewLine,
+#if false
+						PBXBuildFileSection.Append(string.Format("\t\t{0} /* {1} in Embed App Extensions */ = {{isa = PBXBuildFile; fileRef = {2} /* {1} */; settings = {{ATTRIBUTES = (RemoveHeadersOnCopy, ); }}; }};",
 							ExtensionGuid,
 							Extension,
 							ExtensionInfo.ProductGuid));
 
-						PBXFileReferenceSection.Append(string.Format("\t\t{0} /* {1} */ = {{isa = PBXFileReference; explicitFileType = wrapper.app-extension; path = \"{1}\"; sourceTree = BUILT_PRODUCTS_DIR; }};" + ProjectFileGenerator.NewLine,
+						PBXFileReferenceSection.Append(string.Format("\t\t{0} /* {1} */ = {{isa = PBXFileReference; explicitFileType = wrapper.app-extension; path = \"{1}\"; sourceTree = BUILT_PRODUCTS_DIR; }};",
 							ExtensionInfo.ProductGuid,
 							Extension));
 
-						PBXCopyFilesBuildPhaseSection.Append(string.Format("\t\t\t\t{0} /* {1} in Embed App Extensions */," + ProjectFileGenerator.NewLine,
+						PBXCopyFilesBuildPhaseSection.Append(string.Format("\t\t\t\t{0} /* {1} in Embed App Extensions */,",
 							ExtensionGuid,
 							Extension));
 
-						PBXResourcesBuildPhaseSection.Append("/* Begin PBXResourcesBuildPhase section */" + ProjectFileGenerator.NewLine);
-						PBXResourcesBuildPhaseSection.Append("\t\t" + ExtensionInfo.ResourceBuildPhaseGuid + " /* Resources */ = {" + ProjectFileGenerator.NewLine);
-						PBXResourcesBuildPhaseSection.Append("\t\t\tisa = PBXResourcesBuildPhase;" + ProjectFileGenerator.NewLine);
-						PBXResourcesBuildPhaseSection.Append("\t\t\tbuildActionMask = 2147483647;" + ProjectFileGenerator.NewLine);
-						PBXResourcesBuildPhaseSection.Append("\t\t\tfiles = (" + ProjectFileGenerator.NewLine);
+						PBXResourcesBuildPhaseSection.Append("/* Begin PBXResourcesBuildPhase section */");
+						PBXResourcesBuildPhaseSection.Append("\t\t" + ExtensionInfo.ResourceBuildPhaseGuid + " /* Resources */ = {");
+						PBXResourcesBuildPhaseSection.Append("\t\t\tisa = PBXResourcesBuildPhase;");
+						PBXResourcesBuildPhaseSection.Append("\t\t\tbuildActionMask = 2147483647;");
+						PBXResourcesBuildPhaseSection.Append("\t\t\tfiles = (");
 						if (Directory.Exists(Path.Combine(DI.FullName, "Resources")))
 						{
 							DirectoryInfo ResourceDir = new System.IO.DirectoryInfo(Path.Combine(DI.FullName, "Resources"));
@@ -478,25 +520,25 @@ namespace UnrealBuildTool.XcodeProjectXcconfig
 								// for each resource, put it into the File/FileRef section, and into the CopyResuorceBuildPhase
 								string ResourceGuid = XcodeProjectFileGenerator.MakeXcodeGuid();
 								string ResourceRefGuid = XcodeProjectFileGenerator.MakeXcodeGuid();
-								PBXBuildFileSection.Append(string.Format("\t\t{0} /* {1} in Embed App Extensions */ = {{isa = PBXBuildFile; fileRef = {2} /* {1} */; }};" + ProjectFileGenerator.NewLine,
+								PBXBuildFileSection.Append(string.Format("\t\t{0} /* {1} in Embed App Extensions */ = {{isa = PBXBuildFile; fileRef = {2} /* {1} */; }};",
 									ResourceGuid,
 									FSI.Name,
 									ResourceRefGuid));
 
 								// lastKnownFileType = wrapper.app-extension; 
-								PBXFileReferenceSection.Append(string.Format("\t\t{0} /* {1} */ = {{isa = PBXFileReference; lastKnownFileType = folder.assetcatalog; path = \"{2}\"; sourceTree = \"<absolute>\"; }};" + ProjectFileGenerator.NewLine,
+								PBXFileReferenceSection.Append(string.Format("\t\t{0} /* {1} */ = {{isa = PBXFileReference; lastKnownFileType = folder.assetcatalog; path = \"{2}\"; sourceTree = \"<absolute>\"; }};",
 									ResourceRefGuid,
 									FSI.Name,
 									// @todo: make this relative path!! 
 									FSI.FullName));
 
-								PBXResourcesBuildPhaseSection.Append("\t\t\t\t" + ResourceGuid + " /* " + FSI.Name + " in " + ResourceDir.Name + " */," + ProjectFileGenerator.NewLine);
+								PBXResourcesBuildPhaseSection.Append("\t\t\t\t" + ResourceGuid + " /* " + FSI.Name + " in " + ResourceDir.Name + " */,");
 							}
 						}
-						PBXResourcesBuildPhaseSection.Append("\t\t\t);" + ProjectFileGenerator.NewLine);
-						PBXResourcesBuildPhaseSection.Append("\t\t\trunOnlyForDeploymentPostprocessing = 0;" + ProjectFileGenerator.NewLine);
-						PBXResourcesBuildPhaseSection.Append("\t\t};" + ProjectFileGenerator.NewLine);
-						PBXResourcesBuildPhaseSection.Append("/* End PBXResourcesBuildPhase section */" + ProjectFileGenerator.NewLine + ProjectFileGenerator.NewLine);
+						PBXResourcesBuildPhaseSection.Append("\t\t\t);");
+						PBXResourcesBuildPhaseSection.Append("\t\t\trunOnlyForDeploymentPostprocessing = 0;");
+						PBXResourcesBuildPhaseSection.Append("\t\t};");
+						PBXResourcesBuildPhaseSection.Append("/* End PBXResourcesBuildPhase section */");
 
 						StringBuilder ConfigSection = new StringBuilder();
 						// copy over the configs from the general project to the extension
@@ -505,13 +547,13 @@ namespace UnrealBuildTool.XcodeProjectXcconfig
 							string ConfigGuid = XcodeProjectFileGenerator.MakeXcodeGuid();
 							string ConfigName = Configuration.DisplayName;
 
-							ConfigSection.Append("\t\t" + ConfigGuid + " /* " + ConfigName + " */ = {" + ProjectFileGenerator.NewLine);
-							ConfigSection.Append("\t\t\tisa = XCBuildConfiguration;" + ProjectFileGenerator.NewLine);
-							ConfigSection.Append("\t\t\tbuildSettings = {" + ProjectFileGenerator.NewLine);
-							ConfigSection.Append("\t\t\t\tASSETCATALOG_COMPILER_APPICON_NAME = \"iMessage App Icon\";" + ProjectFileGenerator.NewLine);
-							ConfigSection.Append("\t\t\t\tINFOPLIST_FILE = \"" + Path.Combine(DI.FullName, "Info.plist") + "\";" + ProjectFileGenerator.NewLine);
-							ConfigSection.Append("\t\t\t\tSKIP_INSTALL = YES;" + ProjectFileGenerator.NewLine);
-							ConfigSection.Append("\t\t\t\tPRODUCT_NAME = \"$(TARGET_NAME)\";" + ProjectFileGenerator.NewLine);
+							ConfigSection.Append("\t\t" + ConfigGuid + " /* " + ConfigName + " */ = {");
+							ConfigSection.Append("\t\t\tisa = XCBuildConfiguration;");
+							ConfigSection.Append("\t\t\tbuildSettings = {");
+							ConfigSection.Append("\t\t\t\tASSETCATALOG_COMPILER_APPICON_NAME = \"iMessage App Icon\";");
+							ConfigSection.Append("\t\t\t\tINFOPLIST_FILE = \"" + Path.Combine(DI.FullName, "Info.plist") + "\";");
+							ConfigSection.Append("\t\t\t\tSKIP_INSTALL = YES;");
+							ConfigSection.Append("\t\t\t\tPRODUCT_NAME = \"$(TARGET_NAME)\";");
 
 							bool bSupportIOS = true;
 							bool bSupportTVOS = true;
@@ -519,1047 +561,1857 @@ namespace UnrealBuildTool.XcodeProjectXcconfig
 							{
 								IOSPlatform IOSPlatform = ((IOSPlatform)UEBuildPlatform.GetBuildPlatform(UnrealTargetPlatform.IOS));
 								IOSProjectSettings ProjectSettings = IOSPlatform.ReadProjectSettings(UProjectPath);
-								ConfigSection.Append("\t\t\t\t\"PRODUCT_BUNDLE_IDENTIFIER[sdk=iphoneos*]\" = " + ProjectSettings.BundleIdentifier + "." + ExtensionInfo.Name + ";" + ProjectFileGenerator.NewLine);
+								ConfigSection.Append("\t\t\t\t\"PRODUCT_BUNDLE_IDENTIFIER[sdk=iphone*]\" = " + ProjectSettings.BundleIdentifier + "." + ExtensionInfo.Name + ";");
 							}
 
 							if (bSupportTVOS && InstalledPlatformInfo.IsValidPlatform(UnrealTargetPlatform.TVOS, EProjectType.Code))
 							{
 								TVOSPlatform TVOSPlatform = ((TVOSPlatform)UEBuildPlatform.GetBuildPlatform(UnrealTargetPlatform.TVOS));
 								TVOSProjectSettings ProjectSettings = TVOSPlatform.ReadProjectSettings(UProjectPath);
-								ConfigSection.Append("\t\t\t\t\"PRODUCT_BUNDLE_IDENTIFIER[sdk=appletvos*]\" = " + ProjectSettings.BundleIdentifier + "." + ExtensionInfo.Name + ";" + ProjectFileGenerator.NewLine);
+								ConfigSection.Append("\t\t\t\t\"PRODUCT_BUNDLE_IDENTIFIER[sdk=appletvos*]\" = " + ProjectSettings.BundleIdentifier + "." + ExtensionInfo.Name + ";");
 							}
 
 							string? IOSRuntimeVersion, TVOSRuntimeVersion;
 							AppendPlatformConfiguration(ConfigSection, Configuration, ExtensionInfo.Name, UProjectPath, false, bSupportIOS, bSupportTVOS, Logger, out IOSRuntimeVersion, out TVOSRuntimeVersion);
 
-							ConfigSection.Append("\t\t\t};" + ProjectFileGenerator.NewLine);
-							ConfigSection.Append("\t\t\tname = \"" + ConfigName + "\";" + ProjectFileGenerator.NewLine);
-							ConfigSection.Append("\t\t};" + ProjectFileGenerator.NewLine);
+							ConfigSection.Append("\t\t\t};");
+							ConfigSection.Append("\t\t\tname = \"" + ConfigName + "\";");
+							ConfigSection.Append("\t\t};");
 
 							XcodeBuildConfig Config = new XcodeBuildConfig(ConfigName, ExtensionInfo.Name, null, null, null, null, Configuration.BuildConfig);
 							ExtensionInfo.AllConfigs.Add(ConfigGuid, Config);
 						}
 
 						ExtensionInfo.ConfigurationContents = ConfigSection.ToString();
-					}
 #endif
+					}
 				}
 			}
 		}
 
-		private void AppendGroup(XcodeFileGroup Group, StringBuilder Content)
+		public static IEnumerable<UnrealTargetPlatform> GetSupportedPlatforms()
 		{
-			if (!Group.bIsReference)
-			{
-				Content.Append(string.Format("\t\t{0} = {{{1}", Group.GroupGuid, ProjectFileGenerator.NewLine));
-				Content.Append("\t\t\tisa = PBXGroup;" + ProjectFileGenerator.NewLine);
-				Content.Append("\t\t\tchildren = (" + ProjectFileGenerator.NewLine);
-
-				foreach (XcodeFileGroup ChildGroup in Group.Children.Values)
-				{
-					Content.Append(string.Format("\t\t\t\t{0} /* {1} */,{2}", ChildGroup.GroupGuid, ChildGroup.GroupName, ProjectFileGenerator.NewLine));
-				}
+			List<UnrealTargetPlatform> SupportedPlatforms = new List<UnrealTargetPlatform>();
 
-				foreach (XcodeSourceFile File in Group.Files)
-				{
-					Content.Append(string.Format("\t\t\t\t{0} /* {1} */,{2}", File.FileRefGuid, File.Reference.GetFileName(), ProjectFileGenerator.NewLine));
-				}
+			if (XcodeProjectFileGenerator.ProjectFilePlatform.HasFlag(XcodeProjectFileGenerator.XcodeProjectFilePlatform.Mac))
+			{
+				SupportedPlatforms.Add(UnrealTargetPlatform.Mac);
+			}
 
-				Content.Append("\t\t\t);" + ProjectFileGenerator.NewLine);
-				Content.Append("\t\t\tname = \"" + Group.GroupName + "\";" + ProjectFileGenerator.NewLine);
-				Content.Append("\t\t\tpath = \"" + Group.GroupPath + "\";" + ProjectFileGenerator.NewLine);
-				Content.Append("\t\t\tsourceTree = \"<absolute>\";" + ProjectFileGenerator.NewLine);
-				Content.Append("\t\t};" + ProjectFileGenerator.NewLine);
+			if (XcodeProjectFileGenerator.ProjectFilePlatform.HasFlag(XcodeProjectFileGenerator.XcodeProjectFilePlatform.iOS))
+			{
+				SupportedPlatforms.Add(UnrealTargetPlatform.IOS);
+			}
 
-				foreach (XcodeFileGroup ChildGroup in Group.Children.Values)
-				{
-					AppendGroup(ChildGroup, Content);
-				}
+			if (XcodeProjectFileGenerator.ProjectFilePlatform.HasFlag(XcodeProjectFileGenerator.XcodeProjectFilePlatform.tvOS))
+			{
+				SupportedPlatforms.Add(UnrealTargetPlatform.TVOS);
 			}
-		}
 
-		private void AppendBuildFileSection(StringBuilder Content, StringBuilder SectionContent)
-		{
-			Content.Append("/* Begin PBXBuildFile section */" + ProjectFileGenerator.NewLine);
-			Content.Append(SectionContent);
-			Content.Append("/* End PBXBuildFile section */" + ProjectFileGenerator.NewLine + ProjectFileGenerator.NewLine);
+			return SupportedPlatforms;
 		}
 
-		private void AppendFileReferenceSection(StringBuilder Content, StringBuilder SectionContent)
+		public static IEnumerable<UnrealTargetConfiguration> GetSupportedConfigurations()
 		{
-			Content.Append("/* Begin PBXFileReference section */" + ProjectFileGenerator.NewLine);
-			Content.Append(SectionContent);
-			Content.Append("/* End PBXFileReference section */" + ProjectFileGenerator.NewLine + ProjectFileGenerator.NewLine);
+			return new UnrealTargetConfiguration[] {
+				UnrealTargetConfiguration.Debug,
+				UnrealTargetConfiguration.DebugGame,
+				UnrealTargetConfiguration.Development,
+				UnrealTargetConfiguration.Test,
+				UnrealTargetConfiguration.Shipping
+			};
 		}
 
-		private void AppendSourcesBuildPhaseSection(StringBuilder Content, StringBuilder SectionContent, string SourcesBuildPhaseGuid)
+		public static bool ShouldIncludeProjectInWorkspace(ProjectFile Proj, ILogger Logger)
 		{
-			Content.Append("/* Begin PBXSourcesBuildPhase section */" + ProjectFileGenerator.NewLine);
-			Content.Append(string.Format("\t\t{0} = {{{1}", SourcesBuildPhaseGuid, ProjectFileGenerator.NewLine));
-			Content.Append("\t\t\tisa = PBXSourcesBuildPhase;" + ProjectFileGenerator.NewLine);
-			Content.Append("\t\t\tbuildActionMask = 2147483647;" + ProjectFileGenerator.NewLine);
-			Content.Append("\t\t\tfiles = (" + ProjectFileGenerator.NewLine);
-			Content.Append(SectionContent);
-			Content.Append("\t\t\t);" + ProjectFileGenerator.NewLine);
-			Content.Append("\t\t\trunOnlyForDeploymentPostprocessing = 0;" + ProjectFileGenerator.NewLine);
-			Content.Append("\t\t};" + ProjectFileGenerator.NewLine);
-			Content.Append("/* End PBXSourcesBuildPhase section */" + ProjectFileGenerator.NewLine + ProjectFileGenerator.NewLine);
+			return CanBuildProjectLocally(Proj, Logger);
 		}
 
-		private XcodeFileGroup? FindRootFileGroup(Dictionary<string, XcodeFileGroup> GroupsDict)
+		public static bool CanBuildProjectLocally(ProjectFile Proj, ILogger Logger)
 		{
-			foreach (XcodeFileGroup Group in GroupsDict.Values)
+			foreach (Project ProjectTarget in Proj.ProjectTargets)
 			{
-				if (Group.Children.Count > 1 || Group.Files.Count > 0)
-				{
-					return Group;
-				}
-				else
+				foreach (UnrealTargetPlatform Platform in GetSupportedPlatforms())
 				{
-					XcodeFileGroup? Found = FindRootFileGroup(Group.Children);
-					if (Found != null)
+					foreach (UnrealTargetConfiguration Config in GetSupportedConfigurations())
 					{
-						return Found;
+						if (MSBuildProjectFile.IsValidProjectPlatformAndConfiguration(ProjectTarget, Platform, Config, Logger))
+						{
+							return true;
+						}
 					}
 				}
 			}
-			return null;
-		}
 
-		private void AppendCopyExtensionsBuildPhaseSection(StringBuilder Content, StringBuilder SectionContent, string CopyFilesBuildPhaseGuid)
-		{
-			Content.Append("/* Begin PBXCopyFilesBuildPhase section */" + ProjectFileGenerator.NewLine);
-			Content.Append(string.Format("\t{0} /* Embed App Extensions */ = {{{1}", CopyFilesBuildPhaseGuid, ProjectFileGenerator.NewLine));
-			Content.Append("\t\tisa = PBXCopyFilesBuildPhase;" + ProjectFileGenerator.NewLine);
-			Content.Append("\t\tbuildActionMask = 2147483647;" + ProjectFileGenerator.NewLine);
-			Content.Append("\t\tdstPath = \"\";" + ProjectFileGenerator.NewLine);
-			Content.Append("\t\tdstSubfolderSpec = 13;" + ProjectFileGenerator.NewLine);
-			Content.Append("\t\tfiles = (" + ProjectFileGenerator.NewLine);
-			Content.Append(SectionContent);
-			Content.Append("\t\t);" + ProjectFileGenerator.NewLine);
-			Content.Append("\t\tname = \"Embed App Extensions\";" + ProjectFileGenerator.NewLine);
-			Content.Append("\t\trunOnlyForDeploymentPostprocessing = 0;" + ProjectFileGenerator.NewLine);
-			Content.Append("\t};" + ProjectFileGenerator.NewLine);
-			Content.Append("/* End PBXCopyFilesBuildPhase section */" + ProjectFileGenerator.NewLine);
+			return false;
 		}
 
-		private void AppendGroupSection(StringBuilder Content, string MainGroupGuid, string ProductRefGroupGuid, string XcconfigsRefGroupGuid, string TargetAppGuid, string TargetName, List<XcodeExtensionInfo> AllExtensions)
+		private static string MakeExecutableFileName(string BinaryName, UnrealTargetPlatform Platform, UnrealTargetConfiguration Configuration, string Architecture, UnrealTargetConfiguration UndecoratedConfiguration)
 		{
-			XcodeFileGroup? RootGroup = FindRootFileGroup(Groups);
-			if (RootGroup == null)
-			{
-				return;
-			}
-			
-			Content.Append("/* Begin PBXGroup section */" + ProjectFileGenerator.NewLine);
-
-			// Main group
-			Content.Append(string.Format("\t\t{0} = {{{1}", MainGroupGuid, ProjectFileGenerator.NewLine));
-			Content.Append("\t\t\tisa = PBXGroup;" + ProjectFileGenerator.NewLine);
-			Content.Append("\t\t\tchildren = (" + ProjectFileGenerator.NewLine);
-			foreach (XcodeFileGroup Group in RootGroup.Children.Values)
-			{
-				Content.Append(string.Format("\t\t\t\t{0} /* {1} */,{2}", Group.GroupGuid, Group.GroupName, ProjectFileGenerator.NewLine));
-			}
-
-			foreach (XcodeSourceFile File in RootGroup.Files)
-			{
-				Content.Append(string.Format("\t\t\t\t{0} /* {1} */,{2}", File.FileRefGuid, File.Reference.GetFileName(), ProjectFileGenerator.NewLine));
-			}
-
-			Content.Append(string.Format("\t\t\t\t{0} /* Products */,{1}", ProductRefGroupGuid, ProjectFileGenerator.NewLine));
-			Content.Append(string.Format("\t\t\t\t{0} /* Xcconfigs */,{1}", XcconfigsRefGroupGuid, ProjectFileGenerator.NewLine));
-			Content.Append("\t\t\t);" + ProjectFileGenerator.NewLine);
-			Content.Append("\t\t\tsourceTree = \"<group>\";" + ProjectFileGenerator.NewLine);
-			Content.Append("\t\t};" + ProjectFileGenerator.NewLine);
+			StringBuilder Result = new StringBuilder();
 
-			// Sources groups
-			foreach (XcodeFileGroup Group in RootGroup.Children.Values)
-			{
-				AppendGroup(Group, Content);
-			}
+			Result.Append(BinaryName);
 
-			// Products group
-			Content.Append(string.Format("\t\t{0} /* Products */ = {{{1}", ProductRefGroupGuid, ProjectFileGenerator.NewLine));
-			Content.Append("\t\t\tisa = PBXGroup;" + ProjectFileGenerator.NewLine);
-			Content.Append("\t\t\tchildren = (" + ProjectFileGenerator.NewLine);
-			Content.Append(string.Format("\t\t\t\t{0} /* {1} */,{2}", TargetAppGuid, TargetName, ProjectFileGenerator.NewLine));
-			foreach (XcodeExtensionInfo EI in AllExtensions)
+			if (Configuration != UndecoratedConfiguration)
 			{
-				Content.Append(string.Format("\t\t\t\t{0} /* {1} */,{2}", EI.ProductGuid, EI.Name, ProjectFileGenerator.NewLine));
+				Result.AppendFormat("-{0}-{1}", Platform.ToString(), Configuration.ToString());
 			}
-			Content.Append("\t\t\t);" + ProjectFileGenerator.NewLine);
-			Content.Append("\t\t\tname = Products;" + ProjectFileGenerator.NewLine);
-			Content.Append("\t\t\tsourceTree = \"<group>\";" + ProjectFileGenerator.NewLine);
-			Content.Append("\t\t};" + ProjectFileGenerator.NewLine);
 
-			// Xcconfigs group
-			Content.Append(string.Format("\t\t{0} /* XcconfigsRefGroupGuid */ = {{{1}", XcconfigsRefGroupGuid, ProjectFileGenerator.NewLine));
-			Content.Append("\t\t\tisa = PBXGroup;" + ProjectFileGenerator.NewLine);
-			Content.Append("\t\t\tchildren = (" + ProjectFileGenerator.NewLine);
-			foreach (XcconfigFile Xcconfig in AllXcconfigFiles)
+			UEBuildPlatform BuildPlatform = UEBuildPlatform.GetBuildPlatform(Platform);
+			if (BuildPlatform.RequiresArchitectureSuffix())
 			{
-				Content.Append($"\t\t\t\t{Xcconfig.Guid} /* {Xcconfig.Filename} */,{ProjectFileGenerator.NewLine}");
+				Result.Append(Architecture);
 			}
-			Content.Append("\t\t\t);" + ProjectFileGenerator.NewLine);
-			Content.Append("\t\t\tname = Xcconfigs;" + ProjectFileGenerator.NewLine);
-			Content.Append("\t\t\tsourceTree = \"<group>\";" + ProjectFileGenerator.NewLine);
-			Content.Append("\t\t};" + ProjectFileGenerator.NewLine);
 
-			Content.Append("/* End PBXGroup section */" + ProjectFileGenerator.NewLine + ProjectFileGenerator.NewLine);
+			return Result.ToString();
 		}
 
-		private void AppendLegacyTargetSection(StringBuilder Content, string TargetName, string TargetGuid, string TargetBuildConfigGuid, FileReference? UProjectPath)
-		{
-			string UEDir = ConvertPath(Path.GetFullPath(Directory.GetCurrentDirectory() + "../../.."));
-			string BuildToolPath = UEDir + "/Engine/Build/BatchFiles/Mac/XcodeBuild.sh";
-
-			Content.Append("/* Begin PBXLegacyTarget section */" + ProjectFileGenerator.NewLine);
-
-			Content.Append("\t\t" + TargetGuid + " /* " + TargetName + " */ = {" + ProjectFileGenerator.NewLine);
-			Content.Append("\t\t\tisa = PBXLegacyTarget;" + ProjectFileGenerator.NewLine);
-			Content.Append("\t\t\tbuildArgumentsString = \"$(ACTION) $(UE_BUILD_TARGET_NAME) $(PLATFORM_NAME) $(UE_BUILD_TARGET_CONFIG)" 
-				+ (UProjectPath == null ? "" : " \\\"" + UProjectPath.FullName + "\\\"") 
-				+ "\";" + ProjectFileGenerator.NewLine);
-			Content.Append("\t\t\tbuildConfigurationList = "  + TargetBuildConfigGuid + " /* Build configuration list for PBXLegacyTarget \"" + TargetName + "\" */;" + ProjectFileGenerator.NewLine);
-			Content.Append("\t\t\tbuildPhases = (" + ProjectFileGenerator.NewLine);
-			Content.Append("\t\t\t);" + ProjectFileGenerator.NewLine);
-			Content.Append("\t\t\tbuildToolPath = \"" + BuildToolPath + "\";" + ProjectFileGenerator.NewLine);
-			Content.Append("\t\t\tbuildWorkingDirectory = \"" + UEDir + "\";" + ProjectFileGenerator.NewLine);
-			Content.Append("\t\t\tdependencies = (" + ProjectFileGenerator.NewLine);
-			Content.Append("\t\t\t);" + ProjectFileGenerator.NewLine);
-			Content.Append("\t\t\tname = \"" + TargetName + "\";" + ProjectFileGenerator.NewLine);
-			Content.Append("\t\t\tpassBuildSettingsInEnvironment = 1;" + ProjectFileGenerator.NewLine);
-			Content.Append("\t\t\tproductName = \"" + TargetName + "\";" + ProjectFileGenerator.NewLine);
-			Content.Append("\t\t};" + ProjectFileGenerator.NewLine);
 
-			Content.Append("/* End PBXLegacyTarget section */" + ProjectFileGenerator.NewLine + ProjectFileGenerator.NewLine);
-		}
+		// cache for the below function
+		Dictionary<string, IEnumerable<string>> CachedMacProjectArchitectures = new Dictionary<string, IEnumerable<string>>();
 
-		private void AppendRunTargetSection(StringBuilder Content, string TargetName, string TargetGuid, string TargetBuildConfigGuid, string TargetDependencyGuid, 
-				string TargetAppGuid, string CopyExtensionsBuildPhaseGuid, string ShellScriptSectionGuid, List<XcodeExtensionInfo> AllExtensions, bool bIsAppBundle)
+		/// <summary>
+		/// Returns the Mac architectures that should be configured for the provided target. If the target has a project we'll adhere
+		/// to whether it's set as Intel/Universal/Apple unless the type is denied (pretty much just Editor)
+		/// 
+		/// If the target has no project we'll support allow-listed targets for installed builds and all non-editor architectures 
+		/// for source builds. Not all programs are going to compile for Apple Silicon, but being able to build and fail is useful...
+		/// </summary>
+		/// <param name="Config">Build config for the target we're generating</param>
+		/// <param name="InProjectFile">Path to the project file, or null if the target has no project</param>
+		/// <returns></returns>
+		public IEnumerable<string> GetSupportedMacArchitectures(UnrealBuildConfig Config, FileReference? InProjectFile)
 		{
-			List<string> DependencyGuids = new List<string>();
-			// depends on the Run target if we want one
-			if (!XcodeProjectFileGenerator.bGeneratingRunIOSProject && !XcodeProjectFileGenerator.bGeneratingRunTVOSProject)
-			{
-				DependencyGuids.Add(TargetDependencyGuid);
-			}
-			// make sure extensions get built
-			foreach (XcodeExtensionInfo EI in AllExtensions)
-			{
-				DependencyGuids.Add(EI.TargetDependencyGuid);
-			}
-
-			Dictionary<string, string> BuildPhases = new Dictionary<string, string>();
-			// add optional build phases
-			if (!string.IsNullOrEmpty(CopyExtensionsBuildPhaseGuid))
-			{
-				BuildPhases.Add(CopyExtensionsBuildPhaseGuid, "Embed App Extensions");
-			}
-			if (!string.IsNullOrEmpty(ShellScriptSectionGuid))
-			{
-				BuildPhases.Add(ShellScriptSectionGuid, "Shell Script");
-			}
-
-			// use generica target section function for an application type
-			AppendGenericTargetSection(Content, TargetName, TargetGuid, bIsAppBundle ? "com.apple.product-type.application" : "com.apple.product-type.tool", TargetBuildConfigGuid, TargetAppGuid, DependencyGuids, BuildPhases);
-		}
+			// All architectures supported
+			IEnumerable<string> AllArchitectures = new[] { MacExports.IntelArchitecture, MacExports.AppleArchitecture };
 
-		private void AppendGenericTargetSection(StringBuilder Content, string TargetName, string TargetGuid, string TargetType, string TargetBuildConfigGuid, string TargetAppGuid, IEnumerable<string>? TargetDependencyGuids, Dictionary<string, string> BuildPhases)
-		{
-			Content.Append("/* Begin PBXNativeTarget section */" + ProjectFileGenerator.NewLine);
+			// Add a way on the command line of forcing a project file with all architectures (there isn't a good way to let this be
+			// set and checked where we can access it).
+			bool ForceAllArchitectures = Environment.GetCommandLineArgs().Contains("AllArchitectures", StringComparer.OrdinalIgnoreCase);
 
-			Content.Append("\t\t" + TargetGuid + " /* " + TargetName + " */ = {" + ProjectFileGenerator.NewLine);
-			Content.Append("\t\t\tisa = PBXNativeTarget;" + ProjectFileGenerator.NewLine);
-			Content.Append("\t\t\tbuildConfigurationList = "  + TargetBuildConfigGuid + " /* Build configuration list for PBXNativeTarget \"" + TargetName + "\" */;" + ProjectFileGenerator.NewLine);
-			Content.Append("\t\t\tbuildPhases = (" + ProjectFileGenerator.NewLine);
-			if (BuildPhases != null)
-			{
-				foreach (KeyValuePair<string, string> BuildPhasePair in BuildPhases)
-				{
-					Content.Append("\t\t\t\t" + BuildPhasePair.Key + " /* " + BuildPhasePair.Value + " */, " + ProjectFileGenerator.NewLine);
-				}
-			}
-			Content.Append("\t\t\t);" + ProjectFileGenerator.NewLine);
-			Content.Append("\t\t\tdependencies = (" + ProjectFileGenerator.NewLine);
-			if (TargetDependencyGuids != null)
-			{
-				foreach (string DependencyGuid in TargetDependencyGuids)
+			if (ForceAllArchitectures)
 			{
-					Content.Append("\t\t\t\t" + DependencyGuid + " /* PBXTargetDependency */," + ProjectFileGenerator.NewLine);
-				}
+				return AllArchitectures;
 			}
-			Content.Append("\t\t\t);" + ProjectFileGenerator.NewLine);
-			Content.Append("\t\t\tname = \"" + TargetName + "\";" + ProjectFileGenerator.NewLine);
-			Content.Append("\t\t\tpassBuildSettingsInEnvironment = 1;" + ProjectFileGenerator.NewLine);
-			Content.Append("\t\t\tproductName = \"" + TargetName + "\";" + ProjectFileGenerator.NewLine);
-			Content.Append("\t\t\tproductReference = \"" + TargetAppGuid + "\";" + ProjectFileGenerator.NewLine);
-			Content.Append("\t\t\tproductType = \"" + TargetType + "\";" + ProjectFileGenerator.NewLine);
-			Content.Append("\t\t};" + ProjectFileGenerator.NewLine);
-
-			Content.Append("/* End PBXNativeTarget section */" + ProjectFileGenerator.NewLine + ProjectFileGenerator.NewLine);
-		}
-
-		private void AppendIndexTargetSection(StringBuilder Content, string TargetName, string TargetGuid, string TargetBuildConfigGuid, string SourcesBuildPhaseGuid)
-		{
-			Content.Append("/* Begin PBXNativeTarget section */" + ProjectFileGenerator.NewLine);
 
-			Content.Append("\t\t" + TargetGuid + " /* " + TargetName + " */ = {" + ProjectFileGenerator.NewLine);
-			Content.Append("\t\t\tisa = PBXNativeTarget;" + ProjectFileGenerator.NewLine);
-			Content.Append("\t\t\tbuildConfigurationList = "  + TargetBuildConfigGuid + " /* Build configuration list for PBXNativeTarget \"" + TargetName + "\" */;" + ProjectFileGenerator.NewLine);
-			Content.Append("\t\t\tbuildPhases = (" + ProjectFileGenerator.NewLine);
-			Content.Append("\t\t\t\t" + SourcesBuildPhaseGuid + " /* Sources */," + ProjectFileGenerator.NewLine);
-			Content.Append("\t\t\t);" + ProjectFileGenerator.NewLine);
-			Content.Append("\t\t\tdependencies = (" + ProjectFileGenerator.NewLine);
-			Content.Append("\t\t\t);" + ProjectFileGenerator.NewLine);
-			Content.Append("\t\t\tname = \"" + TargetName + "\";" + ProjectFileGenerator.NewLine);
-			Content.Append("\t\t\tpassBuildSettingsInEnvironment = 1;" + ProjectFileGenerator.NewLine);
-			Content.Append("\t\t\tproductName = \"" + TargetName + "\";" + ProjectFileGenerator.NewLine);
-			Content.Append("\t\t\tproductType = \"com.apple.product-type.library.static\";" + ProjectFileGenerator.NewLine);
-			Content.Append("\t\t};" + ProjectFileGenerator.NewLine);
-
-			Content.Append("/* End PBXNativeTarget section */" + ProjectFileGenerator.NewLine + ProjectFileGenerator.NewLine);
-		}
-
-		private void AppendShellScriptSection(StringBuilder Content, string ShellScriptGuid, FileReference? UProjectPath)
-		{
-			StringBuilder FrameworkScript = new StringBuilder();
+			string TargetName = Config.BuildTarget;
 
-			// nothing to do without a project
-			if (UProjectPath == null)
+			// First time seeing this target?
+			if (!CachedMacProjectArchitectures.ContainsKey(TargetName))
 			{
-				return;
-			}
-			
-			// @todo: look also in Project/Build/Frameworks directory!
-			ProjectDescriptor Project = ProjectDescriptor.FromFile(UProjectPath);
-			List<PluginInfo> AvailablePlugins = Plugins.ReadAvailablePlugins(Unreal.EngineDirectory, DirectoryReference.FromFile(UProjectPath), Project.AdditionalPluginDirectories);
+				// Default to Intel
+				IEnumerable<string> TargetArchitectures = new[] { MacExports.IntelArchitecture };
 
-			// look in each plugin for frameworks
-			// @todo: Cache this kind of things since every target will re-do this work!
-			foreach (PluginInfo PI in AvailablePlugins)
-			{
-				if (!Plugins.IsPluginEnabledForTarget(PI, Project, UnrealTargetPlatform.IOS, UnrealTargetConfiguration.Development, TargetRules.TargetType.Game))
-				{
-					continue;
-				}
+				// These targets are known to work so are allow-listed
+				bool IsAllowed = MacExports.TargetsAllowedForAppleSilicon.Contains(TargetName, StringComparer.OrdinalIgnoreCase);
 
-				// for now, we copy and code sign all *.framework.zip, even if the have no code (non-code frameworks are assumed to be *.embeddedframework.zip
-				DirectoryReference FrameworkDir = DirectoryReference.Combine(PI.Directory, "Source/Frameworks");
-				if (!DirectoryReference.Exists(FrameworkDir))
+				// determine the target architectures based on what's allowed/denied
+				if (IsAllowed)
 				{
-					FrameworkDir = DirectoryReference.Combine(PI.Directory, "Frameworks");
+					TargetArchitectures = AllArchitectures;
 				}
-				if (DirectoryReference.Exists(FrameworkDir))
+				else
 				{
-					// look at each zip
-					foreach (FileInfo FI in new System.IO.DirectoryInfo(FrameworkDir.FullName).EnumerateFiles("*.framework.zip"))
+					// if this is an unspecified tool/program, default to Intel for installed builds because we know all of that works. 
+					if (Config.ProjectTarget!.TargetRules!.Type == TargetType.Program)
 					{
-						//string Guid = XcodeProjectFileGenerator.MakeXcodeGuid();
-						//string RefGuid = XcodeProjectFileGenerator.MakeXcodeGuid();
-
-						// for FI of foo.framework.zip, this will give us foo.framework
-						//string Framework = Path.GetFileNameWithoutExtension(FI.FullName);
-
-						// unzip the framework right into the .app
-						FrameworkScript.AppendFormat("\\techo Unzipping {0}...\\n", FI.FullName);
-						FrameworkScript.AppendFormat("\\tunzip -o -q {0} -d ${{FRAMEWORK_DIR}} -x \\\"__MACOSX/*\\\" \\\"*/.DS_Store\\\"\\n", FI.FullName);
+						// For misc tools we default to Intel for installed builds because we know all of that works. 
+						TargetArchitectures = Unreal.IsEngineInstalled() ? new[] { MacExports.IntelArchitecture } : AllArchitectures;
 					}
-				}
-			}
+					else
+					{
+						// For project targets we default to Intel then check the project settings. Note the editor target will have
+						// been denied above already.
+						TargetArchitectures = new[] { MacExports.IntelArchitecture };
 
-			string ShellScript = "set -e\\n\\nIFS=$'\\\\n'\\n\\n" +
-				"if [ $PLATFORM_NAME = iphoneos ] || [ $PLATFORM_NAME = tvos ]; then \\n" +
-				"\\tFRAMEWORK_DIR=$TARGET_BUILD_DIR/$EXECUTABLE_FOLDER_PATH/Frameworks\\n" +
-				FrameworkScript.ToString() + 
-				// and now code sign anything that has been unzipped above
-				"\\tfor FRAMEWORK in ${FRAMEWORK_DIR}/*.framework; do\\n" +
-					"\\t\\t[ -d \\\"${FRAMEWORK}\\\" ] || continue\\n" +
-					"\\t\\techo Codesigning ${FRAMEWORK}\\n" +
-					"\\t\\tcodesign --force --sign ${EXPANDED_CODE_SIGN_IDENTITY} --generate-entitlement-der --verbose --preserve-metadata=identifier,entitlements,flags --timestamp=none \\\"${FRAMEWORK}\\\"\\n" +
-				"\\tdone\\n" +
-				"fi\\n";
+						ConfigHierarchy EngineIni = ConfigCache.ReadHierarchy(ConfigHierarchyType.Engine, InProjectFile?.Directory, UnrealTargetPlatform.Mac);
+						string TargetArchitecture;
+						string Key = Config.ProjectTarget!.TargetRules!.Type == TargetType.Editor ? "EditorTargetArchitecture" : "TargetArchitecture";
+						if (EngineIni.GetString("/Script/MacTargetPlatform.MacTargetSettings", Key, out TargetArchitecture))
+						{
+							if (TargetArchitecture.Contains("Universal", StringComparison.OrdinalIgnoreCase))
+							{
+								TargetArchitectures = AllArchitectures;
+							}
+							else if (TargetArchitecture.Contains("Intel", StringComparison.OrdinalIgnoreCase))
+							{
+								TargetArchitectures = new[] { MacExports.IntelArchitecture };
+							}
+							else if (TargetArchitecture.Contains("Apple", StringComparison.OrdinalIgnoreCase))
+							{
+								TargetArchitectures = new[] { MacExports.AppleArchitecture };
+							}
+						}
+					}
+				}
 
-			Content.Append("/* Begin PBXShellScriptBuildPhase section */" + ProjectFileGenerator.NewLine);
-			Content.Append(string.Format("\t\t{0} /* Sign Frameworks */ = {{" + ProjectFileGenerator.NewLine, ShellScriptGuid));
-			Content.Append("\t\t\tisa = PBXShellScriptBuildPhase;" + ProjectFileGenerator.NewLine);
-			Content.Append("\t\t\tbuildActionMask = 2147483647;" + ProjectFileGenerator.NewLine);
-			Content.Append("\t\t\tfiles = (" + ProjectFileGenerator.NewLine);
-			Content.Append("\t\t\t);" + ProjectFileGenerator.NewLine);
-			Content.Append("\t\t\tinputPaths = (" + ProjectFileGenerator.NewLine);
-			Content.Append("\t\t\t);" + ProjectFileGenerator.NewLine);
-			Content.Append("\t\t\toutputPaths = (" + ProjectFileGenerator.NewLine);
-			Content.Append("\t\t\t);" + ProjectFileGenerator.NewLine);
-			Content.Append("\t\t\tname = \"Sign Manual Frameworks\";" + ProjectFileGenerator.NewLine);
-			Content.Append("\t\t\trunOnlyForDeploymentPostprocessing = 0;" + ProjectFileGenerator.NewLine);
-			Content.Append("\t\t\tshellPath = /bin/sh;" + ProjectFileGenerator.NewLine);
-			Content.Append(string.Format("\t\t\tshellScript = \"{0}\";" + ProjectFileGenerator.NewLine, ShellScript));
-			Content.Append("\t\t};" + ProjectFileGenerator.NewLine);
-			Content.Append("/* End PBXShellScriptBuildPhase section */" + ProjectFileGenerator.NewLine);
+				// Cache this so we don't need to keep checking this file
+				CachedMacProjectArchitectures.Add(TargetName, TargetArchitectures);
+			}
+
+			return CachedMacProjectArchitectures[TargetName];
 		}
 
-		private void AppendExtensionTargetSections(StringBuilder ProjectFileContent, List<XcodeExtensionInfo> AllExtensions)
+	}
+
+	/// <summary>
+	/// Info needed to make a file a member of specific group
+	/// </summary>
+	class XcodeSourceFile : ProjectFile.SourceFile
+	{
+		/// <summary>
+		/// Constructor
+		/// </summary>
+		public XcodeSourceFile(FileReference InitFilePath, DirectoryReference? InitRelativeBaseFolder)
+			: base(InitFilePath, InitRelativeBaseFolder)
 		{
-			foreach (XcodeExtensionInfo EI in AllExtensions)
-			{
-				Dictionary<string, string> BuildPhases = new Dictionary<string, string>();
-				BuildPhases.Add(EI.ResourceBuildPhaseGuid, "Resources");
+			FileGuid = XcodeProjectFileGenerator.MakeXcodeGuid();
+			FileRefGuid = XcodeProjectFileGenerator.MakeXcodeGuid();
+		}
 
-				AppendGenericTargetSection(ProjectFileContent, EI.Name, EI.TargetGuid, "com.apple.product-type.app-extension.messages-sticker-pack", EI.ConfigListGuid, EI.ProductGuid, null, BuildPhases);
-			}
+		/// <summary>
+		/// File Guid for use in Xcode project
+		/// </summary>
+		public string FileGuid
+		{
+			get;
+			private set;
 		}
 
-		private void AppendProjectSection(StringBuilder Content, string TargetName, string TargetGuid, string BuildTargetName, string BuildTargetGuid, string IndexTargetName, string IndexTargetGuid, string MainGroupGuid, string ProductRefGroupGuid, string ProjectGuid, string ProjectBuildConfigGuid, FileReference? ProjectFile, List<XcodeExtensionInfo> AllExtensions)
+		public void ReplaceGuids(string NewFileGuid, string NewFileRefGuid)
 		{
-			Content.Append("/* Begin PBXProject section */" + ProjectFileGenerator.NewLine);
+			FileGuid = NewFileGuid;
+			FileRefGuid = NewFileRefGuid;
+		}
 
-			Content.Append("\t\t" + ProjectGuid + " /* Project object */ = {" + ProjectFileGenerator.NewLine);
-			Content.Append("\t\t\tisa = PBXProject;" + ProjectFileGenerator.NewLine);
-			Content.Append("\t\t\tattributes = {" + ProjectFileGenerator.NewLine);
-			Content.Append("\t\t\t\tLastUpgradeCheck = 2000;" + ProjectFileGenerator.NewLine);
-			Content.Append("\t\t\t\tORGANIZATIONNAME = \"Epic Games, Inc.\";" + ProjectFileGenerator.NewLine);
-            Content.Append("\t\t\t\tTargetAttributes = {" + ProjectFileGenerator.NewLine);
-            Content.Append("\t\t\t\t\t" + TargetGuid + " = {" + ProjectFileGenerator.NewLine);
+		/// <summary>
+		/// File reference Guid for use in Xcode project
+		/// </summary>
+		public string FileRefGuid
+		{
+			get;
+			private set;
+		}
+	}
 
-			bool bAutomaticSigning = false;
-			if (InstalledPlatformInfo.IsValidPlatform(UnrealTargetPlatform.IOS, EProjectType.Code))
-			{
-				IOSPlatform IOSPlatform = ((IOSPlatform)UEBuildPlatform.GetBuildPlatform(UnrealTargetPlatform.IOS));
-				IOSProjectSettings ProjectSettings = IOSPlatform.ReadProjectSettings(ProjectFile);
-				bAutomaticSigning = ProjectSettings.bAutomaticSigning;
-			}
+	/// <summary>
+	/// Represents a group of files shown in Xcode's project navigator as a folder
+	/// </summary>
+	internal class XcodeFileGroup
+	{
+		public XcodeFileGroup(string InName, string InPath, bool InIsReference)
+		{
+			GroupName = InName;
+			GroupPath = InPath;
+			GroupGuid = XcodeProjectFileGenerator.MakeXcodeGuid();
+			bIsReference = InIsReference;
+		}
 
-			if (InstalledPlatformInfo.IsValidPlatform(UnrealTargetPlatform.TVOS, EProjectType.Code))
-			{
-				TVOSPlatform TVOSPlatform = ((TVOSPlatform)UEBuildPlatform.GetBuildPlatform(UnrealTargetPlatform.TVOS));
-				TVOSProjectSettings ProjectSettings = TVOSPlatform.ReadProjectSettings(ProjectFile);
-				//TVOSProvisioningData ProvisioningData = TVOSPlatform.ReadProvisioningData(ProjectSettings, bForDistribution);
-				bAutomaticSigning = ProjectSettings.bAutomaticSigning;
-			}
+		public string GroupGuid;
+		public string GroupName;
+		public string GroupPath;
+		public Dictionary<string, XcodeFileGroup> Children = new Dictionary<string, XcodeFileGroup>();
+		public List<XcodeSourceFile> Files = new List<XcodeSourceFile>();
+		public bool bIsReference;
+	}
 
-			if (bAutomaticSigning)
-			{
-				Content.Append("\t\t\t\t\t\tProvisioningStyle = Automatic;" + ProjectFileGenerator.NewLine);
-			}
-			else
-			{
-				Content.Append("\t\t\t\t\t\tProvisioningStyle = Manual;" + ProjectFileGenerator.NewLine);
-			}
-            Content.Append("\t\t\t\t\t};" + ProjectFileGenerator.NewLine);
-            Content.Append("\t\t\t\t};" + ProjectFileGenerator.NewLine);
-            Content.Append("\t\t\t};" + ProjectFileGenerator.NewLine);
-			Content.Append("\t\t\tbuildConfigurationList = " + ProjectBuildConfigGuid + " /* Build configuration list for PBXProject \"" + TargetName + "\" */;" + ProjectFileGenerator.NewLine);
-			Content.Append("\t\t\tcompatibilityVersion = \"Xcode 8.0\";" + ProjectFileGenerator.NewLine);
-			Content.Append("\t\t\tdevelopmentRegion = English;" + ProjectFileGenerator.NewLine);
-			Content.Append("\t\t\thasScannedForEncodings = 0;" + ProjectFileGenerator.NewLine);
-			Content.Append("\t\t\tknownRegions = (" + ProjectFileGenerator.NewLine);
-			Content.Append("\t\t\t\ten" + ProjectFileGenerator.NewLine);
-			Content.Append("\t\t\t);" + ProjectFileGenerator.NewLine);
-			Content.Append("\t\t\tmainGroup = " + MainGroupGuid + ";" + ProjectFileGenerator.NewLine);
-			Content.Append("\t\t\tproductRefGroup = " + ProductRefGroupGuid + ";" + ProjectFileGenerator.NewLine);
-			Content.Append("\t\t\tprojectDirPath = \"\";" + ProjectFileGenerator.NewLine);
-			Content.Append("\t\t\tprojectRoot = \"\";" + ProjectFileGenerator.NewLine);
-			Content.Append("\t\t\ttargets = (" + ProjectFileGenerator.NewLine);
-			Content.Append("\t\t\t" + TargetGuid + " /* " + TargetName + " */," + ProjectFileGenerator.NewLine);
-			Content.Append("\t\t\t" + BuildTargetGuid + " /* " + BuildTargetName + " */," + ProjectFileGenerator.NewLine);
-			Content.Append("\t\t\t" + IndexTargetGuid + " /* " + IndexTargetName + " */," + ProjectFileGenerator.NewLine);
-			foreach (XcodeExtensionInfo EI in AllExtensions)
-			{
-				Content.Append("\t\t\t" + EI.TargetGuid + " /* " + EI.Name + " */," + ProjectFileGenerator.NewLine);
-			}
-			Content.Append("\t\t\t);" + ProjectFileGenerator.NewLine);
-			Content.Append("\t\t};" + ProjectFileGenerator.NewLine);
+	class UnrealExtensionInfo
+	{
+		public UnrealExtensionInfo(string InName)
+		{
+			Name = InName;
+			TargetDependencyGuid = XcodeProjectFileGenerator.MakeXcodeGuid();
+			TargetProxyGuid = XcodeProjectFileGenerator.MakeXcodeGuid();
+			TargetGuid = XcodeProjectFileGenerator.MakeXcodeGuid();
+			ProductGuid = XcodeProjectFileGenerator.MakeXcodeGuid();
+			ResourceBuildPhaseGuid = XcodeProjectFileGenerator.MakeXcodeGuid();
+			ConfigListGuid = XcodeProjectFileGenerator.MakeXcodeGuid();
+			AllConfigs = new Dictionary<string, UnrealBuildConfig>();
+		}
+
+		public string Name;
+		public string TargetDependencyGuid;
+		public string TargetProxyGuid;
+		public string TargetGuid;
+		public string ProductGuid;
+		public string ResourceBuildPhaseGuid;
+		public string ConfigListGuid;
+		public Dictionary<string, UnrealBuildConfig> AllConfigs;
 
-			Content.Append("/* End PBXProject section */" + ProjectFileGenerator.NewLine + ProjectFileGenerator.NewLine);
+		public string? ConfigurationContents;
+	}
+
+	class XcconfigFile
+	{
+		public string Name;
+		//public string Filename;
+		public string Guid;
+//		public FileReference FileRef;
+		internal StringBuilder Text;
+
+		public XcconfigFile(string ConfigName)
+		{
+			Name = ConfigName;
+//			FileRef = FileReference.Combine(ProjectFile.Directory, "Xcconfigs", $"{ProjectFile.GetFileNameWithoutAnyExtensions()}_{Name.Replace(" ", "")}.xcconfig");
+//			Filename = FileRef.GetFileName();
+			Guid = XcodeProjectFileGenerator.MakeXcodeGuid();
+//			DirectoryReference.CreateDirectory(FileRef.Directory);
+			Text = new StringBuilder();
 		}
 
-		private void AppendContainerItemProxySection(StringBuilder Content, string TargetName, string TargetGuid, string TargetProxyGuid, string ProjectGuid)
+		public void AppendLine(string Line)
 		{
-			Content.Append("/* Begin PBXContainerItemProxy section */" + ProjectFileGenerator.NewLine);
-			Content.Append("\t\t" + TargetProxyGuid + " /* PBXContainerItemProxy */ = {" + ProjectFileGenerator.NewLine);
-			Content.Append("\t\t\tisa = PBXContainerItemProxy;" + ProjectFileGenerator.NewLine);
-			Content.Append("\t\t\tcontainerPortal = " + ProjectGuid + " /* Project object */;" + ProjectFileGenerator.NewLine);
-			Content.Append("\t\t\tproxyType = 1;" + ProjectFileGenerator.NewLine);
-			Content.Append("\t\t\tremoteGlobalIDString = " + TargetGuid + ";" + ProjectFileGenerator.NewLine);
-			Content.Append("\t\t\tremoteInfo = \"" + TargetName + "\";" + ProjectFileGenerator.NewLine);
-			Content.Append("\t\t};" + ProjectFileGenerator.NewLine);
-			Content.Append("/* End PBXContainerItemProxy section */" + ProjectFileGenerator.NewLine + ProjectFileGenerator.NewLine);
+			Text.Append(Line);
+			Text.Append(ProjectFileGenerator.NewLine);
 		}
 
-		private void AppendTargetDependencySection(StringBuilder Content, string TargetName, string TargetGuid, string TargetDependencyGuid, string TargetProxyGuid)
+		public void Write()
 		{
-			Content.Append("/* Begin PBXTargetDependency section */" + ProjectFileGenerator.NewLine);
-			Content.Append("\t\t" + TargetDependencyGuid + " /* PBXTargetDependency */ = {" + ProjectFileGenerator.NewLine);
-			Content.Append("\t\t\tisa = PBXTargetDependency;" + ProjectFileGenerator.NewLine);
-			Content.Append("\t\t\ttarget = " + TargetGuid + " /* " + TargetName + " */;" + ProjectFileGenerator.NewLine);
-			Content.Append("\t\t\ttargetProxy = " + TargetProxyGuid + " /* PBXContainerItemProxy */;" + ProjectFileGenerator.NewLine);
-			Content.Append("\t\t};" + ProjectFileGenerator.NewLine);
-			Content.Append("/* End PBXTargetDependency section */" + ProjectFileGenerator.NewLine + ProjectFileGenerator.NewLine);
+//			FileReference.WriteAllTextIfDifferent(FileRef, Text.ToString());
 		}
+	}
+
+
 
-		// cache for the below function
-		Dictionary<string, IEnumerable<string>> CachedMacProjectArcitectures = new Dictionary<string, IEnumerable<string>>();
+
+	abstract class XcodeProjectNode
+	{
+		// keeps a list of other node this node references, which is used when writing out the whole xcode project file
+		public List<XcodeProjectNode> References = new();
 
 		/// <summary>
-		/// Returns the Mac architectures that should be configured for the provided target. If the target has a project we'll adhere
-		/// to whether it's set as Intel/Universal/Apple unless the type is denied (pretty much just Editor)
-		/// 
-		/// If the target has no project we'll support allow-listed targets for installed builds and all non-editor architectures 
-		/// for source builds. Not all programs are going to compile for Apple Silicon, but being able to build and fail is useful...
+		/// Abstract function the individual node classes must override to write out the node to the project file
 		/// </summary>
-		/// <param name="Config">Build config for the target we're generating</param>
-		/// <param name="InProjectFile">Path to the project file, or null if the target has no project</param>
-		/// <returns></returns>
-		IEnumerable<string> GetSupportedMacArchitectures(XcodeBuildConfig Config, FileReference? InProjectFile)
-		{
-			// All architectures supported
-			IEnumerable<string> AllArchitectures = new[] { MacExports.IntelArchitecture, MacExports.AppleArchitecture};			
+		/// <param name="Content"></param>
+		public abstract void Write(StringBuilder Content);
 
-			// Add a way on the command line of forcing a project file with all architectures (there isn't a good way to let this be
-			// set and checked where we can access it).
-			bool ForceAllArchitectures = Environment.GetCommandLineArgs().Contains("AllArchitectures", StringComparer.OrdinalIgnoreCase);
 
-			if (ForceAllArchitectures)
+		/// <summary>
+		/// Walks the references of the given node to find all nodes of the given type. 
+		/// </summary>
+		/// <typeparam name="T">Parent class of the nodes to return</typeparam>
+		/// <param name="Node">Root node to start with</param>
+		/// <returns>Set of matching nodes</returns>
+		public static IEnumerable<T> GetNodesOfType<T>(XcodeProjectNode Node) where T : XcodeProjectNode
+		{
+			// gather the nodes without recursion
+			LinkedList<XcodeProjectNode> Nodes = new();
+			Nodes.AddLast(Node);
+
+			// pull off the front of the "deque" amd add its references to the back, gather
+			List<XcodeProjectNode> Return = new();
+			while (Nodes.Count() > 0)
 			{
-				return AllArchitectures;
+				XcodeProjectNode Head = Nodes.First();
+				Nodes.RemoveFirst();
+				Head.References.ForEach(x => Nodes.AddLast(x));
+
+				// remember them all 
+				Return.AddRange(Head.References);
 			}
 
-			string TargetName = Config.BuildTarget;
+			// filter down
+			return Return.OfType<T>();
+		}
 
-			// First time seeing this target?
-			if (!CachedMacProjectArcitectures.ContainsKey(TargetName))
-			{
-				// Default to Intel
-				IEnumerable<string> TargetArchitectures = new[] { MacExports.IntelArchitecture };
 
-				// These targets are known to work so are allow-listed
-				bool IsAllowed = MacExports.TargetsAllowedForAppleSilicon.Contains(TargetName, StringComparer.OrdinalIgnoreCase);
+		/// <summary>
+		/// Called by WriteXcconfigFiles, for any node that needs Xcconfig files, the subclass will overrite this to create one or more
+		/// Xcconfig files and add them to the list. WriteXcconfigFiles will then manage those files as needed
+		/// </summary>
+		/// <param name="XcconfigFiles"></param>
+		internal virtual void CreateXcconfigFiles(List<XcconfigFile> XcconfigFiles)
+		{
+		}
 
-				// determine the target architectures based on what's allowed/denied
-				if (IsAllowed)
-				{
-					TargetArchitectures = AllArchitectures;
-				}
-				else
-				{
-					// if this is an unspecified tool/program, default to Intel for installed builds because we know all of that works. 
-					if (Config.ProjectTarget!.TargetRules!.Type == TargetType.Program)
-					{
-						// For misc tools we default to Intel for installed builds because we know all of that works. 
-						TargetArchitectures = Unreal.IsEngineInstalled() ? new[] { MacExports.IntelArchitecture } : AllArchitectures;
-					}
-					else
-					{
-						// For project targets we default to Intel then check the project settings. Note the editor target will have
-						// been denied above already.
-						TargetArchitectures = new[] { MacExports.IntelArchitecture };
+		public static void WriteXcconfigFiles(XcodeProjectNode RootNode, XcodeProjectFile ProjectFile)
+		{
+			List<XcconfigFile> FinalXcconfigFiles = new List<XcconfigFile>();
 
-						ConfigHierarchy EngineIni = ConfigCache.ReadHierarchy(ConfigHierarchyType.Engine, InProjectFile?.Directory, UnrealTargetPlatform.Mac);
-						string TargetArchitecture;
-						string Key = Config.ProjectTarget!.TargetRules!.Type == TargetType.Editor ? "EditorTargetArchitecture" : "TargetArchitecture";
-						if (EngineIni.GetString("/Script/MacTargetPlatform.MacTargetSettings", Key, out TargetArchitecture))
-						{
-							if (TargetArchitecture.Contains("Universal", StringComparison.OrdinalIgnoreCase))
-							{
-								TargetArchitectures = AllArchitectures;
-							}
-							else if (TargetArchitecture.Contains("Intel", StringComparison.OrdinalIgnoreCase))
-							{
-								TargetArchitectures = new[] { MacExports.IntelArchitecture };
-							}
-							else if (TargetArchitecture.Contains("Apple", StringComparison.OrdinalIgnoreCase))
-							{
-								TargetArchitectures = new[] { MacExports.AppleArchitecture };
-							}
-						}
-					}
-				}
+			// recursively walk the nodes
+			WriteXcconfigFilesInternal(FinalXcconfigFiles, RootNode);
 
-				// Cache this so we don't need to keep checking this file
-				CachedMacProjectArcitectures.Add(TargetName, TargetArchitectures);
+			// now write out all of the xcconfig files and add to project
+			foreach (XcconfigFile File in FinalXcconfigFiles)
+			{
+				FileReference FileRef = FileReference.Combine(ProjectFile.ProjectFilePath.Directory, "Xcconfigs", $"{File.Name}.xcconfig");
+				DirectoryReference.CreateDirectory(FileRef.Directory);
+
+				FileReference.WriteAllTextIfDifferent(FileRef, File.Text.ToString());
+
+				ProjectFile.FileCollection.AddFileReference(File.Guid, $"Xcconfigs/{FileRef.GetFileName()}", "test.xcconfig", "\"<group>\"", "Xcconfigs");
 			}
 
-			return CachedMacProjectArcitectures[TargetName];
 		}
 
-		private void WriteProjectXcconfigFiles(List<XcodeBuildConfig> ProjectConfigs, FileReference? ProjectFile, ILogger Logger)
+		private static void WriteXcconfigFilesInternal(List<XcconfigFile> FinalXcconfigFiles, XcodeProjectNode Node)
 		{
-			// Setup all the settings that are shared across all configs
-
-			string UEDir = ConvertPath(Path.GetFullPath(Directory.GetCurrentDirectory() + "../../.."));
-			string SupportedIOSArchitectures = string.Join(" ", this.SupportedIOSArchitectures);
-			string? IOSRunTimeVersion = null;
-			string? TVOSRunTimeVersion = null;
-
-			// shortcut for mac only
-			string? IOSRunTimeDevices = null;
-			string? TVOSRunTimeDevices = null;
+			// allow the node to create optional Xcconfig file(s)
+			Node.CreateXcconfigFiles(FinalXcconfigFiles);
 
-			bool bAutomaticSigning = false;
-			string? UUID_IOS = "";
-			string? UUID_TVOS = "";
-			string? TEAM_IOS = "";
-			string? TEAM_TVOS = "";
-			string? IOS_CERT = "iPhone Developer";
-			string? TVOS_CERT = "iPhone Developer";
-			string IOS_BUNDLE = "";
-			string TVOS_BUNDLE = "";
-				
-			if (InstalledPlatformInfo.IsValidPlatform(UnrealTargetPlatform.IOS, EProjectType.Code))
+			// move on to dependencies
+			foreach (XcodeProjectNode Reference in Node.References)
 			{
-				IOSPlatform IOSPlatform = ((IOSPlatform)UEBuildPlatform.GetBuildPlatform(UnrealTargetPlatform.IOS));
-				IOSProjectSettings ProjectSettings = IOSPlatform.ReadProjectSettings(ProjectFile);
-				IOSProvisioningData ProvisioningData = IOSPlatform.ReadProvisioningData(ProjectSettings, bForDistribution);
-				IOSRunTimeVersion = ProjectSettings.RuntimeVersion;
-				IOSRunTimeDevices = ProjectSettings.RuntimeDevices;
-				bAutomaticSigning = ProjectSettings.bAutomaticSigning;
-				if (!bAutomaticSigning)
-				{
-					UUID_IOS = ProvisioningData.MobileProvisionUUID;
-					IOS_CERT = ProvisioningData.SigningCertificate;
-				}
-				TEAM_IOS = ProvisioningData.TeamUUID;
-				IOS_BUNDLE = ProjectSettings.BundleIdentifier;
+				WriteXcconfigFilesInternal(FinalXcconfigFiles, Reference);
 			}
+		}
 
-			if (InstalledPlatformInfo.IsValidPlatform(UnrealTargetPlatform.TVOS, EProjectType.Code))
+		/// <summary>
+		/// THhis will walk the node reference tree and call WRite on each node to add all needed nodes to the xcode poject file
+		/// </summary>
+		/// <param name="Content"></param>
+		/// <param name="Node"></param>
+		public static void WriteNodeAndReferences(StringBuilder Content, XcodeProjectNode Node)
+		{
+			// write the node into the xcode project file
+			Node.Write(Content);
+
+			foreach (XcodeProjectNode Reference in Node.References)
 			{
-				TVOSPlatform TVOSPlatform = ((TVOSPlatform)UEBuildPlatform.GetBuildPlatform(UnrealTargetPlatform.TVOS));
-				TVOSProjectSettings ProjectSettings = TVOSPlatform.ReadProjectSettings(ProjectFile);
-				TVOSProvisioningData ProvisioningData = TVOSPlatform.ReadProvisioningData(ProjectSettings, bForDistribution);
-				TVOSRunTimeVersion = ProjectSettings.RuntimeVersion;
-				TVOSRunTimeDevices = ProjectSettings.RuntimeDevices;
-				if (!bAutomaticSigning)
-				{
-					UUID_TVOS = ProvisioningData.MobileProvisionUUID;
-					TVOS_CERT = ProvisioningData.SigningCertificate;
-				}
-				TEAM_TVOS = ProvisioningData.TeamUUID;
-				TVOS_BUNDLE = ProjectSettings.BundleIdentifier;
-			}
+				WriteNodeAndReferences(Content, Reference);
+			}	
+		}
+	}
 
-			//#jira UE-50382 Xcode Address Sanitizer feature does not work on iOS
-			// address sanitizer dylib loader depends on the SDKROOT parameter. For macosx or default (missing, translated as macosx), the path is incorrect for iphone/appletv
-			ProjectXcconfig.AppendLine($"SDKROOT[sdk=macosx*] = macosx");
-			ProjectXcconfig.AppendLine($"SDKROOT[sdk=iphoneos] = iphoneos");
-			ProjectXcconfig.AppendLine($"SDKROOT[sdk=appletvos*] = appletvos");
-			ProjectXcconfig.AppendLine("USE_HEADERMAP = NO");
-			ProjectXcconfig.AppendLine("ONLY_ACTIVE_ARCH = YES");
 
-			if (bAutomaticSigning)
-			{
-				ProjectXcconfig.AppendLine("CODE_SIGN_STYLE = Automatic");
-			}
+	class XcodeDependency : XcodeProjectNode
+	{
+		public XcodeTarget Target;
+		public string Guid = XcodeProjectFileGenerator.MakeXcodeGuid();
+		public string ProxyGuid = XcodeProjectFileGenerator.MakeXcodeGuid();
+		public string ProjectGuid;
 
-			if (IOSRunTimeVersion != null)
-			{
-				ProjectXcconfig.AppendLine($"VALID_ARCHS[sdk=iphoneos*] = {SupportedIOSArchitectures}");
+
+		public XcodeDependency(XcodeTarget Target, string ProjectGuid)
+		{
+			this.Target = Target;
+			this.ProjectGuid = ProjectGuid;
+
+			References.Add(Target);
+		}
+
+		public override void Write(StringBuilder Content)
+		{
+			Content.WriteLine("/* Begin PBXContainerItemProxy section */");
+			Content.WriteLine($"\t\t{ProxyGuid} /* PBXContainerItemProxy */ = {{");
+			Content.WriteLine("\t\t\tisa = PBXContainerItemProxy;");
+			Content.WriteLine($"\t\t\tcontainerPortal = {ProjectGuid} /* Project object */;");
+			Content.WriteLine("\t\t\tproxyType = 1;");
+			Content.WriteLine($"\t\t\tremoteGlobalIDString = {Target.Guid};");
+			Content.WriteLine($"\t\t\tremoteInfo = \"{Target.Name}\";");
+			Content.WriteLine("\t\t};");
+			Content.WriteLine("/* End PBXContainerItemProxy section */");
+			Content.WriteLine("");
+
+			Content.WriteLine("/* Begin PBXTargetDependency section */");
+			Content.WriteLine($"\t\t{Guid} /* PBXTargetDependency */ = {{");
+			Content.WriteLine("\t\t\tisa = PBXTargetDependency;");
+			Content.WriteLine($"\t\t\ttarget = {Target.Guid} /* {Target.Name} */;");
+			Content.WriteLine($"\t\t\ttargetProxy = {ProxyGuid} /* PBXContainerItemProxy */;");
+			Content.WriteLine("\t\t};");
+			Content.WriteLine("/* End PBXTargetDependency section */");
+		}
+	}
+
+	class XcodeBuildPhase : XcodeProjectNode
+	{
+		public string Name;
+		public string Guid = XcodeProjectFileGenerator.MakeXcodeGuid();
+		private List<XcodeSourceFile> Items = new();
+
+		public XcodeBuildPhase(string Name)
+		{
+			this.Name = Name;
+		}
+
+		public void AddFile(XcodeSourceFile File)
+		{
+			Items.Add(File);
+		}
+
+		public override void Write(StringBuilder Content)
+		{
+			Content.WriteLine("/* Begin PBXSourcesBuildPhase section */");
+			Content.WriteLine($"\t\t{Guid} = {{");
+			Content.WriteLine("\t\t\tisa = PBXSourcesBuildPhase;");
+			Content.WriteLine("\t\t\tbuildActionMask = 2147483647;");
+			Content.WriteLine("\t\t\tfiles = (");
+			foreach (XcodeSourceFile File in Items)
+			{
+				Content.WriteLine($"\t\t\t\t{File.FileGuid} /* {File.Reference.GetFileName()} in {Name} */,");
+			}
+			Content.WriteLine("\t\t\t);");
+			Content.WriteLine("\t\t\trunOnlyForDeploymentPostprocessing = 0;");
+			Content.WriteLine("\t\t};");
+			Content.WriteLine("/* End PBXSourcesBuildPhase section */");
+		}
+	}
+
+	class XcodeBuildConfig : XcodeProjectNode
+	{
+		public string Guid = XcodeProjectFileGenerator.MakeXcodeGuid();
+		public UnrealBuildConfig Info;
+		public XcconfigFile? Xcconfig = null;
+
+		public XcodeBuildConfig(UnrealBuildConfig Info)
+		{
+			this.Info = Info;
+		}
+
+		public override void Write(StringBuilder Content)
+		{
+			Content.WriteLine(2, $"{Guid} /* {Info.DisplayName} */ = {{");
+			Content.WriteLine(3,	"isa = XCBuildConfiguration;");
+			if (Xcconfig != null)
+			{
+				Content.WriteLine(3, $"baseConfigurationReference = {Xcconfig.Guid} /* {Xcconfig.Name}.xcconfig */;");
+			}
+			Content.WriteLine(3,	"buildSettings = {");
+			Content.WriteLine(3,	"};");
+			Content.WriteLine(3,	$"name = \"{Info.DisplayName}\";");
+			Content.WriteLine(2, "};");
+		}
+	}
+	
+	class XcodeBuildConfigList : XcodeProjectNode
+	{
+		public string Guid = XcodeProjectFileGenerator.MakeXcodeGuid();
+		public string TargetName;
+
+		public List<XcodeBuildConfig> BuildConfigs = new();
+
+		public XcodeBuildConfigList(string TargetName, List<UnrealBuildConfig> BuildConfigInfos)
+		{
+			if (BuildConfigInfos.Count == 0)
+			{
+				throw new BuildException("Created a XcodeBuildConfigList with no BuildConfigs. This likely means a target was created too early");
+			}
+
+			this.TargetName = TargetName;
+
+			// create build config objects for each info passed in, and them as references
+			BuildConfigs = BuildConfigInfos.Select(x => new XcodeBuildConfig(x)).ToList();
+			References.AddRange(BuildConfigs);
+		}
+
+		public override void Write(StringBuilder Content)
+		{
+			// figure out the default configuration to use
+			string Default = BuildConfigs.Any(x => x.Info.DisplayName.Contains(" Editor")) ? "Development Editor" : "Development";
+
+			Content.WriteLine(2, $"{Guid} /* Build configuration list for target {TargetName} */ = {{");
+			Content.WriteLine(3,	"isa = XCConfigurationList;");
+			Content.WriteLine(3,	"buildConfigurations = (");
+			foreach (XcodeBuildConfig Config in BuildConfigs)
+			{
+				Content.WriteLine(4,		$"{Config.Guid} /* {Config.Info.DisplayName} */,");
+			}
+			Content.WriteLine(3,	");");
+			Content.WriteLine(3,	"defaultConfigurationIsVisible = 0;");
+			Content.WriteLine(3,	$"defaultConfigurationName = \"{Default}\";");
+			Content.WriteLine(2, "};");
+
+			// now write the ba
+		}
+	}
+
+	class XcodeTarget : XcodeProjectNode
+	{
+		public enum Type
+		{
+			Run_App,
+			Run_Tool,
+			Build,
+			Index,
+			Extension,
+		}
+
+		// Guid for this target
+		public string Guid = XcodeProjectFileGenerator.MakeXcodeGuid();
+		string TargetAppGuid = XcodeProjectFileGenerator.MakeXcodeGuid();
+
+		// com.apple.product-type.application, etc
+		string ProductType;
+
+		// xcode target type name
+		string TargetTypeName;
+		Type TargetType;
+
+		// UE5_Build, EngineTest_SubIndex1, etc
+		public string Name;
+
+		// the config shared by all BuildConfigs
+//		public XcconfigFile Xcconfig;
+
+		// list of build configs this target supports (for instance, the Index target only indexes a Development config) 
+		public XcodeBuildConfigList? BuildConfigList;
+
+		// dependencies for this target
+		public List<XcodeDependency> Dependencies = new List<XcodeDependency>();
+
+		// build phases for this target (source, resource copying, etc)
+		public List<XcodeBuildPhase> BuildPhases = new List<XcodeBuildPhase>();
+
+		private FileReference? GameProject;
+
+		public XcodeTarget(Type Type, UnrealData UnrealData, string? OverrideName=null)
+		{
+			GameProject = UnrealData.UProjectFileLocation;
+
+			string ConfigName;
+			TargetType = Type;
+			switch (Type)
+			{
+				case Type.Run_App:
+					ProductType = "com.apple.product-type.application";
+					TargetTypeName = "PBXNativeTarget";
+					ConfigName = "Run";
+					break;
+				case Type.Run_Tool:
+					ProductType = "com.apple.product-type.tool";
+					TargetTypeName = "PBXNativeTarget";
+					ConfigName = "Run";
+					break;
+				case Type.Build:
+					ProductType = "com.apple.product-type.library.static";
+					TargetTypeName = "PBXLegacyTarget";
+					ConfigName = "Build";
+					break;
+				case Type.Index:
+					ProductType = "com.apple.product-type.library.static";
+					TargetTypeName = "PBXNativeTarget";
+					ConfigName = "Index";
+					break;
+				case Type.Extension:
+					ProductType = "com.apple.product-type.app-extension.messages-sticker-pack";
+					TargetTypeName = "PBXNativeTarget";
+					ConfigName = "Extension";
+					break;
+				default:
+					throw new BuildException($"Unhandled target type {Type}");
+			}
+
+			// set up names
+			ConfigName = (OverrideName == null) ? ConfigName : OverrideName;
+			Name = $"{UnrealData.ProductName}_{ConfigName}";
+		}
+
+		public void AddDependency(XcodeTarget Target, XcodeProject Project)
+		{
+			XcodeDependency Dependency = new XcodeDependency(Target, Project.Guid);
+			Dependencies.Add(Dependency);
+			References.Add(Dependency);
+		}
+
+		public override void Write(StringBuilder Content)
+		{
+			Content.WriteLine($"/* Begin {TargetType} section */");
+
+			Content.WriteLine(2, $"{Guid} /* {Name} */ = {{");
+			Content.WriteLine(3,	$"isa = {TargetTypeName};");
+				
+			Content.WriteLine(3,	$"buildConfigurationList = {BuildConfigList!.Guid} /* Build configuration list for {TargetTypeName} \"{Name}\" */;");
+
+			if (TargetType == Type.Build)
+			{
+				// get paths to Unreal bits to be able ro tun UBT
+				string UProjectParam = GameProject == null ? "" : $"{GameProject.FullName.Replace(" ", "\\ ")}";
+				string UEDir = XcodeFileCollection.ConvertPath(Path.GetFullPath(Directory.GetCurrentDirectory() + "../../.."));
+				string BuildToolPath = UEDir + "/Engine/Build/BatchFiles/Mac/XcodeBuild.sh";
+
+				// insert elements to call UBT when building
+				Content.WriteLine(3,	$"buildArgumentsString = \"$(ACTION) $(UE_BUILD_TARGET_NAME) $(PLATFORM_NAME) $(UE_BUILD_TARGET_CONFIG) {UProjectParam}\";");
+				Content.WriteLine(3,	$"buildToolPath = \"{BuildToolPath}\";");
+				Content.WriteLine(3,	$"buildWorkingDirectory = \"{UEDir}\";");
+			}
+			Content.WriteLine(3,	"buildPhases = (");
+			foreach (XcodeBuildPhase BuildPhase in BuildPhases)
+			{
+				Content.WriteLine(4,		$"{BuildPhase.Guid} /* {BuildPhase.Name} */,");
+			}
+			Content.WriteLine(3,	");");
+			Content.WriteLine(3,	"dependencies = (");
+			foreach (XcodeDependency Dependency in Dependencies)
+			{
+				Content.WriteLine(4,		$"{Dependency.Guid} /* {Dependency.Target.Name} */,");
+			}
+			Content.WriteLine(3,	");");
+			Content.WriteLine(3,	$"name = \"{Name}\";");
+			Content.WriteLine(3,	"passBuildSettingsInEnvironment = 1;");
+			Content.WriteLine(3,	$"productType = \"{ProductType}\";");
+			WriteExtraTargetProperties(Content);
+			Content.WriteLine(2, "};");
+
+			Content.WriteLine($"/* End {TargetType} section */");
+		}
+
+		/// <summary>
+		/// Let subclasses add extra properties into this target section
+		/// </summary>
+		protected virtual void WriteExtraTargetProperties(StringBuilder Content)
+		{
+			// nothing by default
+		}
+	}
+
+	class XcodeRunTarget : XcodeTarget
+	{
+		private string ProductGuid = XcodeProjectFileGenerator.MakeXcodeGuid();
+		private string ProductName;
+
+		public XcodeRunTarget(XcodeProject Project)
+			: base(Project.UnrealData.bIsAppBundle ? XcodeTarget.Type.Run_App : XcodeTarget.Type.Run_Tool, Project.UnrealData)
+		{
+			BuildConfigList = new XcodeBuildConfigList(Name, Project.UnrealData.AllConfigs);
+			References.Add(BuildConfigList);
+
+			// add the Product item to the project to be visible in left pane
+			ProductName = Project.UnrealData.ProductName;
+			Project.FileCollection.AddFileReference(ProductGuid, ProductName, Project.UnrealData.bIsAppBundle ? "wrapper.application" : "\"compiled.mach-o.executable\"", "BUILT_PRODUCTS_DIR", "Products");
+
+
+			// if we are run-only, then skip some stuff
+			if (!XcodeProjectFileGenerator.bGeneratingRunIOSProject)
+			{
+				// create a biuld target only if we have source files to build
+				if (Project.UnrealData.BatchedFiles.Count != 0)
+				{
+					XcodeBuildTarget BuildTarget = new XcodeBuildTarget(Project.UnrealData);
+					AddDependency(BuildTarget, Project);
+				}
+
+				// add a target for each extension that this depends on
+				foreach (UnrealExtensionInfo Extension in Project.UnrealData.AllExtensions)
+				{
+					XcodeExtensionTarget Target = new XcodeExtensionTarget(Project.UnrealData, Extension);
+					AddDependency(Target, Project);
+				}
+			}
+		}
+
+		protected override void WriteExtraTargetProperties(StringBuilder Content)
+		{
+			Content.WriteLine($"\t\t\tproductReference = {ProductGuid};");
+			Content.WriteLine($"\t\t\tproductName = \"{ProductName}\";");
+		}
+
+		internal override void CreateXcconfigFiles(List<XcconfigFile> XcconfigFiles)
+		{
+			XcconfigFile Xcconfig = new XcconfigFile($"{Name}");
+			XcconfigFiles.Add(Xcconfig);
+
+			Xcconfig.AppendLine("GENERATE_INFOPLIST_FILE[sdk=macosx*] = YES // Xcode 14 code-signing fix for development");
+
+			// #jira UE-143619: Pre Monterey macOS requires this option for a packaged app to run on iOS15 due to new code signature format. Could be removed once Monterey is miniuS.
+			Xcconfig.AppendLine("OTHER_CODE_SIGN_FLAGS = --generate-entitlement-der");
+
+			Xcconfig.AppendLine($"MACOSX_DEPLOYMENT_TARGET = {MacToolChain.Settings.MacOSVersion};");
+			Xcconfig.AppendLine("INFOPLIST_OUTPUT_FORMAT = xml");
+			Xcconfig.AppendLine("COMBINE_HIDPI_IMAGES = YES");
+
+			// hook up each buildconfig to this Xcconfig 
+			BuildConfigList!.BuildConfigs.ForEach(x => x.Xcconfig = Xcconfig);
+		}
+	}
+
+	class XcodeBuildTarget : XcodeTarget
+	{
+		public XcodeBuildTarget(UnrealData UnrealData)
+			: base(XcodeTarget.Type.Build, UnrealData)
+		{
+			BuildConfigList = new XcodeBuildConfigList(Name, UnrealData.AllConfigs);
+			References.Add(BuildConfigList);
+		}
+	}
+
+	class XcodeIndexTarget : XcodeTarget
+	{
+		private UnrealBatchedFiles? SoloTargetBatch;
+
+		// just take the Project since it has everything we need, and is needed when adding target dependencies
+		public XcodeIndexTarget(XcodeProject Project)
+			: base(XcodeTarget.Type.Index, Project.UnrealData)
+		{
+			// do we have an editor? if so, use this Development Editor as the config to index, otherwise, use Development
+			string IndexConfig = Project.UnrealData.AllConfigs.Any(x => x.DisplayName == "Development Editor") ? "Development Editor" : "Development";
+			BuildConfigList = new XcodeBuildConfigList(Name, Project.UnrealData.AllConfigs.Where(x => x.DisplayName == IndexConfig).ToList());
+			References.Add(BuildConfigList);
+
+
+			// now look to see if all files are buildable with one batch, in which case we don't need subtargets, or if we have multiple batches, then 
+			// make a subtarget for each batch of unique build settings
+			if (Project.UnrealData.BatchedFiles.Count == 1)
+			{
+				SoloTargetBatch = Project.UnrealData.BatchedFiles[0];
+
+				XcodeBuildPhase BuildPhase = new XcodeBuildPhase("Sources");
+				BuildPhases.Add(BuildPhase);
+				References.Add(BuildPhase);
+
+				foreach (XcodeSourceFile File in Project.FileCollection.BuildableFiles)
+				{
+					BuildPhase.AddFile(File);
+				}
+			}
+			else
+			{
+				// find the config infos this target is set to use, and copy those in to the subtargets
+				List<UnrealBuildConfig> Configs = BuildConfigList.BuildConfigs.Select(x => x.Info).ToList();
+
+				for (int Index = 0; Index < Project.UnrealData.BatchedFiles.Count; Index++)
+				{
+					XcodeIndexSubTarget SubTarget = new XcodeIndexSubTarget(Project.UnrealData, Index, Configs);
+					AddDependency(SubTarget, Project);
+				}
+			}
+		}
+
+
+		internal override void CreateXcconfigFiles(List<XcconfigFile> XcconfigFiles)
+		{
+			// name the xcconfig after the target name
+			XcconfigFile Xcconfig = new XcconfigFile(Name);
+			XcconfigFiles.Add(Xcconfig);
+
+			// write out settings that apply whether or not we have subtargets, which #include this one, or no subtargets, and we write out more
+			// @todo move tis to the subtarget and remember it from the Module
+			Xcconfig.AppendLine("CLANG_CXX_LANGUAGE_STANDARD = c++17");
+			Xcconfig.AppendLine("GCC_WARN_CHECK_SWITCH_STATEMENTS = NO");
+			Xcconfig.AppendLine("GCC_PRECOMPILE_PREFIX_HEADER = YES");
+			Xcconfig.AppendLine("GCC_OPTIMIZATION_LEVEL = 0");
+
+			// if we have folded a file batch into this main target, write out the settings into this Xcconfig
+			if (SoloTargetBatch != null)
+			{
+				// ask the singl sub target to write out the contents of our Xcconfig file since we are folding them up into this target
+				SoloTargetBatch.WriteXcconfigSettings(Xcconfig.Text);
+			}
+
+			// hook up each buildconfig to this Xcconfig 
+			BuildConfigList!.BuildConfigs.ForEach(x => x.Xcconfig = Xcconfig);
+		}
+	}
+
+	class XcodeIndexSubTarget : XcodeTarget
+	{
+		public string TargetGuid = XcodeProjectFileGenerator.MakeXcodeGuid();
+		public string BuildPhaseGuid = XcodeProjectFileGenerator.MakeXcodeGuid();
+		public string ConfigListGuid = XcodeProjectFileGenerator.MakeXcodeGuid();
+
+		private UnrealBatchedFiles BatchedFiles;
+
+		public XcodeIndexSubTarget(UnrealData UnrealData, int SubTargetIndex, List<UnrealBuildConfig> Configs)
+			: base(Type.Index, UnrealData, $"SubIndex{SubTargetIndex}")
+		{
+			// get the settings and list of files for this subtarget
+			BatchedFiles = UnrealData.BatchedFiles[SubTargetIndex];
+
+			// set the buildconfig list in the SubTarget
+			BuildConfigList = new XcodeBuildConfigList(Name, Configs);
+			References.Add(BuildConfigList);
+
+			// create the build phase for this target for all files in the batch
+			XcodeBuildPhase BuildPhase = new XcodeBuildPhase("Sources");
+			BuildPhases.Add(BuildPhase);
+			References.Add(BuildPhase);
+
+			// and add each file to it
+			foreach (XcodeSourceFile File in BatchedFiles.Files)
+			{
+				BuildPhases[0].AddFile(File);
+			}
+		}
+
+		internal override void CreateXcconfigFiles(List<XcconfigFile> XcconfigFiles)
+		{
+			// name the xcconfig after the target name
+			XcconfigFile Xcconfig = new XcconfigFile(Name);
+			XcconfigFiles.Add(Xcconfig);
+
+			Xcconfig.AppendLine($"#include \"{Name.Split('_')[0]}_Index.xcconfig\"");
+
+			BatchedFiles.WriteXcconfigSettings(Xcconfig.Text);
+		
+			// hook up each buildconfig to this Xcconfig 
+			BuildConfigList!.BuildConfigs.ForEach(x => x.Xcconfig = Xcconfig);
+		}
+	}
+
+	class XcodeExtensionTarget : XcodeTarget
+	{
+		private UnrealExtensionInfo ExtensionInfo;
+
+		public XcodeExtensionTarget(UnrealData UnrealData, UnrealExtensionInfo ExtensionInfo)
+			: base(Type.Extension, UnrealData)
+		{
+			this.ExtensionInfo = ExtensionInfo;
+
+			BuildConfigList = new XcodeBuildConfigList($"{UnrealData.ProductName}_{ExtensionInfo.Name}", UnrealData.AllConfigs.Where(x => x.DisplayName == "Development").ToList());
+		}
+
+		protected override void WriteExtraTargetProperties(StringBuilder Content)
+		{
+			Content.WriteLine($"\t\t\tproductReference = {ExtensionInfo.ProductGuid};");
+			Content.WriteLine($"\t\t\tproductName = \"{ExtensionInfo.Name}\";");
+		}
+	}
+
+
+	class XcodeProject : XcodeProjectNode
+	{
+		// the blob of data coming from unreal that we can pass around
+		public UnrealData UnrealData;
+
+		// container for all files and groups
+		public XcodeFileCollection FileCollection;
+
+		// Guid for the project node
+		public string Guid = XcodeProjectFileGenerator.MakeXcodeGuid();
+		private string ProvisioningStyle;
+
+		private XcodeRunTarget RunTarget;
+
+		public XcodeBuildConfigList ProjectBuildConfigs;
+
+		public XcodeProject(UnrealData UnrealData, XcodeFileCollection FileCollection)
+		{
+			this.UnrealData = UnrealData;
+			this.FileCollection = FileCollection;
+
+			ProvisioningStyle = UnrealData.bUseAutomaticSigning ? "Automatic" : "Manual";
+			RunTarget = new XcodeRunTarget(this);
+			References.Add(RunTarget);
+
+			ProjectBuildConfigs = new XcodeBuildConfigList(UnrealData.ProductName, UnrealData.AllConfigs);
+			References.Add(ProjectBuildConfigs);
+
+			// make an indexing target if we aren't just a run-only project, and it has buildable source files
+			if (!XcodeProjectFileGenerator.bGeneratingRunIOSProject && UnrealData.BatchedFiles.Count != 0)
+			{				
+				// index isn't a dependency of run, it's simply a target that xcode will find to index from
+				XcodeIndexTarget IndexTarget = new XcodeIndexTarget(this);
+				References.Add(IndexTarget);
+			}
+		}
+
+
+		public override void Write(StringBuilder Content)
+		{
+			string RootGroupGuid = XcodeFileCollection.GetRootGroupGuid(FileCollection.Groups);
+
+
+			Content.WriteLine("/* Begin PBXProject section */");
+
+			Content.WriteLine(2, $"{Guid} /* Project object */ = {{");
+			Content.WriteLine(3,	"isa = PBXProject;");
+			Content.WriteLine(3,	"attributes = {");
+			Content.WriteLine(4,		"LastUpgradeCheck = 2000;");
+			Content.WriteLine(4,		"ORGANIZATIONNAME = \"Epic Games, Inc.\";");
+			Content.WriteLine(4,		"TargetAttributes = {");
+			Content.WriteLine(5,			$"{RunTarget.Guid} = {{");
+			Content.WriteLine(6,				$"ProvisioningStyle = {ProvisioningStyle};");
+			Content.WriteLine(5,			"};");
+			Content.WriteLine(4,		"};");
+			Content.WriteLine(3,	"};");
+			Content.WriteLine(3,	$"buildConfigurationList = {ProjectBuildConfigs.Guid} /* Build configuration list for PBXProject \"{ProjectBuildConfigs.TargetName}\" */;");
+			Content.WriteLine(3,	"compatibilityVersion = \"Xcode 8.0\";");
+			Content.WriteLine(3,	"developmentRegion = English;");
+			Content.WriteLine(3,	"hasScannedForEncodings = 0;");
+			Content.WriteLine(3,	"knownRegions = (");
+			Content.WriteLine(4,		"en");
+			Content.WriteLine(3,	");");
+			Content.WriteLine(3,	$"mainGroup = {RootGroupGuid};");
+			Content.WriteLine(3,	$"productRefGroup = {FileCollection.GetProductGroupGuid()};");
+			Content.WriteLine(3,	"projectDirPath = \"\";");
+			Content.WriteLine(3,	"projectRoot = \"\";");
+			Content.WriteLine(3,	"targets = (");
+			foreach (XcodeTarget Target in XcodeProjectNode.GetNodesOfType<XcodeTarget>(this))
+			{
+				Content.WriteLine(4, $"{Target.Guid} /* {Target.Name} */,");
+			}
+			Content.WriteLine(3, ");");
+			Content.WriteLine(2, "};");
+
+			Content.WriteLine("/* End PBXProject section */");
+		}
+
+
+		internal override void CreateXcconfigFiles(List<XcconfigFile> XcconfigFiles)
+		{
+			XcconfigFile ProjectXcconfig = new XcconfigFile($"{UnrealData.ProductName}_Project");
+			XcconfigFiles.Add(ProjectXcconfig);
+
+			string UEDir = XcodeFileCollection.ConvertPath(Path.GetFullPath(Directory.GetCurrentDirectory() + "../../.."));
+			string SupportedIOSArchitectures = string.Join(" ", UnrealData.SupportedIOSArchitectures);
+			string? IOSRunTimeVersion = null;
+			string? TVOSRunTimeVersion = null;
+
+			// shortcut for mac only
+			string? IOSRunTimeDevices = null;
+			string? TVOSRunTimeDevices = null;
+
+			bool bAutomaticSigning = false;
+			string? UUID_IOS = "";
+			string? UUID_TVOS = "";
+			string? TEAM_IOS = "";
+			string? TEAM_TVOS = "";
+			string? IOS_CERT = "iPhone Developer";
+			string? TVOS_CERT = "iPhone Developer";
+			string IOS_BUNDLE = "";
+			string TVOS_BUNDLE = "";
+
+			if (UnrealData.IOSProjectSettings != null)
+			{
+				IOSRunTimeVersion = UnrealData.IOSProjectSettings.RuntimeVersion;
+				IOSRunTimeDevices = UnrealData.IOSProjectSettings.RuntimeDevices;
+				bAutomaticSigning = UnrealData.IOSProjectSettings.bAutomaticSigning;
+				if (!bAutomaticSigning)
+				{
+					UUID_IOS = UnrealData.IOSProvisioningData!.MobileProvisionUUID;
+					IOS_CERT = UnrealData.IOSProvisioningData!.SigningCertificate;
+				}
+				TEAM_IOS = UnrealData.IOSProvisioningData!.TeamUUID;
+				IOS_BUNDLE = UnrealData.IOSProjectSettings.BundleIdentifier;
+			}
+
+			if (UnrealData.TVOSProjectSettings != null)
+			{
+				TVOSRunTimeVersion = UnrealData.TVOSProjectSettings.RuntimeVersion;
+				TVOSRunTimeDevices = UnrealData.TVOSProjectSettings.RuntimeDevices;
+				if (!bAutomaticSigning)
+				{
+					UUID_TVOS = UnrealData.TVOSProvisioningData!.MobileProvisionUUID;
+					TVOS_CERT = UnrealData.TVOSProvisioningData!.SigningCertificate;
+				}
+				TEAM_TVOS = UnrealData.TVOSProvisioningData!.TeamUUID;
+				TVOS_BUNDLE = UnrealData.TVOSProjectSettings.BundleIdentifier;
+			}
+
+			//#jira UE-50382 Xcode Address Sanitizer feature does not work on iOS
+			// address sanitizer dylib loader depends on the SDKROOT parameter. For macosx or default (missing, translated as macosx), the path is incorrect for iphone/appletv
+			ProjectXcconfig.AppendLine($"SDKROOT[sdk=macosx*] = macosx");
+			ProjectXcconfig.AppendLine($"SDKROOT[sdk=iphone*] = iphoneos");
+			ProjectXcconfig.AppendLine($"SDKROOT[sdk=appletvos*] = appletvos");
+			ProjectXcconfig.AppendLine("USE_HEADERMAP = NO");
+			ProjectXcconfig.AppendLine("ONLY_ACTIVE_ARCH = YES");
+
+			if (bAutomaticSigning)
+			{
+				ProjectXcconfig.AppendLine("CODE_SIGN_STYLE = Automatic");
+			}
+
+			if (IOSRunTimeVersion != null)
+			{
+				ProjectXcconfig.AppendLine($"VALID_ARCHS[sdk=iphone*] = {SupportedIOSArchitectures}");
 				ProjectXcconfig.AppendLine($"IPHONEOS_DEPLOYMENT_TARGET = {IOSRunTimeVersion}");
-				ProjectXcconfig.AppendLine($"TARGETED_DEVICE_FAMILY[sdk=iphoneos*] = {IOSRunTimeDevices}");
+				ProjectXcconfig.AppendLine($"TARGETED_DEVICE_FAMILY[sdk=iphone*] = {IOSRunTimeDevices}");
 				if (!string.IsNullOrEmpty(TEAM_IOS))
 				{
-					ProjectXcconfig.AppendLine($"DEVELOPMENT_TEAM[sdk=iphoneos*] = {TEAM_IOS}");
+					ProjectXcconfig.AppendLine($"DEVELOPMENT_TEAM[sdk=iphone*] = {TEAM_IOS}");
 				}
-				ProjectXcconfig.AppendLine($"CODE_SIGN_IDENTITY[sdk=iphoneos*] = {IOS_CERT}");
+				ProjectXcconfig.AppendLine($"CODE_SIGN_IDENTITY[sdk=iphone*] = {IOS_CERT}");
 				if (!bAutomaticSigning && !string.IsNullOrEmpty(UUID_IOS))
 				{
-					ProjectXcconfig.AppendLine($"PROVISIONING_PROFILE_SPECIFIER[sdk=iphoneos*] = {UUID_IOS}");
+					ProjectXcconfig.AppendLine($"PROVISIONING_PROFILE_SPECIFIER[sdk=iphone*] = {UUID_IOS}");
 				}
-				if (ProjectFile != null)
+				if (UnrealData.UProjectFileLocation != null)
 				{
-					ProjectXcconfig.AppendLine($"PRODUCT_BUNDLE_IDENTIFIER[sdk=iphoneos*] = {IOS_BUNDLE}");
+					ProjectXcconfig.AppendLine($"PRODUCT_BUNDLE_IDENTIFIER[sdk=iphone*] = {IOS_BUNDLE}");
 				}
 			}
-			if (TVOSRunTimeVersion != null)
+			if (TVOSRunTimeVersion != null)
+			{
+				ProjectXcconfig.AppendLine($"VALID_ARCHS[sdk=appletvos*] = {SupportedIOSArchitectures}");
+				ProjectXcconfig.AppendLine($"IPHONEOS_DEPLOYMENT_TARGET = {TVOSRunTimeVersion}");
+				ProjectXcconfig.AppendLine($"TARGETED_DEVICE_FAMILY[sdk=appletvos*] = {TVOSRunTimeDevices}");
+				if (!string.IsNullOrEmpty(TEAM_TVOS))
+				{
+					ProjectXcconfig.AppendLine($"DEVELOPMENT_TEAM[sdk=iphone*] = {TEAM_TVOS}");
+				}
+				ProjectXcconfig.AppendLine($"CODE_SIGN_IDENTITY[sdk=appletvos*] = {TVOS_CERT}");
+				if (!bAutomaticSigning && !string.IsNullOrEmpty(UUID_TVOS))
+				{
+					ProjectXcconfig.AppendLine($"PROVISIONING_PROFILE_SPECIFIER[sdk=appletvos*] = {UUID_TVOS}");
+				}
+				if (UnrealData.UProjectFileLocation != null)
+				{
+					ProjectXcconfig.AppendLine($"PRODUCT_BUNDLE_IDENTIFIER[sdk=appletvos*] = {TVOS_BUNDLE}");
+				}
+			}
+
+			// Now for each config write out the specific settings
+
+			DirectoryReference? GameDir = UnrealData.UProjectFileLocation?.Directory;
+			string? GamePath = GameDir != null ? XcodeFileCollection.ConvertPath(GameDir.FullName) : null;
+
+			// Get Mac architectures supported by this project
+			foreach (UnrealBuildConfig Config in UnrealData.AllConfigs)
+			{
+				// make a new xcconfig for this configuration, without spaces (Debug Editor -> DebugEditor)
+				XcconfigFile Xcconfig = new XcconfigFile($"{UnrealData.ProductName}_{Config.DisplayName.Replace(" ", "")}");
+				XcconfigFiles.Add(Xcconfig);
+
+				// hook up the Buildconfig that matches this info to this xcconfig file
+				ProjectBuildConfigs!.BuildConfigs.First(x => x.Info == Config).Xcconfig = Xcconfig;
+
+				Xcconfig.AppendLine("// pull in the shared settings for all configs");
+				Xcconfig.AppendLine($"#include \"{UnrealData.ProductName}_Project.xcconfig\"");
+				Xcconfig.AppendLine("");
+
+				bool bIsUnrealGame = Config.BuildTarget.Equals("UnrealGame", StringComparison.InvariantCultureIgnoreCase);
+				bool bIsUnrealClient = Config.BuildTarget.Equals("UnrealClient", StringComparison.InvariantCultureIgnoreCase);
+
+				// Setup Mac stuff
+				FileReference MacExecutablePath = Config.MacExecutablePath!;
+				string MacExecutableDir = XcodeFileCollection.ConvertPath(MacExecutablePath.Directory.FullName);
+				string MacExecutableFileName = MacExecutablePath.GetFileName();
+				string SupportedMacArchitectures = string.Join(" ", UnrealData.GetSupportedMacArchitectures(Config, UnrealData.UProjectFileLocation));
+
+				// figure out the directory that contains the Binaries and Intermediate directories
+				string ProjectRootDir = (bIsUnrealGame || bIsUnrealClient || GamePath == null) ? $"{UEDir}/Engine" : GamePath;
+
+				// UnrealClient re-uses UnrealGame
+				string PlistTargetName = bIsUnrealClient ? "UnrealGame" : Config.BuildTarget;
+				string MacInfoPlistPath = $"{ProjectRootDir}/Intermediate/Mac/{MacExecutableFileName}-Info.plist";
+				string IOSInfoPlistPath = $"{ProjectRootDir}/Intermediate/IOS/{PlistTargetName}-Info.plist";
+				string TVOSInfoPlistPath = $"{ProjectRootDir}/Intermediate/TVOS/{PlistTargetName}-Info.plist";
+
+				// @todo is really necessary to only set entitlements in the non-UnrealGame, project case? what about mac entitlements?
+				bool bSetEntitlements = !bIsUnrealGame && !bIsUnrealClient && UnrealData.UProjectFileLocation != null;
+
+				string SupportedPlatforms = "";
+				if (Config.ProjectTarget!.TargetRules != null)
+				{
+					if (Config.bSupportsIOS)
+					{
+						SupportedPlatforms += "iphoneos iphonesimulator ";
+					}
+					if (Config.bSupportsTVOS)
+					{
+						SupportedPlatforms += "appletvos ";
+					}
+					if (Config.bSupportsMac)
+					{
+						SupportedPlatforms += "macosx ";
+					}
+				}
+				else
+				{
+					// @todo when does this case happen?
+					SupportedPlatforms = "macosx";
+				}
+
+				// debug settings
+				if (Config.BuildConfig == UnrealTargetConfiguration.Debug)
+				{
+					Xcconfig.AppendLine("ENABLE_TESTABILITY = YES");
+				}
+
+				Xcconfig.AppendLine($"UE_BUILD_TARGET_NAME = {Config.BuildTarget}");
+				Xcconfig.AppendLine($"UE_BUILD_TARGET_CONFIG = {Config.BuildConfig}");
+
+				// list the supported platforms, which will narrow down what other settings are active
+				Xcconfig.AppendLine($"SUPPORTED_PLATFORMS = {SupportedPlatforms}");
+
+				// @otodo move VALID_ARCHS up to Project config once we are always universal
+				if (Config.bSupportsMac)
+				{
+					Xcconfig.AppendLine("");
+					Xcconfig.AppendLine("// Mac setup");
+					Xcconfig.AppendLine($"VALID_ARCHS[sdk=macosx*] = {SupportedMacArchitectures}");
+					Xcconfig.AppendLine($"PRODUCT_NAME[sdk=macosx*] = {MacExecutableFileName}");
+					Xcconfig.AppendLine($"CONFIGURATION_BUILD_DIR[sdk=macosx*] = {MacExecutableDir}");
+					Xcconfig.AppendLine($"INFOPLIST_FILE[sdk=macosx*] = {MacInfoPlistPath}");
+				}
+
+				if (Config.bSupportsIOS)
+				{
+					Xcconfig.AppendLine("");
+					Xcconfig.AppendLine("// IOS setup");
+					Xcconfig.AppendLine($"PRODUCT_NAME[sdk=iphone*] = {Config.BuildTarget}"); // @todo: change to Path.GetFileName(Config.IOSExecutablePath) when we stop using payload
+					Xcconfig.AppendLine($"CONFIGURATION_BUILD_DIR[sdk=iphone*] = {ProjectRootDir}/Binaries/IOS/Payload");
+					Xcconfig.AppendLine($"INFOPLIST_FILE[sdk=iphone*] = {IOSInfoPlistPath}");
+					if (bSetEntitlements)
+					{
+						Xcconfig.AppendLine($"CODE_SIGN_ENTITLEMENTS[sdk=iphone*] = {ProjectRootDir}/Intermediate/IOS/{PlistTargetName}.entitlements");
+					}
+				}
+				if (Config.bSupportsTVOS)
+				{
+					Xcconfig.AppendLine("");
+					Xcconfig.AppendLine("// TVOS setup");
+					Xcconfig.AppendLine($"PRODUCT_NAME[sdk=appletvos*] = {Config.BuildTarget}"); // @todo: change to Path.GetFileName(Config.IOSExecutablePath) when we stop using payload
+					Xcconfig.AppendLine($"CONFIGURATION_BUILD_DIR[sdk=appletvos*] = {ProjectRootDir}/Binaries/TVOS/Payload");
+					Xcconfig.AppendLine($"INFOPLIST_FILE[sdk=appletvos*] = {TVOSInfoPlistPath}");
+					if (bSetEntitlements)
+					{
+						Xcconfig.AppendLine($"CODE_SIGN_ENTITLEMENTS[sdk=appletvos*] = {ProjectRootDir}/Intermediate/TVOS/{PlistTargetName}.entitlements");
+					}
+				}
+
+
+				// @todo can we remove this and use GENERATE_INFOPLIST_FILE when plist doesn't exist?
+
+				// Prepare a temp Info.plist file so Xcode has some basic info about the target immediately after opening the project.
+				// This is needed for the target to pass the settings validation before code signing. UBT will overwrite this plist file later, with proper contents.
+				if (BuildHostPlatform.Current.Platform == UnrealTargetPlatform.Mac)
+				{
+					bool bCreateMacInfoPlist = !File.Exists(MacInfoPlistPath);
+					bool bCreateIOSInfoPlist = !File.Exists(IOSInfoPlistPath) && IOSRunTimeVersion != null;
+					bool bCreateTVOSInfoPlist = !File.Exists(TVOSInfoPlistPath) && TVOSRunTimeVersion != null;
+					if (bCreateMacInfoPlist || bCreateIOSInfoPlist || bCreateTVOSInfoPlist)
+					{
+						DirectoryReference? ProjectPath = GameDir;
+						DirectoryReference EngineDir = DirectoryReference.Combine(new DirectoryReference(UEDir), "Engine");
+						string GameName = Config.BuildTarget;
+						bool bIsClient = false;
+						if (ProjectPath == null)
+						{
+							ProjectPath = EngineDir;
+						}
+						if (bIsUnrealGame)
+						{
+							ProjectPath = EngineDir;
+							GameName = "UnrealGame";
+							bIsClient = (UnrealData.AppName == "UnrealClient");
+						}
+
+						if (bCreateMacInfoPlist)
+						{
+							Directory.CreateDirectory(Path.GetDirectoryName(MacInfoPlistPath)!);
+							UEDeployMac.GeneratePList(ProjectPath.FullName, bIsUnrealGame, GameName, Config.BuildTarget, EngineDir.FullName, MacExecutableFileName);
+						}
+						if (bCreateIOSInfoPlist)
+						{
+							// get the receipt
+							FileReference ReceiptFilename;
+							if (bIsUnrealGame)
+							{
+								ReceiptFilename = TargetReceipt.GetDefaultPath(Unreal.EngineDirectory, "UnrealGame", UnrealTargetPlatform.IOS, Config.BuildConfig, "");
+							}
+							else
+							{
+								ReceiptFilename = TargetReceipt.GetDefaultPath(ProjectPath, GameName, UnrealTargetPlatform.IOS, Config.BuildConfig, "");
+							}
+							Directory.CreateDirectory(Path.GetDirectoryName(IOSInfoPlistPath)!);
+							bool bSupportPortrait, bSupportLandscape;
+							TargetReceipt? Receipt;
+							TargetReceipt.TryRead(ReceiptFilename, out Receipt);
+							bool bBuildAsFramework = UEDeployIOS.GetCompileAsDll(Receipt);
+							UEDeployIOS.GenerateIOSPList(UnrealData.UProjectFileLocation, Config.BuildConfig, ProjectPath.FullName, bIsUnrealGame, GameName, bIsClient, Config.BuildTarget, EngineDir.FullName, ProjectPath + "/Binaries/IOS/Payload", null, UnrealData.BundleIdentifier, bBuildAsFramework, UnrealData.Logger!, out bSupportPortrait, out bSupportLandscape);
+						}
+						if (bCreateTVOSInfoPlist)
+						{
+							Directory.CreateDirectory(Path.GetDirectoryName(TVOSInfoPlistPath)!);
+							UEDeployTVOS.GenerateTVOSPList(ProjectPath.FullName, bIsUnrealGame, GameName, bIsClient, Config.BuildTarget, EngineDir.FullName, ProjectPath + "/Binaries/TVOS/Payload", null, UnrealData.BundleIdentifier, UnrealData.Logger!);
+						}
+					}
+				}
+			}
+		}
+	}
+
+	class XcodeFileCollection
+	{
+		class ManualFileReference
+		{
+			public string Guid = XcodeProjectFileGenerator.MakeXcodeGuid();
+			public string RelativePath = "";
+			public string FileType = "";
+			public string SourceTree = "";
+			public string GroupName = "";
+		}
+
+		// we need relativ paths, so this is the directory that paths are relative to
+		private DirectoryReference ProjectDirectory;
+
+		public string MainGroupGuid = "";
+		//public string ProductRefGroupGuid = XcodeProjectFileGenerator.MakeXcodeGuid();
+
+		internal Dictionary<string, XcodeFileGroup> Groups = new();
+		internal List<XcodeSourceFile> BuildableFiles = new();
+		internal List<XcodeSourceFile> AllFiles = new();
+
+		private List<ManualFileReference> ManualFiles = new();
+		private Dictionary<string, string> GuidsForGroups = new();
+
+
+		public XcodeFileCollection(XcodeProjectFile ProjectFile)
+		{
+			ProjectDirectory = ProjectFile.ProjectFilePath.Directory;
+		}
+
+		public void ProcessFile(XcodeSourceFile File, bool bIsForBuild)
+		{
+			//string FileName = File.Reference.GetFileName();
+			//string FileExtension = Path.GetExtension(FileName);
+			//string FilePath = File.Reference.MakeRelativeTo(ProjectFile.Directory);
+			//string FilePathMac = Utils.CleanDirectorySeparators(FilePath, '/');
+			//string FileGuid = File.FileGuid;
+			//string FileRefGuid = File.FileRefGuid;
+
+			// this is the main entry into the project for the file tree on the left of Xcode
+			//			AddFileReference(FileRefGuid, FilePathMac, GetFileType(FileExtension));
+
+			AllFiles.Add(File);
+
+			// remember all buildable files - we may further filter this later when looking for SubTargets - the module the file is
+			// in may not be built, which means it cannot be indexed
+			if (bIsForBuild && ShouldIncludeFileInBuildPhaseSection(File))
+			{
+				BuildableFiles.Add(File);
+			}
+
+			// group the files by path
+			XcodeFileGroup? Group = FindGroupByAbsolutePath(File.Reference.Directory.FullName);
+			if (Group != null)
+			{
+				Group.Files.Add(File);
+			}
+		}
+
+		/// <summary>
+		/// The project needs to point to the product group, so we expose the Guid here
+		/// </summary>
+		/// <returns></returns>
+		public string GetProductGroupGuid()
+		{
+			return GuidsForGroups["Products"];
+		}
+
+		public void AddFileReference(string Guid, string RelativePath, string FileType, string SourceTree/*="SOURCE_ROOT"*/, string GroupName)
+		{
+			ManualFiles.Add(new ManualFileReference()
+			{
+				Guid = Guid,
+				RelativePath = RelativePath,
+				FileType = FileType,
+				SourceTree = SourceTree,
+				GroupName = GroupName,
+			});
+
+			// make sure each unique group has a guid
+			if (!GuidsForGroups.ContainsKey(GroupName))
+			{
+				GuidsForGroups[GroupName] = XcodeProjectFileGenerator.MakeXcodeGuid();
+			}
+		}
+
+
+		public void Write(StringBuilder Content)
+		{
+			Content.WriteLine("/* Begin PBXBuildFile section */");
+			foreach (XcodeSourceFile File in BuildableFiles)
+			{
+				Content.WriteLine($"\t\t{File.FileGuid} = {{isa = PBXBuildFile; fileRef = {File.FileRefGuid}; }}; /* {File.Reference.GetFileName()} */");
+			}
+			Content.WriteLine("/* End PBXBuildFile section */");
+			Content.WriteLine();
+
+
+			Content.WriteLine("/* Begin PBXFileReference section */");
+			foreach (XcodeSourceFile File in AllFiles)
+			{
+				string FileName = File.Reference.GetFileName();
+				string FileExtension = Path.GetExtension(FileName);
+				string FileType = GetFileType(FileExtension);
+				string RelativePath = Utils.CleanDirectorySeparators(File.Reference.MakeRelativeTo(ProjectDirectory), '/');
+				string SourceTree = "SOURCE_ROOT";
+
+				Content.WriteLine($"\t\t{File.FileRefGuid} = {{isa = PBXFileReference; explicitFileType = {FileType}; name = \"{FileName}\"; path = \"{RelativePath}\"; sourceTree = {SourceTree}; }};");
+			}
+			foreach (ManualFileReference File in ManualFiles)
+			{
+				string FileName = Path.GetFileName(File.RelativePath);
+				Content.WriteLine($"\t\t{File.Guid} = {{isa = PBXFileReference; explicitFileType = {File.FileType}; name = \"{FileName}\"; path = \"{File.RelativePath}\"; sourceTree = {File.SourceTree}; }};");
+			}
+			Content.WriteLine("/* End PBXFileReference section */");
+			Content.WriteLine();
+
+			WriteGroups(Content);
+		}
+
+		private void WriteGroups(StringBuilder Content /*, List<UnrealExtensionInfo> AllExtensions*/)
+		{
+			XcodeFileGroup? RootGroup = FindRootFileGroup(Groups);
+			if (RootGroup == null)
+			{
+				return;
+			}
+
+			string XcconfigsRefGroupGuid = XcodeProjectFileGenerator.MakeXcodeGuid();
+
+			Content.WriteLine("/* Begin PBXGroup section */");
+
+
+			// write main/root group and it's children (adding the manual groups to the main group, which is the true param here0
+			WriteGroup(Content, RootGroup, bIsRootGroup: true);
+
+			// write some manual groups
+			foreach (var Pair in GuidsForGroups)
+			{
+				Content.WriteLine($"\t\t{Pair.Value} /* {Pair.Key} */ = {{");
+				Content.WriteLine("\t\t\tisa = PBXGroup;");
+				Content.WriteLine("\t\t\tchildren = (");
+				foreach (ManualFileReference Ref in ManualFiles.Where(x => x.GroupName == Pair.Key))
+				{
+					Content.WriteLine($"\t\t\t\t{Ref.Guid} /* {Path.GetFileName(Ref.RelativePath)} */,");
+				}
+				Content.WriteLine("\t\t\t);");
+				Content.WriteLine($"\t\t\tname = {Pair.Key};");
+				Content.WriteLine("\t\t\tsourceTree = \"<group>\";");
+				Content.WriteLine("\t\t};");
+
+			}
+
+			// Products group
+			//			Content.WriteLine($"\t\t\t\t{XcodeProject.TargetAppGuid} /* {XcodeProject.TargetName} */,");
+
+			// @todo extensions
+			//foreach (UnrealExtensionInfo EI in AllExtensions)
+			//{
+			//	Content.WriteLine(string.Format("\t\t\t\t{0} /* {1} */,{2}", EI.ProductGuid, EI.Name, ProjectFileGenerator.NewLine));
+			//}
+			// Xcconfigs group
+			//Content.WriteLine($"\t\t{XcconfigsRefGroupGuid} /* XcconfigsRefGroupGuid */ = {{");
+			//	Content.WriteLine("\t\t\tisa = PBXGroup;");
+			//	Content.WriteLine("\t\t\tchildren = (");
+			//		foreach (XcconfigFile Xcconfig in XcodeProject.AllXcconfigFiles)
+			//		{
+			//			Content.WriteLine($"\t\t\t\t{Xcconfig.Guid} /* {Xcconfig.Filename} */,{ProjectFileGenerator.NewLine}");
+			//		}
+			//	Content.WriteLine("\t\t\t);");
+			//	Content.WriteLine("\t\t\tname = Xcconfigs;");
+			//	Content.WriteLine("\t\t\tsourceTree = \"<group>\";");
+			//Content.WriteLine("\t\t};");
+
+			Content.WriteLine("/* End PBXGroup section */");
+		}
+
+
+		private void WriteGroup(StringBuilder Content, XcodeFileGroup Group, bool bIsRootGroup)
+		{
+			if (!Group.bIsReference)
+			{
+				Content.WriteLine($"\t\t{Group.GroupGuid} = {{");
+				Content.WriteLine("\t\t\tisa = PBXGroup;");
+				Content.WriteLine("\t\t\tchildren = (");
+				foreach (XcodeFileGroup ChildGroup in Group.Children.Values)
+				{
+					Content.WriteLine($"\t\t\t\t{ChildGroup.GroupGuid} /* {ChildGroup.GroupName} */,");
+				}
+				foreach (XcodeSourceFile File in Group.Files)
+				{
+					Content.WriteLine($"\t\t\t\t{File.FileRefGuid} /* {File.Reference.GetFileName()} */,");
+				}
+
+				if (bIsRootGroup)
+				{
+					foreach (var Pair in GuidsForGroups)
+					{
+						Content.WriteLine($"\t\t\t\t{Pair.Value} /* {Pair.Key} */,");
+					}
+				}
+				Content.WriteLine("\t\t\t);");
+				if (!bIsRootGroup)
+				{
+					Content.WriteLine("\t\t\tname = \"" + Group.GroupName + "\";");
+					Content.WriteLine("\t\t\tpath = \"" + Group.GroupPath + "\";");
+					Content.WriteLine("\t\t\tsourceTree = \"<absolute>\";");
+				}
+				else
+				{
+					Content.WriteLine("\t\t\tsourceTree = \"<group>\";");
+				}
+				Content.WriteLine("\t\t};");
+
+				foreach (XcodeFileGroup ChildGroup in Group.Children.Values)
+				{
+					WriteGroup(Content, ChildGroup, bIsRootGroup: false);
+				}
+			}
+		}
+
+		public static string GetRootGroupGuid(Dictionary<string, XcodeFileGroup> GroupsDict)
+		{
+			return FindRootFileGroup(GroupsDict)!.GroupGuid;
+		}
+
+		private static XcodeFileGroup? FindRootFileGroup(Dictionary<string, XcodeFileGroup> GroupsDict)
+		{
+			foreach (XcodeFileGroup Group in GroupsDict.Values)
+			{
+				if (Group.Children.Count > 1 || Group.Files.Count > 0)
+				{
+					return Group;
+				}
+				else
+				{
+					XcodeFileGroup? Found = FindRootFileGroup(Group.Children);
+					if (Found != null)
+					{
+						return Found;
+					}
+				}
+			}
+			return null;
+		}
+
+
+		/// <summary>
+		/// Gets Xcode file category based on its extension
+		/// </summary>
+		internal static string GetFileCategory(string Extension)
+		{
+			// @todo Mac: Handle more categories
+			switch (Extension)
+			{
+				case ".framework":
+					return "Frameworks";
+				default:
+					return "Sources";
+			}
+		}
+
+		/// <summary>
+		/// Gets Xcode file type based on its extension
+		/// </summary>
+		internal static string GetFileType(string Extension)
+		{
+			// @todo Mac: Handle more file types
+			switch (Extension)
 			{
-				ProjectXcconfig.AppendLine($"VALID_ARCHS[sdk=appletvos*] = {SupportedIOSArchitectures}");
-				ProjectXcconfig.AppendLine($"IPHONEOS_DEPLOYMENT_TARGET = {TVOSRunTimeVersion}");
-				ProjectXcconfig.AppendLine($"TARGETED_DEVICE_FAMILY[sdk=appletvos*] = {TVOSRunTimeDevices}");
-				if (!string.IsNullOrEmpty(TEAM_TVOS))
-				{
-					ProjectXcconfig.AppendLine($"DEVELOPMENT_TEAM[sdk=iphoneos*] = {TEAM_TVOS}");
-				}
-				ProjectXcconfig.AppendLine($"CODE_SIGN_IDENTITY[sdk=appletvos*] = {TVOS_CERT}");
-				if (!bAutomaticSigning && !string.IsNullOrEmpty(UUID_TVOS))
-				{
-					ProjectXcconfig.AppendLine($"PROVISIONING_PROFILE_SPECIFIER[sdk=appletvos*] = {UUID_TVOS}");
-				}
-				if (ProjectFile != null)
-				{
-					ProjectXcconfig.AppendLine($"PRODUCT_BUNDLE_IDENTIFIER[sdk=appletvos*] = {TVOS_BUNDLE}");
-				}
+				case ".c":
+				case ".m":
+					return "sourcecode.c.objc";
+				case ".cc":
+				case ".cpp":
+				case ".mm":
+					return "sourcecode.cpp.objcpp";
+				case ".h":
+				case ".inl":
+				case ".pch":
+					return "sourcecode.c.h";
+				case ".framework":
+					return "wrapper.framework";
+				case ".plist":
+					return "text.plist.xml";
+				case ".png":
+					return "image.png";
+				case ".icns":
+					return "image.icns";
+				default:
+					return "file.text";
 			}
+		}
 
-			// Now for each config write out the specific settings
+		/// <summary>
+		/// Returns true if Extension is a known extension for files containing source code
+		/// </summary>
+		internal static bool IsSourceCode(string Extension)
+		{
+			return Extension == ".c" || Extension == ".cc" || Extension == ".cpp" || Extension == ".m" || Extension == ".mm";
+		}
 
-			DirectoryReference? GameDir = ProjectFile?.Directory;
-			string? GamePath = GameDir != null ? ConvertPath(GameDir.FullName) : null;
+		static string[] ExcludedFolders = PlatformExports.GetExcludedFolderNames(UnrealTargetPlatform.Mac);
+		internal static bool ShouldIncludeFileInBuildPhaseSection(XcodeSourceFile SourceFile)
+		{
+			string FileExtension = SourceFile.Reference.GetExtension();
 
-			// Get Mac architectures supported by this project
-			foreach (XcodeBuildConfig Config in ProjectConfigs)
+			if (IsSourceCode(FileExtension))
 			{
-				// find the config file we made earlier - we know it exists
-				XcconfigFile Xcconfig = AllXcconfigFiles.First(x => x.Name == Config.DisplayName);
+				// look if it contains any of the exluded names, and is so, don't include it
+				return !ExcludedFolders.Any(x => SourceFile.Reference.ContainsName(x, 0));
+			}
 
-				string TargetName = Config.BuildTarget;
-				bool bIsUnrealGame = TargetName.Equals("UnrealGame", StringComparison.InvariantCultureIgnoreCase);
-				bool bIsUnrealClient = TargetName.Equals("UnrealClient", StringComparison.InvariantCultureIgnoreCase);
+			return false;
+		}
 
-				// Setup Mac stuff
-				FileReference MacExecutablePath = Config.MacExecutablePath!;
-				string MacExecutableDir = ConvertPath(MacExecutablePath.Directory.FullName);
-				string MacExecutableFileName = MacExecutablePath.GetFileName();
-				string SupportedMacArchitectures = string.Join(" ", GetSupportedMacArchitectures(Config, ProjectFile));
+		/// <summary>
+		/// Returns a project navigator group to which the file should belong based on its path.
+		/// Creates a group tree if it doesn't exist yet.
+		/// </summary>
+		internal XcodeFileGroup? FindGroupByAbsolutePath(string AbsolutePath)
+		{
+			string[] Parts = AbsolutePath.Split(Path.DirectorySeparatorChar);
+			string CurrentPath = "/";
+			Dictionary<string, XcodeFileGroup> CurrentSubGroups = Groups;
 
-				// figure out the directory that contains the Binaries and Intermediate directories
-				string ProjectRootDir = (bIsUnrealGame || bIsUnrealClient || GamePath == null) ? $"{UEDir}/Engine" : GamePath;
+			for (int Index = 1; Index < Parts.Count(); ++Index)
+			{
+				string Part = Parts[Index];
 
-				// UnrealClient re-uses UnrealGame
-				string PlistTargetName = bIsUnrealClient ? "UnrealGame" : Config.BuildTarget;
-				string MacInfoPlistPath = $"{ProjectRootDir}/Intermediate/Mac/{MacExecutableFileName}-Info.plist";
-				string IOSInfoPlistPath = $"{ProjectRootDir}/Intermediate/IOS/{PlistTargetName}-Info.plist";
-				string TVOSInfoPlistPath = $"{ProjectRootDir}/Intermediate/TVOS/{PlistTargetName}-Info.plist";
+				if (CurrentPath.Length > 1)
+				{
+					CurrentPath += Path.DirectorySeparatorChar;
+				}
 
-				// @todo is really necessary to only set entitlements in the non-UnrealGame, project case? what about mac entitlements?
-				bool bSetEntitlements = !bIsUnrealGame && !bIsUnrealClient && ProjectFile != null;
+				CurrentPath += Part;
 
-				string SupportedPlatforms = "";
-				if (Config.ProjectTarget!.TargetRules != null)
+				XcodeFileGroup CurrentGroup;
+				if (!CurrentSubGroups.ContainsKey(CurrentPath))
 				{
-					if (Config.bSupportsIOS)
-					{
-						SupportedPlatforms += "iphoneos ";
-					}
-					if (Config.bSupportsTVOS)
-					{
-						SupportedPlatforms += "appletvos ";
-					}
-					if (Config.bSupportsMac)
+					CurrentGroup = new XcodeFileGroup(Path.GetFileName(CurrentPath), CurrentPath, CurrentPath.EndsWith(".xcassets"));
+					if (Groups.Count == 0)
 					{
-						SupportedPlatforms += "macosx ";
+						MainGroupGuid = CurrentGroup.GroupGuid;
 					}
+					CurrentSubGroups.Add(CurrentPath, CurrentGroup);
 				}
 				else
 				{
-					// @todo when does this case happen?
-					SupportedPlatforms = "macosx";
+					CurrentGroup = CurrentSubGroups[CurrentPath];
 				}
 
-				// debug settings
-				if (Config.BuildConfig == UnrealTargetConfiguration.Debug)
+				if (CurrentPath == AbsolutePath)
 				{
-					Xcconfig.AppendLine("ENABLE_TESTABILITY = YES");
+					return CurrentGroup;
 				}
 
-				Xcconfig.AppendLine($"UE_BUILD_TARGET_NAME = {TargetName}");
-				Xcconfig.AppendLine($"UE_BUILD_TARGET_CONFIG = {Config.BuildConfig}");
+				CurrentSubGroups = CurrentGroup.Children;
+			}
 
-				// list the supported platforms, which will narrow down what other settings are active
-				Xcconfig.AppendLine($"SUPPORTED_PLATFORMS = {SupportedPlatforms}");
+			return null;
+		}
 
-				// @otodo move VALID_ARCHS up to Project config once we are always universal
-				if (Config.bSupportsMac)
-				{
-					Xcconfig.AppendLine("");
-					Xcconfig.AppendLine("// Mac setup");
-					Xcconfig.AppendLine($"VALID_ARCHS[sdk=macosx*] = {SupportedMacArchitectures}");
-					Xcconfig.AppendLine($"PRODUCT_NAME[sdk=macosx*] = {MacExecutableFileName}");
-					Xcconfig.AppendLine($"CONFIGURATION_BUILD_DIR[sdk=macosx*] = {MacExecutableDir}");
-					Xcconfig.AppendLine($"INFOPLIST_FILE[sdk=macosx*] = {MacInfoPlistPath}");
-				}
+		/// <summary>
+		/// Convert all paths to Apple/Unix format (with forward slashes)
+		/// </summary>
+		/// <param name="InPath">The path to convert</param>
+		/// <returns>The normalized path</returns>
+		internal static string ConvertPath(string InPath)
+		{
+			return InPath.Replace("\\", "/");
+		}
+	}
 
-				if (Config.bSupportsIOS)
-				{
-					Xcconfig.AppendLine("");
-					Xcconfig.AppendLine("// IOS setup");
-					Xcconfig.AppendLine($"PRODUCT_NAME[sdk=iphoneos*] = {TargetName}"); // @todo: change to Path.GetFileName(Config.IOSExecutablePath) when we stop using payload
-					Xcconfig.AppendLine($"CONFIGURATION_BUILD_DIR[sdk=iphoneos*] = {ProjectRootDir}/Binaries/IOS/Payload");
-					Xcconfig.AppendLine($"INFOPLIST_FILE[sdk=iphoneos*] = {IOSInfoPlistPath}");
-					if (bSetEntitlements)
-					{
-						Xcconfig.AppendLine($"CODE_SIGN_ENTITLEMENTS[sdk=iphoneos*] = {ProjectRootDir}/Intermediate/IOS/{PlistTargetName}.entitlements");
-					}
-				}
-				if (Config.bSupportsTVOS)
-				{
-					Xcconfig.AppendLine("");
-					Xcconfig.AppendLine("// TVOS setup");
-					Xcconfig.AppendLine($"PRODUCT_NAME[sdk=appletvos*] = {TargetName}"); // @todo: change to Path.GetFileName(Config.IOSExecutablePath) when we stop using payload
-					Xcconfig.AppendLine($"CONFIGURATION_BUILD_DIR[sdk=appletvos*] = {ProjectRootDir}/Binaries/TVOS/Payload");
-					Xcconfig.AppendLine($"INFOPLIST_FILE[sdk=appletvos*] = {TVOSInfoPlistPath}");
-					if (bSetEntitlements)
-					{
-						Xcconfig.AppendLine($"CODE_SIGN_ENTITLEMENTS[sdk=appletvos*] = {ProjectRootDir}/Intermediate/TVOS/{PlistTargetName}.entitlements");
-					}
-				}
 
+	class XcodeProjectFile : ProjectFile
+	{
 
-				// @todo can we remove this and use GENERATE_INFOPLIST_FILE when plist doesn't exist?
+		/// <summary>
+		/// Constructs a new project file object
+		/// </summary>
+		/// <param name="InitFilePath">The path to the project file on disk</param>
+		/// <param name="BaseDir">The base directory for files within this project</param>
+		/// <param name="bIsForDistribution">True for distribution builds</param>
+		/// <param name="BundleID">Override option for bundle identifier</param>
+		/// <param name="AppName"></param>
+		public XcodeProjectFile(FileReference InitFilePath, DirectoryReference BaseDir, bool bIsForDistribution, string BundleID, string AppName)
+			: base(InitFilePath, BaseDir)
+		{
+			UnrealData = new UnrealData();
 
-				// Prepare a temp Info.plist file so Xcode has some basic info about the target immediately after opening the project.
-				// This is needed for the target to pass the settings validation before code signing. UBT will overwrite this plist file later, with proper contents.
-				if (BuildHostPlatform.Current.Platform == UnrealTargetPlatform.Mac)
-				{
-					bool bCreateMacInfoPlist = !File.Exists(MacInfoPlistPath);
-					bool bCreateIOSInfoPlist = !File.Exists(IOSInfoPlistPath) && IOSRunTimeVersion != null;
-					bool bCreateTVOSInfoPlist = !File.Exists(TVOSInfoPlistPath) && TVOSRunTimeVersion != null;
-					if (bCreateMacInfoPlist || bCreateIOSInfoPlist || bCreateTVOSInfoPlist)
-					{
-						DirectoryReference? ProjectPath = GameDir;
-						DirectoryReference EngineDir = DirectoryReference.Combine(new DirectoryReference(UEDir), "Engine");
-						string GameName = Config.BuildTarget;
-						bool bIsClient = false;
-						if (ProjectPath == null)
-						{
-							ProjectPath = EngineDir;
-						}
-						if (bIsUnrealGame)
-						{
-							ProjectPath = EngineDir;
-							GameName = "UnrealGame";
-							bIsClient = (AppName == "UnrealClient");
-						}
+			UnrealData.bForDistribution = bIsForDistribution;
+			UnrealData.BundleIdentifier = BundleID;
+			UnrealData.AppName = AppName;
 
-						if (bCreateMacInfoPlist)
-						{
-							Directory.CreateDirectory(Path.GetDirectoryName(MacInfoPlistPath)!);
-							UEDeployMac.GeneratePList(ProjectPath.FullName, bIsUnrealGame, GameName, Config.BuildTarget, EngineDir.FullName, MacExecutableFileName);
-						}
-						if (bCreateIOSInfoPlist)
-						{
-							// get the receipt
-							FileReference ReceiptFilename;
-							if (bIsUnrealGame)
-							{
-								ReceiptFilename = TargetReceipt.GetDefaultPath(Unreal.EngineDirectory, "UnrealGame", UnrealTargetPlatform.IOS, Config.BuildConfig, "");
-							}
-							else
-							{
-								ReceiptFilename = TargetReceipt.GetDefaultPath(ProjectPath, GameName, UnrealTargetPlatform.IOS, Config.BuildConfig, "");
-							}
-							Directory.CreateDirectory(Path.GetDirectoryName(IOSInfoPlistPath)!);
-							bool bSupportPortrait, bSupportLandscape;
-							TargetReceipt? Receipt;
-							TargetReceipt.TryRead(ReceiptFilename, out Receipt);
-							bool bBuildAsFramework = UEDeployIOS.GetCompileAsDll(Receipt);
-							UEDeployIOS.GenerateIOSPList(ProjectFile, Config.BuildConfig, ProjectPath.FullName, bIsUnrealGame, GameName, bIsClient, Config.BuildTarget, EngineDir.FullName, ProjectPath + "/Binaries/IOS/Payload", null, BundleIdentifier, bBuildAsFramework, Logger, out bSupportPortrait, out bSupportLandscape);
-						}
-						if (bCreateTVOSInfoPlist)
-						{
-							Directory.CreateDirectory(Path.GetDirectoryName(TVOSInfoPlistPath)!);
-							UEDeployTVOS.GenerateTVOSPList(ProjectPath.FullName, bIsUnrealGame, GameName, bIsClient, Config.BuildTarget, EngineDir.FullName, ProjectPath + "/Binaries/TVOS/Payload", null, BundleIdentifier, Logger);
-						}
-					}
-				}
-			}
+			// create the container for all the files that will 
+			FileCollection = new XcodeFileCollection(this);
 		}
 
-		private void AppendBuildConfiguration(StringBuilder Content, XcodeBuildConfig Config, string ConfigGuid,  XcconfigFile Xcconfig)
-		{
-			Content.Append("\t\t" + ConfigGuid + " /* \"" + Config.DisplayName + "\" */ = {" + ProjectFileGenerator.NewLine);
-			Content.Append("\t\t\tisa = XCBuildConfiguration;" + ProjectFileGenerator.NewLine);
-			Content.Append($"\t\t\tbaseConfigurationReference = {Xcconfig.Guid} /* \"{Xcconfig.Filename}\" */;" + ProjectFileGenerator.NewLine);
-			Content.Append("\t\t\tbuildSettings = {" + ProjectFileGenerator.NewLine);
-			Content.Append("\t\t\t};" + ProjectFileGenerator.NewLine);
-			Content.Append("\t\t\tname = \"" + Config.DisplayName + "\";" + ProjectFileGenerator.NewLine);
-			Content.Append("\t\t};" + ProjectFileGenerator.NewLine);
-		}
+		public UnrealData UnrealData;
+
+		/// <summary>
+		///  The PBXPRoject node, root of everything
+		/// </summary>
+		public XcodeProject? RootProject;
 
+		/// <summary>
+		/// Gathers the files and generates project sections
+		/// </summary>
+		public XcodeFileCollection FileCollection;
 
-		private void AppendXCBuildConfigurationSection(StringBuilder Content, Dictionary<string, XcodeBuildConfig> ProjectBuildConfigs, Dictionary<string, XcodeBuildConfig> TargetBuildConfigs,
-				Dictionary<string, XcodeBuildConfig> BuildTargetBuildConfigs, Dictionary<string, XcodeBuildConfig> IndexTargetBuildConfigs, FileReference? GameProjectPath,
-				List<XcodeExtensionInfo> AllExtensions, ILogger Logger)
+		/// <summary>
+		/// Allocates a generator-specific source file object
+		/// </summary>
+		/// <param name="InitFilePath">Path to the source file on disk</param>
+		/// <param name="InitProjectSubFolder">Optional sub-folder to put the file in.  If empty, this will be determined automatically from the file's path relative to the project file</param>
+		/// <returns>The newly allocated source file object</returns>
+		public override SourceFile? AllocSourceFile(FileReference InitFilePath, DirectoryReference? InitProjectSubFolder)
 		{
-			Content.Append("/* Begin XCBuildConfiguration section */" + ProjectFileGenerator.NewLine);
-
-			foreach (KeyValuePair<string, XcodeBuildConfig> Config in ProjectBuildConfigs)
+			if (InitFilePath.GetFileName().StartsWith("."))
 			{
-				// find the premade file that matches the config name (like "Debug Editor")
-				XcconfigFile Xcconfig = AllXcconfigFiles.First(x => x.Name == Config.Value.DisplayName);
-				AppendBuildConfiguration(Content, Config.Value, Config.Key, Xcconfig);
+				return null;
 			}
+			return new XcodeSourceFile(InitFilePath, InitProjectSubFolder);
+		}
 
-			foreach (KeyValuePair<string, XcodeBuildConfig> Config in TargetBuildConfigs)
-			{
-				AppendBuildConfiguration(Content, Config.Value, Config.Key, RunXcconfig);
-			}
+		public override void AddModule(UEBuildModuleCPP Module, CppCompileEnvironment CompileEnvironment)
+		{
+			UnrealData.AddModule(Module, CompileEnvironment);
+		}
 
-			foreach (KeyValuePair<string, XcodeBuildConfig> Config in BuildTargetBuildConfigs)
-			{
-				AppendBuildConfiguration(Content, Config.Value, Config.Key, BuildXcconfig);
-			}
+		/// <summary>
+		/// Generates bodies of all sections that contain a list of source files plus a dictionary of project navigator groups.
+		/// </summary>
+		private void ProcessSourceFiles()
+		{
+			// process the files that came from UE/cross-platform land
+			SourceFiles.SortBy(x => x.Reference.FullName);
 
-			foreach (KeyValuePair<string, XcodeBuildConfig> Config in IndexTargetBuildConfigs)
+			foreach (XcodeSourceFile SourceFile in SourceFiles.OfType<XcodeSourceFile>())
 			{
-				AppendBuildConfiguration(Content, Config.Value, Config.Key, IndexXcconfig);
+				FileCollection.ProcessFile(SourceFile, bIsForBuild:IsGeneratedProject);
 			}
 
-			foreach (XcodeExtensionInfo EI in AllExtensions)
-			{
-				Content.Append(EI.ConfigurationContents);
-			}
+			UnrealData.FilterBuildableFiles(FileCollection);
+		}
+
+
 
-			Content.Append("/* End XCBuildConfiguration section */" + ProjectFileGenerator.NewLine + ProjectFileGenerator.NewLine);
-		}
 
-		private void AppendXCConfigurationList(StringBuilder Content, string TypeName, string TargetName, string ConfigListGuid, Dictionary<string, XcodeBuildConfig> BuildConfigs, string Default = "Development")
+		/// Implements Project interface
+		public override bool WriteProjectFile(List<UnrealTargetPlatform> InPlatforms, List<UnrealTargetConfiguration> InConfigurations, PlatformProjectGeneratorCollection PlatformProjectGenerators, ILogger Logger)
 		{
-			Content.Append("\t\t" + ConfigListGuid + " /* Build configuration list for " + TypeName + " \"" + TargetName + "\" */ = {" + ProjectFileGenerator.NewLine);
-			Content.Append("\t\t\tisa = XCConfigurationList;" + ProjectFileGenerator.NewLine);
-			Content.Append("\t\t\tbuildConfigurations = (" + ProjectFileGenerator.NewLine);
-			foreach (KeyValuePair<string, XcodeBuildConfig> Config in BuildConfigs)
+			if (UnrealData.Initialize(this, InPlatforms, InConfigurations, Logger) == false)
 			{
-				Content.Append("\t\t\t\t" + Config.Key + " /* \"" + Config.Value.DisplayName + "\" */," + ProjectFileGenerator.NewLine);
+				// if we failed to initialize, we silently return to move on (it's not an error, it's a project with nothing to do)
+				return true;
 			}
-			Content.Append("\t\t\t);" + ProjectFileGenerator.NewLine);
-			Content.Append("\t\t\tdefaultConfigurationIsVisible = 0;" + ProjectFileGenerator.NewLine);
-			Content.Append("\t\t\tdefaultConfigurationName = " + Default + ";" + ProjectFileGenerator.NewLine);
-			Content.Append("\t\t};" + ProjectFileGenerator.NewLine);
-		}
 
-		private void AppendXCConfigurationListSection(StringBuilder Content, string TargetName, string BuildTargetName, string IndexTargetName, string ProjectConfigListGuid,
-			Dictionary<string, XcodeBuildConfig> ProjectBuildConfigs, string TargetConfigListGuid, Dictionary<string, XcodeBuildConfig> TargetBuildConfigs,
-			string BuildTargetConfigListGuid, Dictionary<string, XcodeBuildConfig> BuildTargetBuildConfigs,
-			string IndexTargetConfigListGuid, Dictionary<string, XcodeBuildConfig> IndexTargetBuildConfigs,
-			List<XcodeExtensionInfo> AllExtensions)
-		{
-			Content.Append("/* Begin XCConfigurationList section */" + ProjectFileGenerator.NewLine);
+			// turn all UE files into internal representation
+			ProcessSourceFiles();
 
-			AppendXCConfigurationList(Content, "PBXProject", TargetName, ProjectConfigListGuid, ProjectBuildConfigs);
-			AppendXCConfigurationList(Content, "PBXLegacyTarget", BuildTargetName, BuildTargetConfigListGuid, BuildTargetBuildConfigs);
-			AppendXCConfigurationList(Content, "PBXNativeTarget", TargetName, TargetConfigListGuid, TargetBuildConfigs);
-			AppendXCConfigurationList(Content, "PBXNativeTarget", IndexTargetName, IndexTargetConfigListGuid, IndexTargetBuildConfigs);
+			// now create the xcodeproject elements (project -> target -> buildconfigs, etc)
+			RootProject = new XcodeProject(UnrealData, FileCollection);
 
-			foreach (XcodeExtensionInfo EI in AllExtensions)
-			{
-				AppendXCConfigurationList(Content, "PBXNativeTarget", EI.Name, EI.ConfigListGuid, EI.AllConfigs);
-		}
+			// now that the node hierarchy is made, walk the nodes and create Xcconfig files
+			XcodeProjectNode.WriteXcconfigFiles(RootProject, this);
 
-			Content.Append("/* End XCConfigurationList section */" + ProjectFileGenerator.NewLine);
-			}
+			StringBuilder Content = new StringBuilder();
 
-		private List<XcodeBuildConfig> GetSupportedBuildConfigs(List<UnrealTargetPlatform> Platforms, List<UnrealTargetConfiguration> Configurations, ILogger Logger)
-		{
-			List<XcodeBuildConfig> BuildConfigs = new List<XcodeBuildConfig>();
+			// look for an existing project to use as a template (if none found, create one from scratch)
+			DirectoryReference BuildDirLocation = UnrealData.UProjectFileLocation == null ? Unreal.EngineDirectory : UnrealData.UProjectFileLocation.Directory;
+			string ExistingProjectName = UnrealData.XcodeProjectFileLocation!.GetFileNameWithoutAnyExtensions();
+			FileReference TemplateProject = FileReference.Combine(BuildDirLocation, "Build", "IOS", ExistingProjectName + ".xcodeproj", "project.pbxproj");
 
-			//string ProjectName = ProjectFilePath.GetFileNameWithoutExtension();
+			if (TemplateProject.FullName.Contains("ThirdPerson"))
+			{
+				Console.WriteLine("foo");
+			}
 
-			foreach (UnrealTargetConfiguration Configuration in Configurations)
+			FileReference PBXProjFilePath = ProjectFilePath + "/project.pbxproj";
+
+			bool bSuccess;
+			if (FileReference.Exists(TemplateProject))
 			{
-				if (InstalledPlatformInfo.IsValidConfiguration(Configuration, EProjectType.Code))
+				// copy existing template project to final location
+				if (FileReference.Exists(PBXProjFilePath))
 				{
-					foreach (UnrealTargetPlatform Platform in Platforms)
+					FileReference.Delete(PBXProjFilePath);
+				}
+				DirectoryReference.CreateDirectory(PBXProjFilePath.Directory);
+				FileReference.Copy(TemplateProject, PBXProjFilePath);
+
+				// write the nodes we need to add (Build/Index targets)
+				XcodeRunTarget RunTarget = XcodeProjectNode.GetNodesOfType<XcodeRunTarget>(RootProject).First();
+				XcodeBuildTarget BuildTarget = XcodeProjectNode.GetNodesOfType<XcodeBuildTarget>(RunTarget).First();
+				XcodeIndexTarget IndexTarget = XcodeProjectNode.GetNodesOfType<XcodeIndexTarget>(RootProject).First();
+				XcodeDependency BuildDependency = XcodeProjectNode.GetNodesOfType<XcodeDependency>(RunTarget).First();
+
+				Content.WriteLine(0, "{");
+				FileCollection.Write(Content);
+				XcodeProjectNode.WriteNodeAndReferences(Content, BuildTarget);
+				XcodeProjectNode.WriteNodeAndReferences(Content, IndexTarget);
+				XcodeProjectNode.WriteNodeAndReferences(Content, BuildDependency);
+				Content.WriteLine(0, "}");
+
+				// write to disk
+				FileReference ImportFile = FileReference.Combine(PBXProjFilePath.Directory, "import.plist");
+				File.WriteAllText(ImportFile.FullName, Content.ToString());
+
+				Func<string, string> Plist = (string Command) =>
 					{
-						if (InstalledPlatformInfo.IsValidPlatform(Platform, EProjectType.Code) && (Platform == UnrealTargetPlatform.Mac || Platform == UnrealTargetPlatform.IOS || Platform == UnrealTargetPlatform.TVOS)) // @todo support other platforms
-						{
-							UEBuildPlatform? BuildPlatform;
-							if (UEBuildPlatform.TryGetBuildPlatform(Platform, out BuildPlatform) && (BuildPlatform.HasRequiredSDKsInstalled() == SDKStatus.Valid))
-							{
-								// Check we have targets (Expected to be no Engine targets when generating for a single .uproject)
-								if (ProjectTargets.Count == 0 && BaseDir != Unreal.EngineDirectory)
-								{
-									throw new BuildException("Expecting at least one ProjectTarget to be associated with project '{0}' in the TargetProjects list ", ProjectFilePath);
-								}
+						Command = Command.Replace("\"", "\\\"");
+						return Utils.RunLocalProcessAndReturnStdOut("/usr/libexec/PlistBuddy", $"-c \"{Command}\" \"{PBXProjFilePath.FullName}\"");
+					};
 
-								// Now go through all of the target types for this project
-								foreach (ProjectTarget ProjectTarget in ProjectTargets.OfType<ProjectTarget>())
-								{
-									if (MSBuildProjectFile.IsValidProjectPlatformAndConfiguration(ProjectTarget, Platform, Configuration, Logger))
-									{
-										// Figure out if this is a monolithic build
-										bool bShouldCompileMonolithic = BuildPlatform.ShouldCompileMonolithicBinary(Platform);
-										bShouldCompileMonolithic |= (ProjectTarget.CreateRulesDelegate(Platform, Configuration).LinkType == TargetLinkType.Monolithic);
+				// and import it into the template
+				Plist($"Merge \"{ImportFile.FullName}\" :objects");
 
-										string ConfigName = Configuration.ToString();
-										if (ProjectTarget.TargetRules!.Type != TargetType.Game && ProjectTarget.TargetRules.Type != TargetType.Program)
-										{
-											ConfigName += " " + ProjectTarget.TargetRules.Type.ToString();
-										}
+				// get the project guid
+				string ProjectGuid = Plist($"Print :rootObject");
 
-										if (BuildConfigs.Where(Config => Config.DisplayName == ConfigName).ToList().Count == 0)
-										{
-											string TargetName = ProjectTarget.TargetFilePath.GetFileNameWithoutAnyExtensions();
+				// get all the targets in the template that are application types
+				IEnumerable<string> AppTargetGuids = Plist($"Print :objects:{ProjectGuid}:targets")
+					.ReplaceLineEndings()
+					.Split(Environment.NewLine)
+					.Where(TargetGuid => (Plist($"Print :objects:{TargetGuid}:productType") == "com.apple.product-type.application"));
 
-											// Get the output directory
-											DirectoryReference RootDirectory = Unreal.EngineDirectory;
-											// Unique and Monolithic both need to use the target directory not the engine directory
-											if (ProjectTarget.TargetRules.Type != TargetType.Program && (bShouldCompileMonolithic || ProjectTarget.TargetRules.BuildEnvironment == TargetBuildEnvironment.Unique))
-											{
-												if (ProjectTarget.UnrealProjectFilePath != null)
-												{
-													RootDirectory = ProjectTarget.UnrealProjectFilePath.Directory;
-												}
-											}
+				// add a dependency on the build target from the app target(s)
+				foreach (string AppTargetGuid in AppTargetGuids)
+				{
+					Plist($"Add :objects:{AppTargetGuid}:dependencies:0 string {BuildTarget.Guid}");
+				}
 
-											if (ProjectTarget.TargetRules.Type == TargetType.Program && ProjectTarget.UnrealProjectFilePath != null)
-											{
-												RootDirectory = ProjectTarget.UnrealProjectFilePath.Directory;
-											}
+				// now add all the non-run targets from the generated
+				foreach (XcodeTarget Target in XcodeProjectNode.GetNodesOfType<XcodeTarget>(RootProject).Where(x => x.GetType() != typeof(XcodeRunTarget)))
+				{
+					Plist($"Add :objects:{ProjectGuid}:targets:0 string {Target.Guid}");
+				}
 
-											// Get the output directory
-											DirectoryReference OutputDirectory = DirectoryReference.Combine(RootDirectory, "Binaries");
 
-											string ExeName = TargetName;
-											if (!bShouldCompileMonolithic && ProjectTarget.TargetRules.Type != TargetType.Program)
-											{
-												// Figure out what the compiled binary will be called so that we can point the IDE to the correct file
-												if (ProjectTarget.TargetRules.Type != TargetType.Game)
-												{
-													// Only if shared - unique retains the Target Name
-													if(ProjectTarget.TargetRules.BuildEnvironment == TargetBuildEnvironment.Shared)
-													{
-														ExeName = "Unreal" + ProjectTarget.TargetRules.Type.ToString();
-													}
-												}
-											}
 
-											if (BuildPlatform.Platform == UnrealTargetPlatform.Mac)
-											{
-												string MacExecutableName = MakeExecutableFileName(ExeName, UnrealTargetPlatform.Mac, Configuration, ProjectTarget.TargetRules.Architecture, ProjectTarget.TargetRules.UndecoratedConfiguration);
-												string IOSExecutableName = MacExecutableName.Replace("-Mac-", "-IOS-");
-												string TVOSExecutableName = MacExecutableName.Replace("-Mac-", "-TVOS-");
-												BuildConfigs.Add(new XcodeBuildConfig(ConfigName, TargetName, FileReference.Combine(OutputDirectory, "Mac", MacExecutableName), FileReference.Combine(OutputDirectory, "IOS", IOSExecutableName), FileReference.Combine(OutputDirectory, "TVOS", TVOSExecutableName), ProjectTarget, Configuration));
-											}
-											else if (BuildPlatform.Platform == UnrealTargetPlatform.IOS || BuildPlatform.Platform == UnrealTargetPlatform.TVOS)
-											{
-												string IOSExecutableName = MakeExecutableFileName(ExeName, UnrealTargetPlatform.IOS, Configuration, ProjectTarget.TargetRules.Architecture, ProjectTarget.TargetRules.UndecoratedConfiguration);
-												string TVOSExecutableName = IOSExecutableName.Replace("-IOS-", "-TVOS-");
-												//string MacExecutableName = IOSExecutableName.Replace("-IOS-", "-Mac-");
-												BuildConfigs.Add(new XcodeBuildConfig(ConfigName, TargetName, FileReference.Combine(OutputDirectory, "Mac", IOSExecutableName), FileReference.Combine(OutputDirectory, "IOS", IOSExecutableName), FileReference.Combine(OutputDirectory, "TVOS", TVOSExecutableName), ProjectTarget, Configuration));
-											}
-										}
-									}
-								}
-							}
-						}
-					}
-				}
+				bSuccess = true;
 			}
+			else
+			{
+				Content.WriteLine(0, "// !$*UTF8*$!");
+				Content.WriteLine(0, "{");
+				Content.WriteLine(1, "archiveVersion = 1;");
+				Content.WriteLine(1, "classes = {");
+				Content.WriteLine(1, "};");
+				Content.WriteLine(1, "objectVersion = 46;");
+				Content.WriteLine(1, "objects = {");
 
-			return BuildConfigs;
-		}
+				// write out the list of files and groups
+				FileCollection.Write(Content);
 
-        private static string MakeExecutableFileName(string BinaryName, UnrealTargetPlatform Platform, UnrealTargetConfiguration Configuration, string Architecture, UnrealTargetConfiguration UndecoratedConfiguration)
-		{
-			StringBuilder Result = new StringBuilder();
+				// now write out the project node and its recursive dependent nodes
+				XcodeProjectNode.WriteNodeAndReferences(Content, RootProject);
 
-			Result.Append(BinaryName);
+				Content.WriteLine(1, "};");
+				Content.WriteLine(1, $"rootObject = {RootProject.Guid} /* Project object */;");
+				Content.WriteLine(0, "}");
 
-			if (Configuration != UndecoratedConfiguration)
-			{
-				Result.AppendFormat("-{0}-{1}", Platform.ToString(), Configuration.ToString());
+				// finally write out the pbxproj file!
+				bSuccess = ProjectFileGenerator.WriteFileIfChanged(PBXProjFilePath.FullName, Content.ToString(), Logger, new UTF8Encoding());
 			}
 
-			UEBuildPlatform BuildPlatform = UEBuildPlatform.GetBuildPlatform(Platform);
-			if(BuildPlatform.RequiresArchitectureSuffix())
+
+
+
+			bool bNeedScheme = UnrealData.CanBuildProjectLocally(this, Logger);
+			if (bNeedScheme)
 			{
-				Result.Append(Architecture);
+				if (bSuccess)
+				{
+					string TargetName = ProjectFilePath.GetFileNameWithoutAnyExtensions();
+					string RunTargetGuid = XcodeProjectNode.GetNodesOfType<XcodeRunTarget>(RootProject).First()!.Guid;
+					string? BuildTargetGuid = XcodeProjectNode.GetNodesOfType<XcodeBuildTarget>(RootProject).FirstOrDefault()?.Guid;
+					string? IndexTargetGuid = XcodeProjectNode.GetNodesOfType<XcodeIndexTarget>(RootProject).FirstOrDefault()?.Guid;
+					WriteSchemeFile(TargetName, RunTargetGuid, BuildTargetGuid, IndexTargetGuid, UnrealData.bHasEditorConfiguration, UnrealData.UProjectFileLocation != null ? UnrealData.UProjectFileLocation.FullName : "");
+				}
 			}
+			else
+			{
+				// clean this up because we don't want it persisting if we narrow our project list
+				DirectoryReference SchemeDir = GetProjectSchemeDirectory();
+
+				if (DirectoryReference.Exists(SchemeDir))
+				{
+					DirectoryReference.Delete(SchemeDir, true);
+				}
+			}
+
+			return bSuccess;
 
-            return Result.ToString();
+//			ProcessExtensions(PBXBuildFileSection, PBXFileReferenceSection, PBXCopyExtensionsBuildPhaseSection, PBXResourcesBuildPhaseSection, AllExtensions, GameProjectPath, BuildConfigs, Logger);
 		}
 
+#region Schemes
+
 		private FileReference GetUserSchemeManagementFilePath()
 		{
 			return new FileReference(ProjectFilePath.FullName + "/xcuserdata/" + Environment.UserName + ".xcuserdatad/xcschemes/xcschememanagement.plist");
@@ -1575,9 +2427,9 @@ namespace UnrealBuildTool.XcodeProjectXcconfig
 			return FileReference.Combine(GetProjectSchemeDirectory(), TargetName + ".xcscheme");
 		}
 
-		private void WriteSchemeFile(string TargetName, string TargetGuid, string BuildTargetGuid, string IndexTargetGuid, bool bHasEditorConfiguration, string GameProjectPath)
+		private void WriteSchemeFile(string TargetName, string TargetGuid, string? BuildTargetGuid, string? IndexTargetGuid, bool bHasEditorConfiguration, string GameProjectPath)
 		{
-		
+
 			FileReference SchemeFilePath = GetProjectSchemeFilePathForTarget(TargetName);
 
 			DirectoryReference.CreateDirectory(SchemeFilePath.Directory);
@@ -1598,165 +2450,171 @@ namespace UnrealBuildTool.XcodeProjectXcconfig
 
 			StringBuilder Content = new StringBuilder();
 
-			Content.Append("<?xml version=\"1.0\" encoding=\"UTF-8\"?>" + ProjectFileGenerator.NewLine);
-			Content.Append("<Scheme" + ProjectFileGenerator.NewLine);
-			Content.Append("   LastUpgradeVersion = \"2000\"" + ProjectFileGenerator.NewLine);
-			Content.Append("   version = \"1.3\">" + ProjectFileGenerator.NewLine);
-			Content.Append("   <BuildAction" + ProjectFileGenerator.NewLine);
-			Content.Append("      parallelizeBuildables = \"YES\"" + ProjectFileGenerator.NewLine);
-			Content.Append("      buildImplicitDependencies = \"YES\">" + ProjectFileGenerator.NewLine);
-			Content.Append("      <BuildActionEntries>" + ProjectFileGenerator.NewLine);
-			Content.Append("         <BuildActionEntry" + ProjectFileGenerator.NewLine);
-			Content.Append("            buildForTesting = \"YES\"" + ProjectFileGenerator.NewLine);
-			Content.Append("            buildForRunning = \"YES\"" + ProjectFileGenerator.NewLine);
-			Content.Append("            buildForProfiling = \"YES\"" + ProjectFileGenerator.NewLine);
-			Content.Append("            buildForArchiving = \"YES\"" + ProjectFileGenerator.NewLine);
-			Content.Append("            buildForAnalyzing = \"YES\">" + ProjectFileGenerator.NewLine);
-			Content.Append("            <BuildableReference" + ProjectFileGenerator.NewLine);
-			Content.Append("               BuildableIdentifier = \"primary\"" + ProjectFileGenerator.NewLine);
-			Content.Append("               BlueprintIdentifier = \"" + TargetGuid + "\"" + ProjectFileGenerator.NewLine);
-			Content.Append("               BuildableName = \"" + TargetName + ".app\"" + ProjectFileGenerator.NewLine);
-			Content.Append("               BlueprintName = \"" + TargetName + "\"" + ProjectFileGenerator.NewLine);
-			Content.Append("               ReferencedContainer = \"container:" + TargetName + ".xcodeproj\">" + ProjectFileGenerator.NewLine);
-			Content.Append("            </BuildableReference>" + ProjectFileGenerator.NewLine);
-			Content.Append("         </BuildActionEntry>" + ProjectFileGenerator.NewLine);
-			Content.Append("      </BuildActionEntries>" + ProjectFileGenerator.NewLine);
-			Content.Append("   </BuildAction>" + ProjectFileGenerator.NewLine);
-			Content.Append("   <TestAction" + ProjectFileGenerator.NewLine);
-			Content.Append("      buildConfiguration = \"" + DefaultConfiguration + "\"" + ProjectFileGenerator.NewLine);
-			Content.Append("      selectedDebuggerIdentifier = \"Xcode.DebuggerFoundation.Debugger.LLDB\"" + ProjectFileGenerator.NewLine);
-			Content.Append("      selectedLauncherIdentifier = \"Xcode.DebuggerFoundation.Launcher.LLDB\"" + ProjectFileGenerator.NewLine);
-			Content.Append("      shouldUseLaunchSchemeArgsEnv = \"YES\">" + ProjectFileGenerator.NewLine);
-			Content.Append("      <Testables>" + ProjectFileGenerator.NewLine);
-			Content.Append("      </Testables>" + ProjectFileGenerator.NewLine);
-			Content.Append("      <MacroExpansion>" + ProjectFileGenerator.NewLine);
-			Content.Append("            <BuildableReference" + ProjectFileGenerator.NewLine);
-			Content.Append("               BuildableIdentifier = \"primary\"" + ProjectFileGenerator.NewLine);
-			Content.Append("               BlueprintIdentifier = \"" + TargetGuid + "\"" + ProjectFileGenerator.NewLine);
-			Content.Append("               BuildableName = \"" + TargetName + ".app\"" + ProjectFileGenerator.NewLine);
-			Content.Append("               BlueprintName = \"" + TargetName + "\"" + ProjectFileGenerator.NewLine);
-			Content.Append("               ReferencedContainer = \"container:" + TargetName + ".xcodeproj\">" + ProjectFileGenerator.NewLine);
-			Content.Append("            </BuildableReference>" + ProjectFileGenerator.NewLine);
-			Content.Append("      </MacroExpansion>" + ProjectFileGenerator.NewLine);
-			Content.Append("      <AdditionalOptions>" + ProjectFileGenerator.NewLine);
-			Content.Append("      </AdditionalOptions>" + ProjectFileGenerator.NewLine);
-			Content.Append("   </TestAction>" + ProjectFileGenerator.NewLine);
-			Content.Append("   <LaunchAction" + ProjectFileGenerator.NewLine);
-			Content.Append("      buildConfiguration = \"" + DefaultConfiguration + "\"" + ProjectFileGenerator.NewLine);
-			Content.Append("      selectedDebuggerIdentifier = \"Xcode.DebuggerFoundation.Debugger.LLDB\"" + ProjectFileGenerator.NewLine);
-			Content.Append("      selectedLauncherIdentifier = \"Xcode.DebuggerFoundation.Launcher.LLDB\"" + ProjectFileGenerator.NewLine);
-			Content.Append("      launchStyle = \"0\"" + ProjectFileGenerator.NewLine);
-			Content.Append("      useCustomWorkingDirectory = \"NO\"" + ProjectFileGenerator.NewLine);
-			Content.Append("      ignoresPersistentStateOnLaunch = \"NO\"" + ProjectFileGenerator.NewLine);
-			Content.Append("      debugDocumentVersioning = \"YES\"" + ProjectFileGenerator.NewLine);
-			Content.Append("      debugServiceExtension = \"internal\"" + ProjectFileGenerator.NewLine);
-			Content.Append("      allowLocationSimulation = \"YES\">" + ProjectFileGenerator.NewLine);
-			Content.Append("      <BuildableProductRunnable" + ProjectFileGenerator.NewLine);
-			Content.Append("         runnableDebuggingMode = \"0\">" + ProjectFileGenerator.NewLine);
-			Content.Append("            <BuildableReference" + ProjectFileGenerator.NewLine);
-			Content.Append("               BuildableIdentifier = \"primary\"" + ProjectFileGenerator.NewLine);
-			Content.Append("               BlueprintIdentifier = \"" + TargetGuid + "\"" + ProjectFileGenerator.NewLine);
-			Content.Append("               BuildableName = \"" + TargetName + ".app\"" + ProjectFileGenerator.NewLine);
-			Content.Append("               BlueprintName = \"" + TargetName + "\"" + ProjectFileGenerator.NewLine);
-			Content.Append("               ReferencedContainer = \"container:" + TargetName + ".xcodeproj\">" + ProjectFileGenerator.NewLine);
-			Content.Append("            </BuildableReference>" + ProjectFileGenerator.NewLine);
-			Content.Append("      </BuildableProductRunnable>" + ProjectFileGenerator.NewLine);
+			Content.WriteLine("<?xml version=\"1.0\" encoding=\"UTF-8\"?>");
+			Content.WriteLine("<Scheme");
+			Content.WriteLine("   LastUpgradeVersion = \"2000\"");
+			Content.WriteLine("   version = \"1.3\">");
+			Content.WriteLine("   <BuildAction");
+			Content.WriteLine("      parallelizeBuildables = \"YES\"");
+			Content.WriteLine("      buildImplicitDependencies = \"YES\">");
+			Content.WriteLine("      <BuildActionEntries>");
+			Content.WriteLine("         <BuildActionEntry");
+			Content.WriteLine("            buildForTesting = \"YES\"");
+			Content.WriteLine("            buildForRunning = \"YES\"");
+			Content.WriteLine("            buildForProfiling = \"YES\"");
+			Content.WriteLine("            buildForArchiving = \"YES\"");
+			Content.WriteLine("            buildForAnalyzing = \"YES\">");
+			Content.WriteLine("            <BuildableReference");
+			Content.WriteLine("               BuildableIdentifier = \"primary\"");
+			Content.WriteLine("               BlueprintIdentifier = \"" + TargetGuid + "\"");
+			Content.WriteLine("               BuildableName = \"" + TargetName + ".app\"");
+			Content.WriteLine("               BlueprintName = \"" + TargetName + "\"");
+			Content.WriteLine("               ReferencedContainer = \"container:" + TargetName + ".xcodeproj\">");
+			Content.WriteLine("            </BuildableReference>");
+			Content.WriteLine("         </BuildActionEntry>");
+			Content.WriteLine("      </BuildActionEntries>");
+			Content.WriteLine("   </BuildAction>");
+			Content.WriteLine("   <TestAction");
+			Content.WriteLine("      buildConfiguration = \"" + DefaultConfiguration + "\"");
+			Content.WriteLine("      selectedDebuggerIdentifier = \"Xcode.DebuggerFoundation.Debugger.LLDB\"");
+			Content.WriteLine("      selectedLauncherIdentifier = \"Xcode.DebuggerFoundation.Launcher.LLDB\"");
+			Content.WriteLine("      shouldUseLaunchSchemeArgsEnv = \"YES\">");
+			Content.WriteLine("      <Testables>");
+			Content.WriteLine("      </Testables>");
+			Content.WriteLine("      <MacroExpansion>");
+			Content.WriteLine("            <BuildableReference");
+			Content.WriteLine("               BuildableIdentifier = \"primary\"");
+			Content.WriteLine("               BlueprintIdentifier = \"" + TargetGuid + "\"");
+			Content.WriteLine("               BuildableName = \"" + TargetName + ".app\"");
+			Content.WriteLine("               BlueprintName = \"" + TargetName + "\"");
+			Content.WriteLine("               ReferencedContainer = \"container:" + TargetName + ".xcodeproj\">");
+			Content.WriteLine("            </BuildableReference>");
+			Content.WriteLine("      </MacroExpansion>");
+			Content.WriteLine("      <AdditionalOptions>");
+			Content.WriteLine("      </AdditionalOptions>");
+			Content.WriteLine("   </TestAction>");
+			Content.WriteLine("   <LaunchAction");
+			Content.WriteLine("      buildConfiguration = \"" + DefaultConfiguration + "\"");
+			Content.WriteLine("      selectedDebuggerIdentifier = \"Xcode.DebuggerFoundation.Debugger.LLDB\"");
+			Content.WriteLine("      selectedLauncherIdentifier = \"Xcode.DebuggerFoundation.Launcher.LLDB\"");
+			Content.WriteLine("      launchStyle = \"0\"");
+			Content.WriteLine("      useCustomWorkingDirectory = \"NO\"");
+			Content.WriteLine("      ignoresPersistentStateOnLaunch = \"NO\"");
+			Content.WriteLine("      debugDocumentVersioning = \"YES\"");
+			Content.WriteLine("      debugServiceExtension = \"internal\"");
+			Content.WriteLine("      allowLocationSimulation = \"YES\">");
+			Content.WriteLine("      <BuildableProductRunnable");
+			Content.WriteLine("         runnableDebuggingMode = \"0\">");
+			Content.WriteLine("            <BuildableReference");
+			Content.WriteLine("               BuildableIdentifier = \"primary\"");
+			Content.WriteLine("               BlueprintIdentifier = \"" + TargetGuid + "\"");
+			Content.WriteLine("               BuildableName = \"" + TargetName + ".app\"");
+			Content.WriteLine("               BlueprintName = \"" + TargetName + "\"");
+			Content.WriteLine("               ReferencedContainer = \"container:" + TargetName + ".xcodeproj\">");
+			Content.WriteLine("            </BuildableReference>");
+			Content.WriteLine("      </BuildableProductRunnable>");
 			if (string.IsNullOrEmpty(OldCommandLineArguments))
 			{
 				if (bHasEditorConfiguration && TargetName != "UE5")
 				{
-					Content.Append("      <CommandLineArguments>" + ProjectFileGenerator.NewLine);
+					Content.WriteLine("      <CommandLineArguments>");
 					if (IsForeignProject)
 					{
-						Content.Append("         <CommandLineArgument" + ProjectFileGenerator.NewLine);
-						Content.Append("            argument = \"&quot;" + GameProjectPath + "&quot;\"" + ProjectFileGenerator.NewLine);
-						Content.Append("            isEnabled = \"YES\">" + ProjectFileGenerator.NewLine);
-						Content.Append("         </CommandLineArgument>" + ProjectFileGenerator.NewLine);
+						Content.WriteLine("         <CommandLineArgument");
+						Content.WriteLine("            argument = \"&quot;" + GameProjectPath + "&quot;\"");
+						Content.WriteLine("            isEnabled = \"YES\">");
+						Content.WriteLine("         </CommandLineArgument>");
 					}
 					else
 					{
-						Content.Append("         <CommandLineArgument" + ProjectFileGenerator.NewLine);
-						Content.Append("            argument = \"" + TargetName + "\"" + ProjectFileGenerator.NewLine);
-						Content.Append("            isEnabled = \"YES\">" + ProjectFileGenerator.NewLine);
-						Content.Append("         </CommandLineArgument>" + ProjectFileGenerator.NewLine);
+						Content.WriteLine("         <CommandLineArgument");
+						Content.WriteLine("            argument = \"" + TargetName + "\"");
+						Content.WriteLine("            isEnabled = \"YES\">");
+						Content.WriteLine("         </CommandLineArgument>");
 					}
 					// Always add a configuration argument
-					Content.Append("         <CommandLineArgument" + ProjectFileGenerator.NewLine);
-					Content.Append("            argument = \"-RunConfig=$(Configuration)\"" + ProjectFileGenerator.NewLine);
-					Content.Append("            isEnabled = \"YES\">" + ProjectFileGenerator.NewLine);
-					Content.Append("         </CommandLineArgument>" + ProjectFileGenerator.NewLine);
-					Content.Append("      </CommandLineArguments>" + ProjectFileGenerator.NewLine);
+					Content.WriteLine("         <CommandLineArgument");
+					Content.WriteLine("            argument = \"-RunConfig=$(Configuration)\"");
+					Content.WriteLine("            isEnabled = \"YES\">");
+					Content.WriteLine("         </CommandLineArgument>");
+					Content.WriteLine("      </CommandLineArguments>");
 				}
 			}
 			else
 			{
-				Content.Append("      <CommandLineArguments>" + OldCommandLineArguments + "</CommandLineArguments>");
-			}
-			Content.Append("      <AdditionalOptions>" + ProjectFileGenerator.NewLine);
-			Content.Append("      </AdditionalOptions>" + ProjectFileGenerator.NewLine);
-			Content.Append("   </LaunchAction>" + ProjectFileGenerator.NewLine);
-			Content.Append("   <ProfileAction" + ProjectFileGenerator.NewLine);
-			Content.Append("      buildConfiguration = \"" + DefaultConfiguration + "\"" + ProjectFileGenerator.NewLine);
-			Content.Append("      shouldUseLaunchSchemeArgsEnv = \"YES\"" + ProjectFileGenerator.NewLine);
-			Content.Append("      savedToolIdentifier = \"\"" + ProjectFileGenerator.NewLine);
-			Content.Append("      useCustomWorkingDirectory = \"NO\"" + ProjectFileGenerator.NewLine);
-			Content.Append("      debugDocumentVersioning = \"YES\">" + ProjectFileGenerator.NewLine);
-			Content.Append("      <BuildableProductRunnable" + ProjectFileGenerator.NewLine);
-			Content.Append("         runnableDebuggingMode = \"0\">" + ProjectFileGenerator.NewLine);
-			Content.Append("            <BuildableReference" + ProjectFileGenerator.NewLine);
-			Content.Append("               BuildableIdentifier = \"primary\"" + ProjectFileGenerator.NewLine);
-			Content.Append("               BlueprintIdentifier = \"" + TargetGuid + "\"" + ProjectFileGenerator.NewLine);
-			Content.Append("               BuildableName = \"" + TargetName + ".app\"" + ProjectFileGenerator.NewLine);
-			Content.Append("               BlueprintName = \"" + TargetName + "\"" + ProjectFileGenerator.NewLine);
-			Content.Append("               ReferencedContainer = \"container:" + TargetName + ".xcodeproj\">" + ProjectFileGenerator.NewLine);
-			Content.Append("            </BuildableReference>" + ProjectFileGenerator.NewLine);
-			Content.Append("      </BuildableProductRunnable>" + ProjectFileGenerator.NewLine);
-			Content.Append("   </ProfileAction>" + ProjectFileGenerator.NewLine);
-			Content.Append("   <AnalyzeAction" + ProjectFileGenerator.NewLine);
-			Content.Append("      buildConfiguration = \"" + DefaultConfiguration + "\">" + ProjectFileGenerator.NewLine);
-			Content.Append("   </AnalyzeAction>" + ProjectFileGenerator.NewLine);
-			Content.Append("   <ArchiveAction" + ProjectFileGenerator.NewLine);
-			Content.Append("      buildConfiguration = \"" + DefaultConfiguration + "\"" + ProjectFileGenerator.NewLine);
-			Content.Append("      revealArchiveInOrganizer = \"YES\">" + ProjectFileGenerator.NewLine);
-			Content.Append("   </ArchiveAction>" + ProjectFileGenerator.NewLine);
-			Content.Append("</Scheme>" + ProjectFileGenerator.NewLine);
+				Content.WriteLine("      <CommandLineArguments>" + OldCommandLineArguments + "</CommandLineArguments>");
+			}
+			Content.WriteLine("      <AdditionalOptions>");
+			Content.WriteLine("      </AdditionalOptions>");
+			Content.WriteLine("   </LaunchAction>");
+			Content.WriteLine("   <ProfileAction");
+			Content.WriteLine("      buildConfiguration = \"" + DefaultConfiguration + "\"");
+			Content.WriteLine("      shouldUseLaunchSchemeArgsEnv = \"YES\"");
+			Content.WriteLine("      savedToolIdentifier = \"\"");
+			Content.WriteLine("      useCustomWorkingDirectory = \"NO\"");
+			Content.WriteLine("      debugDocumentVersioning = \"YES\">");
+			Content.WriteLine("      <BuildableProductRunnable");
+			Content.WriteLine("         runnableDebuggingMode = \"0\">");
+			Content.WriteLine("            <BuildableReference");
+			Content.WriteLine("               BuildableIdentifier = \"primary\"");
+			Content.WriteLine("               BlueprintIdentifier = \"" + TargetGuid + "\"");
+			Content.WriteLine("               BuildableName = \"" + TargetName + ".app\"");
+			Content.WriteLine("               BlueprintName = \"" + TargetName + "\"");
+			Content.WriteLine("               ReferencedContainer = \"container:" + TargetName + ".xcodeproj\">");
+			Content.WriteLine("            </BuildableReference>");
+			Content.WriteLine("      </BuildableProductRunnable>");
+			Content.WriteLine("   </ProfileAction>");
+			Content.WriteLine("   <AnalyzeAction");
+			Content.WriteLine("      buildConfiguration = \"" + DefaultConfiguration + "\">");
+			Content.WriteLine("   </AnalyzeAction>");
+			Content.WriteLine("   <ArchiveAction");
+			Content.WriteLine("      buildConfiguration = \"" + DefaultConfiguration + "\"");
+			Content.WriteLine("      revealArchiveInOrganizer = \"YES\">");
+			Content.WriteLine("   </ArchiveAction>");
+			Content.WriteLine("</Scheme>");
 
 			File.WriteAllText(SchemeFilePath.FullName, Content.ToString(), new UTF8Encoding());
 
 			Content.Clear();
 
-			Content.Append("<?xml version=\"1.0\" encoding=\"UTF-8\"?>" + ProjectFileGenerator.NewLine);
-			Content.Append("<!DOCTYPE plist PUBLIC \"-//Apple//DTD PLIST 1.0//EN\" \"http://www.apple.com/DTDs/PropertyList-1.0.dtd\">" + ProjectFileGenerator.NewLine);
-			Content.Append("<plist version=\"1.0\">" + ProjectFileGenerator.NewLine);
-			Content.Append("<dict>" + ProjectFileGenerator.NewLine);
-			Content.Append("\t<key>SchemeUserState</key>" + ProjectFileGenerator.NewLine);
-			Content.Append("\t<dict>" + ProjectFileGenerator.NewLine);
-			Content.Append("\t\t<key>" + TargetName + ".xcscheme_^#shared#^_</key>" + ProjectFileGenerator.NewLine);
-			Content.Append("\t\t<dict>" + ProjectFileGenerator.NewLine);
-			Content.Append("\t\t\t<key>orderHint</key>" + ProjectFileGenerator.NewLine);
-			Content.Append("\t\t\t<integer>1</integer>" + ProjectFileGenerator.NewLine);
-			Content.Append("\t\t</dict>" + ProjectFileGenerator.NewLine);
-			Content.Append("\t</dict>" + ProjectFileGenerator.NewLine);
-			Content.Append("\t<key>SuppressBuildableAutocreation</key>" + ProjectFileGenerator.NewLine);
-			Content.Append("\t<dict>" + ProjectFileGenerator.NewLine);
-			Content.Append("\t\t<key>" + TargetGuid + "</key>" + ProjectFileGenerator.NewLine);
-			Content.Append("\t\t<dict>" + ProjectFileGenerator.NewLine);
-			Content.Append("\t\t\t<key>primary</key>" + ProjectFileGenerator.NewLine);
-			Content.Append("\t\t\t<true/>" + ProjectFileGenerator.NewLine);
-			Content.Append("\t\t</dict>" + ProjectFileGenerator.NewLine);
-			Content.Append("\t\t<key>" + BuildTargetGuid + "</key>" + ProjectFileGenerator.NewLine);
-			Content.Append("\t\t<dict>" + ProjectFileGenerator.NewLine);
-			Content.Append("\t\t\t<key>primary</key>" + ProjectFileGenerator.NewLine);
-			Content.Append("\t\t\t<true/>" + ProjectFileGenerator.NewLine);
-			Content.Append("\t\t</dict>" + ProjectFileGenerator.NewLine);
-			Content.Append("\t\t<key>" + IndexTargetGuid + "</key>" + ProjectFileGenerator.NewLine);
-			Content.Append("\t\t<dict>" + ProjectFileGenerator.NewLine);
-			Content.Append("\t\t\t<key>primary</key>" + ProjectFileGenerator.NewLine);
-			Content.Append("\t\t\t<true/>" + ProjectFileGenerator.NewLine);
-			Content.Append("\t\t</dict>" + ProjectFileGenerator.NewLine);
-			Content.Append("\t</dict>" + ProjectFileGenerator.NewLine);
-			Content.Append("</dict>" + ProjectFileGenerator.NewLine);
-			Content.Append("</plist>" + ProjectFileGenerator.NewLine);
+			Content.WriteLine("<?xml version=\"1.0\" encoding=\"UTF-8\"?>");
+			Content.WriteLine("<!DOCTYPE plist PUBLIC \"-//Apple//DTD PLIST 1.0//EN\" \"http://www.apple.com/DTDs/PropertyList-1.0.dtd\">");
+			Content.WriteLine("<plist version=\"1.0\">");
+			Content.WriteLine("<dict>");
+			Content.WriteLine("\t<key>SchemeUserState</key>");
+			Content.WriteLine("\t<dict>");
+			Content.WriteLine("\t\t<key>" + TargetName + ".xcscheme_^#shared#^_</key>");
+			Content.WriteLine("\t\t<dict>");
+			Content.WriteLine("\t\t\t<key>orderHint</key>");
+			Content.WriteLine("\t\t\t<integer>1</integer>");
+			Content.WriteLine("\t\t</dict>");
+			Content.WriteLine("\t</dict>");
+			Content.WriteLine("\t<key>SuppressBuildableAutocreation</key>");
+			Content.WriteLine("\t<dict>");
+			Content.WriteLine("\t\t<key>" + TargetGuid + "</key>");
+			Content.WriteLine("\t\t<dict>");
+			Content.WriteLine("\t\t\t<key>primary</key>");
+			Content.WriteLine("\t\t\t<true/>");
+			Content.WriteLine("\t\t</dict>");
+			if (BuildTargetGuid != null)
+			{
+				Content.WriteLine("\t\t<key>" + BuildTargetGuid + "</key>");
+				Content.WriteLine("\t\t<dict>");
+				Content.WriteLine("\t\t\t<key>primary</key>");
+				Content.WriteLine("\t\t\t<true/>");
+				Content.WriteLine("\t\t</dict>");
+			}
+			if (IndexTargetGuid != null)
+			{
+				Content.WriteLine("\t\t<key>" + IndexTargetGuid + "</key>");
+				Content.WriteLine("\t\t<dict>");
+				Content.WriteLine("\t\t\t<key>primary</key>");
+				Content.WriteLine("\t\t\t<true/>");
+				Content.WriteLine("\t\t</dict>");
+			}
+			Content.WriteLine("\t</dict>");
+			Content.WriteLine("</dict>");
+			Content.WriteLine("</plist>");
 
 			FileReference ManagementFile = GetUserSchemeManagementFilePath();
 			if (!DirectoryReference.Exists(ManagementFile.Directory))
@@ -1767,274 +2625,21 @@ namespace UnrealBuildTool.XcodeProjectXcconfig
 			File.WriteAllText(ManagementFile.FullName, Content.ToString(), new UTF8Encoding());
 		}
 
+#endregion
 
-		private void WriteTargetXcconfigFiles()
-		{
-			string UEDir = ConvertPath(Path.GetFullPath(Directory.GetCurrentDirectory() + "../../.."));
-
-			//////////////////// RUN ////////////////////////
-			RunXcconfig.AppendLine("GENERATE_INFOPLIST_FILE[sdk=macosx*] = YES // Xcode 14 code-signing fix for development");
-	
-			// #jira UE-143619: Pre Monterey macOS requires this option for a packaged app to run on iOS15 due to new code signature format. Could be removed once Monterey is miniuS.
-			RunXcconfig.AppendLine("OTHER_CODE_SIGN_FLAGS = --generate-entitlement-der");
-
-			RunXcconfig.AppendLine($"MACOSX_DEPLOYMENT_TARGET = {MacToolChain.Settings.MacOSVersion};");
-			RunXcconfig.AppendLine("INFOPLIST_OUTPUT_FORMAT = xml");
-			RunXcconfig.AppendLine("COMBINE_HIDPI_IMAGES = YES");
-
-			RunXcconfig.AppendLine("SDKROOT[sdk=iphoneos*] = iphoneos");
-			RunXcconfig.AppendLine("SDKROOT[sdk=appletvos*] = appletvos");
-			RunXcconfig.AppendLine("SDKROOT[sdk=macosx*] = macosx");
-
-
-			//////////////////// INDEX ////////////////////////
-			IndexXcconfig.AppendLine("ALWAYS_SEARCH_USER_PATHS = NO");
-			IndexXcconfig.AppendLine("CLANG_CXX_LANGUAGE_STANDARD = c++14");
-			IndexXcconfig.AppendLine("GCC_ENABLE_CPP_RTTI = NO");
-			IndexXcconfig.AppendLine("GCC_WARN_CHECK_SWITCH_STATEMENTS = NO");
-			IndexXcconfig.AppendLine("GCC_PRECOMPILE_PREFIX_HEADER = YES");
-			IndexXcconfig.AppendLine($"GCC_PREFIX_HEADER = {UEDir}/Engine/Source/Editor/UnrealEd/Public/UnrealEd.h");
-
-			if (IntelliSensePreprocessorDefinitions.Count > 0)
-			{
-				throw new BuildException("project has defines! " + ProjectFilePath.FullName);
-			}
-			string Definitions = string.Join(" ", IntelliSensePreprocessorDefinitions);
-			IndexXcconfig.AppendLine($"GCC_PREPROCESSOR_DEFINITIONS = __INTELLISENSE__ MONOLITHIC_BUILD=1 {Definitions}");
-
-			string SearchPaths = string.Join(" ", IntelliSenseSystemIncludeSearchPaths.Select(x => x.Contains(" ") ? $"\"{x}\"" : x));
-			IndexXcconfig.AppendLine($"HEADER_SEARCH_PATHS = {SearchPaths}");
-
-			SearchPaths = string.Join(" ", IntelliSenseIncludeSearchPaths.Select(x => x.Contains(" ") ? $"\"{x}\"" : x));
-			IndexXcconfig.AppendLine($"USER_HEADER_SEARCH_PATHS = {SearchPaths}");
-
-
-			//////////////////// BUILD ////////////////////////
-			BuildXcconfig.AppendLine("// Will add more here later");
-
-
-			//////////////////// CLEANUP ////////////////////////
-			// make sure they all written out
-			foreach (XcconfigFile Xcconfig in AllXcconfigFiles)
-			{
-				Xcconfig.Write();
-			}
-		}
-
-
-		public static IEnumerable<UnrealTargetPlatform> GetSupportedPlatforms()
-		{
-			List<UnrealTargetPlatform> SupportedPlatforms = new List<UnrealTargetPlatform>();
-
-			if (XcodeProjectFileGenerator.ProjectFilePlatform.HasFlag(XcodeProjectFileGenerator.XcodeProjectFilePlatform.Mac))
-			{
-				SupportedPlatforms.Add(UnrealTargetPlatform.Mac);
-			}
-
-			if (XcodeProjectFileGenerator.ProjectFilePlatform.HasFlag(XcodeProjectFileGenerator.XcodeProjectFilePlatform.iOS))
-			{
-				SupportedPlatforms.Add(UnrealTargetPlatform.IOS);
-			}
-
-			if (XcodeProjectFileGenerator.ProjectFilePlatform.HasFlag(XcodeProjectFileGenerator.XcodeProjectFilePlatform.tvOS))
-			{
-				SupportedPlatforms.Add(UnrealTargetPlatform.TVOS);
-			}
-
-			return SupportedPlatforms;
-		}
-
-		public static IEnumerable<UnrealTargetConfiguration> GetSupportedConfigurations()
-		{
-			return new UnrealTargetConfiguration[] { 
-				UnrealTargetConfiguration.Debug,
-				UnrealTargetConfiguration.DebugGame,
-				UnrealTargetConfiguration.Development,
-				UnrealTargetConfiguration.Test,
-				UnrealTargetConfiguration.Shipping
-			};
-		}
 
-		public static bool ShouldIncludeProjectInWorkspace(ProjectFile Proj, ILogger Logger)
-		{
-			return CanBuildProjectLocally(Proj, Logger);
-		}
+#region Utilities
 
-		public static bool CanBuildProjectLocally(ProjectFile Proj, ILogger Logger)
+		public override string ToString()
 		{
-			foreach (Project ProjectTarget in Proj.ProjectTargets)
-			{
-				foreach (UnrealTargetPlatform Platform in GetSupportedPlatforms())
-				{
-					foreach (UnrealTargetConfiguration Config in GetSupportedConfigurations())
-					{
-						if (MSBuildProjectFile.IsValidProjectPlatformAndConfiguration(ProjectTarget, Platform, Config, Logger))
-						{
-							return true;
-						}
-					}
-				}
-			}
-
-			return false;
+			return ProjectFilePath.GetFileNameWithoutExtension();
 		}
 
-		/// Implements Project interface
-		public override bool WriteProjectFile(List<UnrealTargetPlatform> InPlatforms, List<UnrealTargetConfiguration> InConfigurations, PlatformProjectGeneratorCollection PlatformProjectGenerators, ILogger Logger)
-		{
-			bool bSuccess = true;
-
-			string TargetName = ProjectFilePath.GetFileNameWithoutExtension();
-			string TargetGuid = XcodeProjectFileGenerator.MakeXcodeGuid();
-			string TargetConfigListGuid = XcodeProjectFileGenerator.MakeXcodeGuid();
-			string TargetDependencyGuid = XcodeProjectFileGenerator.MakeXcodeGuid();
-			string TargetProxyGuid = XcodeProjectFileGenerator.MakeXcodeGuid();
-			string TargetAppGuid = XcodeProjectFileGenerator.MakeXcodeGuid();
-			string BuildTargetName = TargetName + "_Build";
-			string BuildTargetGuid = XcodeProjectFileGenerator.MakeXcodeGuid();
-			string BuildTargetConfigListGuid = XcodeProjectFileGenerator.MakeXcodeGuid();
-			string IndexTargetName = TargetName + "_Index";
-			string IndexTargetGuid = XcodeProjectFileGenerator.MakeXcodeGuid();
-			string IndexTargetConfigListGuid = XcodeProjectFileGenerator.MakeXcodeGuid();
-			string ProjectGuid = XcodeProjectFileGenerator.MakeXcodeGuid();
-			string ProjectConfigListGuid = XcodeProjectFileGenerator.MakeXcodeGuid();
-			string MainGroupGuid = XcodeProjectFileGenerator.MakeXcodeGuid();
-			string ProductRefGroupGuid = XcodeProjectFileGenerator.MakeXcodeGuid();
-			string XcconfigsRefGroupGuid = XcodeProjectFileGenerator.MakeXcodeGuid();
-			string SourcesBuildPhaseGuid = XcodeProjectFileGenerator.MakeXcodeGuid();
-			string CopyExtensionsBuildPhaseGuid = XcodeProjectFileGenerator.MakeXcodeGuid();
-			string ShellScriptSectionGuid = XcodeProjectFileGenerator.MakeXcodeGuid();
-
-
-			// Figure out all the desired configurations
-			List<XcodeBuildConfig> BuildConfigs = GetSupportedBuildConfigs(InPlatforms, InConfigurations, Logger);
-			if (BuildConfigs.Count == 0)
-			{
-				return true;
-			}
-
-			FileReference? GameProjectPath = null;
-			foreach(Project Target in ProjectTargets)
-			{
-				if(Target.UnrealProjectFilePath != null)
-				{
-					GameProjectPath = Target.UnrealProjectFilePath;
-					break;
-				}
-			}
-
-			bool bHasEditorConfiguration = false;
-			bool bIsAppBundle = true;
-
-			Dictionary<string, XcodeBuildConfig> ProjectBuildConfigs = new Dictionary<string, XcodeBuildConfig>();
-			Dictionary<string, XcodeBuildConfig> TargetBuildConfigs = new Dictionary<string, XcodeBuildConfig>();
-			Dictionary<string, XcodeBuildConfig> BuildTargetBuildConfigs = new Dictionary<string, XcodeBuildConfig>();
-			Dictionary<string, XcodeBuildConfig> IndexTargetBuildConfigs = new Dictionary<string, XcodeBuildConfig>();
-			foreach (XcodeBuildConfig Config in BuildConfigs)
-			{
-				ProjectBuildConfigs[XcodeProjectFileGenerator.MakeXcodeGuid()] = Config;
-				TargetBuildConfigs[XcodeProjectFileGenerator.MakeXcodeGuid()] = Config;
-				BuildTargetBuildConfigs[XcodeProjectFileGenerator.MakeXcodeGuid()] = Config;
-				IndexTargetBuildConfigs[XcodeProjectFileGenerator.MakeXcodeGuid()] = Config;
-
-				if (Config.ProjectTarget!.TargetRules!.Type == TargetType.Editor)
-				{
-					bHasEditorConfiguration = true;
-				}
-				else if (Config.ProjectTarget.TargetRules.bIsBuildingConsoleApplication || Config.ProjectTarget.TargetRules.bShouldCompileAsDLL)
-				{
-					bIsAppBundle = false;
-				}
-
-				// make a project-level xcconfig for this configuration
-				XcconfigFile Xcconfig = new XcconfigFile(ProjectFilePath, Config.DisplayName);
-				AllXcconfigFiles.Add(Xcconfig);
-				Xcconfig.AppendLine("// pull in the shared settings for all configs");
-				Xcconfig.AppendLine($"#include \"{ProjectFilePath.GetFileNameWithoutAnyExtensions()}_Project.xcconfig\"");
-				Xcconfig.AppendLine("");
-			}
-
-			WriteProjectXcconfigFiles(ProjectBuildConfigs.Values.ToList(), GameProjectPath, Logger);
-
-			StringBuilder PBXBuildFileSection = new StringBuilder();
-			StringBuilder PBXFileReferenceSection = new StringBuilder();
-			StringBuilder PBXSourcesBuildPhaseSection = new StringBuilder();
-			StringBuilder PBXCopyExtensionsBuildPhaseSection = new StringBuilder();
-			StringBuilder PBXResourcesBuildPhaseSection = new StringBuilder();
-			List<XcodeExtensionInfo> AllExtensions = new List<XcodeExtensionInfo>();
-			GenerateSectionsWithSourceFiles(PBXBuildFileSection, PBXFileReferenceSection, PBXSourcesBuildPhaseSection, TargetAppGuid, TargetName, bIsAppBundle);
-			GenerateSectionsWithExtensions(PBXBuildFileSection, PBXFileReferenceSection, PBXCopyExtensionsBuildPhaseSection, PBXResourcesBuildPhaseSection, AllExtensions, GameProjectPath, BuildConfigs, Logger);
-
-			StringBuilder ProjectFileContent = new StringBuilder();
-
-			ProjectFileContent.Append("// !$*UTF8*$!" + ProjectFileGenerator.NewLine);
-			ProjectFileContent.Append("{" + ProjectFileGenerator.NewLine);
-			ProjectFileContent.Append("\tarchiveVersion = 1;" + ProjectFileGenerator.NewLine);
-			ProjectFileContent.Append("\tclasses = {" + ProjectFileGenerator.NewLine);
-			ProjectFileContent.Append("\t};" + ProjectFileGenerator.NewLine);
-			ProjectFileContent.Append("\tobjectVersion = 46;" + ProjectFileGenerator.NewLine);
-			ProjectFileContent.Append("\tobjects = {" + ProjectFileGenerator.NewLine + ProjectFileGenerator.NewLine);
-
-			AppendBuildFileSection(ProjectFileContent, PBXBuildFileSection);
-			AppendFileReferenceSection(ProjectFileContent, PBXFileReferenceSection);
-			AppendSourcesBuildPhaseSection(ProjectFileContent, PBXSourcesBuildPhaseSection, SourcesBuildPhaseGuid);
-			AppendCopyExtensionsBuildPhaseSection(ProjectFileContent, PBXCopyExtensionsBuildPhaseSection, CopyExtensionsBuildPhaseGuid);
-			ProjectFileContent.Append(PBXResourcesBuildPhaseSection);
-			AppendContainerItemProxySection(ProjectFileContent, BuildTargetName, BuildTargetGuid, TargetProxyGuid, ProjectGuid);
-			if (!XcodeProjectFileGenerator.bGeneratingRunIOSProject)
-			{
-				AppendTargetDependencySection(ProjectFileContent, BuildTargetName, BuildTargetGuid, TargetDependencyGuid, TargetProxyGuid);
-			}
-			foreach (XcodeExtensionInfo EI in AllExtensions)
-			{
-				AppendContainerItemProxySection(ProjectFileContent, EI.Name, EI.TargetGuid, EI.TargetProxyGuid, ProjectGuid);
-				AppendTargetDependencySection(ProjectFileContent, EI.Name, EI.TargetGuid, EI.TargetDependencyGuid, EI.TargetProxyGuid);
-			}
-			AppendGroupSection(ProjectFileContent, MainGroupGuid, ProductRefGroupGuid, XcconfigsRefGroupGuid, TargetAppGuid, TargetName, AllExtensions);
-			AppendLegacyTargetSection(ProjectFileContent, BuildTargetName, BuildTargetGuid, BuildTargetConfigListGuid, GameProjectPath);
-			AppendRunTargetSection(ProjectFileContent, TargetName, TargetGuid, TargetConfigListGuid, TargetDependencyGuid, TargetAppGuid, CopyExtensionsBuildPhaseGuid, ShellScriptSectionGuid, AllExtensions, bIsAppBundle);
-			AppendIndexTargetSection(ProjectFileContent, IndexTargetName, IndexTargetGuid, IndexTargetConfigListGuid, SourcesBuildPhaseGuid);
-			AppendExtensionTargetSections(ProjectFileContent, AllExtensions);
-			AppendProjectSection(ProjectFileContent, TargetName, TargetGuid, BuildTargetName, BuildTargetGuid, IndexTargetName, IndexTargetGuid, MainGroupGuid, ProductRefGroupGuid, ProjectGuid, ProjectConfigListGuid, GameProjectPath, AllExtensions);
-			AppendXCBuildConfigurationSection(ProjectFileContent, ProjectBuildConfigs, TargetBuildConfigs, BuildTargetBuildConfigs, IndexTargetBuildConfigs, GameProjectPath, AllExtensions, Logger);
-			AppendXCConfigurationListSection(ProjectFileContent, TargetName, BuildTargetName, IndexTargetName, ProjectConfigListGuid, ProjectBuildConfigs,
-				TargetConfigListGuid, TargetBuildConfigs, BuildTargetConfigListGuid, BuildTargetBuildConfigs, IndexTargetConfigListGuid, IndexTargetBuildConfigs, AllExtensions);
-			AppendShellScriptSection(ProjectFileContent, ShellScriptSectionGuid, GameProjectPath);
-
-			ProjectFileContent.Append("\t};" + ProjectFileGenerator.NewLine);
-			ProjectFileContent.Append("\trootObject = " + ProjectGuid + " /* Project object */;" + ProjectFileGenerator.NewLine);
-			ProjectFileContent.Append("}" + ProjectFileGenerator.NewLine);
-
-			// write out .xcconfig file(s) for settings we are moving out of the xcode project file
-			WriteTargetXcconfigFiles();
 
-			if (bSuccess)
-			{
-				FileReference PBXProjFilePath = ProjectFilePath + "/project.pbxproj";
-				bSuccess = ProjectFileGenerator.WriteFileIfChanged(PBXProjFilePath.FullName, ProjectFileContent.ToString(), Logger, new UTF8Encoding());
-			}
 
-			bool bNeedScheme = CanBuildProjectLocally(this, Logger);
 
-			if (bNeedScheme)
-			{
-				if (bSuccess)
-				{
-					WriteSchemeFile(TargetName, TargetGuid, BuildTargetGuid, IndexTargetGuid, bHasEditorConfiguration, GameProjectPath != null ? GameProjectPath.FullName : "");
-				}
-			}
-			else
-			{
-				// clean this up because we don't want it persisting if we narrow our project list
-				DirectoryReference SchemeDir = GetProjectSchemeDirectory();
 
-				if (DirectoryReference.Exists(SchemeDir))
-				{
-					DirectoryReference.Delete(SchemeDir, true);
-				}
-			}				
+#endregion
 
-			return bSuccess;
-		}
 	}
 }
diff --git a/Engine/Source/Programs/UnrealBuildTool/ProjectFiles/Xcode/XcodeProjectFileGenerator.cs b/Engine/Source/Programs/UnrealBuildTool/ProjectFiles/Xcode/XcodeProjectFileGenerator.cs
--- a/Engine/Source/Programs/UnrealBuildTool/ProjectFiles/Xcode/XcodeProjectFileGenerator.cs
+++ b/Engine/Source/Programs/UnrealBuildTool/ProjectFiles/Xcode/XcodeProjectFileGenerator.cs
@@ -199,7 +199,7 @@ namespace UnrealBuildTool
 						// Filter out anything that isn't an XC project, and that shouldn't be in the workspace
 						IEnumerable<ProjectFile> SupportedProjects =
 								CurFolder.ChildProjects.Where(P => P.GetType() == typeof(XcodeProjectXcconfig.XcodeProjectFile) || P.GetType() == typeof(XcodeProjectLegacy.XcodeProjectFile))
-									.Where(P => XcodeProjectXcconfig.XcodeProjectFile.ShouldIncludeProjectInWorkspace(P, Logger))
+									.Where(P => XcodeProjectXcconfig.UnrealData.ShouldIncludeProjectInWorkspace(P, Logger))
 									.OrderBy(P => P.ProjectFilePath.GetFileName());
 
 						foreach (ProjectFile XcodeProject in SupportedProjects)
