- Cleaned up Xcconfig file creation/writing (needed for re-using when using a Template project, which is coming up)
- Deleted remnants of Extension code

diff --git a/Engine/Source/Programs/UnrealBuildTool/ProjectFiles/Xcode/XcodeProject.cs b/Engine/Source/Programs/UnrealBuildTool/ProjectFiles/Xcode/XcodeProject.cs
--- a/Engine/Source/Programs/UnrealBuildTool/ProjectFiles/Xcode/XcodeProject.cs
+++ b/Engine/Source/Programs/UnrealBuildTool/ProjectFiles/Xcode/XcodeProject.cs
@@ -137,11 +137,6 @@ namespace UnrealBuildTool.XcodeProjectXcconfig
 		}
 	}
 
-	//class UnrealRemoteExtension
-	//{
-
-	//}
-
 	class UnrealData
 	{
 		public bool bIsAppBundle;
@@ -209,9 +204,6 @@ namespace UnrealBuildTool.XcodeProjectXcconfig
 				return false;
 			}
 
-			// cache any extension info
-			ProcessExtensions(Logger);
-
 
 			// find a uproject file (UE5 target won't have one)
 			foreach (Project Target in ProjectFile.ProjectTargets)
@@ -467,128 +459,6 @@ namespace UnrealBuildTool.XcodeProjectXcconfig
 			return BuildConfigs;
 		}
 
-		private void ProcessExtensions(ILogger Logger)
-		{
-			if (UProjectFileLocation != null)
-			{
-				DirectoryReference ProjectExtensionsDir = DirectoryReference.Combine(UProjectFileLocation.Directory, "Build/IOS/Extensions");
-				//string ProjectIntermediateDir = Path.Combine(Path.GetDirectoryName(UProjectPath.FullName), "Intermediate/IOS/Extensions");
-
-				if (DirectoryReference.Exists(ProjectExtensionsDir))
-				{
-					foreach (DirectoryReference ExtensionDir in DirectoryReference.EnumerateDirectories(ProjectExtensionsDir))
-					{
-						// assume each Extension directory in here will create a resulting Extension.appex
-						string ExtensionName = ExtensionDir.GetDirectoryName();
-						Logger.LogInformation("  Project {0} has Extension {1}!", UProjectFileLocation.GetFileName(), ExtensionName);
-
-						string Extension = ExtensionName + ".appex";
-
-						string ExtensionGuid = XcodeProjectFileGenerator.MakeXcodeGuid();
-
-						// make an extension info object
-						UnrealExtensionInfo ExtensionInfo = new UnrealExtensionInfo(ExtensionName);
-						AllExtensions.Add(ExtensionInfo);
-#if false
-						PBXBuildFileSection.Append(string.Format("\t\t{0} /* {1} in Embed App Extensions */ = {{isa = PBXBuildFile; fileRef = {2} /* {1} */; settings = {{ATTRIBUTES = (RemoveHeadersOnCopy, ); }}; }};",
-							ExtensionGuid,
-							Extension,
-							ExtensionInfo.ProductGuid));
-
-						PBXFileReferenceSection.Append(string.Format("\t\t{0} /* {1} */ = {{isa = PBXFileReference; explicitFileType = wrapper.app-extension; path = \"{1}\"; sourceTree = BUILT_PRODUCTS_DIR; }};",
-							ExtensionInfo.ProductGuid,
-							Extension));
-
-						PBXCopyFilesBuildPhaseSection.Append(string.Format("\t\t\t\t{0} /* {1} in Embed App Extensions */,",
-							ExtensionGuid,
-							Extension));
-
-						PBXResourcesBuildPhaseSection.Append("/* Begin PBXResourcesBuildPhase section */");
-						PBXResourcesBuildPhaseSection.Append("\t\t" + ExtensionInfo.ResourceBuildPhaseGuid + " /* Resources */ = {");
-						PBXResourcesBuildPhaseSection.Append("\t\t\tisa = PBXResourcesBuildPhase;");
-						PBXResourcesBuildPhaseSection.Append("\t\t\tbuildActionMask = 2147483647;");
-						PBXResourcesBuildPhaseSection.Append("\t\t\tfiles = (");
-						if (Directory.Exists(Path.Combine(DI.FullName, "Resources")))
-						{
-							DirectoryInfo ResourceDir = new System.IO.DirectoryInfo(Path.Combine(DI.FullName, "Resources"));
-							foreach (FileSystemInfo FSI in ResourceDir.EnumerateFileSystemInfos())
-							{
-								if (FSI.Name.StartsWith("."))
-								{
-									continue;
-								}
-								// for each resource, put it into the File/FileRef section, and into the CopyResuorceBuildPhase
-								string ResourceGuid = XcodeProjectFileGenerator.MakeXcodeGuid();
-								string ResourceRefGuid = XcodeProjectFileGenerator.MakeXcodeGuid();
-								PBXBuildFileSection.Append(string.Format("\t\t{0} /* {1} in Embed App Extensions */ = {{isa = PBXBuildFile; fileRef = {2} /* {1} */; }};",
-									ResourceGuid,
-									FSI.Name,
-									ResourceRefGuid));
-
-								// lastKnownFileType = wrapper.app-extension; 
-								PBXFileReferenceSection.Append(string.Format("\t\t{0} /* {1} */ = {{isa = PBXFileReference; lastKnownFileType = folder.assetcatalog; path = \"{2}\"; sourceTree = \"<absolute>\"; }};",
-									ResourceRefGuid,
-									FSI.Name,
-									// @todo: make this relative path!! 
-									FSI.FullName));
-
-								PBXResourcesBuildPhaseSection.Append("\t\t\t\t" + ResourceGuid + " /* " + FSI.Name + " in " + ResourceDir.Name + " */,");
-							}
-						}
-						PBXResourcesBuildPhaseSection.Append("\t\t\t);");
-						PBXResourcesBuildPhaseSection.Append("\t\t\trunOnlyForDeploymentPostprocessing = 0;");
-						PBXResourcesBuildPhaseSection.Append("\t\t};");
-						PBXResourcesBuildPhaseSection.Append("/* End PBXResourcesBuildPhase section */");
-
-						StringBuilder ConfigSection = new StringBuilder();
-						// copy over the configs from the general project to the extension
-						foreach (XcodeBuildConfig Configuration in BuildConfigs)
-						{
-							string ConfigGuid = XcodeProjectFileGenerator.MakeXcodeGuid();
-							string ConfigName = Configuration.DisplayName;
-
-							ConfigSection.Append("\t\t" + ConfigGuid + " /* " + ConfigName + " */ = {");
-							ConfigSection.Append("\t\t\tisa = XCBuildConfiguration;");
-							ConfigSection.Append("\t\t\tbuildSettings = {");
-							ConfigSection.Append("\t\t\t\tASSETCATALOG_COMPILER_APPICON_NAME = \"iMessage App Icon\";");
-							ConfigSection.Append("\t\t\t\tINFOPLIST_FILE = \"" + Path.Combine(DI.FullName, "Info.plist") + "\";");
-							ConfigSection.Append("\t\t\t\tSKIP_INSTALL = YES;");
-							ConfigSection.Append("\t\t\t\tPRODUCT_NAME = \"$(TARGET_NAME)\";");
-
-							bool bSupportIOS = true;
-							bool bSupportTVOS = true;
-							if (bSupportIOS && InstalledPlatformInfo.IsValidPlatform(UnrealTargetPlatform.IOS, EProjectType.Code))
-							{
-								IOSPlatform IOSPlatform = ((IOSPlatform)UEBuildPlatform.GetBuildPlatform(UnrealTargetPlatform.IOS));
-								IOSProjectSettings ProjectSettings = IOSPlatform.ReadProjectSettings(UProjectPath);
-								ConfigSection.Append("\t\t\t\t\"PRODUCT_BUNDLE_IDENTIFIER[sdk=iphone*]\" = " + ProjectSettings.BundleIdentifier + "." + ExtensionInfo.Name + ";");
-							}
-
-							if (bSupportTVOS && InstalledPlatformInfo.IsValidPlatform(UnrealTargetPlatform.TVOS, EProjectType.Code))
-							{
-								TVOSPlatform TVOSPlatform = ((TVOSPlatform)UEBuildPlatform.GetBuildPlatform(UnrealTargetPlatform.TVOS));
-								TVOSProjectSettings ProjectSettings = TVOSPlatform.ReadProjectSettings(UProjectPath);
-								ConfigSection.Append("\t\t\t\t\"PRODUCT_BUNDLE_IDENTIFIER[sdk=appletvos*]\" = " + ProjectSettings.BundleIdentifier + "." + ExtensionInfo.Name + ";");
-							}
-
-							string? IOSRuntimeVersion, TVOSRuntimeVersion;
-							AppendPlatformConfiguration(ConfigSection, Configuration, ExtensionInfo.Name, UProjectPath, false, bSupportIOS, bSupportTVOS, Logger, out IOSRuntimeVersion, out TVOSRuntimeVersion);
-
-							ConfigSection.Append("\t\t\t};");
-							ConfigSection.Append("\t\t\tname = \"" + ConfigName + "\";");
-							ConfigSection.Append("\t\t};");
-
-							XcodeBuildConfig Config = new XcodeBuildConfig(ConfigName, ExtensionInfo.Name, null, null, null, null, Configuration.BuildConfig);
-							ExtensionInfo.AllConfigs.Add(ConfigGuid, Config);
-						}
-
-						ExtensionInfo.ConfigurationContents = ConfigSection.ToString();
-#endif
-					}
-				}
-			}
-		}
-
 		public static IEnumerable<UnrealTargetPlatform> GetSupportedPlatforms()
 		{
 			List<UnrealTargetPlatform> SupportedPlatforms = new List<UnrealTargetPlatform>();
@@ -844,18 +714,15 @@ namespace UnrealBuildTool.XcodeProjectXcconfig
 	class XcconfigFile
 	{
 		public string Name;
-		//public string Filename;
 		public string Guid;
-//		public FileReference FileRef;
+		public FileReference FileRef;
 		internal StringBuilder Text;
 
-		public XcconfigFile(string ConfigName)
+		public XcconfigFile(DirectoryReference XcodeProjectDirectory, string ConfigName)
 		{
 			Name = ConfigName;
-//			FileRef = FileReference.Combine(ProjectFile.Directory, "Xcconfigs", $"{ProjectFile.GetFileNameWithoutAnyExtensions()}_{Name.Replace(" ", "")}.xcconfig");
-//			Filename = FileRef.GetFileName();
+			FileRef = FileReference.Combine(XcodeProjectDirectory, "Xcconfigs", $"{ConfigName}.xcconfig");
 			Guid = XcodeProjectFileGenerator.MakeXcodeGuid();
-//			DirectoryReference.CreateDirectory(FileRef.Directory);
 			Text = new StringBuilder();
 		}
 
@@ -867,7 +734,9 @@ namespace UnrealBuildTool.XcodeProjectXcconfig
 
 		public void Write()
 		{
-//			FileReference.WriteAllTextIfDifferent(FileRef, Text.ToString());
+			// write the file to disk
+			DirectoryReference.CreateDirectory(FileRef.Directory);
+			FileReference.WriteAllTextIfDifferent(FileRef, Text.ToString());
 		}
 	}
 
@@ -879,6 +748,9 @@ namespace UnrealBuildTool.XcodeProjectXcconfig
 		// keeps a list of other node this node references, which is used when writing out the whole xcode project file
 		public List<XcodeProjectNode> References = new();
 
+		// optional Xcconfig file 
+		public XcconfigFile? Xcconfig = null;
+
 		/// <summary>
 		/// Abstract function the individual node classes must override to write out the node to the project file
 		/// </summary>
@@ -915,45 +787,10 @@ namespace UnrealBuildTool.XcodeProjectXcconfig
 		}
 
 
-		/// <summary>
-		/// Called by WriteXcconfigFiles, for any node that needs Xcconfig files, the subclass will overrite this to create one or more
-		/// Xcconfig files and add them to the list. WriteXcconfigFiles will then manage those files as needed
-		/// </summary>
-		/// <param name="XcconfigFiles"></param>
-		internal virtual void CreateXcconfigFiles(List<XcconfigFile> XcconfigFiles)
-		{
-		}
-
-		public static void WriteXcconfigFiles(XcodeProjectNode RootNode, XcodeProjectFile ProjectFile)
-		{
-			List<XcconfigFile> FinalXcconfigFiles = new List<XcconfigFile>();
-
-			// recursively walk the nodes
-			WriteXcconfigFilesInternal(FinalXcconfigFiles, RootNode);
-
-			// now write out all of the xcconfig files and add to project
-			foreach (XcconfigFile File in FinalXcconfigFiles)
-			{
-				FileReference FileRef = FileReference.Combine(ProjectFile.ProjectFilePath.Directory, "Xcconfigs", $"{File.Name}.xcconfig");
-				DirectoryReference.CreateDirectory(FileRef.Directory);
-
-				FileReference.WriteAllTextIfDifferent(FileRef, File.Text.ToString());
-
-				ProjectFile.FileCollection.AddFileReference(File.Guid, $"Xcconfigs/{FileRef.GetFileName()}", "test.xcconfig", "\"<group>\"", "Xcconfigs");
-			}
-
-		}
-
-		private static void WriteXcconfigFilesInternal(List<XcconfigFile> FinalXcconfigFiles, XcodeProjectNode Node)
+		public void CreateXcconfigFile(XcodeProject Project, string Name)
 		{
-			// allow the node to create optional Xcconfig file(s)
-			Node.CreateXcconfigFiles(FinalXcconfigFiles);
-
-			// move on to dependencies
-			foreach (XcodeProjectNode Reference in Node.References)
-			{
-				WriteXcconfigFilesInternal(FinalXcconfigFiles, Reference);
-			}
+			Xcconfig = new XcconfigFile(Project.UnrealData.XcodeProjectFileLocation!.Directory, Name);
+			Project.FileCollection.AddFileReference(Xcconfig.Guid, $"Xcconfigs/{Xcconfig.FileRef.GetFileName()}", "test.xcconfig", "\"<group>\"", "Xcconfigs");
 		}
 
 		/// <summary>
@@ -965,6 +802,10 @@ namespace UnrealBuildTool.XcodeProjectXcconfig
 		{
 			// write the node into the xcode project file
 			Node.Write(Content);
+			if (Node.Xcconfig != null)
+			{
+				Node.Xcconfig.Write();
+			}
 
 			foreach (XcodeProjectNode Reference in Node.References)
 			{
@@ -1051,7 +892,6 @@ namespace UnrealBuildTool.XcodeProjectXcconfig
 	{
 		public string Guid = XcodeProjectFileGenerator.MakeXcodeGuid();
 		public UnrealBuildConfig Info;
-		public XcconfigFile? Xcconfig = null;
 
 		public XcodeBuildConfig(UnrealBuildConfig Info)
 		{
@@ -1110,8 +950,6 @@ namespace UnrealBuildTool.XcodeProjectXcconfig
 			Content.WriteLine(3,	"defaultConfigurationIsVisible = 0;");
 			Content.WriteLine(3,	$"defaultConfigurationName = \"{Default}\";");
 			Content.WriteLine(2, "};");
-
-			// now write the ba
 		}
 	}
 
@@ -1123,7 +961,6 @@ namespace UnrealBuildTool.XcodeProjectXcconfig
 			Run_Tool,
 			Build,
 			Index,
-			Extension,
 		}
 
 		// Guid for this target
@@ -1140,9 +977,6 @@ namespace UnrealBuildTool.XcodeProjectXcconfig
 		// UE5_Build, EngineTest_SubIndex1, etc
 		public string Name;
 
-		// the config shared by all BuildConfigs
-//		public XcconfigFile Xcconfig;
-
 		// list of build configs this target supports (for instance, the Index target only indexes a Development config) 
 		public XcodeBuildConfigList? BuildConfigList;
 
@@ -1182,11 +1016,6 @@ namespace UnrealBuildTool.XcodeProjectXcconfig
 					TargetTypeName = "PBXNativeTarget";
 					ConfigName = "Index";
 					break;
-				case Type.Extension:
-					ProductType = "com.apple.product-type.app-extension.messages-sticker-pack";
-					TargetTypeName = "PBXNativeTarget";
-					ConfigName = "Extension";
-					break;
 				default:
 					throw new BuildException($"Unhandled target type {Type}");
 			}
@@ -1279,28 +1108,20 @@ namespace UnrealBuildTool.XcodeProjectXcconfig
 					XcodeBuildTarget BuildTarget = new XcodeBuildTarget(Project.UnrealData);
 					AddDependency(BuildTarget, Project);
 				}
-
-				// add a target for each extension that this depends on
-				foreach (UnrealExtensionInfo Extension in Project.UnrealData.AllExtensions)
-				{
-					XcodeExtensionTarget Target = new XcodeExtensionTarget(Project.UnrealData, Extension);
-					AddDependency(Target, Project);
-				}
 			}
-		}
 
-		protected override void WriteExtraTargetProperties(StringBuilder Content)
-		{
-			Content.WriteLine($"\t\t\tproductReference = {ProductGuid};");
-			Content.WriteLine($"\t\t\tproductName = \"{ProductName}\";");
+			CreateXcconfigFile(Project, Name);
+
+			// hook up each buildconfig to this Xcconfig 
+			BuildConfigList!.BuildConfigs.ForEach(x => x.Xcconfig = Xcconfig);
 		}
 
-		internal override void CreateXcconfigFiles(List<XcconfigFile> XcconfigFiles)
+		public override void Write(StringBuilder Content)
 		{
-			XcconfigFile Xcconfig = new XcconfigFile($"{Name}");
-			XcconfigFiles.Add(Xcconfig);
+			base.Write(Content);
 
-			Xcconfig.AppendLine("GENERATE_INFOPLIST_FILE[sdk=macosx*] = YES // Xcode 14 code-signing fix for development");
+			// also write to the Xcconfig file
+			Xcconfig!.AppendLine("GENERATE_INFOPLIST_FILE[sdk=macosx*] = YES // Xcode 14 code-signing fix for development");
 
 			// #jira UE-143619: Pre Monterey macOS requires this option for a packaged app to run on iOS15 due to new code signature format. Could be removed once Monterey is miniuS.
 			Xcconfig.AppendLine("OTHER_CODE_SIGN_FLAGS = --generate-entitlement-der");
@@ -1308,9 +1129,12 @@ namespace UnrealBuildTool.XcodeProjectXcconfig
 			Xcconfig.AppendLine($"MACOSX_DEPLOYMENT_TARGET = {MacToolChain.Settings.MacOSVersion};");
 			Xcconfig.AppendLine("INFOPLIST_OUTPUT_FORMAT = xml");
 			Xcconfig.AppendLine("COMBINE_HIDPI_IMAGES = YES");
+		}
 
-			// hook up each buildconfig to this Xcconfig 
-			BuildConfigList!.BuildConfigs.ForEach(x => x.Xcconfig = Xcconfig);
+		protected override void WriteExtraTargetProperties(StringBuilder Content)
+		{
+			Content.WriteLine($"\t\t\tproductReference = {ProductGuid};");
+			Content.WriteLine($"\t\t\tproductName = \"{ProductName}\";");
 		}
 	}
 
@@ -1337,6 +1161,10 @@ namespace UnrealBuildTool.XcodeProjectXcconfig
 			BuildConfigList = new XcodeBuildConfigList(Name, Project.UnrealData.AllConfigs.Where(x => x.DisplayName == IndexConfig).ToList());
 			References.Add(BuildConfigList);
 
+			CreateXcconfigFile(Project, Name);
+			// hook up each buildconfig to this Xcconfig 
+			BuildConfigList!.BuildConfigs.ForEach(x => x.Xcconfig = Xcconfig);
+
 
 			// now look to see if all files are buildable with one batch, in which case we don't need subtargets, or if we have multiple batches, then 
 			// make a subtarget for each batch of unique build settings
@@ -1360,22 +1188,19 @@ namespace UnrealBuildTool.XcodeProjectXcconfig
 
 				for (int Index = 0; Index < Project.UnrealData.BatchedFiles.Count; Index++)
 				{
-					XcodeIndexSubTarget SubTarget = new XcodeIndexSubTarget(Project.UnrealData, Index, Configs);
+					XcodeIndexSubTarget SubTarget = new XcodeIndexSubTarget(Project, Index, Configs);
 					AddDependency(SubTarget, Project);
 				}
 			}
 		}
 
-
-		internal override void CreateXcconfigFiles(List<XcconfigFile> XcconfigFiles)
+		public override void Write(StringBuilder Content)
 		{
-			// name the xcconfig after the target name
-			XcconfigFile Xcconfig = new XcconfigFile(Name);
-			XcconfigFiles.Add(Xcconfig);
+			base.Write(Content);
 
 			// write out settings that apply whether or not we have subtargets, which #include this one, or no subtargets, and we write out more
 			// @todo move tis to the subtarget and remember it from the Module
-			Xcconfig.AppendLine("CLANG_CXX_LANGUAGE_STANDARD = c++17");
+			Xcconfig!.AppendLine("CLANG_CXX_LANGUAGE_STANDARD = c++17");
 			Xcconfig.AppendLine("GCC_WARN_CHECK_SWITCH_STATEMENTS = NO");
 			Xcconfig.AppendLine("GCC_PRECOMPILE_PREFIX_HEADER = YES");
 			Xcconfig.AppendLine("GCC_OPTIMIZATION_LEVEL = 0");
@@ -1383,12 +1208,9 @@ namespace UnrealBuildTool.XcodeProjectXcconfig
 			// if we have folded a file batch into this main target, write out the settings into this Xcconfig
 			if (SoloTargetBatch != null)
 			{
-				// ask the singl sub target to write out the contents of our Xcconfig file since we are folding them up into this target
+				// ask the single sub target to write out the contents of our Xcconfig file since we are folding them up into this target
 				SoloTargetBatch.WriteXcconfigSettings(Xcconfig.Text);
 			}
-
-			// hook up each buildconfig to this Xcconfig 
-			BuildConfigList!.BuildConfigs.ForEach(x => x.Xcconfig = Xcconfig);
 		}
 	}
 
@@ -1400,16 +1222,20 @@ namespace UnrealBuildTool.XcodeProjectXcconfig
 
 		private UnrealBatchedFiles BatchedFiles;
 
-		public XcodeIndexSubTarget(UnrealData UnrealData, int SubTargetIndex, List<UnrealBuildConfig> Configs)
-			: base(Type.Index, UnrealData, $"SubIndex{SubTargetIndex}")
+		public XcodeIndexSubTarget(XcodeProject Project, int SubTargetIndex, List<UnrealBuildConfig> Configs)
+			: base(Type.Index, Project.UnrealData, $"SubIndex{SubTargetIndex}")
 		{
 			// get the settings and list of files for this subtarget
-			BatchedFiles = UnrealData.BatchedFiles[SubTargetIndex];
+			BatchedFiles = Project.UnrealData.BatchedFiles[SubTargetIndex];
 
 			// set the buildconfig list in the SubTarget
 			BuildConfigList = new XcodeBuildConfigList(Name, Configs);
 			References.Add(BuildConfigList);
 
+			CreateXcconfigFile(Project, Name);
+			// hook up each buildconfig to this Xcconfig 
+			BuildConfigList!.BuildConfigs.ForEach(x => x.Xcconfig = Xcconfig);
+
 			// create the build phase for this target for all files in the batch
 			XcodeBuildPhase BuildPhase = new XcodeBuildPhase("Sources");
 			BuildPhases.Add(BuildPhase);
@@ -1422,41 +1248,18 @@ namespace UnrealBuildTool.XcodeProjectXcconfig
 			}
 		}
 
-		internal override void CreateXcconfigFiles(List<XcconfigFile> XcconfigFiles)
+		public override void Write(StringBuilder Content)
 		{
-			// name the xcconfig after the target name
-			XcconfigFile Xcconfig = new XcconfigFile(Name);
-			XcconfigFiles.Add(Xcconfig);
+			base.Write(Content);
 
-			Xcconfig.AppendLine($"#include \"{Name.Split('_')[0]}_Index.xcconfig\"");
+			// include the share Index config
+			Xcconfig!.AppendLine($"#include \"{Name.Split('_')[0]}_Index.xcconfig\"");
 
+			// then right the build settings from the modules we are batched together
 			BatchedFiles.WriteXcconfigSettings(Xcconfig.Text);
-		
-			// hook up each buildconfig to this Xcconfig 
-			BuildConfigList!.BuildConfigs.ForEach(x => x.Xcconfig = Xcconfig);
-		}
-	}
-
-	class XcodeExtensionTarget : XcodeTarget
-	{
-		private UnrealExtensionInfo ExtensionInfo;
-
-		public XcodeExtensionTarget(UnrealData UnrealData, UnrealExtensionInfo ExtensionInfo)
-			: base(Type.Extension, UnrealData)
-		{
-			this.ExtensionInfo = ExtensionInfo;
-
-			BuildConfigList = new XcodeBuildConfigList($"{UnrealData.ProductName}_{ExtensionInfo.Name}", UnrealData.AllConfigs.Where(x => x.DisplayName == "Development").ToList());
-		}
-
-		protected override void WriteExtraTargetProperties(StringBuilder Content)
-		{
-			Content.WriteLine($"\t\t\tproductReference = {ExtensionInfo.ProductGuid};");
-			Content.WriteLine($"\t\t\tproductName = \"{ExtensionInfo.Name}\";");
 		}
 	}
 
-
 	class XcodeProject : XcodeProjectNode
 	{
 		// the blob of data coming from unreal that we can pass around
@@ -1485,6 +1288,16 @@ namespace UnrealBuildTool.XcodeProjectXcconfig
 			ProjectBuildConfigs = new XcodeBuildConfigList(UnrealData.ProductName, UnrealData.AllConfigs);
 			References.Add(ProjectBuildConfigs);
 
+			// create the Projet xcconfig
+			CreateXcconfigFile(this, $"{UnrealData.ProductName}_Project");
+
+			// create per-config Xcconfig files
+			foreach (XcodeBuildConfig Config in ProjectBuildConfigs.BuildConfigs)
+			{
+				Config.CreateXcconfigFile(this, $"{UnrealData.ProductName}_{Config.Info.DisplayName.Replace(" ", "")}");
+			}
+
+
 			// make an indexing target if we aren't just a run-only project, and it has buildable source files
 			if (!XcodeProjectFileGenerator.bGeneratingRunIOSProject && UnrealData.BatchedFiles.Count != 0)
 			{				
@@ -1533,14 +1346,13 @@ namespace UnrealBuildTool.XcodeProjectXcconfig
 			Content.WriteLine(2, "};");
 
 			Content.WriteLine("/* End PBXProject section */");
+
+			WriteXcconfigFiles();
 		}
 
 
-		internal override void CreateXcconfigFiles(List<XcconfigFile> XcconfigFiles)
+		private void WriteXcconfigFiles()
 		{
-			XcconfigFile ProjectXcconfig = new XcconfigFile($"{UnrealData.ProductName}_Project");
-			XcconfigFiles.Add(ProjectXcconfig);
-
 			string UEDir = XcodeFileCollection.ConvertPath(Path.GetFullPath(Directory.GetCurrentDirectory() + "../../.."));
 			string SupportedIOSArchitectures = string.Join(" ", UnrealData.SupportedIOSArchitectures);
 			string? IOSRunTimeVersion = null;
@@ -1589,53 +1401,53 @@ namespace UnrealBuildTool.XcodeProjectXcconfig
 
 			//#jira UE-50382 Xcode Address Sanitizer feature does not work on iOS
 			// address sanitizer dylib loader depends on the SDKROOT parameter. For macosx or default (missing, translated as macosx), the path is incorrect for iphone/appletv
-			ProjectXcconfig.AppendLine($"SDKROOT[sdk=macosx*] = macosx");
-			ProjectXcconfig.AppendLine($"SDKROOT[sdk=iphone*] = iphoneos");
-			ProjectXcconfig.AppendLine($"SDKROOT[sdk=appletvos*] = appletvos");
-			ProjectXcconfig.AppendLine("USE_HEADERMAP = NO");
-			ProjectXcconfig.AppendLine("ONLY_ACTIVE_ARCH = YES");
+			Xcconfig!.AppendLine($"SDKROOT[sdk=macosx*] = macosx");
+			Xcconfig.AppendLine($"SDKROOT[sdk=iphone*] = iphoneos");
+			Xcconfig.AppendLine($"SDKROOT[sdk=appletvos*] = appletvos");
+			Xcconfig.AppendLine("USE_HEADERMAP = NO");
+			Xcconfig.AppendLine("ONLY_ACTIVE_ARCH = YES");
 
 			if (bAutomaticSigning)
 			{
-				ProjectXcconfig.AppendLine("CODE_SIGN_STYLE = Automatic");
+				Xcconfig.AppendLine("CODE_SIGN_STYLE = Automatic");
 			}
 
 			if (IOSRunTimeVersion != null)
 			{
-				ProjectXcconfig.AppendLine($"VALID_ARCHS[sdk=iphone*] = {SupportedIOSArchitectures}");
-				ProjectXcconfig.AppendLine($"IPHONEOS_DEPLOYMENT_TARGET = {IOSRunTimeVersion}");
-				ProjectXcconfig.AppendLine($"TARGETED_DEVICE_FAMILY[sdk=iphone*] = {IOSRunTimeDevices}");
+				Xcconfig.AppendLine($"VALID_ARCHS[sdk=iphone*] = {SupportedIOSArchitectures}");
+				Xcconfig.AppendLine($"IPHONEOS_DEPLOYMENT_TARGET = {IOSRunTimeVersion}");
+				Xcconfig.AppendLine($"TARGETED_DEVICE_FAMILY[sdk=iphone*] = {IOSRunTimeDevices}");
 				if (!string.IsNullOrEmpty(TEAM_IOS))
 				{
-					ProjectXcconfig.AppendLine($"DEVELOPMENT_TEAM[sdk=iphone*] = {TEAM_IOS}");
+					Xcconfig.AppendLine($"DEVELOPMENT_TEAM[sdk=iphone*] = {TEAM_IOS}");
 				}
-				ProjectXcconfig.AppendLine($"CODE_SIGN_IDENTITY[sdk=iphone*] = {IOS_CERT}");
+				Xcconfig.AppendLine($"CODE_SIGN_IDENTITY[sdk=iphone*] = {IOS_CERT}");
 				if (!bAutomaticSigning && !string.IsNullOrEmpty(UUID_IOS))
 				{
-					ProjectXcconfig.AppendLine($"PROVISIONING_PROFILE_SPECIFIER[sdk=iphone*] = {UUID_IOS}");
+					Xcconfig.AppendLine($"PROVISIONING_PROFILE_SPECIFIER[sdk=iphone*] = {UUID_IOS}");
 				}
 				if (UnrealData.UProjectFileLocation != null)
 				{
-					ProjectXcconfig.AppendLine($"PRODUCT_BUNDLE_IDENTIFIER[sdk=iphone*] = {IOS_BUNDLE}");
+					Xcconfig.AppendLine($"PRODUCT_BUNDLE_IDENTIFIER[sdk=iphone*] = {IOS_BUNDLE}");
 				}
 			}
 			if (TVOSRunTimeVersion != null)
 			{
-				ProjectXcconfig.AppendLine($"VALID_ARCHS[sdk=appletvos*] = {SupportedIOSArchitectures}");
-				ProjectXcconfig.AppendLine($"IPHONEOS_DEPLOYMENT_TARGET = {TVOSRunTimeVersion}");
-				ProjectXcconfig.AppendLine($"TARGETED_DEVICE_FAMILY[sdk=appletvos*] = {TVOSRunTimeDevices}");
+				Xcconfig.AppendLine($"VALID_ARCHS[sdk=appletvos*] = {SupportedIOSArchitectures}");
+				Xcconfig.AppendLine($"IPHONEOS_DEPLOYMENT_TARGET = {TVOSRunTimeVersion}");
+				Xcconfig.AppendLine($"TARGETED_DEVICE_FAMILY[sdk=appletvos*] = {TVOSRunTimeDevices}");
 				if (!string.IsNullOrEmpty(TEAM_TVOS))
 				{
-					ProjectXcconfig.AppendLine($"DEVELOPMENT_TEAM[sdk=iphone*] = {TEAM_TVOS}");
+					Xcconfig.AppendLine($"DEVELOPMENT_TEAM[sdk=iphone*] = {TEAM_TVOS}");
 				}
-				ProjectXcconfig.AppendLine($"CODE_SIGN_IDENTITY[sdk=appletvos*] = {TVOS_CERT}");
+				Xcconfig.AppendLine($"CODE_SIGN_IDENTITY[sdk=appletvos*] = {TVOS_CERT}");
 				if (!bAutomaticSigning && !string.IsNullOrEmpty(UUID_TVOS))
 				{
-					ProjectXcconfig.AppendLine($"PROVISIONING_PROFILE_SPECIFIER[sdk=appletvos*] = {UUID_TVOS}");
+					Xcconfig.AppendLine($"PROVISIONING_PROFILE_SPECIFIER[sdk=appletvos*] = {UUID_TVOS}");
 				}
 				if (UnrealData.UProjectFileLocation != null)
 				{
-					ProjectXcconfig.AppendLine($"PRODUCT_BUNDLE_IDENTIFIER[sdk=appletvos*] = {TVOS_BUNDLE}");
+					Xcconfig.AppendLine($"PRODUCT_BUNDLE_IDENTIFIER[sdk=appletvos*] = {TVOS_BUNDLE}");
 				}
 			}
 
@@ -1647,16 +1459,12 @@ namespace UnrealBuildTool.XcodeProjectXcconfig
 			// Get Mac architectures supported by this project
 			foreach (UnrealBuildConfig Config in UnrealData.AllConfigs)
 			{
-				// make a new xcconfig for this configuration, without spaces (Debug Editor -> DebugEditor)
-				XcconfigFile Xcconfig = new XcconfigFile($"{UnrealData.ProductName}_{Config.DisplayName.Replace(" ", "")}");
-				XcconfigFiles.Add(Xcconfig);
-
 				// hook up the Buildconfig that matches this info to this xcconfig file
-				ProjectBuildConfigs!.BuildConfigs.First(x => x.Info == Config).Xcconfig = Xcconfig;
+				XcconfigFile ConfigXcconfig = ProjectBuildConfigs!.BuildConfigs.First(x => x.Info == Config).Xcconfig!;
 
-				Xcconfig.AppendLine("// pull in the shared settings for all configs");
-				Xcconfig.AppendLine($"#include \"{UnrealData.ProductName}_Project.xcconfig\"");
-				Xcconfig.AppendLine("");
+				ConfigXcconfig.AppendLine("// pull in the shared settings for all configs");
+				ConfigXcconfig.AppendLine($"#include \"{UnrealData.ProductName}_Project.xcconfig\"");
+				ConfigXcconfig.AppendLine("");
 
 				bool bIsUnrealGame = Config.BuildTarget.Equals("UnrealGame", StringComparison.InvariantCultureIgnoreCase);
 				bool bIsUnrealClient = Config.BuildTarget.Equals("UnrealClient", StringComparison.InvariantCultureIgnoreCase);
@@ -1704,48 +1512,48 @@ namespace UnrealBuildTool.XcodeProjectXcconfig
 				// debug settings
 				if (Config.BuildConfig == UnrealTargetConfiguration.Debug)
 				{
-					Xcconfig.AppendLine("ENABLE_TESTABILITY = YES");
+					ConfigXcconfig.AppendLine("ENABLE_TESTABILITY = YES");
 				}
 
-				Xcconfig.AppendLine($"UE_BUILD_TARGET_NAME = {Config.BuildTarget}");
-				Xcconfig.AppendLine($"UE_BUILD_TARGET_CONFIG = {Config.BuildConfig}");
+				ConfigXcconfig.AppendLine($"UE_BUILD_TARGET_NAME = {Config.BuildTarget}");
+				ConfigXcconfig.AppendLine($"UE_BUILD_TARGET_CONFIG = {Config.BuildConfig}");
 
 				// list the supported platforms, which will narrow down what other settings are active
-				Xcconfig.AppendLine($"SUPPORTED_PLATFORMS = {SupportedPlatforms}");
+				ConfigXcconfig.AppendLine($"SUPPORTED_PLATFORMS = {SupportedPlatforms}");
 
 				// @otodo move VALID_ARCHS up to Project config once we are always universal
 				if (Config.bSupportsMac)
 				{
-					Xcconfig.AppendLine("");
-					Xcconfig.AppendLine("// Mac setup");
-					Xcconfig.AppendLine($"VALID_ARCHS[sdk=macosx*] = {SupportedMacArchitectures}");
-					Xcconfig.AppendLine($"PRODUCT_NAME[sdk=macosx*] = {MacExecutableFileName}");
-					Xcconfig.AppendLine($"CONFIGURATION_BUILD_DIR[sdk=macosx*] = {MacExecutableDir}");
-					Xcconfig.AppendLine($"INFOPLIST_FILE[sdk=macosx*] = {MacInfoPlistPath}");
+					ConfigXcconfig.AppendLine("");
+					ConfigXcconfig.AppendLine("// Mac setup");
+					ConfigXcconfig.AppendLine($"VALID_ARCHS[sdk=macosx*] = {SupportedMacArchitectures}");
+					ConfigXcconfig.AppendLine($"PRODUCT_NAME[sdk=macosx*] = {MacExecutableFileName}");
+					ConfigXcconfig.AppendLine($"CONFIGURATION_BUILD_DIR[sdk=macosx*] = {MacExecutableDir}");
+					ConfigXcconfig.AppendLine($"INFOPLIST_FILE[sdk=macosx*] = {MacInfoPlistPath}");
 				}
 
 				if (Config.bSupportsIOS)
 				{
-					Xcconfig.AppendLine("");
-					Xcconfig.AppendLine("// IOS setup");
-					Xcconfig.AppendLine($"PRODUCT_NAME[sdk=iphone*] = {Config.BuildTarget}"); // @todo: change to Path.GetFileName(Config.IOSExecutablePath) when we stop using payload
-					Xcconfig.AppendLine($"CONFIGURATION_BUILD_DIR[sdk=iphone*] = {ProjectRootDir}/Binaries/IOS/Payload");
-					Xcconfig.AppendLine($"INFOPLIST_FILE[sdk=iphone*] = {IOSInfoPlistPath}");
+					ConfigXcconfig.AppendLine("");
+					ConfigXcconfig.AppendLine("// IOS setup");
+					ConfigXcconfig.AppendLine($"PRODUCT_NAME[sdk=iphone*] = {Config.BuildTarget}"); // @todo: change to Path.GetFileName(Config.IOSExecutablePath) when we stop using payload
+					ConfigXcconfig.AppendLine($"CONFIGURATION_BUILD_DIR[sdk=iphone*] = {ProjectRootDir}/Binaries/IOS/Payload");
+					ConfigXcconfig.AppendLine($"INFOPLIST_FILE[sdk=iphone*] = {IOSInfoPlistPath}");
 					if (bSetEntitlements)
 					{
-						Xcconfig.AppendLine($"CODE_SIGN_ENTITLEMENTS[sdk=iphone*] = {ProjectRootDir}/Intermediate/IOS/{PlistTargetName}.entitlements");
+						ConfigXcconfig.AppendLine($"CODE_SIGN_ENTITLEMENTS[sdk=iphone*] = {ProjectRootDir}/Intermediate/IOS/{PlistTargetName}.entitlements");
 					}
 				}
 				if (Config.bSupportsTVOS)
 				{
-					Xcconfig.AppendLine("");
-					Xcconfig.AppendLine("// TVOS setup");
-					Xcconfig.AppendLine($"PRODUCT_NAME[sdk=appletvos*] = {Config.BuildTarget}"); // @todo: change to Path.GetFileName(Config.IOSExecutablePath) when we stop using payload
-					Xcconfig.AppendLine($"CONFIGURATION_BUILD_DIR[sdk=appletvos*] = {ProjectRootDir}/Binaries/TVOS/Payload");
-					Xcconfig.AppendLine($"INFOPLIST_FILE[sdk=appletvos*] = {TVOSInfoPlistPath}");
+					ConfigXcconfig.AppendLine("");
+					ConfigXcconfig.AppendLine("// TVOS setup");
+					ConfigXcconfig.AppendLine($"PRODUCT_NAME[sdk=appletvos*] = {Config.BuildTarget}"); // @todo: change to Path.GetFileName(Config.IOSExecutablePath) when we stop using payload
+					ConfigXcconfig.AppendLine($"CONFIGURATION_BUILD_DIR[sdk=appletvos*] = {ProjectRootDir}/Binaries/TVOS/Payload");
+					ConfigXcconfig.AppendLine($"INFOPLIST_FILE[sdk=appletvos*] = {TVOSInfoPlistPath}");
 					if (bSetEntitlements)
 					{
-						Xcconfig.AppendLine($"CODE_SIGN_ENTITLEMENTS[sdk=appletvos*] = {ProjectRootDir}/Intermediate/TVOS/{PlistTargetName}.entitlements");
+						ConfigXcconfig.AppendLine($"CODE_SIGN_ENTITLEMENTS[sdk=appletvos*] = {ProjectRootDir}/Intermediate/TVOS/{PlistTargetName}.entitlements");
 					}
 				}
 
@@ -2276,9 +2084,6 @@ namespace UnrealBuildTool.XcodeProjectXcconfig
 			// now create the xcodeproject elements (project -> target -> buildconfigs, etc)
 			RootProject = new XcodeProject(UnrealData, FileCollection);
 
-			// now that the node hierarchy is made, walk the nodes and create Xcconfig files
-			XcodeProjectNode.WriteXcconfigFiles(RootProject, this);
-
 			StringBuilder Content = new StringBuilder();
 
 			// look for an existing project to use as a template (if none found, create one from scratch)
@@ -2406,8 +2211,6 @@ namespace UnrealBuildTool.XcodeProjectXcconfig
 			}
 
 			return bSuccess;
-
-//			ProcessExtensions(PBXBuildFileSection, PBXFileReferenceSection, PBXCopyExtensionsBuildPhaseSection, PBXResourcesBuildPhaseSection, AllExtensions, GameProjectPath, BuildConfigs, Logger);
 		}
 
 #region Schemes
