Removed legacy VS Android debugger toolchains support

diff --git a/Engine/Source/Programs/UnrealBuildTool/Platform/Android/AndroidProjectGenerator.cs b/Engine/Source/Programs/UnrealBuildTool/Platform/Android/AndroidProjectGenerator.cs
--- a/Engine/Source/Programs/UnrealBuildTool/Platform/Android/AndroidProjectGenerator.cs
+++ b/Engine/Source/Programs/UnrealBuildTool/Platform/Android/AndroidProjectGenerator.cs
@@ -3,7 +3,6 @@
 using System;
 using System.Collections.Generic;
 using System.Text;
-using System.Diagnostics;
 using System.IO;
 using EpicGames.Core;
 using Microsoft.Extensions.Logging;
@@ -16,23 +15,8 @@ namespace UnrealBuildTool
 	/// </summary>
 	class AndroidProjectGenerator : PlatformProjectGenerator
 	{
-		bool CheckedForNsight = false;		// whether we have checked for a recent enough version of Nsight yet
-		bool NsightInstalled = false;		// true if a recent enough version of Nsight is installed
-		int NsightVersionCode = 0;           // version code matching detected Nsight
-
-		bool VSDebugCommandLineOptionPresent = false;    //User must put -vsdebugandroid commandline option to build the debug projects
-		bool VSDebuggingEnabled = false;      // When set to true, allows debugging with built in MS Cross Platform Android tools.  
-													 //  It adds separate projects ending in .androidproj and a file VSAndroidUnreal.props for the engine and all game projects
-
-		bool VSSupportChecked = false;       // Don't want to check multiple times
-
-		bool VSPropsFileWritten = false;     // This is for the file VSAndroidUnreal.props which only needs to be written once
-
-		bool VSSandboxedSDK = false;         // Checks if VS has installed the sandboxed SDK version to support Unreal
-											 // If this sandboxed SDK is present change the config to consume it instead of the main VS Android SDK
-
 		/// <summary>
-		/// Whether Android Game Development Extension is installed in the system. See https://developer.android.com/games/preview for more details.
+		/// Whether Android Game Development Extension is installed in the system. See https://developer.android.com/games/agde for more details.
 		/// May be disabled by using -noagde on commandline
 		/// </summary>
 		private bool AGDEInstalled = false;
@@ -42,11 +26,6 @@ namespace UnrealBuildTool
 		{
 			if (OperatingSystem.IsWindows())
 			{
-				if (Arguments.HasOption("-vsdebugandroid"))
-				{
-					VSDebugCommandLineOptionPresent = true;
-				}
-
 				AGDEInstalled = Microsoft.Win32.Registry.LocalMachine.OpenSubKey(@"SOFTWARE\WOW6432Node\Google\AndroidGameDevelopmentExtension")?.ValueCount > 0;
 				if (Arguments.HasOption("-noagde"))
 				{
@@ -55,136 +34,6 @@ namespace UnrealBuildTool
 			}
 		}
 
-		bool HostSupportsVSAndroid()
-		{
-			// Only supported on Windows (VS Mac is a thing..)
-			return BuildHostPlatform.Current.Platform.IsInGroup(UnrealPlatformGroup.Microsoft);
-		}
-
-		bool IsVSAndroidSupportInstalled()
-		{
-			if (VSSupportChecked)
-			{
-				return VSDebuggingEnabled;
-			}
-			VSSupportChecked = true;
-
-
-			if (!HostSupportsVSAndroid())
-			{
-				return false;
-			}
-
-			//check to make sure Cross Platform Tools are installed for MS
-
-			// First check to see if the sandboxed SDK has been installed, it's always dropped to the same ProgramData file location
-			if (File.Exists(Path.Combine(Environment.GetFolderPath(Environment.SpecialFolder.CommonApplicationData), "Microsoft\\AndroidSDK\\25\\tools\\android.bat"))
-				&& VSDebugCommandLineOptionPresent)
-			{
-				VSDebuggingEnabled = true;
-				VSSandboxedSDK = true;
-			}
-			else if (OperatingSystem.IsWindows())
-			{
-				// If the sandboxed SDK is not present (pre Visual Studio 15.4) then the non-Sandboxed SDK tools should have the correct build tools for building
-				string? SDKToolsPath = Microsoft.Win32.Registry.GetValue("HKEY_LOCAL_MACHINE\\SOFTWARE\\Wow6432Node\\Android SDK Tools", "Path", null) as string;
-				if (!String.IsNullOrEmpty(SDKToolsPath) && VSDebugCommandLineOptionPresent)
-				{
-					VSDebuggingEnabled = true;
-				}
-			}
-
-			if (VSDebugCommandLineOptionPresent && VSDebuggingEnabled == false)
-			{
-				Logger.LogWarning("Android SDK tools have to be installed to use this.  Please Install Visual C++ for Cross-Platform Mobile Development https://msdn.microsoft.com/en-us/library/dn707598.aspx");
-			}
-
-			return VSDebuggingEnabled;
-
-		}
-
-		/// <summary>
-		/// Check to see if a recent enough version of Nsight is installed.
-		/// </summary>
-		bool IsNsightInstalled(VCProjectFileFormat ProjectFileFormat)
-		{
-			// cache the results since this gets called a number of times
-			if (CheckedForNsight)
-			{
-				return NsightInstalled;
-			}
-
-			CheckedForNsight = true;
-
-			if (!HostSupportsVSAndroid())
-			{
-				NsightInstalled = false;
-				return false;
-			}
-
-			// NOTE: there is now a registry key that we can use instead at:
-			//			HKEY_LOCAL_MACHINE\SOFTWARE\Wow6432Node\NVIDIA Corporation\Nsight Tegra\Version
-
-			string ProgramFilesPath = Environment.GetFolderPath(Environment.SpecialFolder.ProgramFilesX86);
-
-			string PlatformToolsetVersion = VCProjectFileGenerator.GetProjectFilePlatformToolsetVersionString(ProjectFileFormat);
-			if (String.IsNullOrEmpty(PlatformToolsetVersion))
-			{
-				// future maintainer: add toolset version and verify that the rest of the msbuild path, version, and location in ProgramFiles(x86) is still valid
-				Logger.LogInformation("Android project generation needs to be updated for this version of Visual Studio.");
-				return false;
-			}
-
-			// build the path to where the Nsight DLL we'll be checking should sit
-			string NsightDllPath = Path.Combine(ProgramFilesPath, @"MSBuild\Microsoft.Cpp\v4.0", PlatformToolsetVersion, @"Platforms\Tegra-Android\Nvidia.Build.CPPTasks.Tegra-Android.Extensibility.dll");
-
-			if (!File.Exists(NsightDllPath))
-			{
-				return false;
-			}
-
-			// grab the version info from the DLL
-			FileVersionInfo NsightVersion = FileVersionInfo.GetVersionInfo(NsightDllPath);
-
-			if (NsightVersion.ProductMajorPart > 3)
-			{
-				// Mark as Nsight 3.1 (project will be updated)
-				NsightVersionCode = 11;
-				NsightInstalled = true;
-			}
-			else if (NsightVersion.ProductMajorPart == 3)
-			{
-				// Nsight 3.0 supported
-				NsightVersionCode = 9;
-				NsightInstalled = true;
-
-				if (NsightVersion.ProductMinorPart >= 1)
-				{
-					// Nsight 3.1+ should be valid (will update project if newer)
-					NsightVersionCode = 11;
-				}
-			}
-			else if (NsightVersion.ProductMajorPart == 2)
-			{
-				// Nsight 2.0+ should be valid
-				NsightVersionCode = 6;
-				NsightInstalled = true;
-			}
-			else if ((NsightVersion.ProductMajorPart == 1) && (NsightVersion.ProductMinorPart >= 5))
-			{
-				// Nsight 1.5+ should be valid
-				NsightVersionCode = 6;
-				NsightInstalled = true;
-			}
-
-			if (!NsightInstalled)
-			{
-				Logger.LogInformation("\nNsight Tegra {NsightVersionProductMajorPart}.{NsightVersionProductMinorPart} found, but Nsight Tegra 1.5 or higher is required for debugging support.", NsightVersion.ProductMajorPart, NsightVersion.ProductMinorPart);
-			}
-
-			return NsightInstalled;
-		}
-
 		/// <summary>
 		/// Enumerate all the platforms that this generator supports
 		/// </summary>
@@ -203,7 +52,7 @@ namespace UnrealBuildTool
 		public override bool HasVisualStudioSupport(UnrealTargetPlatform InPlatform, UnrealTargetConfiguration InConfiguration, VCProjectFileFormat ProjectFileFormat)
 		{
 			// Debugging, etc. are dependent on the TADP being installed
-			return AGDEInstalled || IsNsightInstalled(ProjectFileFormat) || IsVSAndroidSupportInstalled();
+			return AGDEInstalled;
 		}
 
 		/// <summary>
@@ -216,20 +65,9 @@ namespace UnrealBuildTool
 		{
 			string PlatformName = InPlatform.ToString();
 
-			if (InPlatform == UnrealTargetPlatform.Android)
+			if (InPlatform == UnrealTargetPlatform.Android && AGDEInstalled)
 			{
-				if (AGDEInstalled)
-				{
-					PlatformName = "Android-arm64-v8a";
-				}
-				else if (IsVSAndroidSupportInstalled())
-				{
-					PlatformName = "ARM";
-				}
-				else
-				{
-					PlatformName = "Tegra-Android";
-				}
+				PlatformName = "Android-arm64-v8a";
 			}
 
 			return PlatformName;
@@ -243,16 +81,10 @@ namespace UnrealBuildTool
 		/// <param name="ProjectFileBuilder">String builder for the project file</param>
 		public override void GetAdditionalVisualStudioPropertyGroups(UnrealTargetPlatform InPlatform, VCProjectFileFormat ProjectFileFormat, StringBuilder ProjectFileBuilder)
 		{
-			if (AGDEInstalled || IsVSAndroidSupportInstalled() || !IsNsightInstalled(ProjectFileFormat))
+			if (AGDEInstalled)
 			{
 				base.GetAdditionalVisualStudioPropertyGroups(InPlatform, ProjectFileFormat, ProjectFileBuilder);
 			}
-			else
-			{
-				ProjectFileBuilder.AppendLine("  <PropertyGroup Label=\"NsightTegraProject\">");
-				ProjectFileBuilder.AppendLine("    <NsightTegraProjectRevisionNumber>" + NsightVersionCode.ToString() + "</NsightTegraProjectRevisionNumber>");
-				ProjectFileBuilder.AppendLine("  </PropertyGroup>");
-			}
 		}
 
 		/// <summary>
@@ -265,19 +97,15 @@ namespace UnrealBuildTool
 		{
 			string ConfigurationType = "";
 
-			if (AGDEInstalled || IsVSAndroidSupportInstalled())
+			if (AGDEInstalled)
 			{
 				ConfigurationType = "Makefile";
 			}
-			else if (IsNsightInstalled(ProjectFileFormat))
-			{
-				ConfigurationType = "ExternalBuildSystem";				
-			}
 			else
 			{
 				ConfigurationType = base.GetVisualStudioPlatformConfigurationType(InPlatform, ProjectFileFormat);
 			}
-			
+
 			return ConfigurationType;
 		}
 
@@ -292,14 +120,6 @@ namespace UnrealBuildTool
 		public override void GetVisualStudioPlatformToolsetString(UnrealTargetPlatform InPlatform, UnrealTargetConfiguration InConfiguration, VCProjectFileFormat InProjectFileFormat, StringBuilder ProjectFileBuilder)
 		{
 			VCProjectFileGenerator.AppendPlatformToolsetProperty(ProjectFileBuilder, InProjectFileFormat);
-
-			if (AGDEInstalled || IsVSAndroidSupportInstalled() || !IsNsightInstalled(InProjectFileFormat))
-			{
-			}
-			else
-			{
-				ProjectFileBuilder.AppendLine("    <AndroidNativeAPI>UseTarget</AndroidNativeAPI>");
-			}
 		}
 
 		/// <summary>
@@ -328,418 +148,10 @@ namespace UnrealBuildTool
 				string symbolLocations = $@"{intermediatePath}jni\arm64-v8a;{intermediatePath}libs\arm64-v8a";
 				ProjectFileBuilder.AppendLine($"    <AndroidSymbolDirectories>{symbolLocations}</AndroidSymbolDirectories>");
 			}
-			else if (!IsVSAndroidSupportInstalled() && IsNsightInstalled(InProjectFileFormat))
-			{
-
-				// NOTE: We are intentionally overriding defaults for these paths with empty strings.  We never want Visual Studio's
-				//       defaults for these fields to be propagated, since they are version-sensitive paths that may not reflect
-				//       the environment that UBT is building in.  We'll set these environment variables ourselves!
-				// NOTE: We don't touch 'ExecutablePath' because that would result in Visual Studio clobbering the system "Path"
-				//       environment variable
-
-				//@todo android: clean up debug path generation
-				string GameName = TargetRulesPath.GetFileNameWithoutExtension();
-				GameName = Path.GetFileNameWithoutExtension(GameName);
-
-
-				// intermediate path for Engine or Game's intermediate
-				string IntermediateDirectoryPath;
-				IntermediateDirectoryPath = Path.GetDirectoryName(NMakeOutputPath.FullName) + "/../../Intermediate/Android/APK";
-
-				// string for <OverrideAPKPath>
-				string APKPath = Path.Combine(
-					Path.GetDirectoryName(NMakeOutputPath.FullName)!,
-					Path.GetFileNameWithoutExtension(NMakeOutputPath.FullName) + "-armv7.apk");
-
-				// string for <BuildXmlPath> and <AndroidManifestPath>
-				string BuildXmlPath = IntermediateDirectoryPath;
-				string AndroidManifestPath = Path.Combine(IntermediateDirectoryPath, "AndroidManifest.xml");
-
-				// string for <AdditionalLibraryDirectories>
-				string AdditionalLibDirs = "";
-				AdditionalLibDirs += IntermediateDirectoryPath + @"\obj\local\arm64-v8a";
-				AdditionalLibDirs += ";" + IntermediateDirectoryPath + @"\obj\local\x86_64";
-				AdditionalLibDirs += @";$(AdditionalLibraryDirectories)";
-
-				ProjectFileBuilder.AppendLine("    <IncludePath />");
-				ProjectFileBuilder.AppendLine("    <ReferencePath />");
-				ProjectFileBuilder.AppendLine("    <LibraryPath />");
-				ProjectFileBuilder.AppendLine("    <LibraryWPath />");
-				ProjectFileBuilder.AppendLine("    <SourcePath />");
-				ProjectFileBuilder.AppendLine("    <ExcludePath />");
-				ProjectFileBuilder.AppendLine("    <AndroidAttach>False</AndroidAttach>");
-				ProjectFileBuilder.AppendLine("    <DebuggerFlavor>AndroidDebugger</DebuggerFlavor>");
-				ProjectFileBuilder.AppendLine("    <OverrideAPKPath>" + APKPath + "</OverrideAPKPath>");
-				ProjectFileBuilder.AppendLine("    <AdditionalLibraryDirectories>" + AdditionalLibDirs + "</AdditionalLibraryDirectories>");
-				ProjectFileBuilder.AppendLine("    <BuildXmlPath>" + BuildXmlPath + "</BuildXmlPath>");
-				ProjectFileBuilder.AppendLine("    <AndroidManifestPath>" + AndroidManifestPath + "</AndroidManifestPath>");
-			}
-			else 
-			{
-				base.GetVisualStudioPathsEntries(InPlatform, InConfiguration, TargetType, TargetRulesPath, ProjectFilePath, NMakeOutputPath, InProjectFileFormat, ProjectFileBuilder);
-			}
-		}
-		
-		/// <summary>
-		/// Return any custom property settings. These will be included right after Global properties to make values available to all other imports.
-		/// </summary>
-		/// <param name="InPlatform">  The UnrealTargetPlatform being built</param>
-		/// <param name="ProjectFileBuilder">String builder for the project file</param>
-		/// <returns>string    The custom property import lines for the project file; Empty string if it doesn't require one</returns>
-		public override void GetVisualStudioGlobalProperties(UnrealTargetPlatform InPlatform, StringBuilder ProjectFileBuilder)
-		{
-			if (IsVSAndroidSupportInstalled())
-			{
-				ProjectFileBuilder.AppendLine("  <Import Project=\"$(AndroidTargetsPath)\\Android.Tools.props\"/>");
-			}
 			else
 			{
-				base.GetVisualStudioGlobalProperties(InPlatform, ProjectFileBuilder);
-			}
-		}
-
-		/// <summary>
-		/// Return any custom property settings. These will be included in the ImportGroup section
-		/// </summary>
-		/// <param name="InPlatform">  The UnrealTargetPlatform being built</param>
-		/// <param name="ProjectFileBuilder">String builder for the project file</param>
-		/// <returns>string    The custom property import lines for the project file; Empty string if it doesn't require one</returns>
-		public override void GetVisualStudioImportGroupProperties(UnrealTargetPlatform InPlatform, StringBuilder ProjectFileBuilder)
-		{
-			if (IsVSAndroidSupportInstalled())
-			{
-				ProjectFileBuilder.AppendLine("    <Import Project=\"VSAndroidUnreal.props\" />");
-			}
-			else
-			{
-				base.GetVisualStudioImportGroupProperties(InPlatform, ProjectFileBuilder);
-			}
-		}
-
-		/// <summary>
-		/// For Additional Project Property file  VSAndroidUnreal.props file that need to be written out.  This is currently used only on Android. 
-		/// </summary>
-		public override void WriteAdditionalPropFile()
-		{
-			if(!IsVSAndroidSupportInstalled())
-			{
-				return;
-			}
-
-			//This file only needs to be written once and doesn't change.
-			if(VSPropsFileWritten)
-			{
-				return;
-			}
-			VSPropsFileWritten = true;
-
-
-			string FileName = "VSAndroidUnreal.props";
-
-			// Change the path to the Android SDK based on if we are using the Visual Studio Sandboxed SDK or not
-			string vsAndroidSDKPath;
-			if (VSSandboxedSDK)
-			{
-				vsAndroidSDKPath = "$(ProgramData)\\Microsoft\\AndroidSDK\\25";
-			}
-			else
-			{
-				vsAndroidSDKPath = "$(VS_AndroidHome)";
-			}
-
-			string FileText = "<?xml version=\"1.0\" encoding=\"utf-8\"?>" + ProjectFileGenerator.NewLine +
-								"<Project ToolsVersion=\"4.0\" xmlns=\"http://schemas.microsoft.com/developer/msbuild/2003\">" + ProjectFileGenerator.NewLine +
-								"	<ImportGroup Label=\"PropertySheets\" />" + ProjectFileGenerator.NewLine +
-								"	<PropertyGroup Label=\"UserMacros\">" + ProjectFileGenerator.NewLine +
-								"		<ANDROID_HOME>" + vsAndroidSDKPath + "</ANDROID_HOME>" + ProjectFileGenerator.NewLine +
-								"		<JAVA_HOME>$(VS_JavaHome)</JAVA_HOME>" + ProjectFileGenerator.NewLine +
-								"		<NDKROOT>$(VS_NdkRoot)</NDKROOT>" + ProjectFileGenerator.NewLine +
-								"	</PropertyGroup>" + ProjectFileGenerator.NewLine +
-								"	<PropertyGroup />" + ProjectFileGenerator.NewLine +
-								"	<ItemDefinitionGroup />" + ProjectFileGenerator.NewLine +
-								"	<ItemGroup>" + ProjectFileGenerator.NewLine +
-								"		<BuildMacro Include=\"ANDROID_HOME\">" + ProjectFileGenerator.NewLine +
-								"			<Value>$(ANDROID_HOME)</Value>" + ProjectFileGenerator.NewLine +
-								"			<EnvironmentVariable>true</EnvironmentVariable>" + ProjectFileGenerator.NewLine +
-								"		</BuildMacro>" + ProjectFileGenerator.NewLine +
-								"		<BuildMacro Include=\"JAVA_HOME\">" + ProjectFileGenerator.NewLine +
-								"			<Value>$(JAVA_HOME)</Value>" + ProjectFileGenerator.NewLine +
-								"			<EnvironmentVariable>true</EnvironmentVariable>" + ProjectFileGenerator.NewLine +
-								"		</BuildMacro>" + ProjectFileGenerator.NewLine +
-								"		<BuildMacro Include=\"NDKROOT\">" + ProjectFileGenerator.NewLine +
-								"			<Value>$(NDKROOT)</Value>" + ProjectFileGenerator.NewLine +
-								"			<EnvironmentVariable>true</EnvironmentVariable>" + ProjectFileGenerator.NewLine +
-								"		</BuildMacro>" + ProjectFileGenerator.NewLine +
-								"	</ItemGroup>" + ProjectFileGenerator.NewLine +
-								"</Project>";
-
-			ProjectFileGenerator.WriteFileIfChanged(ProjectFileGenerator.IntermediateProjectFilesPath + "\\" + FileName, FileText, Logger);
-
-		}
-
-
-		/// <summary>
-		/// For additional Project file *PROJECTNAME*-AndroidRun.androidproj.user that needs to be written out.  This is currently used only on Android. 
-		/// </summary>
-		/// <param name="ProjectFile">ProjectFile object</param>
-		public override void WriteAdditionalProjUserFile(ProjectFile ProjectFile)
-		{
-			if (!IsVSAndroidSupportInstalled() || ProjectFile.SourceFiles.Count == 0)
-			{
-				return;
-			}
-
-			string BaseDirectory = ProjectFile.SourceFiles[0].BaseFolder!.FullName;
-
-			string ProjectName = ProjectFile.ProjectFilePath.GetFileNameWithoutExtension();
-
-			string FileName = ProjectName + "-AndroidRun.androidproj.user";
-
-			string FileText =   "<?xml version=\"1.0\" encoding=\"utf-8\"?>" + ProjectFileGenerator.NewLine +
-								"<Project ToolsVersion=\"14.0\" xmlns=\"http://schemas.microsoft.com/developer/msbuild/2003\">" + ProjectFileGenerator.NewLine +
-								"	<PropertyGroup Condition=\"\'$(Configuration)|$(Platform)\'==\'Debug|ARM\'\">" + ProjectFileGenerator.NewLine +
-								"		<PackagePath>" + BaseDirectory + "\\Binaries\\Android\\" + ProjectName + "-armv7.apk</PackagePath>" + ProjectFileGenerator.NewLine +
-								"		<LaunchActivity>" + ProjectFileGenerator.NewLine +
-								"		</LaunchActivity>" + ProjectFileGenerator.NewLine +
-								"		<AdditionalSymbolSearchPaths>" + BaseDirectory + "\\Intermediate\\Android\\APK\\obj\\local\\arm64-v8a; " + BaseDirectory + "\\Intermediate\\Android\\APK\\jni\\arm64-v8a;" + BaseDirectory + "\\Binaries\\Android;$(AdditionalSymbolSearchPaths)</AdditionalSymbolSearchPaths>" + ProjectFileGenerator.NewLine +
-								"		<DebuggerFlavor>AndroidDebugger</DebuggerFlavor>" + ProjectFileGenerator.NewLine +
-								"	</PropertyGroup>" + ProjectFileGenerator.NewLine +
-								"</Project>";
-
-			ProjectFileGenerator.WriteFileIfChanged(ProjectFileGenerator.IntermediateProjectFilesPath + "\\" + FileName, FileText, Logger);
-
-		}
-
-
-		/// <summary>
-		/// For additional Project file *PROJECTNAME*-AndroidRun.androidproj that needs to be written out.  This is currently used only on Android. 
-		/// </summary>
-		/// <param name="ProjectFile">ProjectFile object</param>
-		public override Tuple<ProjectFile, string>? WriteAdditionalProjFile(ProjectFile ProjectFile)
-		{
-			if (!IsVSAndroidSupportInstalled())
-			{
-				return null;
+				base.GetVisualStudioPathsEntries(InPlatform, InConfiguration, TargetType, TargetRulesPath, ProjectFilePath, NMakeOutputPath, InProjectFileFormat, ProjectFileBuilder);
 			}
-
-			string ProjectName = ProjectFile.ProjectFilePath.GetFileNameWithoutExtension() + "-AndroidRun";
-			
-			string FileName = ProjectName + ".androidproj";
-
-			string FileText = "<?xml version=\"1.0\" encoding=\"utf-8\"?> " + ProjectFileGenerator.NewLine +
-								"<Project DefaultTargets=\"Build\" ToolsVersion=\"14.0\" xmlns=\"http://schemas.microsoft.com/developer/msbuild/2003\"> " + ProjectFileGenerator.NewLine +
-								"	<ItemGroup Label=\"ProjectConfigurations\"> " + ProjectFileGenerator.NewLine +
-								"		<ProjectConfiguration Include=\"Debug|ARM\"> " + ProjectFileGenerator.NewLine +
-								"			<Configuration>Debug</Configuration> " + ProjectFileGenerator.NewLine +
-								"			<Platform>ARM</Platform> " + ProjectFileGenerator.NewLine +
-								"		</ProjectConfiguration> " + ProjectFileGenerator.NewLine +
-								"		<ProjectConfiguration Include=\"Development|ARM\"> " + ProjectFileGenerator.NewLine +
-								"			<Configuration>Development</Configuration> " + ProjectFileGenerator.NewLine +
-								"			<Platform>ARM</Platform> " + ProjectFileGenerator.NewLine +
-								"		</ProjectConfiguration> " + ProjectFileGenerator.NewLine +
-								"		<ProjectConfiguration Include=\"Release|ARM\"> " + ProjectFileGenerator.NewLine +
-								"			<Configuration>Release</Configuration> " + ProjectFileGenerator.NewLine +
-								"			<Platform>ARM</Platform> " + ProjectFileGenerator.NewLine +
-								"		</ProjectConfiguration> " + ProjectFileGenerator.NewLine +
-								"		<ProjectConfiguration Include=\"Debug|ARM64\"> " + ProjectFileGenerator.NewLine +
-								"			<Configuration>Debug</Configuration> " + ProjectFileGenerator.NewLine +
-								"			<Platform>ARM64</Platform> " + ProjectFileGenerator.NewLine +
-								"		</ProjectConfiguration> " + ProjectFileGenerator.NewLine +
-								"		<ProjectConfiguration Include=\"Development|ARM64\"> " + ProjectFileGenerator.NewLine +
-								"			<Configuration>Development</Configuration> " + ProjectFileGenerator.NewLine +
-								"			<Platform>ARM64</Platform> " + ProjectFileGenerator.NewLine +
-								"		</ProjectConfiguration> " + ProjectFileGenerator.NewLine +
-								"		<ProjectConfiguration Include=\"Release|ARM64\"> " + ProjectFileGenerator.NewLine +
-								"			<Configuration>Release</Configuration> " + ProjectFileGenerator.NewLine +
-								"			<Platform>ARM64</Platform> " + ProjectFileGenerator.NewLine +
-								"		</ProjectConfiguration> " + ProjectFileGenerator.NewLine +
-								"		<ProjectConfiguration Include=\"Debug|x64\"> " + ProjectFileGenerator.NewLine +
-								"			<Configuration>Debug</Configuration> " + ProjectFileGenerator.NewLine +
-								"			<Platform>x64</Platform> " + ProjectFileGenerator.NewLine +
-								"		</ProjectConfiguration> " + ProjectFileGenerator.NewLine +
-								"		<ProjectConfiguration Include=\"Development|x64\"> " + ProjectFileGenerator.NewLine +
-								"			<Configuration>Development</Configuration> " + ProjectFileGenerator.NewLine +
-								"			<Platform>x64</Platform> " + ProjectFileGenerator.NewLine +
-								"		</ProjectConfiguration> " + ProjectFileGenerator.NewLine +
-								"		<ProjectConfiguration Include=\"Release|x64\"> " + ProjectFileGenerator.NewLine +
-								"			<Configuration>Release</Configuration> " + ProjectFileGenerator.NewLine +
-								"			<Platform>x64</Platform> " + ProjectFileGenerator.NewLine +
-								"		</ProjectConfiguration> " + ProjectFileGenerator.NewLine +
-								"		<ProjectConfiguration Include=\"Debug|x86\"> " + ProjectFileGenerator.NewLine +
-								"			<Configuration>Debug</Configuration> " + ProjectFileGenerator.NewLine +
-								"			<Platform>x86</Platform> " + ProjectFileGenerator.NewLine +
-								"		</ProjectConfiguration> " + ProjectFileGenerator.NewLine +
-								"		<ProjectConfiguration Include=\"Development|x86\"> " + ProjectFileGenerator.NewLine +
-								"			<Configuration>Development</Configuration> " + ProjectFileGenerator.NewLine +
-								"			<Platform>x86</Platform> " + ProjectFileGenerator.NewLine +
-								"		</ProjectConfiguration> " + ProjectFileGenerator.NewLine +
-								"		<ProjectConfiguration Include=\"Release|x86\"> " + ProjectFileGenerator.NewLine +
-								"			<Configuration>Release</Configuration> " + ProjectFileGenerator.NewLine +
-								"			<Platform>x86</Platform> " + ProjectFileGenerator.NewLine +
-								"		</ProjectConfiguration> " + ProjectFileGenerator.NewLine +
-								"	</ItemGroup> " + ProjectFileGenerator.NewLine +
-								"	<PropertyGroup Label=\"Globals\"> " + ProjectFileGenerator.NewLine +
-								"		<RootNamespace>" + ProjectName + "</RootNamespace> " + ProjectFileGenerator.NewLine +
-								"		<MinimumVisualStudioVersion>14.0</MinimumVisualStudioVersion> " + ProjectFileGenerator.NewLine +
-                                "		<ProjectVersion>1.0</ProjectVersion> " + ProjectFileGenerator.NewLine +
-								
-								//Set the project guid 
-								"		<ProjectGuid>" + System.Guid.NewGuid().ToString("B").ToUpper() + "</ProjectGuid> " + ProjectFileGenerator.NewLine +
-								
-								"		<ConfigurationType>Application</ConfigurationType> " + ProjectFileGenerator.NewLine +
-								"		<_PackagingProjectWithoutNativeComponent>true</_PackagingProjectWithoutNativeComponent> " + ProjectFileGenerator.NewLine +
-								"		<LaunchActivity Condition=\"\'$(LaunchActivity)\' == \'\'\">com." + ProjectName + "." + ProjectName + "</LaunchActivity> " + ProjectFileGenerator.NewLine +
-								"		<JavaSourceRoots>src</JavaSourceRoots> " + ProjectFileGenerator.NewLine +
-								"	</PropertyGroup> " + ProjectFileGenerator.NewLine +
-								"	<Import Project=\"$(AndroidTargetsPath)\\Android.Default.props\" /> " + ProjectFileGenerator.NewLine +
-								"	<PropertyGroup Condition=\"\'$(Configuration)|$(Platform)\'==\'Debug|ARM\'\" Label=\"Configuration\"> " + ProjectFileGenerator.NewLine +
-								"		<UseDebugLibraries>true</UseDebugLibraries> " + ProjectFileGenerator.NewLine +
-								"		<TargetName>$(RootNamespace)</TargetName> " + ProjectFileGenerator.NewLine +
-								"	</PropertyGroup> " + ProjectFileGenerator.NewLine +
-								"	<PropertyGroup Condition=\"\'$(Configuration)|$(Platform)\'==\'Development|ARM\'\" Label=\"Configuration\"> " + ProjectFileGenerator.NewLine +
-								"		<UseDebugLibraries>false</UseDebugLibraries> " + ProjectFileGenerator.NewLine +
-								"		<TargetName>$(RootNamespace)</TargetName> " + ProjectFileGenerator.NewLine +
-								"	</PropertyGroup> " + ProjectFileGenerator.NewLine +
-								"	<PropertyGroup Condition=\"\'$(Configuration)|$(Platform)\'==\'Release|ARM\'\" Label=\"Configuration\"> " + ProjectFileGenerator.NewLine +
-								"		<UseDebugLibraries>false</UseDebugLibraries> " + ProjectFileGenerator.NewLine +
-								"		<TargetName>$(RootNamespace)</TargetName> " + ProjectFileGenerator.NewLine +
-								"	</PropertyGroup> " + ProjectFileGenerator.NewLine +
-								"	<PropertyGroup Condition=\"\'$(Configuration)|$(Platform)\'==\'Debug|ARM64\'\" Label=\"Configuration\"> " + ProjectFileGenerator.NewLine +
-								"		<UseDebugLibraries>true</UseDebugLibraries> " + ProjectFileGenerator.NewLine +
-								"		<TargetName>$(RootNamespace)</TargetName> " + ProjectFileGenerator.NewLine +
-								"	</PropertyGroup> " + ProjectFileGenerator.NewLine +
-								"	<PropertyGroup Condition=\"\'$(Configuration)|$(Platform)\'==\'Development|ARM64\'\" Label=\"Configuration\"> " + ProjectFileGenerator.NewLine +
-								"		<UseDebugLibraries>false</UseDebugLibraries> " + ProjectFileGenerator.NewLine +
-								"		<TargetName>$(RootNamespace)</TargetName> " + ProjectFileGenerator.NewLine +
-								"	</PropertyGroup> " + ProjectFileGenerator.NewLine +
-								"	<PropertyGroup Condition=\"\'$(Configuration)|$(Platform)\'==\'Release|ARM64\'\" Label=\"Configuration\"> " + ProjectFileGenerator.NewLine +
-								"		<UseDebugLibraries>false</UseDebugLibraries> " + ProjectFileGenerator.NewLine +
-								"		<TargetName>$(RootNamespace)</TargetName> " + ProjectFileGenerator.NewLine +
-								"	</PropertyGroup> " + ProjectFileGenerator.NewLine +
-								"	<PropertyGroup Condition=\"\'$(Configuration)|$(Platform)\'==\'Debug|x64\'\" Label=\"Configuration\"> " + ProjectFileGenerator.NewLine +
-								"		<UseDebugLibraries>true</UseDebugLibraries> " + ProjectFileGenerator.NewLine +
-								"		<TargetName>$(RootNamespace)</TargetName> " + ProjectFileGenerator.NewLine +
-								"	</PropertyGroup> " + ProjectFileGenerator.NewLine +
-								"	<PropertyGroup Condition=\"\'$(Configuration)|$(Platform)\'==\'Development|x64\'\" Label=\"Configuration\"> " + ProjectFileGenerator.NewLine +
-								"		<UseDebugLibraries>false</UseDebugLibraries> " + ProjectFileGenerator.NewLine +
-								"		<TargetName>$(RootNamespace)</TargetName> " + ProjectFileGenerator.NewLine +
-								"	</PropertyGroup> " + ProjectFileGenerator.NewLine +
-								"	<PropertyGroup Condition=\"\'$(Configuration)|$(Platform)\'==\'Release|x64\'\" Label=\"Configuration\"> " + ProjectFileGenerator.NewLine +
-								"		<UseDebugLibraries>false</UseDebugLibraries> " + ProjectFileGenerator.NewLine +
-								"		<TargetName>$(RootNamespace)</TargetName> " + ProjectFileGenerator.NewLine +
-								"	</PropertyGroup> " + ProjectFileGenerator.NewLine +
-								"	<PropertyGroup Condition=\"\'$(Configuration)|$(Platform)\'==\'Debug|x86\'\" Label=\"Configuration\"> " + ProjectFileGenerator.NewLine +
-								"		<UseDebugLibraries>true</UseDebugLibraries> " + ProjectFileGenerator.NewLine +
-								"		<TargetName>$(RootNamespace)</TargetName> " + ProjectFileGenerator.NewLine +
-								"	</PropertyGroup> " + ProjectFileGenerator.NewLine +
-								"	<PropertyGroup Condition=\"\'$(Configuration)|$(Platform)\'==\'Development|x86\'\" Label=\"Configuration\"> " + ProjectFileGenerator.NewLine +
-								"		<UseDebugLibraries>false</UseDebugLibraries> " + ProjectFileGenerator.NewLine +
-								"		<TargetName>$(RootNamespace)</TargetName> " + ProjectFileGenerator.NewLine +
-								"	</PropertyGroup> " + ProjectFileGenerator.NewLine +
-								"	<PropertyGroup Condition=\"\'$(Configuration)|$(Platform)\'==\'Release|x86\'\" Label=\"Configuration\"> " + ProjectFileGenerator.NewLine +
-								"		<UseDebugLibraries>false</UseDebugLibraries> " + ProjectFileGenerator.NewLine +
-								"		<TargetName>$(RootNamespace)</TargetName> " + ProjectFileGenerator.NewLine +
-								"	</PropertyGroup> " + ProjectFileGenerator.NewLine +
-								"	<Import Project=\"$(AndroidTargetsPath)\\Android.props\" /> " + ProjectFileGenerator.NewLine +
-								"	<ImportGroup Label=\"ExtensionSettings\" /> " + ProjectFileGenerator.NewLine +
-								"	<PropertyGroup Label=\"UserMacros\" /> " + ProjectFileGenerator.NewLine +
-								"	<PropertyGroup Condition=\"\'$(Configuration)|$(Platform)\'==\'Debug|ARM\'\"> " + ProjectFileGenerator.NewLine +
-								"	</PropertyGroup> " + ProjectFileGenerator.NewLine +
-								"	<PropertyGroup Condition=\"\'$(Configuration)|$(Platform)\'==\'Development|ARM\'\"> " + ProjectFileGenerator.NewLine +
-								"	</PropertyGroup> " + ProjectFileGenerator.NewLine +
-								"	<PropertyGroup Condition=\"\'$(Configuration)|$(Platform)\'==\'Release|ARM\'\"> " + ProjectFileGenerator.NewLine +
-								"	</PropertyGroup> " + ProjectFileGenerator.NewLine +
-								"	<PropertyGroup Condition=\"\'$(Configuration)|$(Platform)\'==\'Debug|ARM64\'\"> " + ProjectFileGenerator.NewLine +
-								"	</PropertyGroup> " + ProjectFileGenerator.NewLine +
-								"	<PropertyGroup Condition=\"\'$(Configuration)|$(Platform)\'==\'Development|ARM64\'\"> " + ProjectFileGenerator.NewLine +
-								"	</PropertyGroup> " + ProjectFileGenerator.NewLine +
-								"	<PropertyGroup Condition=\"\'$(Configuration)|$(Platform)\'==\'Release|ARM64\'\"> " + ProjectFileGenerator.NewLine +
-								"	</PropertyGroup> " + ProjectFileGenerator.NewLine +
-								"	<PropertyGroup Condition=\"\'$(Configuration)|$(Platform)\'==\'Debug|x64\'\"> " + ProjectFileGenerator.NewLine +
-								"	</PropertyGroup> " + ProjectFileGenerator.NewLine +
-								"	<PropertyGroup Condition=\"\'$(Configuration)|$(Platform)\'==\'Development|x64\'\"> " + ProjectFileGenerator.NewLine +
-								"	</PropertyGroup> " + ProjectFileGenerator.NewLine +
-								"	<PropertyGroup Condition=\"\'$(Configuration)|$(Platform)\'==\'Release|x64\'\"> " + ProjectFileGenerator.NewLine +
-								"	</PropertyGroup> " + ProjectFileGenerator.NewLine +
-								"	<PropertyGroup Condition=\"\'$(Configuration)|$(Platform)\'==\'Debug|x86\'\"> " + ProjectFileGenerator.NewLine +
-								"	</PropertyGroup> " + ProjectFileGenerator.NewLine +
-								"	<PropertyGroup Condition=\"\'$(Configuration)|$(Platform)\'==\'Development|x86\'\"> " + ProjectFileGenerator.NewLine +
-								"	</PropertyGroup> " + ProjectFileGenerator.NewLine +
-								"	<PropertyGroup Condition=\"\'$(Configuration)|$(Platform)\'==\'Release|x86\'\"> " + ProjectFileGenerator.NewLine +
-								"	</PropertyGroup> " + ProjectFileGenerator.NewLine +
-								"	<ItemDefinitionGroup Condition=\"\'$(Configuration)|$(Platform)\'==\'Debug|ARM\'\"> " + ProjectFileGenerator.NewLine +
-								"		<AntPackage> " + ProjectFileGenerator.NewLine +
-								"			<AndroidAppLibName /> " + ProjectFileGenerator.NewLine +
-								"		</AntPackage> " + ProjectFileGenerator.NewLine +
-								"	</ItemDefinitionGroup> " + ProjectFileGenerator.NewLine +
-								"	<ItemDefinitionGroup Condition=\"\'$(Configuration)|$(Platform)\'==\'Development|ARM\'\"> " + ProjectFileGenerator.NewLine +
-								"		<AntPackage> " + ProjectFileGenerator.NewLine +
-								"			<AndroidAppLibName /> " + ProjectFileGenerator.NewLine +
-								"		</AntPackage> " + ProjectFileGenerator.NewLine +
-								"	</ItemDefinitionGroup> " + ProjectFileGenerator.NewLine +
-								"	<ItemDefinitionGroup Condition=\"\'$(Configuration)|$(Platform)\'==\'Release|ARM\'\"> " + ProjectFileGenerator.NewLine +
-								"		<AntPackage> " + ProjectFileGenerator.NewLine +
-								"			<AndroidAppLibName /> " + ProjectFileGenerator.NewLine +
-								"		</AntPackage> " + ProjectFileGenerator.NewLine +
-								"	</ItemDefinitionGroup> " + ProjectFileGenerator.NewLine +
-								"	<ItemDefinitionGroup Condition=\"\'$(Configuration)|$(Platform)\'==\'Debug|ARM64\'\"> " + ProjectFileGenerator.NewLine +
-								"		<AntPackage> " + ProjectFileGenerator.NewLine +
-								"			<AndroidAppLibName /> " + ProjectFileGenerator.NewLine +
-								"		</AntPackage> " + ProjectFileGenerator.NewLine +
-								"	</ItemDefinitionGroup> " + ProjectFileGenerator.NewLine +
-								"	<ItemDefinitionGroup Condition=\"\'$(Configuration)|$(Platform)\'==\'Development|ARM64\'\"> " + ProjectFileGenerator.NewLine +
-								"		<AntPackage> " + ProjectFileGenerator.NewLine +
-								"			<AndroidAppLibName /> " + ProjectFileGenerator.NewLine +
-								"		</AntPackage> " + ProjectFileGenerator.NewLine +
-								"	</ItemDefinitionGroup> " + ProjectFileGenerator.NewLine +
-								"	<ItemDefinitionGroup Condition=\"\'$(Configuration)|$(Platform)\'==\'Release|ARM64\'\"> " + ProjectFileGenerator.NewLine +
-								"		<AntPackage> " + ProjectFileGenerator.NewLine +
-								"			<AndroidAppLibName /> " + ProjectFileGenerator.NewLine +
-								"		</AntPackage> " + ProjectFileGenerator.NewLine +
-								"	</ItemDefinitionGroup> " + ProjectFileGenerator.NewLine +
-								"	<ItemDefinitionGroup Condition=\"\'$(Configuration)|$(Platform)\'==\'Debug|x64\'\"> " + ProjectFileGenerator.NewLine +
-								"		<AntPackage> " + ProjectFileGenerator.NewLine +
-								"			<AndroidAppLibName /> " + ProjectFileGenerator.NewLine +
-								"		</AntPackage> " + ProjectFileGenerator.NewLine +
-								"	</ItemDefinitionGroup> " + ProjectFileGenerator.NewLine +
-								"	<ItemDefinitionGroup Condition=\"\'$(Configuration)|$(Platform)\'==\'Development|x64\'\"> " + ProjectFileGenerator.NewLine +
-								"		<AntPackage> " + ProjectFileGenerator.NewLine +
-								"			<AndroidAppLibName /> " + ProjectFileGenerator.NewLine +
-								"		</AntPackage> " + ProjectFileGenerator.NewLine +
-								"	</ItemDefinitionGroup> " + ProjectFileGenerator.NewLine +
-								"	<ItemDefinitionGroup Condition=\"\'$(Configuration)|$(Platform)\'==\'Release|x64\'\"> " + ProjectFileGenerator.NewLine +
-								"		<AntPackage> " + ProjectFileGenerator.NewLine +
-								"			<AndroidAppLibName /> " + ProjectFileGenerator.NewLine +
-								"		</AntPackage> " + ProjectFileGenerator.NewLine +
-								"	</ItemDefinitionGroup> " + ProjectFileGenerator.NewLine +
-								"	<ItemDefinitionGroup Condition=\"\'$(Configuration)|$(Platform)\'==\'Debug|x86\'\"> " + ProjectFileGenerator.NewLine +
-								"		<AntPackage> " + ProjectFileGenerator.NewLine +
-								"			<AndroidAppLibName /> " + ProjectFileGenerator.NewLine +
-								"		</AntPackage> " + ProjectFileGenerator.NewLine +
-								"	</ItemDefinitionGroup> " + ProjectFileGenerator.NewLine +
-								"	<ItemDefinitionGroup Condition=\"\'$(Configuration)|$(Platform)\'==\'Development|x86\'\"> " + ProjectFileGenerator.NewLine +
-								"		<AntPackage> " + ProjectFileGenerator.NewLine +
-								"			<AndroidAppLibName /> " + ProjectFileGenerator.NewLine +
-								"		</AntPackage> " + ProjectFileGenerator.NewLine +
-								"	</ItemDefinitionGroup> " + ProjectFileGenerator.NewLine +
-								"	<ItemDefinitionGroup Condition=\"\'$(Configuration)|$(Platform)\'==\'Release|x86\'\"> " + ProjectFileGenerator.NewLine +
-								"		<AntPackage> " + ProjectFileGenerator.NewLine +
-								"			<AndroidAppLibName /> " + ProjectFileGenerator.NewLine +
-								"		</AntPackage> " + ProjectFileGenerator.NewLine +
-								"	</ItemDefinitionGroup> " + ProjectFileGenerator.NewLine +
-								"	<Import Project=\"$(AndroidTargetsPath)\\Android.targets\" /> " + ProjectFileGenerator.NewLine +
-                                "	<ImportGroup Label=\"ExtensionTargets\" /> " + ProjectFileGenerator.NewLine +
-								"</Project>";
-			
-			bool Success = ProjectFileGenerator.WriteFileIfChanged(ProjectFileGenerator.IntermediateProjectFilesPath + "\\" + FileName, FileText, Logger);
-
-			FileReference ProjectFilePath = FileReference.Combine(ProjectFileGenerator.IntermediateProjectFilesPath, FileName);
-			AndroidDebugProjectFile Project = new AndroidDebugProjectFile(ProjectFilePath, ProjectFile.BaseDir);
-			Project.ShouldBuildForAllSolutionTargets = false;
-			Project.ShouldBuildByDefaultForSolutionTargets = false;
-
-			return Success ? new Tuple<ProjectFile, string>(Project, "Unreal Android Debug Projects") : null;
-
 		}
 
 		public override string GetExtraBuildArguments(UnrealTargetPlatform InPlatform, UnrealTargetConfiguration InConfiguration)
@@ -764,43 +176,4 @@ namespace UnrealBuildTool
 			return base.GetVisualStudioUserFileStrings(InPlatform, InConfiguration, InConditionString, InTargetRules, TargetRulesPath, ProjectFilePath);
 		}
 	}
-
-
-	/// <summary>
-	/// An Android Debug Project
-	/// </summary>
-	class AndroidDebugProjectFile : MSBuildProjectFile
-	{
-		/// <summary>
-		/// Constructs a new project file object
-		/// </summary>
-		/// <param name="InitFilePath">The path to the project file on disk</param>
-		/// <param name="BaseDir">The base directory for files within this project</param>
-		public AndroidDebugProjectFile(FileReference InitFilePath, DirectoryReference BaseDir)
-			: base(InitFilePath, BaseDir)
-		{
-		}
-
-		/// <summary>
-		/// This is the GUID that Visual Studio uses to identify an Android project file in the solution
-		/// </summary>
-		public override string ProjectTypeGUID
-		{
-			get { return "{39E2626F-3545-4960-A6E8-258AD8476CE5}"; }
-		}
-
-		/// <summary>
-		/// The only valid configuration for these to be run in is Debug|ARM
-		/// </summary>
-		/// <param name="SolutionTarget">The solution target type</param>
-		/// <param name="SolutionConfiguration">The solution configuration</param>
-		/// <param name="SolutionPlatform">The solution platform</param>
-		/// <param name="PlatformProjectGenerators">Set of platform project generators</param>
-		/// <param name="Logger"></param>
-		/// <returns>Project context matching the given solution context</returns>
-		public override MSBuildProjectContext GetMatchingProjectContext(TargetType SolutionTarget, UnrealTargetConfiguration SolutionConfiguration, UnrealTargetPlatform SolutionPlatform, PlatformProjectGeneratorCollection PlatformProjectGenerators, ILogger Logger)
-		{
-			return new MSBuildProjectContext("Debug", "ARM"){ bBuildByDefault = (SolutionPlatform == UnrealTargetPlatform.Android) };
-		}
-	}
 }
