PR #9455: Merged FASTBuild PR #8784 with recent UE5 changes (Contributed by [REDACTED_PERSON])

c# changes only

Additional changes:
* Instead of filtering the action list and running pre and post actions through the local executor, instead build everything that won't distribute through FASTBuild using local Exec() functions
* Remove custom link support and just run through Exec
* Handle /external:I arguments
* /sourceDependencies doesn't actually need to be handled manually because it will run during preprocessing and will be stripped for distribution

diff --git a/Engine/Source/Programs/UnrealBuildTool/Executors/Experimental/FASTBuild.cs b/Engine/Source/Programs/UnrealBuildTool/Executors/Experimental/FASTBuild.cs
--- a/Engine/Source/Programs/UnrealBuildTool/Executors/Experimental/FASTBuild.cs
+++ b/Engine/Source/Programs/UnrealBuildTool/Executors/Experimental/FASTBuild.cs
@@ -432,65 +432,17 @@ namespace UnrealBuildTool
 			if (Actions.Count <= 0)
 				return true;
 
-			IEnumerable<LinkedAction> CompileActions		= Actions.Where(Action => (Action.ActionType == ActionType.Compile && Action.bCanExecuteRemotely == true));
-			IEnumerable<LinkedAction> NonCompileActions	= Actions.Where(Action => (Action.ActionType != ActionType.Compile || Action.bCanExecuteRemotely == false));
-
-			///////////////////////////////////////////////////////////////
-			// Pre Compile Stage
-
-			// We want to complete any non-compile actions locally that are necessary for the distributed compile step
-			List<LinkedAction> PreCompileActions =
-				NonCompileActions
-				.Where(NonCompileAction => CompileActions.Any(CompileAction => CompileAction.PrerequisiteActions.Contains(NonCompileAction)))
-				.ToList();
-
-			// Precompile actions may have their own prerequisites which need to be executed along with them
-			List<LinkedAction> Prerequisites = new List<LinkedAction>();
-			foreach (LinkedAction PreCompileAction in PreCompileActions)
-			{
-				Prerequisites.AddRange(PreCompileAction.PrerequisiteActions);
-			}
-
-			PreCompileActions.AddRange(Prerequisites);
-
-			if (PreCompileActions.Any())
-			{
-				bool bResult = LocalExecutor.ExecuteActions(PreCompileActions.Distinct().ToList(), Logger);
-
-				if (!bResult)
-					return false;
-			}
-
-			///////////////////////////////////////////////////////////////
-			// Compile Stage
-
-			if (CompileActions.Any())
+			IEnumerable<LinkedAction> CompileActions = Actions.Where(Action => Action.ActionType == ActionType.Compile && Action.bCanExecuteRemotely && Action.bCanExecuteRemotelyWithSNDBS);
+			if (CompileActions.Any() && DetectBuildType(CompileActions, Logger))
 			{
-				if (!DetectBuildType(CompileActions, Logger))
-					return false;
-
 				string FASTBuildFilePath = Path.Combine(Unreal.EngineDirectory.FullName, "Intermediate", "Build", "fbuild.bff");
-				if (!CreateBffFile(CompileActions, FASTBuildFilePath, Logger))
-					return false;
-
-				if (!ExecuteBffFile(FASTBuildFilePath, Logger))
+				if (!CreateBffFile(Actions, FASTBuildFilePath, Logger))
 					return false;
-			}
-
-			///////////////////////////////////////////////////////////////
-			// Post Compile Stage
-
-			List<LinkedAction> PostCompileActions = NonCompileActions.Except(PreCompileActions).ToList();
 
-			if (PostCompileActions.Any())
-			{
-				bool bResult = LocalExecutor.ExecuteActions(PostCompileActions, Logger);
-
-				if (!bResult)
-					return false;
+				return ExecuteBffFile(FASTBuildFilePath, Logger);
 			}
 
-			return true;
+			return LocalExecutor.ExecuteActions(Actions.Distinct().ToList(), Logger); ;
 		}
 
 		private void AddText(string StringToWrite)
@@ -710,6 +662,7 @@ namespace UnrealBuildTool
 
 				// Skip the following tokens:
 				if ((Token == "/I")						||
+					(Token == "/external:I")			||
 					(Token == "/l")						||
 					(Token == "/D")						||
 					(Token == "-D")						||
@@ -723,15 +676,12 @@ namespace UnrealBuildTool
 					(Token == "-rpath")					||
 					(Token == "-weak_library")			||
 					(Token == "-weak_framework")		||
-					(Token == "-framework"))
+					(Token == "-framework")				||
+					(Token == "/sourceDependencies")	||
+					(Token == "/sourceDependencies:directives"))
 				{
 					++i;
 				}
-				else if (Token.StartsWith("/sourceDependencies"))
-				{
-					++i;
-					++i;
-				}
 				else if (Token == "/we4668")
 				{
 					// Replace this to make Windows builds compile happily
@@ -1036,7 +986,7 @@ namespace UnrealBuildTool
 			AddText("\t.Environment = \n\t{\n");
 			if (VCEnv != null)
             {
-                AddText(string.Format("\t\t\"PATH={0}\\Common7\\IDE\\;{1}\",\n", VCEnv.GetVCInstallDirectory(), VCEnv.GetToolPath()));
+                AddText(string.Format("\t\t\"PATH={0}\\Common7\\IDE\\;{1};{2}\\bin\\{3}\\x64\",\n", VCEnv.GetVCInstallDirectory(), VCEnv.GetToolPath(), VCEnv.WindowsSdkDir, VCEnv.WindowsSdkVersion));
             }
 
 			if (!IsApple())
@@ -1173,199 +1123,16 @@ namespace UnrealBuildTool
 			AddText("}\n\n");
 		}
 
-		private void AddLinkAction(LinkedAction Action, IEnumerable<LinkedAction> DependencyActions, ILogger Logger)
+		private void AddExecAction(LinkedAction Action, IEnumerable<LinkedAction> DependencyActions, ILogger Logger)
 		{
-			string[] SpecialLinkerOptions = IsApple() ? new string[] { "-o" } : new string[] { "/OUT:", "@", "-o" };
-
-			List<string> SplitCommandArgs = new List<string>();
-
-			string CommandArgsToParse = Action.CommandArguments;
-            if (CommandArgsToParse.StartsWith("-c \'") && CommandArgsToParse.EndsWith("\'"))
-            {
-                CommandArgsToParse = CommandArgsToParse.Remove(0, 4).TrimEnd('\'').Trim(' ');
-				SplitCommandArgs = CommandArgsToParse.Split(';').Select(SplitCommandArg => SplitCommandArg.Trim()).ToList();
-				CommandArgsToParse = SplitCommandArgs[0];
-			}
-
-			Dictionary<string, string> ParsedLinkerOptions = ParseCommandLineOptions(Action.CommandPath.GetFileName(), CommandArgsToParse, SpecialLinkerOptions, Logger, SaveResponseFile: true);
-
-			string OutputFile;
-
-			if (IsMSVC())
-			{
-				OutputFile = GetOptionValue(ParsedLinkerOptions, "/OUT:", Action, Logger, ProblemIfNotFound: true);
-			}
-			else // Apple
-			{
-				OutputFile = GetOptionValue(ParsedLinkerOptions, "-o", Action, Logger, ProblemIfNotFound: false);
-				if (string.IsNullOrEmpty(OutputFile))
-				{
-					OutputFile = GetOptionValue(ParsedLinkerOptions, "InputFile", Action, Logger, ProblemIfNotFound: true);
-				}
-			}
-
-			if (string.IsNullOrEmpty(OutputFile))
-			{
-				Logger.LogError("Failed to find output file. Bailing.");
-				return;
-			}
-
-			string ResponseFilePath = GetOptionValue(ParsedLinkerOptions, "@", Action, Logger);
-			string OtherCompilerOptions = GetOptionValue(ParsedLinkerOptions, "OtherOptions", Action, Logger);
-
-			IEnumerable<LinkedAction>? PrebuildDependencies = null;
-
-			if (Action.CommandPath.FullName.Contains("lib.exe"))
-			{
-				if (DependencyActions.Any())
-				{
-					Func<LinkedAction, bool> DoesActionProducePCH = (LinkedAction ActionToCheck) =>
-					{
-						foreach (FileItem ProducedItem in ActionToCheck.ProducedItems)
-						{
-							if (ProducedItem.ToString().Contains(".pch") || ProducedItem.ToString().Contains(".res"))
-							{
-								return true;
-							}
-						}
-						return false;
-					};
-
-					// Don't specify pch or resource files, they have the wrong name and the response file will have them anyways.
-					PrebuildDependencies = DependencyActions.Where(DoesActionProducePCH);
-					DependencyActions = DependencyActions.Where((ActionToCheck) => !DoesActionProducePCH(ActionToCheck));
-				}
-
-				AddText($"Library('{ActionToActionString(Action)}')\n{{\n");
-				AddText($"\t.Compiler = '{GetCompilerName()}'\n");
-				if (IsMSVC())
-					AddText("\t.CompilerOptions = '\"%1\" /Fo\"%2\" /c'\n");
-				else
-					AddText("\t.CompilerOptions = '\"%1\" -o \"%2\" -c'\n");
-				AddText($"\t.CompilerOutputPath = \"{Path.GetDirectoryName(OutputFile)}\"\n");
-				AddText($"\t.Librarian = '{Action.CommandPath.FullName}' \n");
-
-				if (!string.IsNullOrEmpty(ResponseFilePath))
-				{
-					if (IsMSVC())
-                        AddText($"\t.LibrarianOptions = ' /OUT:\"%2\" @\"{ResponseFilePath}\" \"%1\"' \n");
-					else
-						AddText($"\t.LibrarianOptions = '\"%2\" @\"%1\" {OtherCompilerOptions}' \n");
-				}
-				else
-				{
-					if (IsMSVC())
-						AddText($"\t.LibrarianOptions = ' /OUT:\"%2\" {OtherCompilerOptions} \"%1\"' \n");
-				}
-
-				if (DependencyActions.Any())
-				{
-					List<string> DependencyNames = DependencyActions.Select(ActionToDependencyString).ToList();
-
-                    if (!string.IsNullOrEmpty(ResponseFilePath))
-                        AddText($"\t.LibrarianAdditionalInputs = {{\n{DependencyNames[0]}\n\t}} \n"); // Hack...Because FASTBuild needs at least one Input file
-                    else if (IsMSVC())
-                        AddText($"\t.LibrarianAdditionalInputs = {{\n{string.Join(",", DependencyNames.ToArray())}\n\t}} \n");
-
-                    PrebuildDependencies = PrebuildDependencies!.Concat(DependencyActions);
-				}
-				else
-				{
-					AddText(string.Format("\t.LibrarianAdditionalInputs = {{ '{0}' }} \n", GetOptionValue(ParsedLinkerOptions, "InputFile", Action, Logger, ProblemIfNotFound: true)));
-				}
-
-				AddText($"\t.LibrarianOutput = '{OutputFile}' \n");
-				AddPreBuildDependenciesText(PrebuildDependencies);
-				AddText($"}}\n\n");
-			}
-			else if (Action.CommandPath.FullName.Contains("link.exe"))
-			{
-				AddText($"Executable('{ActionToActionString(Action)}')\n{{ \n");
-				AddText($"\t.Linker = '{Action.CommandPath.FullName}' \n");
-				if (DependencyActions.Any())
-				{
-					AddText($"\t.Libraries = {{ '{ResponseFilePath}' }} \n");
-					if (IsMSVC())
-					{
-						AddText($"\t.LinkerOptions = '/TLBOUT:\"%1\" /Out:\"%2\" @\"{ResponseFilePath}\" ' \n"); // The TLBOUT is a huge bodge to consume the %1.
-					}
-					else
-					{
-						AddText($"\t.LinkerOptions = '-o \"%2\" @\"{ResponseFilePath}\" {OtherCompilerOptions} -MQ \"%1\"' \n"); // The MQ is a huge bodge to consume the %1.
-					}
-				}
-				else
-				{
-					AddText($"\t.Libraries = '{ActionToActionString(DependencyActions.First())}' \n");
-
-					if (IsMSVC())
-					{
-						AddText($"\t.LinkerOptions = '/TLBOUT:\"%1\" /Out:\"%2\" @\"{ResponseFilePath}\" ' \n"); // The TLBOUT is a huge bodge to consume the %1.
-					}
-					else
-					{
-						AddText($"\t.LinkerOptions = '-o \"%2\" @\"{ResponseFilePath}\" {OtherCompilerOptions} -MQ \"%1\"' \n"); // The MQ is a huge bodge to consume the %1.
-					}
-				}
-
-				AddText($"\t.LinkerOutput = '{OutputFile}' \n");
-				AddPreBuildDependenciesText(DependencyActions);
-				AddText($"}}\n\n");
-			}
-			else if (Action.CommandArguments.Contains("clang++"))
-			{
-				AddText($"Executable('{ActionToActionString(Action)}')\n{{ \n");
-				AddText("\t.Linker = '$MacToolchainDir$/clang++' \n");
-
-				string InputFile = GetOptionValue(ParsedLinkerOptions, "InputFile", Action, Logger, ProblemIfNotFound: true);
-				if (!string.IsNullOrEmpty(InputFile))
-				{
-					LinkedAction? InputFileAction = DependencyActions
-						.Where(ActionToInspect =>
-							ActionToInspect.ProducedItems.Any(Item => Item.AbsolutePath == InputFile)
-						).FirstOrDefault();
-
-					if (InputFileAction != null)
-						InputFile = ActionToActionString(InputFileAction);
-				}
-
-				AddText($"\t.Libraries = {{ '{InputFile}' }} \n");
-				AddText($"\t.LinkerOptions = '{OtherCompilerOptions} \"%1\" -o \"%2\"' \n");
-				AddText($"\t.LinkerOutput = '{OutputFile}' \n");
-				AddPreBuildDependenciesText(DependencyActions);
-
-				IEnumerable<string> PostLinkCommands = SplitCommandArgs.Skip(1);
-				if (PostLinkCommands.Any())
-				{
-					List<string> ChangeCommands = new List<string>();
-					foreach (string PostLinkCommand in PostLinkCommands)
-					{
-						int ChangeIndex = PostLinkCommand.IndexOf("-change");
-						if (ChangeIndex == -1)
-							continue;
-
-						string LastDylibString = ".dylib ";
-						int LastDylibIndex = PostLinkCommand.LastIndexOf(LastDylibString);
-						if (LastDylibIndex == -1)
-							continue;
-
-						ChangeCommands.Add(PostLinkCommand.Substring(ChangeIndex, LastDylibIndex + LastDylibString.Count() - ChangeIndex));
-					}
-
-					if (ChangeCommands.Any())
-					{
-						AddText($"\t.LinkerStampExe = '$MacToolchainDir$/install_name_tool' \n");
-						AddText($"\t.LinkerStampExeArgs = '{string.Join(" ", ChangeCommands)} {OutputFile}' \n");
-					}
-				}
-
-				AddText($"}}\n\n");
-			}
-			else
-			{
-				Logger.LogError("Failed to add link action!");
-				PrintActionDetails(Action, Logger);
-			}
+			AddText($"Exec('{ActionToActionString(Action)}')\n{{\n");
+			AddText($"\t.ExecExecutable = '{Action.CommandPath.FullName}' \n");
+			AddText($"\t.ExecArguments = '{Action.CommandArguments}' \n");
+			AddText($"\t.ExecWorkingDir = '{Action.WorkingDirectory.FullName}' \n");
+			AddText($"\t.ExecOutput = '{Action.ProducedItems.First().FullName}' \n");
+			AddText($"\t.ExecAlways = true \n");
+			AddPreBuildDependenciesText(DependencyActions);
+			AddText($"}}\n\n");
 		}
 
 		private void PrintActionDetails(LinkedAction ActionToPrint, ILogger Logger)
@@ -1402,18 +1169,13 @@ namespace UnrealBuildTool
 					AddText($";** CommandArguments: {Action.CommandArguments}\n");
 					AddText("\n");
 
-					switch (Action.ActionType)
+					if (Action.ActionType == ActionType.Compile && Action.bCanExecuteRemotely && Action.bCanExecuteRemotelyWithSNDBS)
+					{
+						AddCompileAction(Action, DependencyActions, Logger);
+					}
+					else
 					{
-						case ActionType.Compile:
-							AddCompileAction(Action, DependencyActions, Logger);
-							break;
-						case ActionType.Link:
-							AddLinkAction(Action, DependencyActions, Logger);
-							break;
-						default:
-							Logger.LogWarning("FASTBuild is ignoring an unsupported action!");
-							PrintActionDetails(Action, Logger);
-							break;
+						AddExecAction(Action, DependencyActions, Logger);
 					}
 				}
 
