PR #9449: Add support for specifying C Standard when compiling plain C files  (Contributed by [REDACTED_PERSON])

diff --git a/Engine/Source/Programs/UnrealBuildTool/Configuration/ModuleRules.cs b/Engine/Source/Programs/UnrealBuildTool/Configuration/ModuleRules.cs
--- a/Engine/Source/Programs/UnrealBuildTool/Configuration/ModuleRules.cs
+++ b/Engine/Source/Programs/UnrealBuildTool/Configuration/ModuleRules.cs
@@ -1098,6 +1098,12 @@ namespace UnrealBuildTool
 		/// Which stanard to use for compiling this module
 		/// </summary>
 		public CppStandardVersion CppStandard = CppStandardVersion.Default;
+		
+		/// <summary>
+		/// Which standard to use for compiling this module
+		/// </summary>
+		public CStandardVersion CStandard = CStandardVersion.Default;
+
 
 		/// <summary>
 		///  Control visibility of symbols
diff --git a/Engine/Source/Programs/UnrealBuildTool/Configuration/TargetRules.cs b/Engine/Source/Programs/UnrealBuildTool/Configuration/TargetRules.cs
--- a/Engine/Source/Programs/UnrealBuildTool/Configuration/TargetRules.cs
+++ b/Engine/Source/Programs/UnrealBuildTool/Configuration/TargetRules.cs
@@ -1786,6 +1786,14 @@ namespace UnrealBuildTool
 		[XmlConfigFile(Category = "BuildConfiguration")]
 		public CppStandardVersion CppStandard = CppStandardVersion.Default;
 
+		/// <summary>
+		/// Which C standard to use for compiling this target
+		/// </summary>
+		[RequiresUniqueBuildEnvironment]
+		[CommandLine("-CStd")]
+		[XmlConfigFile(Category = "BuildConfiguration")]
+		public CStandardVersion CStandard = CStandardVersion.Default;
+
 		/// <summary>
 		/// Do not allow manifest changes when building this target. Used to cause earlier errors when building multiple targets with a shared build environment.
 		/// </summary>
@@ -3329,6 +3337,11 @@ namespace UnrealBuildTool
 			get { return Inner.CppStandard; }
 		}
 
+		public CStandardVersion CStandard
+		{
+			get { return Inner.CStandard; }
+		}
+
 		internal bool bNoManifestChanges
 		{
 			get { return Inner.bNoManifestChanges; }
diff --git a/Engine/Source/Programs/UnrealBuildTool/Configuration/UEBuildModuleCPP.cs b/Engine/Source/Programs/UnrealBuildTool/Configuration/UEBuildModuleCPP.cs
--- a/Engine/Source/Programs/UnrealBuildTool/Configuration/UEBuildModuleCPP.cs
+++ b/Engine/Source/Programs/UnrealBuildTool/Configuration/UEBuildModuleCPP.cs
@@ -856,6 +856,12 @@ namespace UnrealBuildTool
 			{
 				return false;
 			}
+			
+			if (ModuleCompileEnvironment.CStandard != CompileEnvironment.CStandard)
+			{
+				return false;
+			}
+			
 			if (ModuleCompileEnvironment.IncludeOrderVersion != CompileEnvironment.IncludeOrderVersion)
 			{
 				return false;
@@ -954,6 +960,11 @@ namespace UnrealBuildTool
 				Variant += String.Format(".{0}", CompileEnvironment.CppStandard);
 			}
 
+			if (CompileEnvironment.CStandard != BaseCompileEnvironment.CStandard)
+			{
+				Variant += String.Format(".{0}", CompileEnvironment.CStandard);
+			}
+
 			if (CompileEnvironment.IncludeOrderVersion != BaseCompileEnvironment.IncludeOrderVersion)
 			{
 				if (CompileEnvironment.IncludeOrderVersion != EngineIncludeOrderVersion.Latest)
@@ -979,6 +990,7 @@ namespace UnrealBuildTool
 			CompileEnvironment.UnsafeTypeCastWarningLevel = ModuleCompileEnvironment.UnsafeTypeCastWarningLevel;
 			CompileEnvironment.bEnableUndefinedIdentifierWarnings = ModuleCompileEnvironment.bEnableUndefinedIdentifierWarnings;
 			CompileEnvironment.CppStandard = ModuleCompileEnvironment.CppStandard;
+			CompileEnvironment.CStandard = ModuleCompileEnvironment.CStandard;
 			CompileEnvironment.IncludeOrderVersion = ModuleCompileEnvironment.IncludeOrderVersion;
 		}
 
@@ -1485,6 +1497,12 @@ namespace UnrealBuildTool
 				Result.CppStandard = CppStandardVersion.Cpp20;
 			}
 
+			// If the module overrides the C language version, override it on the compile environment
+			if (Rules.CStandard != CStandardVersion.Default)
+			{
+				Result.CStandard = Rules.CStandard;
+			}
+
 			// Set the macro used to check whether monolithic headers can be used
 			if (Rules.bTreatAsEngineModule && (!Rules.bEnforceIWYU || !Target.bEnforceIWYU))
 			{
diff --git a/Engine/Source/Programs/UnrealBuildTool/Configuration/UEBuildTarget.cs b/Engine/Source/Programs/UnrealBuildTool/Configuration/UEBuildTarget.cs
--- a/Engine/Source/Programs/UnrealBuildTool/Configuration/UEBuildTarget.cs
+++ b/Engine/Source/Programs/UnrealBuildTool/Configuration/UEBuildTarget.cs
@@ -3887,6 +3887,7 @@ namespace UnrealBuildTool
 			GlobalCompileEnvironment.bCompileISPC = Rules.bCompileISPC;
 			GlobalCompileEnvironment.bHideSymbolsByDefault = !Rules.bPublicSymbolsByDefault;
 			GlobalCompileEnvironment.CppStandard = Rules.CppStandard;
+			GlobalCompileEnvironment.CStandard = Rules.CStandard;
 			GlobalCompileEnvironment.AdditionalArguments = Rules.AdditionalCompilerArguments ?? String.Empty;
 			GlobalCompileEnvironment.CrashDiagnosticDirectory = Rules.CrashDiagnosticDirectory;
 
diff --git a/Engine/Source/Programs/UnrealBuildTool/Modes/GenerateClangDatabase.cs b/Engine/Source/Programs/UnrealBuildTool/Modes/GenerateClangDatabase.cs
--- a/Engine/Source/Programs/UnrealBuildTool/Modes/GenerateClangDatabase.cs
+++ b/Engine/Source/Programs/UnrealBuildTool/Modes/GenerateClangDatabase.cs
@@ -84,26 +84,24 @@
 							{
 								UEBuildModuleCPP.InputFileCollection InputFileCollection = Module.FindInputFiles(Target.Platform, new Dictionary<DirectoryItem, FileItem[]>());
 
-								List<FileItem> InputFiles = new List<FileItem>();
-								InputFiles.AddRange(InputFileCollection.CPPFiles);
-								InputFiles.AddRange(InputFileCollection.CCFiles);
-
 								CppCompileEnvironment ModuleCompileEnvironment = Module.CreateModuleCompileEnvironment(Target.Rules, BinaryCompileEnvironment);
 
 								StringBuilder CommandBuilder = new StringBuilder();
+								StringBuilder CppCommandBuilder = new StringBuilder();
+								StringBuilder CCommandBuilder = new StringBuilder();
 								CommandBuilder.AppendFormat("\"{0}\"", ClangPath.FullName);
 
 								switch (ModuleCompileEnvironment.CppStandard)
 								{
 									case CppStandardVersion.Cpp14:
-										CommandBuilder.AppendFormat(IsWindowsClang ? " /std:c++14" : " -std=c++14");
+										CppCommandBuilder.AppendFormat(IsWindowsClang ? " /std:c++14" : " -std=c++14");
 										break;
 									case CppStandardVersion.Cpp17:
-										CommandBuilder.AppendFormat(IsWindowsClang ? " /std:c++17" : " -std=c++17");
+										CppCommandBuilder.AppendFormat(IsWindowsClang ? " /std:c++17" : " -std=c++17");
 										break;
 									case CppStandardVersion.Cpp20:
 									case CppStandardVersion.Latest:
-										CommandBuilder.AppendFormat(IsWindowsClang ? " /std:c++latest" : " -std=c++20");
+										CppCommandBuilder.AppendFormat(IsWindowsClang ? " /std:c++latest" : " -std=c++20");
 										break;
 									default:
 										throw new BuildException($"Unsupported C++ standard type set: {ModuleCompileEnvironment.CppStandard}");
@@ -111,13 +109,36 @@
 
 								if (ModuleCompileEnvironment.bEnableCoroutines)
 								{
-									CommandBuilder.AppendFormat(" -fcoroutines-ts");
+									CppCommandBuilder.AppendFormat(" -fcoroutines-ts");
 									if (!ModuleCompileEnvironment.bEnableExceptions)
 									{
-										CommandBuilder.AppendFormat(" -Wno-coroutine-missing-unhandled-exception");
+										CppCommandBuilder.AppendFormat(" -Wno-coroutine-missing-unhandled-exception");
 									}
 								}
 
+								switch (ModuleCompileEnvironment.CStandard)
+								{
+									case CStandardVersion.Default:
+										break;
+									case CStandardVersion.C89:
+										CCommandBuilder.AppendFormat(IsWindowsClang ? "" : " -std=c89");
+										break;
+									case CStandardVersion.C99:
+										CCommandBuilder.AppendFormat(IsWindowsClang ? "" : " -std=c99");
+										break;
+									case CStandardVersion.C11:
+										CCommandBuilder.AppendFormat(IsWindowsClang ? " /std:c11" : " -std=c11");
+										break;
+									case CStandardVersion.C17:
+										CCommandBuilder.AppendFormat(IsWindowsClang ? " /std:c17" : " -std=c17");
+										break;
+									case CStandardVersion.Latest:
+										CCommandBuilder.AppendFormat(IsWindowsClang ? " /std:c17" : " -std=c2x");
+										break;
+									default:
+										throw new BuildException($"Unsupported C standard type set: {ModuleCompileEnvironment.CStandard}");
+								}
+
 								foreach (FileItem ForceIncludeFile in ModuleCompileEnvironment.ForceIncludeFiles)
 								{
 									CommandBuilder.AppendFormat(" -include \"{0}\"", ForceIncludeFile.FullName);
@@ -135,13 +156,21 @@
 									CommandBuilder.AppendFormat(" -I\"{0}\"", IncludePath);
 								}
 
-								foreach (FileItem InputFile in InputFiles)
+								foreach (FileItem InputFile in InputFileCollection.CPPFiles)
 								{
 									if (FileFilter == null || FileFilter.Matches(InputFile.Location.MakeRelativeTo(Unreal.RootDirectory)))
 									{
-										FileToCommand[InputFile.Location] = String.Format("{0} \"{1}\"", CommandBuilder, InputFile.FullName);
+										FileToCommand[InputFile.Location] = String.Format("{0} {1} \"{2}\"", CommandBuilder, CppCommandBuilder, InputFile.FullName);
 									}
 								}
+
+								foreach (FileItem InputFile in InputFileCollection.CFiles)
+								{
+									if (FileFilter == null || FileFilter.Matches(InputFile.Location.MakeRelativeTo(Unreal.RootDirectory)))
+									{
+										FileToCommand[InputFile.Location] = String.Format("{0} {1} \"{2}\"", CommandBuilder, CCommandBuilder, InputFile.FullName);
+									}
+								}
 							}
 						}
 					}
diff --git a/Engine/Source/Programs/UnrealBuildTool/Platform/Windows/VCToolChain.cs b/Engine/Source/Programs/UnrealBuildTool/Platform/Windows/VCToolChain.cs
--- a/Engine/Source/Programs/UnrealBuildTool/Platform/Windows/VCToolChain.cs
+++ b/Engine/Source/Programs/UnrealBuildTool/Platform/Windows/VCToolChain.cs
@@ -972,13 +972,23 @@ namespace UnrealBuildTool
 			}
 		}
 
-		static void AppendCLArguments_C(List<string> Arguments)
+		static void AppendCLArguments_C(CppCompileEnvironment CompileEnvironment, List<string> Arguments)
 		{
 			// Explicitly compile the file as C.
 			Arguments.Add("/TC");
 
 			// Level 0 warnings.  Needed for external C projects that produce warnings at higher warning levels.
 			Arguments.Add("/W0");
+
+			// Select C Standard version available
+			if (CompileEnvironment.CStandard == CStandardVersion.C11)
+			{
+				Arguments.Add("/std:c11");
+			}
+			else if (CompileEnvironment.CStandard >= CStandardVersion.C17)
+			{
+				Arguments.Add("/std:c17");
+			}
 		}
 
 		protected virtual void AppendLinkArguments(LinkEnvironment LinkEnvironment, List<string> Arguments)
@@ -1379,7 +1389,7 @@ namespace UnrealBuildTool
 				// Add C or C++ specific compiler arguments.
 				if (bIsPlainCFile)
 				{
-					AppendCLArguments_C(CompileAction.Arguments);
+					AppendCLArguments_C(CompileEnvironment, CompileAction.Arguments);
 				}
 				else
 				{
diff --git a/Engine/Source/Programs/UnrealBuildTool/System/CppCompileEnvironment.cs b/Engine/Source/Programs/UnrealBuildTool/System/CppCompileEnvironment.cs
--- a/Engine/Source/Programs/UnrealBuildTool/System/CppCompileEnvironment.cs
+++ b/Engine/Source/Programs/UnrealBuildTool/System/CppCompileEnvironment.cs
@@ -53,6 +53,42 @@ namespace UnrealBuildTool
 		/// </summary>
 		Default = Cpp17,
 	}
+	
+	/// <summary>
+	/// Specifies which C language standard to use. This enum should be kept in order, so that toolchains can check whether the requested setting is >= values that they support.
+	/// </summary>
+	public enum CStandardVersion
+	{
+		/// <summary>
+		/// Use the default standard version
+		/// </summary>
+		Default,
+
+		/// <summary>
+		/// Supports C89
+		/// </summary>
+		C89,
+
+		/// <summary>
+		/// Supports C99
+		/// </summary>
+		C99,
+
+		/// <summary>
+		/// Supports C11
+		/// </summary>
+		C11,
+
+		/// <summary>
+		/// Supports C17
+		/// </summary>
+		C17,
+
+		/// <summary>
+		/// Latest standard supported by the compiler
+		/// </summary>
+		Latest,
+	}
 
 	/// <summary>
 	/// The optimization level that may be compilation targets for C# files.
@@ -397,6 +433,11 @@ namespace UnrealBuildTool
 		/// Which C++ standard to support. May not be compatible with all platforms.
 		/// </summary>
 		public CppStandardVersion CppStandard = CppStandardVersion.Default;
+		
+		/// <summary>
+		/// Which C standard to support. May not be compatible with all platforms.
+		/// </summary>
+		public CStandardVersion CStandard = CStandardVersion.Default;
 
 		/// <summary>
 		/// The amount of the stack usage to report static analysis warnings.
@@ -501,6 +542,7 @@ namespace UnrealBuildTool
 			bHackHeaderGenerator = Other.bHackHeaderGenerator;
 			bHideSymbolsByDefault = Other.bHideSymbolsByDefault;
 			CppStandard = Other.CppStandard;
+			CStandard = Other.CStandard;
 			bEnableCoroutines = Other.bEnableCoroutines;
 			IncludeOrderVersion = Other.IncludeOrderVersion;
 			CrashDiagnosticDirectory = Other.CrashDiagnosticDirectory;
diff --git a/Engine/Source/Programs/UnrealBuildTool/ToolChain/ClangToolChain.cs b/Engine/Source/Programs/UnrealBuildTool/ToolChain/ClangToolChain.cs
--- a/Engine/Source/Programs/UnrealBuildTool/ToolChain/ClangToolChain.cs
+++ b/Engine/Source/Programs/UnrealBuildTool/ToolChain/ClangToolChain.cs
@@ -324,6 +324,32 @@ namespace UnrealBuildTool
 			}
 		}
 
+		protected virtual void GetCStandardCompileArgument(CppCompileEnvironment CompileEnvironment, List<string> Arguments)
+		{
+			switch (CompileEnvironment.CStandard)
+			{
+				case CStandardVersion.Default:
+					break;
+				case CStandardVersion.C89:
+					Arguments.Add("-std=c89");
+					break;
+				case CStandardVersion.C99:
+					Arguments.Add("-std=c99");
+					break;
+				case CStandardVersion.C11:
+					Arguments.Add("-std=c11");
+					break;
+				case CStandardVersion.C17:
+					Arguments.Add("-std=c17");
+					break;
+				case CStandardVersion.Latest:
+					Arguments.Add("-std=c2x");
+					break;
+				default:
+					throw new BuildException($"Unsupported C standard type set: {CompileEnvironment.CStandard}");
+			}
+		}
+
 		protected virtual void GetCompileArguments_CPP(CppCompileEnvironment CompileEnvironment, List<string> Arguments)
 		{
 			Arguments.Add("-x c++");
@@ -333,6 +359,7 @@ namespace UnrealBuildTool
 		protected virtual void GetCompileArguments_C(CppCompileEnvironment CompileEnvironment, List<string> Arguments)
 		{
 			Arguments.Add("-x c");
+			GetCStandardCompileArgument(CompileEnvironment, Arguments);
 		}
 
 		protected virtual void GetCompileArguments_MM(CppCompileEnvironment CompileEnvironment, List<string> Arguments)
