- First version of support for allowing a project to have a template project that GenProjectFiles will merge the generated project into the template project to be the final project. This allows for standard Xcode poject setup for signing, extensions, asset catalogs, etc. There are alot of edge cases still to deal with (conflicting settings from UBT vs Xcode, moving more stuff out of UBT since Xcode can do it, stuff like that).
- The merging of xcode project files is done with PlistBuddy, but with the amount of operations we do now, it takes seconds to make one project. This will need to change
- This is still only enabled when GPF is run with -xcconfig

diff --git a/Engine/Source/Programs/UnrealBuildTool/ProjectFiles/Xcode/XcodeProject.cs b/Engine/Source/Programs/UnrealBuildTool/ProjectFiles/Xcode/XcodeProject.cs
--- a/Engine/Source/Programs/UnrealBuildTool/ProjectFiles/Xcode/XcodeProject.cs
+++ b/Engine/Source/Programs/UnrealBuildTool/ProjectFiles/Xcode/XcodeProject.cs
@@ -142,6 +142,8 @@ namespace UnrealBuildTool.XcodeProjectXcconfig
 		public bool bIsAppBundle;
 		public bool bHasEditorConfiguration;
 		public bool bUseAutomaticSigning = false;
+		public bool bIsMergingProjects = false;
+		public bool bWriteCodeSigningSettings = true;
 
 		public List<UnrealBuildConfig> AllConfigs = new();
 
@@ -793,6 +795,12 @@ namespace UnrealBuildTool.XcodeProjectXcconfig
 			Project.FileCollection.AddFileReference(Xcconfig.Guid, $"Xcconfigs/{Xcconfig.FileRef.GetFileName()}", "test.xcconfig", "\"<group>\"", "Xcconfigs");
 		}
 
+		public virtual void WriteXcconfigFile()
+		{
+			
+		}
+
+
 		/// <summary>
 		/// THhis will walk the node reference tree and call WRite on each node to add all needed nodes to the xcode poject file
 		/// </summary>
@@ -802,10 +810,7 @@ namespace UnrealBuildTool.XcodeProjectXcconfig
 		{
 			// write the node into the xcode project file
 			Node.Write(Content);
-			if (Node.Xcconfig != null)
-			{
-				Node.Xcconfig.Write();
-			}
+			Node.WriteXcconfigFile();
 
 			foreach (XcodeProjectNode Reference in Node.References)
 			{
@@ -1116,10 +1121,14 @@ namespace UnrealBuildTool.XcodeProjectXcconfig
 			BuildConfigList!.BuildConfigs.ForEach(x => x.Xcconfig = Xcconfig);
 		}
 
-		public override void Write(StringBuilder Content)
+		protected override void WriteExtraTargetProperties(StringBuilder Content)
 		{
-			base.Write(Content);
+			Content.WriteLine($"\t\t\tproductReference = {ProductGuid};");
+			Content.WriteLine($"\t\t\tproductName = \"{ProductName}\";");
+		}
 
+		public override void WriteXcconfigFile()
+		{
 			// also write to the Xcconfig file
 			Xcconfig!.AppendLine("GENERATE_INFOPLIST_FILE[sdk=macosx*] = YES // Xcode 14 code-signing fix for development");
 
@@ -1129,12 +1138,8 @@ namespace UnrealBuildTool.XcodeProjectXcconfig
 			Xcconfig.AppendLine($"MACOSX_DEPLOYMENT_TARGET = {MacToolChain.Settings.MacOSVersion};");
 			Xcconfig.AppendLine("INFOPLIST_OUTPUT_FORMAT = xml");
 			Xcconfig.AppendLine("COMBINE_HIDPI_IMAGES = YES");
-		}
 
-		protected override void WriteExtraTargetProperties(StringBuilder Content)
-		{
-			Content.WriteLine($"\t\t\tproductReference = {ProductGuid};");
-			Content.WriteLine($"\t\t\tproductName = \"{ProductName}\";");
+			Xcconfig.Write();
 		}
 	}
 
@@ -1194,10 +1199,8 @@ namespace UnrealBuildTool.XcodeProjectXcconfig
 			}
 		}
 
-		public override void Write(StringBuilder Content)
+		public override void WriteXcconfigFile()
 		{
-			base.Write(Content);
-
 			// write out settings that apply whether or not we have subtargets, which #include this one, or no subtargets, and we write out more
 			// @todo move tis to the subtarget and remember it from the Module
 			Xcconfig!.AppendLine("CLANG_CXX_LANGUAGE_STANDARD = c++17");
@@ -1211,6 +1214,8 @@ namespace UnrealBuildTool.XcodeProjectXcconfig
 				// ask the single sub target to write out the contents of our Xcconfig file since we are folding them up into this target
 				SoloTargetBatch.WriteXcconfigSettings(Xcconfig.Text);
 			}
+
+			Xcconfig.Write();
 		}
 	}
 
@@ -1248,15 +1253,15 @@ namespace UnrealBuildTool.XcodeProjectXcconfig
 			}
 		}
 
-		public override void Write(StringBuilder Content)
+		public override void WriteXcconfigFile()
 		{
-			base.Write(Content);
-
 			// include the share Index config
 			Xcconfig!.AppendLine($"#include \"{Name.Split('_')[0]}_Index.xcconfig\"");
 
 			// then right the build settings from the modules we are batched together
 			BatchedFiles.WriteXcconfigSettings(Xcconfig.Text);
+
+			Xcconfig.Write();
 		}
 	}
 
@@ -1310,9 +1315,6 @@ namespace UnrealBuildTool.XcodeProjectXcconfig
 
 		public override void Write(StringBuilder Content)
 		{
-			string RootGroupGuid = XcodeFileCollection.GetRootGroupGuid(FileCollection.Groups);
-
-
 			Content.WriteLine("/* Begin PBXProject section */");
 
 			Content.WriteLine(2, $"{Guid} /* Project object */ = {{");
@@ -1333,7 +1335,7 @@ namespace UnrealBuildTool.XcodeProjectXcconfig
 			Content.WriteLine(3,	"knownRegions = (");
 			Content.WriteLine(4,		"en");
 			Content.WriteLine(3,	");");
-			Content.WriteLine(3,	$"mainGroup = {RootGroupGuid};");
+			Content.WriteLine(3,	$"mainGroup = {FileCollection.MainGroupGuid};");
 			Content.WriteLine(3,	$"productRefGroup = {FileCollection.GetProductGroupGuid()};");
 			Content.WriteLine(3,	"projectDirPath = \"\";");
 			Content.WriteLine(3,	"projectRoot = \"\";");
@@ -1346,12 +1348,10 @@ namespace UnrealBuildTool.XcodeProjectXcconfig
 			Content.WriteLine(2, "};");
 
 			Content.WriteLine("/* End PBXProject section */");
-
-			WriteXcconfigFiles();
 		}
 
 
-		private void WriteXcconfigFiles()
+		public override void WriteXcconfigFile()
 		{
 			string UEDir = XcodeFileCollection.ConvertPath(Path.GetFullPath(Directory.GetCurrentDirectory() + "../../.."));
 			string SupportedIOSArchitectures = string.Join(" ", UnrealData.SupportedIOSArchitectures);
@@ -1407,7 +1407,7 @@ namespace UnrealBuildTool.XcodeProjectXcconfig
 			Xcconfig.AppendLine("USE_HEADERMAP = NO");
 			Xcconfig.AppendLine("ONLY_ACTIVE_ARCH = YES");
 
-			if (bAutomaticSigning)
+			if (bAutomaticSigning && UnrealData.bWriteCodeSigningSettings)
 			{
 				Xcconfig.AppendLine("CODE_SIGN_STYLE = Automatic");
 			}
@@ -1417,14 +1417,17 @@ namespace UnrealBuildTool.XcodeProjectXcconfig
 				Xcconfig.AppendLine($"VALID_ARCHS[sdk=iphone*] = {SupportedIOSArchitectures}");
 				Xcconfig.AppendLine($"IPHONEOS_DEPLOYMENT_TARGET = {IOSRunTimeVersion}");
 				Xcconfig.AppendLine($"TARGETED_DEVICE_FAMILY[sdk=iphone*] = {IOSRunTimeDevices}");
-				if (!string.IsNullOrEmpty(TEAM_IOS))
-				{
-					Xcconfig.AppendLine($"DEVELOPMENT_TEAM[sdk=iphone*] = {TEAM_IOS}");
-				}
-				Xcconfig.AppendLine($"CODE_SIGN_IDENTITY[sdk=iphone*] = {IOS_CERT}");
-				if (!bAutomaticSigning && !string.IsNullOrEmpty(UUID_IOS))
+				if (UnrealData.bWriteCodeSigningSettings)
 				{
-					Xcconfig.AppendLine($"PROVISIONING_PROFILE_SPECIFIER[sdk=iphone*] = {UUID_IOS}");
+					if (!string.IsNullOrEmpty(TEAM_IOS))
+					{
+						Xcconfig.AppendLine($"DEVELOPMENT_TEAM[sdk=iphone*] = {TEAM_IOS}");
+					}
+					Xcconfig.AppendLine($"CODE_SIGN_IDENTITY[sdk=iphone*] = {IOS_CERT}");
+					if (!bAutomaticSigning && !string.IsNullOrEmpty(UUID_IOS))
+					{
+						Xcconfig.AppendLine($"PROVISIONING_PROFILE_SPECIFIER[sdk=iphone*] = {UUID_IOS}");
+					}
 				}
 				if (UnrealData.UProjectFileLocation != null)
 				{
@@ -1434,16 +1437,19 @@ namespace UnrealBuildTool.XcodeProjectXcconfig
 			if (TVOSRunTimeVersion != null)
 			{
 				Xcconfig.AppendLine($"VALID_ARCHS[sdk=appletvos*] = {SupportedIOSArchitectures}");
-				Xcconfig.AppendLine($"IPHONEOS_DEPLOYMENT_TARGET = {TVOSRunTimeVersion}");
+				Xcconfig.AppendLine($"TVOS_DEPLOYMENT_TARGET = {TVOSRunTimeVersion}");
 				Xcconfig.AppendLine($"TARGETED_DEVICE_FAMILY[sdk=appletvos*] = {TVOSRunTimeDevices}");
-				if (!string.IsNullOrEmpty(TEAM_TVOS))
-				{
-					Xcconfig.AppendLine($"DEVELOPMENT_TEAM[sdk=iphone*] = {TEAM_TVOS}");
-				}
-				Xcconfig.AppendLine($"CODE_SIGN_IDENTITY[sdk=appletvos*] = {TVOS_CERT}");
-				if (!bAutomaticSigning && !string.IsNullOrEmpty(UUID_TVOS))
+				if (UnrealData.bWriteCodeSigningSettings)
 				{
-					Xcconfig.AppendLine($"PROVISIONING_PROFILE_SPECIFIER[sdk=appletvos*] = {UUID_TVOS}");
+					if (!string.IsNullOrEmpty(TEAM_TVOS))
+					{
+						Xcconfig.AppendLine($"DEVELOPMENT_TEAM[sdk=iphone*] = {TEAM_TVOS}");
+					}
+					Xcconfig.AppendLine($"CODE_SIGN_IDENTITY[sdk=appletvos*] = {TVOS_CERT}");
+					if (!bAutomaticSigning && !string.IsNullOrEmpty(UUID_TVOS))
+					{
+						Xcconfig.AppendLine($"PROVISIONING_PROFILE_SPECIFIER[sdk=appletvos*] = {UUID_TVOS}");
+					}
 				}
 				if (UnrealData.UProjectFileLocation != null)
 				{
@@ -1451,6 +1457,8 @@ namespace UnrealBuildTool.XcodeProjectXcconfig
 				}
 			}
 
+			Xcconfig.Write();
+
 			// Now for each config write out the specific settings
 
 			DirectoryReference? GameDir = UnrealData.UProjectFileLocation?.Directory;
@@ -1615,6 +1623,8 @@ namespace UnrealBuildTool.XcodeProjectXcconfig
 						}
 					}
 				}
+
+				ConfigXcconfig.Write();
 			}
 		}
 	}
@@ -2063,6 +2073,9 @@ namespace UnrealBuildTool.XcodeProjectXcconfig
 				FileCollection.ProcessFile(SourceFile, bIsForBuild:IsGeneratedProject);
 			}
 
+			// cache the main group
+			FileCollection.MainGroupGuid = XcodeFileCollection.GetRootGroupGuid(FileCollection.Groups);
+
 			UnrealData.FilterBuildableFiles(FileCollection);
 		}
 
@@ -2078,90 +2091,38 @@ namespace UnrealBuildTool.XcodeProjectXcconfig
 				return true;
 			}
 
+			// look for an existing project to use as a template (if none found, create one from scratch)
+			DirectoryReference BuildDirLocation = UnrealData.UProjectFileLocation == null ? Unreal.EngineDirectory : UnrealData.UProjectFileLocation.Directory;
+			string ExistingProjectName = UnrealData.XcodeProjectFileLocation!.GetFileNameWithoutAnyExtensions();
+			FileReference TemplateProject = FileReference.Combine(BuildDirLocation, "Build", "IOS", ExistingProjectName + ".xcodeproj", "project.pbxproj");
+
+			UnrealData.bIsMergingProjects = FileReference.Exists(TemplateProject);
+			UnrealData.bWriteCodeSigningSettings = !UnrealData.bIsMergingProjects;
+
+
 			// turn all UE files into internal representation
 			ProcessSourceFiles();
 
 			// now create the xcodeproject elements (project -> target -> buildconfigs, etc)
 			RootProject = new XcodeProject(UnrealData, FileCollection);
 
-			StringBuilder Content = new StringBuilder();
-
-			// look for an existing project to use as a template (if none found, create one from scratch)
-			DirectoryReference BuildDirLocation = UnrealData.UProjectFileLocation == null ? Unreal.EngineDirectory : UnrealData.UProjectFileLocation.Directory;
-			string ExistingProjectName = UnrealData.XcodeProjectFileLocation!.GetFileNameWithoutAnyExtensions();
-			FileReference TemplateProject = FileReference.Combine(BuildDirLocation, "Build", "IOS", ExistingProjectName + ".xcodeproj", "project.pbxproj");
 
 			if (TemplateProject.FullName.Contains("ThirdPerson"))
 			{
 				Console.WriteLine("foo");
 			}
 
-			FileReference PBXProjFilePath = ProjectFilePath + "/project.pbxproj";
-
 			bool bSuccess;
 			if (FileReference.Exists(TemplateProject))
 			{
-				// copy existing template project to final location
-				if (FileReference.Exists(PBXProjFilePath))
-				{
-					FileReference.Delete(PBXProjFilePath);
-				}
-				DirectoryReference.CreateDirectory(PBXProjFilePath.Directory);
-				FileReference.Copy(TemplateProject, PBXProjFilePath);
-
-				// write the nodes we need to add (Build/Index targets)
-				XcodeRunTarget RunTarget = XcodeProjectNode.GetNodesOfType<XcodeRunTarget>(RootProject).First();
-				XcodeBuildTarget BuildTarget = XcodeProjectNode.GetNodesOfType<XcodeBuildTarget>(RunTarget).First();
-				XcodeIndexTarget IndexTarget = XcodeProjectNode.GetNodesOfType<XcodeIndexTarget>(RootProject).First();
-				XcodeDependency BuildDependency = XcodeProjectNode.GetNodesOfType<XcodeDependency>(RunTarget).First();
-
-				Content.WriteLine(0, "{");
-				FileCollection.Write(Content);
-				XcodeProjectNode.WriteNodeAndReferences(Content, BuildTarget);
-				XcodeProjectNode.WriteNodeAndReferences(Content, IndexTarget);
-				XcodeProjectNode.WriteNodeAndReferences(Content, BuildDependency);
-				Content.WriteLine(0, "}");
-
-				// write to disk
-				FileReference ImportFile = FileReference.Combine(PBXProjFilePath.Directory, "import.plist");
-				File.WriteAllText(ImportFile.FullName, Content.ToString());
-
-				Func<string, string> Plist = (string Command) =>
-					{
-						Command = Command.Replace("\"", "\\\"");
-						return Utils.RunLocalProcessAndReturnStdOut("/usr/libexec/PlistBuddy", $"-c \"{Command}\" \"{PBXProjFilePath.FullName}\"");
-					};
-
-				// and import it into the template
-				Plist($"Merge \"{ImportFile.FullName}\" :objects");
-
-				// get the project guid
-				string ProjectGuid = Plist($"Print :rootObject");
-
-				// get all the targets in the template that are application types
-				IEnumerable<string> AppTargetGuids = Plist($"Print :objects:{ProjectGuid}:targets")
-					.ReplaceLineEndings()
-					.Split(Environment.NewLine)
-					.Where(TargetGuid => (Plist($"Print :objects:{TargetGuid}:productType") == "com.apple.product-type.application"));
-
-				// add a dependency on the build target from the app target(s)
-				foreach (string AppTargetGuid in AppTargetGuids)
-				{
-					Plist($"Add :objects:{AppTargetGuid}:dependencies:0 string {BuildTarget.Guid}");
-				}
-
-				// now add all the non-run targets from the generated
-				foreach (XcodeTarget Target in XcodeProjectNode.GetNodesOfType<XcodeTarget>(RootProject).Where(x => x.GetType() != typeof(XcodeRunTarget)))
-				{
-					Plist($"Add :objects:{ProjectGuid}:targets:0 string {Target.Guid}");
-				}
-
-
-
-				bSuccess = true;
+				bSuccess = MergeIntoTemplateProject(TemplateProject);
 			}
 			else
 			{
+				FileReference PBXProjFilePath = ProjectFilePath + "/project.pbxproj";
+
+				StringBuilder Content = new StringBuilder();
+
 				Content.WriteLine(0, "// !$*UTF8*$!");
 				Content.WriteLine(0, "{");
 				Content.WriteLine(1, "archiveVersion = 1;");
@@ -2213,7 +2174,257 @@ namespace UnrealBuildTool.XcodeProjectXcconfig
 			return bSuccess;
 		}
 
-#region Schemes
+
+		private string Plist(string Command)
+		{
+			Command = Command.Replace("\"", "\\\"");
+			string Output = Utils.RunLocalProcessAndReturnStdOut("/usr/libexec/PlistBuddy", $"-c \"{Command}\" \"{ProjectFilePath.FullName}/project.pbxproj\"");
+			
+			//Console.WriteLine($"{Command} ==> {Output}");
+
+			return Output;
+		}
+
+		private void PlistSetAdd(string Entry, string Value, string Type="string")
+		{
+			string AddOutput = Plist($"Add {Entry} {Type} {Value}");
+			// error will be non-empty string
+			if (AddOutput != "")
+			{
+				Plist($"Set {Entry} {Value}");
+			}
+		}
+
+		private bool PlistSetUpdate(string Entry, string Value)
+		{
+			// see if the setting is already there
+			string ExistingSetting = Plist($"Print {Entry}");
+
+			// Print errors start with Print
+			if (!ExistingSetting.StartsWith("Print:") && ExistingSetting != Value)
+			{
+				Plist($"Set {Entry} {Value}");
+
+				return true;
+			}
+
+			return false;
+		}
+
+		private IEnumerable<string> PlistArray(string Entry)
+		{
+			return Plist($"Print {Entry}")
+				.Replace("Array {", "")
+				.Replace("}", "")
+				.Trim()
+				.ReplaceLineEndings()
+				.Split(Environment.NewLine)
+				.Select(x => x.Trim());
+		}
+
+		private List<string> PlistObjects()
+		{
+			List<string> Result = new();
+
+			IEnumerable<string> Lines = Plist("print :objects")
+				.ReplaceLineEndings()
+				.Split(Environment.NewLine);
+
+			Regex Regex = new Regex("^    (\\S*) = Dict {$");
+			foreach (string Line in Lines)
+			{
+				Match Match = Regex.Match(Line);
+				if (Match.Success)
+				{
+					Result.Add(Match.Groups[1].Value);
+				}
+			}
+			return Result;
+		}
+
+		private string? PlistFixPath(string Entry, string RelativeToProject)
+		{
+			string ExistingPath = Plist($"Print {Entry}");
+			// skip of errors, or it's an absolute path
+			if (!ExistingPath.StartsWith("Print:") && !ExistingPath.StartsWith("/"))
+			{
+				// fixup the path to be relative to new project instead of old
+				string FixedPath = Utils.CollapseRelativeDirectories(Path.Combine(RelativeToProject, ExistingPath));
+				// and set it back
+				Plist($"Set {Entry} {FixedPath}");
+
+				return FixedPath;
+			}
+
+			return null;
+		}
+
+
+		bool MergeIntoTemplateProject(FileReference TemplateProject)
+		{
+			FileReference PBXProjFilePath = ProjectFilePath + "/project.pbxproj";
+
+			// copy existing template project to final location
+			if (FileReference.Exists(PBXProjFilePath))
+			{
+				FileReference.Delete(PBXProjFilePath);
+			}
+			DirectoryReference.CreateDirectory(PBXProjFilePath.Directory);
+			FileReference.Copy(TemplateProject, PBXProjFilePath);
+
+			// write the nodes we need to add (Build/Index targets)
+			XcodeRunTarget RunTarget = XcodeProjectNode.GetNodesOfType<XcodeRunTarget>(RootProject!).First();
+			XcodeBuildTarget BuildTarget = XcodeProjectNode.GetNodesOfType<XcodeBuildTarget>(RunTarget).First();
+			XcodeIndexTarget IndexTarget = XcodeProjectNode.GetNodesOfType<XcodeIndexTarget>(RootProject!).First();
+			XcodeDependency BuildDependency = XcodeProjectNode.GetNodesOfType<XcodeDependency>(RunTarget).First();
+
+			// the runtarget and project need to write out so all of their xcconfigs get written as well,
+			// so write everything to a temp string that is tossed, but all xcconfigs will be done at least
+			StringBuilder Temp = new StringBuilder();
+			XcodeProjectNode.WriteNodeAndReferences(Temp, RootProject!);
+
+			StringBuilder Content = new StringBuilder();
+			Content.WriteLine(0, "{");
+			FileCollection.Write(Content);
+			XcodeProjectNode.WriteNodeAndReferences(Content, BuildTarget);
+			XcodeProjectNode.WriteNodeAndReferences(Content, IndexTarget);
+			XcodeProjectNode.WriteNodeAndReferences(Content, BuildDependency);
+			Content.WriteLine(0, "}");
+
+			// write to disk
+			FileReference ImportFile = FileReference.Combine(PBXProjFilePath.Directory, "import.plist");
+			File.WriteAllText(ImportFile.FullName, Content.ToString());
+
+
+			// cache some standard guids from the template project
+			string ProjectGuid = Plist($"Print :rootObject");
+			string TemplateMainGroupGuid = Plist($"Print :objects:{ProjectGuid}:mainGroup");
+
+			// fixup paths that were relative to original project to be relative to merged project
+//			List<string> ObjectGuids = PlistObjects();
+			IEnumerable<string> MainGroupChildrenGuids = PlistArray($":objects:{TemplateMainGroupGuid}:children");
+
+			string RelativeFromMergedToTemplate = TemplateProject.Directory.ParentDirectory!.MakeRelativeTo(PBXProjFilePath.Directory.ParentDirectory!);
+
+			// look for groups with a 'path' element that is in the main group, so that it and everything will get redirected to new location
+			string? FixedPath;
+			foreach (string ChildGuid in MainGroupChildrenGuids)
+			{
+				string IsA = Plist($"Print :objects:{ChildGuid}:isa");
+				// if a Group has a path
+				if (IsA == "PBXGroup")
+				{
+					if ((FixedPath = PlistFixPath($":objects:{ChildGuid}:path", RelativeFromMergedToTemplate)) != null)
+					{
+						// if there wasn't a name before, it will now have a nasty path as the name, so add it now
+						PlistSetAdd($":objects:{ChildGuid}:name", Path.GetFileName(FixedPath));
+					}
+				}
+			}
+
+			// and import it into the template
+			Plist($"Merge \"{ImportFile.FullName}\" :objects");
+
+
+			// get all the targets in the template that are application types
+			IEnumerable<string> AppTargetGuids = PlistArray($":objects:{ProjectGuid}:targets")
+				.Where(TargetGuid => (Plist($"Print :objects:{TargetGuid}:productType") == "com.apple.product-type.application"));
+
+			// add a dependency on the build target from the app target(s)
+			foreach (string AppTargetGuid in AppTargetGuids)
+			{
+				Plist($"Add :objects:{AppTargetGuid}:dependencies:0 string {BuildDependency.Guid}");
+			}
+
+			// the BuildDependency object was in the "container" of the generated project, not the merged one, so fix it up now
+			Plist($"Set :objects:{BuildDependency.ProxyGuid}:containerPortal {ProjectGuid}");
+
+			// now add all the non-run targets from the generated
+			foreach (XcodeTarget Target in XcodeProjectNode.GetNodesOfType<XcodeTarget>(RootProject!).Where(x => x.GetType() != typeof(XcodeRunTarget)))
+			{
+				Plist($"Add :objects:{ProjectGuid}:targets:0 string {Target.Guid}");
+			}
+
+
+			// hook up Xcconfig files to the project and the project configs
+			// @todo how to manage with conflicts already present...
+			//PlistSetAdd($":objects:{ProjectGuid}:baseConfigurationReference", RootProject.Xcconfig!.Guid, "string");
+
+			// re-get the list of targets now that we merged in the other file
+			IEnumerable<string> AllTargetGuids = PlistArray($":objects:{ProjectGuid}:targets");
+
+			List<string> NodesToFix = new() { ProjectGuid };
+			NodesToFix.AddRange(AllTargetGuids);
+
+			bool bIsProject = true;
+			foreach (string NodeGuid in NodesToFix)
+			{
+				bool bIsAppTarget = AppTargetGuids.Contains(NodeGuid);
+
+				// get the config list, and from there we can get the configs
+				string ProjectBuildConfigListGuid = Plist($"Print :objects:{NodeGuid}:buildConfigurationList");
+
+				IEnumerable<string> ConfigGuids = PlistArray($":objects:{ProjectBuildConfigListGuid}:buildConfigurations");
+				foreach (string ConfigGuid in ConfigGuids)
+				{
+					// find the matching unreal generated project build config to hook up to
+					// for now we assume Release is Development [Editor], but we should make sure the template project has good configs
+					// we have to rename the template config from Release because it won't find the matching config in the build target
+					string ConfigName = Plist($"Print :objects:{ConfigGuid}:name");
+					if (ConfigName == "Release")
+					{
+						ConfigName = UnrealData.bHasEditorConfiguration ? "Development Editor" : "Development";
+						Plist($"Set :objects:{ConfigGuid}:name \"{ConfigName}\"");
+					}
+
+					// if there's a plist path, then it will need to be fixed up
+					PlistFixPath($":objects:{ConfigGuid}:buildSettings:INFOPLIST_FILE", RelativeFromMergedToTemplate);
+
+					if (bIsProject)
+					{
+						Console.WriteLine("Looking for " + ConfigName);
+						XcodeBuildConfig Config = RootProject!.ProjectBuildConfigs.BuildConfigs.First(x => x.Info.DisplayName == ConfigName);
+						PlistSetAdd($":objects:{ConfigGuid}:baseConfigurationReference", Config.Xcconfig!.Guid, "string");
+					}
+
+					// the Build target used some ini settings to compile, and Run target must match, so we override a few settings, at
+					// whatever level they were already specified at (Projet and/or Target)
+					PlistSetUpdate($":objects:{ConfigGuid}:buildSettings:MACOSX_DEPLOYMENT_TARGET", MacToolChain.Settings.MacOSVersion);
+					if (UnrealData.IOSProjectSettings != null)
+					{
+						PlistSetUpdate($":objects:{ConfigGuid}:buildSettings:IPHONEOS_DEPLOYMENT_TARGET", UnrealData.IOSProjectSettings.RuntimeVersion);
+					}
+					if (UnrealData.TVOSProjectSettings != null)
+					{
+						PlistSetUpdate($":objects:{ConfigGuid}:buildSettings:TVOS_DEPLOYMENT_TARGET", UnrealData.TVOSProjectSettings.RuntimeVersion);
+					}
+				}
+
+				bIsProject = false;
+			}
+
+			// now we need to merge the main groups together
+			string GeneratedMainGroupGuid = FileCollection.MainGroupGuid;
+			int Index = 0;
+			while (true)
+			{
+				// we copy to a high index to put the copied entries at the end in the same order
+				string Output = Plist($"Copy :objects:{GeneratedMainGroupGuid}:children:{Index} :objects:{TemplateMainGroupGuid}:children:100000000");
+
+				// loop until error
+				if (Output != "")
+				{
+					break;
+				}
+				Index++;
+			}
+			// and remove the one we copied from
+			Plist($"Delete :objects:{GeneratedMainGroupGuid}");
+
+			return true;
+		}
+
+		#region Schemes
 
 		private FileReference GetUserSchemeManagementFilePath()
 		{
