UnrealBuildTool: Preferred vs2022 version

#rnx

diff --git a/Engine/Source/Programs/UnrealBuildTool/Platform/Windows/MicrosoftPlatformSDK.Versions.cs b/Engine/Source/Programs/UnrealBuildTool/Platform/Windows/MicrosoftPlatformSDK.Versions.cs
--- a/Engine/Source/Programs/UnrealBuildTool/Platform/Windows/MicrosoftPlatformSDK.Versions.cs
+++ b/Engine/Source/Programs/UnrealBuildTool/Platform/Windows/MicrosoftPlatformSDK.Versions.cs
@@ -34,7 +34,8 @@ namespace UnrealBuildTool
 		/// </summary>
 		static readonly VersionNumberRange[] PreferredVisualCppVersions = new VersionNumberRange[]
 		{
-			VersionNumberRange.Parse("14.29.30133", "14.29.30145"), // VS2019 16.11.32428.217
+			VersionNumberRange.Parse("14.29.30133", "14.29.30146"), // VS2019 16.11.x
+			VersionNumberRange.Parse("14.32.31326", "14.32.31333"), // VS2022 17.2.x
 		};
 
 		/// <summary>
