Workaround for debugger issues with Android NKD25 related to debug symbols storing relative path to the engine's source code instead of full path

diff --git a/Engine/Source/Programs/UnrealBuildTool/Platform/Android/AndroidProjectGenerator.cs b/Engine/Source/Programs/UnrealBuildTool/Platform/Android/AndroidProjectGenerator.cs
--- a/Engine/Source/Programs/UnrealBuildTool/Platform/Android/AndroidProjectGenerator.cs
+++ b/Engine/Source/Programs/UnrealBuildTool/Platform/Android/AndroidProjectGenerator.cs
@@ -21,6 +21,11 @@ namespace UnrealBuildTool
 		/// </summary>
 		private bool AGDEInstalled = false;
 
+		/// <summary>
+		/// Cached version of Android NDK toolchain
+		/// </summary>
+		private UInt64 NDKVersionInt;
+
 		public AndroidProjectGenerator(CommandLineArguments Arguments, ILogger Logger)
 			: base(Arguments, Logger)
 		{
@@ -31,6 +36,12 @@ namespace UnrealBuildTool
 				{
 					AGDEInstalled = false;
 				}
+				else
+				{
+					UEBuildPlatformSDK SDK = UEBuildPlatform.GetSDK(UnrealTargetPlatform.Android)!;
+					string? NDKToolchainVersion = SDK.GetInstalledVersion();
+					SDK.TryConvertVersionToInt(NDKToolchainVersion, out NDKVersionInt);
+				}
 			}
 		}
 
@@ -168,7 +179,11 @@ namespace UnrealBuildTool
 				&& ((InTargetRules.Type == TargetRules.TargetType.Client) || (InTargetRules.Type == TargetRules.TargetType.Game)))
 			{
 				string UserFileEntry = "<PropertyGroup " + InConditionString + ">\n";
-				UserFileEntry		+= "	<AndroidLldbStartupCommands>command script import \"" + Path.Combine(Unreal.EngineDirectory.FullName, "Extras", "LLDBDataFormatters", "UEDataFormatters_2ByteChars.py") + "\"</AndroidLldbStartupCommands>\n";
+				UserFileEntry		+= "	<AndroidLldbStartupCommands>" +
+												"command script import \"" + Path.Combine(Unreal.EngineDirectory.FullName, "Extras", "LLDBDataFormatters", "UEDataFormatters_2ByteChars.py") + "\";" +
+												(NDKVersionInt >= 250000 ? "settings set target.source-map \"\" \"" + Path.Combine(Unreal.EngineDirectory.FullName, "Source") + "\";" : "") +
+												"$(AndroidLldbStartupCommands)" +
+											"</AndroidLldbStartupCommands>\n";
 				UserFileEntry		+= "</PropertyGroup>\n";
 				return UserFileEntry;
 			}
