UnrealBuildTool: Handle PCH correctly when running clang analyze

#rnx

diff --git a/Engine/Source/Programs/UnrealBuildTool/Platform/Windows/VCToolChain.cs b/Engine/Source/Programs/UnrealBuildTool/Platform/Windows/VCToolChain.cs
--- a/Engine/Source/Programs/UnrealBuildTool/Platform/Windows/VCToolChain.cs
+++ b/Engine/Source/Programs/UnrealBuildTool/Platform/Windows/VCToolChain.cs
@@ -316,6 +316,10 @@ namespace UnrealBuildTool
 						}
 					}
 				}
+				else if (Target.StaticAnalyzer == StaticAnalyzer.Default && CompileEnvironment.PrecompiledHeaderAction == PrecompiledHeaderAction.Create)
+				{
+					AddDefinition(Arguments, "__clang_analyzer__");
+				}
 			}
 			else if (Target.StaticAnalyzer == StaticAnalyzer.Default)
 			{
diff --git a/Engine/Source/Programs/UnrealBuildTool/ToolChain/ClangToolChain.cs b/Engine/Source/Programs/UnrealBuildTool/ToolChain/ClangToolChain.cs
--- a/Engine/Source/Programs/UnrealBuildTool/ToolChain/ClangToolChain.cs
+++ b/Engine/Source/Programs/UnrealBuildTool/ToolChain/ClangToolChain.cs
@@ -420,6 +420,11 @@ namespace UnrealBuildTool
 		protected virtual void GetCompileArguments_PreprocessorDefinitions(CppCompileEnvironment CompileEnvironment, List<string> Arguments)
 		{
 			Arguments.AddRange(CompileEnvironment.Definitions.Select(Definition => GetPreprocessorDefinitionArgument(Definition)));
+
+			if (CompileEnvironment.PrecompiledHeaderAction == PrecompiledHeaderAction.Create && StaticAnalyzer == StaticAnalyzer.Default)
+			{
+				Arguments.Add(GetPreprocessorDefinitionArgument("__clang_analyzer__"));
+			}
 		}
 
 		protected virtual string GetForceIncludeFileArgument(FileReference ForceIncludeFile)
@@ -535,7 +540,7 @@ namespace UnrealBuildTool
 			}
 
 			// shipping builds will cause this warning with "ensure", so disable only in those case
-			if (CompileEnvironment.Configuration == CppConfiguration.Shipping)
+			if (CompileEnvironment.Configuration == CppConfiguration.Shipping || StaticAnalyzer != StaticAnalyzer.None)
 			{
 				Arguments.Add("-Wno-unused-value");                     // https://clang.llvm.org/docs/DiagnosticsReference.html#wunused-value
 			}
