UnrealBuildTool: Move static analysis checker options into ModuleRules rather than TargetRules

#rnx

diff --git a/Engine/Source/Programs/UnrealBuildTool/Configuration/ModuleRules.cs b/Engine/Source/Programs/UnrealBuildTool/Configuration/ModuleRules.cs
--- a/Engine/Source/Programs/UnrealBuildTool/Configuration/ModuleRules.cs
+++ b/Engine/Source/Programs/UnrealBuildTool/Configuration/ModuleRules.cs
@@ -800,6 +800,37 @@ namespace UnrealBuildTool
 		/// </summary>
 		public bool bEnableUndefinedIdentifierWarnings = true;
 
+		/// <summary>
+		/// The static analyzer checkers that should be enabled rather than the defaults. This is only supported for Clang.
+		/// See https://clang.llvm.org/docs/analyzer/checkers.html for a full list. Or run:
+		///    'clang -Xclang -analyzer-checker-help' 
+		/// or: 
+		///    'clang -Xclang -analyzer-checker-help-alpha' 
+		/// for the list of experimental checkers.
+		/// </summary>
+		public HashSet<string> StaticAnalyzerCheckers = new HashSet<string>();
+
+		/// <summary>
+		/// The static analyzer default checkers that should be disabled. Unused if StaticAnalyzerCheckers is populated. This is only supported for Clang.
+		/// This overrides the default disabled checkers, which are deadcode.DeadStores and security.FloatLoopCounter
+		/// See https://clang.llvm.org/docs/analyzer/checkers.html for a full list. Or run:
+		///    'clang -Xclang -analyzer-checker-help' 
+		/// or: 
+		///    'clang -Xclang -analyzer-checker-help-alpha' 
+		/// for the list of experimental checkers.
+		/// </summary>
+		public HashSet<string> StaticAnalyzerDisabledCheckers = new HashSet<string>() { "deadcode.DeadStores", "security.FloatLoopCounter" };
+
+		/// <summary>
+		/// The static analyzer non-default checkers that should be enabled. Unused if StaticAnalyzerCheckers is populated. This is only supported for Clang.
+		/// See https://clang.llvm.org/docs/analyzer/checkers.html for a full list. Or run:
+		///    'clang -Xclang -analyzer-checker-help' 
+		/// or: 
+		///    'clang -Xclang -analyzer-checker-help-alpha' 
+		/// for the list of experimental checkers.
+		/// </summary>
+		public HashSet<string> StaticAnalyzerAdditionalCheckers = new HashSet<string>();
+
 		private bool? bUseUnityOverride;
 		/// <summary>
 		/// If unity builds are enabled this can be used to override if this specific module will build using Unity.
diff --git a/Engine/Source/Programs/UnrealBuildTool/Configuration/TargetRules.cs b/Engine/Source/Programs/UnrealBuildTool/Configuration/TargetRules.cs
--- a/Engine/Source/Programs/UnrealBuildTool/Configuration/TargetRules.cs
+++ b/Engine/Source/Programs/UnrealBuildTool/Configuration/TargetRules.cs
@@ -1469,57 +1469,6 @@ namespace UnrealBuildTool
 		[XmlConfigFile(Category = "BuildConfiguration")]
 		public StaticAnalyzerMode StaticAnalyzerMode = StaticAnalyzerMode.Deep;
 
-		/// <summary>
-		/// The static analyzer checkers that should be enabled rather than the defaults. This is only supported for Clang.
-		/// See https://clang.llvm.org/docs/analyzer/checkers.html for a full list. Or run:
-		///    'clang -Xclang -analyzer-checker-help' 
-		/// or: 
-		///    'clang -Xclang -analyzer-checker-help-alpha' 
-		/// for the list of experimental checkers.
-		/// </summary>
-		[CommandLine("-StaticAnalyzerChecker", ListSeparator = '+')]
-		public List<string> StaticAnalyzerCheckers = new List<string>();
-
-		/// <summary>
-		/// The static analyzer default checkers that should be disabled. Unused if StaticAnalyzerCheckers is set. This is only supported for Clang.
-		/// This overrides the default disabled checkers, which are deadcode.DeadStores and security.FloatLoopCounter
-		/// See https://clang.llvm.org/docs/analyzer/checkers.html for a full list. Or run:
-		///    'clang -Xclang -analyzer-checker-help' 
-		/// or: 
-		///    'clang -Xclang -analyzer-checker-help-alpha' 
-		/// for the list of experimental checkers.
-		/// </summary>
-		[CommandLine("-StaticAnalyzerDisableChecker", ListSeparator = '+')]
-		private List<string> StaticAnalyzerDisabledCheckersOverride = new List<string>();
-
-		/// <summary>
-		/// The static analyzer default checkers that should be disabled. Unused if StaticAnalyzerCheckers is set. This is only supported for Clang.
-		/// See https://clang.llvm.org/docs/analyzer/checkers.html for a full list. Or run:
-		///    'clang -Xclang -analyzer-checker-help' 
-		/// or: 
-		///    'clang -Xclang -analyzer-checker-help-alpha' 
-		/// for the list of experimental checkers.
-		/// </summary>
-		public List<string> StaticAnalyzerDisabledCheckers
-		{
-			get
-			{
-				return StaticAnalyzerDisabledCheckersOverride.Count > 0 ? StaticAnalyzerDisabledCheckersOverride : new List<string>() { "deadcode.DeadStores", "security.FloatLoopCounter" };
-			}
-			set { StaticAnalyzerDisabledCheckersOverride = value; }
-		}
-
-		/// <summary>
-		/// The static analyzer non-default checkers that should be enabled. Unused if StaticAnalyzerCheckers is set. This is only supported for Clang.
-		/// See https://clang.llvm.org/docs/analyzer/checkers.html for a full list. Or run:
-		///    'clang -Xclang -analyzer-checker-help' 
-		/// or: 
-		///    'clang -Xclang -analyzer-checker-help-alpha' 
-		/// for the list of experimental checkers.
-		/// </summary>
-		[CommandLine("-StaticAnalyzerAdditionalChecker", ListSeparator = '+')]
-		public List<string> StaticAnalyzerAdditionalCheckers = new List<string>();
-
 		/// <summary>
 		/// The minimum number of files that must use a pre-compiled header before it will be created and used.
 		/// </summary>
@@ -3137,21 +3086,6 @@ namespace UnrealBuildTool
 			get { return Inner.StaticAnalyzerMode; }
 		}
 
-		public IReadOnlyList<string> StaticAnalyzerCheckers
-		{
-			get { return Inner.StaticAnalyzerCheckers; }
-		}
-
-		public IReadOnlyList<string> StaticAnalyzerDisabledCheckers
-		{
-			get { return Inner.StaticAnalyzerDisabledCheckers; }
-		}
-
-		public IReadOnlyList<string> StaticAnalyzerAdditionalCheckers
-		{
-			get { return Inner.StaticAnalyzerAdditionalCheckers; }
-		}
-
 		public int MinFilesUsingPrecompiledHeader
 		{
 			get { return Inner.MinFilesUsingPrecompiledHeader; }
diff --git a/Engine/Source/Programs/UnrealBuildTool/Configuration/UEBuildModuleCPP.cs b/Engine/Source/Programs/UnrealBuildTool/Configuration/UEBuildModuleCPP.cs
--- a/Engine/Source/Programs/UnrealBuildTool/Configuration/UEBuildModuleCPP.cs
+++ b/Engine/Source/Programs/UnrealBuildTool/Configuration/UEBuildModuleCPP.cs
@@ -1484,6 +1484,9 @@ namespace UnrealBuildTool
 			Result.bEnableObjCExceptions |= Rules.bEnableObjCExceptions;
 			Result.ShadowVariableWarningLevel = Rules.ShadowVariableWarningLevel;
 			Result.UnsafeTypeCastWarningLevel = Rules.UnsafeTypeCastWarningLevel;
+			Result.StaticAnalyzerCheckers = Rules.StaticAnalyzerCheckers;
+			Result.StaticAnalyzerDisabledCheckers = Rules.StaticAnalyzerDisabledCheckers;
+			Result.StaticAnalyzerAdditionalCheckers = Rules.StaticAnalyzerAdditionalCheckers;
 			Result.bEnableUndefinedIdentifierWarnings = Rules.bEnableUndefinedIdentifierWarnings;
 			Result.IncludeOrderVersion = Rules.IncludeOrderVersion;
 
diff --git a/Engine/Source/Programs/UnrealBuildTool/Platform/Windows/VCToolChain.cs b/Engine/Source/Programs/UnrealBuildTool/Platform/Windows/VCToolChain.cs
--- a/Engine/Source/Programs/UnrealBuildTool/Platform/Windows/VCToolChain.cs
+++ b/Engine/Source/Programs/UnrealBuildTool/Platform/Windows/VCToolChain.cs
@@ -291,13 +291,13 @@ namespace UnrealBuildTool
 					// Run shallow analyze if requested.
 					if (Target.StaticAnalyzerMode == StaticAnalyzerMode.Shallow) Arguments.Add("-Xclang -analyzer-config -Xclang mode=shallow");
 
-					if (Target.StaticAnalyzerCheckers.Count > 0)
+					if (CompileEnvironment.StaticAnalyzerCheckers.Count > 0)
 					{
 						// Disable all default checks
 						Arguments.Add("--analyzer-no-default-checks");
 
 						// Only enable specific checks.
-						foreach (string Checker in Target.StaticAnalyzerCheckers)
+						foreach (string Checker in CompileEnvironment.StaticAnalyzerCheckers)
 						{
 							Arguments.Add($"-Xclang -analyzer-checker -Xclang {Checker}");
 						}
@@ -305,12 +305,12 @@ namespace UnrealBuildTool
 					else
 					{
 						// Disable default checks.
-						foreach (string Checker in Target.StaticAnalyzerDisabledCheckers)
+						foreach (string Checker in CompileEnvironment.StaticAnalyzerDisabledCheckers)
 						{
 							Arguments.Add($"-Xclang -analyzer-disable-checker -Xclang {Checker}");
 						}
 						// Enable additional non-default checks.
-						foreach (string Checker in Target.StaticAnalyzerAdditionalCheckers)
+						foreach (string Checker in CompileEnvironment.StaticAnalyzerAdditionalCheckers)
 						{
 							Arguments.Add($"-Xclang -analyzer-checker -Xclang {Checker}");
 						}
diff --git a/Engine/Source/Programs/UnrealBuildTool/System/CppCompileEnvironment.cs b/Engine/Source/Programs/UnrealBuildTool/System/CppCompileEnvironment.cs
--- a/Engine/Source/Programs/UnrealBuildTool/System/CppCompileEnvironment.cs
+++ b/Engine/Source/Programs/UnrealBuildTool/System/CppCompileEnvironment.cs
@@ -263,6 +263,21 @@ namespace UnrealBuildTool
 		/// </summary>
 		public bool bWarningsAsErrors = false;
 
+		/// <summary>
+		/// The static analyzer checkers that should be enabled rather than the defaults. This is only supported for Clang.
+		/// </summary>
+		public HashSet<string> StaticAnalyzerCheckers = new HashSet<string>();
+
+		/// <summary>
+		/// The static analyzer default checkers that should be disabled. Unused if StaticAnalyzerCheckers is populated. This is only supported for Clang.
+		/// </summary>
+		public HashSet<string> StaticAnalyzerDisabledCheckers = new HashSet<string>();
+
+		/// <summary>
+		/// The static analyzer non-default checkers that should be enabled. Unused if StaticAnalyzerCheckers is populated. This is only supported for Clang.
+		/// </summary>
+		public HashSet<string> StaticAnalyzerAdditionalCheckers = new HashSet<string>();
+
 		/// <summary>
 		/// True if compiler optimizations should be enabled. This setting is distinct from the configuration (see CPPTargetConfiguration).
 		/// </summary>
@@ -508,6 +523,9 @@ namespace UnrealBuildTool
 			bUndefinedIdentifierWarningsAsErrors = Other.bUndefinedIdentifierWarningsAsErrors;
 			bEnableUndefinedIdentifierWarnings = Other.bEnableUndefinedIdentifierWarnings;
 			bWarningsAsErrors = Other.bWarningsAsErrors;
+			StaticAnalyzerCheckers = new HashSet<string>(Other.StaticAnalyzerCheckers);
+			StaticAnalyzerDisabledCheckers = new HashSet<string>(Other.StaticAnalyzerDisabledCheckers);
+			StaticAnalyzerAdditionalCheckers = new HashSet<string>(Other.StaticAnalyzerAdditionalCheckers);
 			bOptimizeCode = Other.bOptimizeCode;
 			bOptimizeForSize = Other.bOptimizeForSize;
 			bCreateDebugInfo = Other.bCreateDebugInfo;
diff --git a/Engine/Source/Programs/UnrealBuildTool/ToolChain/ClangToolChain.cs b/Engine/Source/Programs/UnrealBuildTool/ToolChain/ClangToolChain.cs
--- a/Engine/Source/Programs/UnrealBuildTool/ToolChain/ClangToolChain.cs
+++ b/Engine/Source/Programs/UnrealBuildTool/ToolChain/ClangToolChain.cs
@@ -187,9 +187,6 @@ namespace UnrealBuildTool
 		protected StaticAnalyzer StaticAnalyzer = StaticAnalyzer.None;
 		protected StaticAnalyzerMode StaticAnalyzerMode = StaticAnalyzerMode.Deep;
 		protected StaticAnalyzerOutputType StaticAnalyzerOutputType = StaticAnalyzerOutputType.Text;
-		protected IReadOnlyList<string> StaticAnalyzerCheckers = new List<string>();
-		protected IReadOnlyList<string> StaticAnalyzerDisabledCheckers = new List<string>();
-		protected IReadOnlyList<string> StaticAnalyzerAdditionalCheckers = new List<string>();
 
 		static ClangToolChain()
 		{
@@ -216,9 +213,6 @@ namespace UnrealBuildTool
 			StaticAnalyzer = Target.StaticAnalyzer;
 			StaticAnalyzerMode = Target.StaticAnalyzerMode;
 			StaticAnalyzerOutputType = Target.StaticAnalyzerOutputType;
-			StaticAnalyzerCheckers = Target.StaticAnalyzerCheckers;
-			StaticAnalyzerDisabledCheckers = Target.StaticAnalyzerDisabledCheckers;
-			StaticAnalyzerAdditionalCheckers = Target.StaticAnalyzerAdditionalCheckers;
 		}
 
 		public override void FinalizeOutput(ReadOnlyTargetRules Target, TargetMakefileBuilder MakefileBuilder)
@@ -747,13 +741,13 @@ namespace UnrealBuildTool
 			// Run shallow analyze if requested.
 			if (StaticAnalyzerMode == StaticAnalyzerMode.Shallow) Arguments.Add("-Xclang -analyzer-config -Xclang mode=shallow");
 
-			if (StaticAnalyzerCheckers.Count > 0)
+			if (CompileEnvironment.StaticAnalyzerCheckers.Count > 0)
 			{
 				// Disable all default checks
 				Arguments.Add("--analyzer-no-default-checks");
 
 				// Only enable specific checkers.
-				foreach (string Checker in StaticAnalyzerCheckers)
+				foreach (string Checker in CompileEnvironment.StaticAnalyzerCheckers)
 				{
 					Arguments.Add($"-Xclang -analyzer-checker -Xclang {Checker}");
 				}
@@ -761,12 +755,12 @@ namespace UnrealBuildTool
 			else
 			{
 				// Disable default checks.
-				foreach (string Checker in StaticAnalyzerDisabledCheckers)
+				foreach (string Checker in CompileEnvironment.StaticAnalyzerDisabledCheckers)
 				{
 					Arguments.Add($"-Xclang -analyzer-disable-checker -Xclang {Checker}");
 				}
 				// Enable additional non-default checks.
-				foreach (string Checker in StaticAnalyzerAdditionalCheckers)
+				foreach (string Checker in CompileEnvironment.StaticAnalyzerAdditionalCheckers)
 				{
 					Arguments.Add($"-Xclang -analyzer-checker -Xclang {Checker}");
 				}
