[UnrealBuildTool]
* Fixed so .uproject added on commandline is fixed up with correct casing to make sure it generate the same solution as when discovering uproject file when not provided

Fixes so UGS, UnrealVs.RefreshProjects and GenerateProjectFiles.bat all create the same thing

diff --git a/Engine/Source/Programs/UnrealBuildTool/Modes/GenerateProjectFilesMode.cs b/Engine/Source/Programs/UnrealBuildTool/Modes/GenerateProjectFilesMode.cs
--- a/Engine/Source/Programs/UnrealBuildTool/Modes/GenerateProjectFilesMode.cs
+++ b/Engine/Source/Programs/UnrealBuildTool/Modes/GenerateProjectFilesMode.cs
@@ -277,7 +277,7 @@ namespace UnrealBuildTool
 			// so now validate it
 			if (!string.IsNullOrEmpty(CandidateProjectPath))
 			{
-				FileReference? CandidateProjectFile = new FileReference(CandidateProjectPath);
+				FileReference? CandidateProjectFile = FileReference.FindCorrectCase(new FileReference(CandidateProjectPath));
 
 				// if the path doesn't exist then check native paths (ueprojectdirs)
 				if (!FileReference.Exists(CandidateProjectFile))
