Fixed TSan deploy code for Android

diff --git a/Engine/Source/Programs/UnrealBuildTool/Platform/Android/UEDeployAndroid.cs b/Engine/Source/Programs/UnrealBuildTool/Platform/Android/UEDeployAndroid.cs
--- a/Engine/Source/Programs/UnrealBuildTool/Platform/Android/UEDeployAndroid.cs
+++ b/Engine/Source/Programs/UnrealBuildTool/Platform/Android/UEDeployAndroid.cs
@@ -1526,7 +1526,7 @@ namespace UnrealBuildTool
 					LibName = "ubsan_minimal";
 					break;
 				case AndroidToolChain.ClangSanitizer.Thread:
-					LibName = "thread";
+					LibName = "tsan";
 					break;
 			}
 			
