Make CreateTargetRulesInstance able to return null in the case a rules assembly does not exist
This fixes issues with platform extensions when using -TargetType

diff --git a/Engine/Source/Programs/UnrealBuildTool/System/RulesAssembly.cs b/Engine/Source/Programs/UnrealBuildTool/System/RulesAssembly.cs
--- a/Engine/Source/Programs/UnrealBuildTool/System/RulesAssembly.cs
+++ b/Engine/Source/Programs/UnrealBuildTool/System/RulesAssembly.cs
@@ -536,19 +536,20 @@ namespace UnrealBuildTool
 
 		/// <summary>
 		/// Construct an instance of the given target rules
+		/// Will return null if the requested type name does not exist in the assembly 
 		/// </summary>
 		/// <param name="TypeName">Type name of the target rules</param>
 		/// <param name="TargetInfo">Target configuration information to pass to the constructor</param>
 		/// <param name="Logger">Logger for output</param>
 		/// <param name="IsTestTarget">If building a low level tests target</param>
-		/// <returns>Instance of the corresponding TargetRules</returns>
-		protected TargetRules CreateTargetRulesInstance(string TypeName, TargetInfo TargetInfo, ILogger Logger, bool IsTestTarget = false)
+		/// <returns>Instance of the corresponding TargetRules or null if requested type name does not exist</returns>
+		protected TargetRules? CreateTargetRulesInstance(string TypeName, TargetInfo TargetInfo, ILogger Logger, bool IsTestTarget = false)
 		{
 			// The build module must define a type named '<TargetName>Target' that derives from our 'TargetRules' type.  
 			Type? BaseRulesType = CompiledAssembly?.GetType(TypeName);
 			if (BaseRulesType == null)
 			{
-				throw new BuildException("Expecting to find a type to be declared in a target rules named '{0}'.  This type must derive from the 'TargetRules' type defined by Unreal Build Tool.", TypeName);
+				return null;
 			}
 
 			// Look for platform/group rules that we will use instead of the base rules
@@ -729,7 +730,14 @@ namespace UnrealBuildTool
 			string TargetTypeName = TargetName + "Target";
 
 			// The build module must define a type named '<TargetName>Target' that derives from our 'TargetRules' type.  
-			return CreateTargetRulesInstance(TargetTypeName, new TargetInfo(TargetName, Platform, Configuration, Architecture, ProjectFile, Arguments), Logger, IsTestTarget);
+			TargetRules? TargetRules = CreateTargetRulesInstance(TargetTypeName, new TargetInfo(TargetName, Platform, Configuration, Architecture, ProjectFile, Arguments), Logger, IsTestTarget);
+
+			if (TargetRules == null)
+            {
+				throw new BuildException("Expecting to find a type to be declared in a target rules named '{0}'.  This type must derive from the 'TargetRules' type defined by Unreal Build Tool.", TargetTypeName);
+			}
+
+			return TargetRules;
 		}
 
 		/// <summary>
@@ -748,8 +756,8 @@ namespace UnrealBuildTool
 			List<string> Matches = new List<string>();
 			foreach(KeyValuePair<string, FileReference> TargetPair in TargetNameToTargetFile)
 			{
-				TargetRules Rules = CreateTargetRulesInstance(TargetPair.Key + "Target", new TargetInfo(TargetPair.Key, Platform, Configuration, Architecture, ProjectFile, null), Logger);
-				if(Rules.Type == Type)
+				TargetRules? Rules = CreateTargetRulesInstance(TargetPair.Key + "Target", new TargetInfo(TargetPair.Key, Platform, Configuration, Architecture, ProjectFile, null), Logger);
+				if(Rules != null && Rules.Type == Type)
 				{
 					Matches.Add(TargetPair.Key);
 				}
