UnrealBuildTool: Block msvc 14.33.31629 due to internal compiler error

#rnx

diff --git a/Engine/Source/Programs/UnrealBuildTool/Platform/Windows/MicrosoftPlatformSDK.Versions.cs b/Engine/Source/Programs/UnrealBuildTool/Platform/Windows/MicrosoftPlatformSDK.Versions.cs
--- a/Engine/Source/Programs/UnrealBuildTool/Platform/Windows/MicrosoftPlatformSDK.Versions.cs
+++ b/Engine/Source/Programs/UnrealBuildTool/Platform/Windows/MicrosoftPlatformSDK.Versions.cs
@@ -34,8 +34,8 @@ namespace UnrealBuildTool
 		/// </summary>
 		static readonly VersionNumberRange[] PreferredVisualCppVersions = new VersionNumberRange[]
 		{
-			VersionNumberRange.Parse("14.29.30133", "14.29.30146"), // VS2019 16.11.x
-			VersionNumberRange.Parse("14.32.31326", "14.32.31333"), // VS2022 17.2.x
+			VersionNumberRange.Parse("14.29.30133", "14.29.99999"), // VS2019 16.11.x
+			VersionNumberRange.Parse("14.32.31326", "14.32.99999"), // VS2022 17.2.x
 		};
 
 		/// <summary>
@@ -43,6 +43,7 @@ namespace UnrealBuildTool
 		/// </summary>
 		static readonly VersionNumberRange[] BannedVisualCppVersions = new VersionNumberRange[]
 		{
+			VersionNumberRange.Parse("14.33.31629", "14.33.99999"), // VS2022 17.3.x: https://developercommunity.visualstudio.com/t/clexe-143331629-Access-violation-IN/10122382
 		};
 
 		static readonly VersionNumber MinimumVisualCppVersion = new VersionNumber(14, 29, 30133);
