UnrealBuildTool: Stop passing /Zc:lambda- for VS2022 toolchains 14.33.x or newer

#rnx

diff --git a/Engine/Source/Programs/UnrealBuildTool/Platform/Windows/VCToolChain.cs b/Engine/Source/Programs/UnrealBuildTool/Platform/Windows/VCToolChain.cs
--- a/Engine/Source/Programs/UnrealBuildTool/Platform/Windows/VCToolChain.cs
+++ b/Engine/Source/Programs/UnrealBuildTool/Platform/Windows/VCToolChain.cs
@@ -417,7 +417,7 @@ namespace UnrealBuildTool
 				
 				Arguments.Add("/permissive-");
 				Arguments.Add("/Zc:strictStrings-"); // Have to disable strict const char* semantics due to Windows headers not being compliant.
-				if (EnvVars.CompilerVersion >= new VersionNumber(14, 32))
+				if (EnvVars.CompilerVersion >= new VersionNumber(14, 32) && EnvVars.CompilerVersion < new VersionNumber(14, 33))
 				{
 					Arguments.Add("/Zc:lambda-");
 				}
