UnrealBuildTool: Support clang 14 by suppressing new warning.
Bump minimum windows clang to 13 to match other platforms

#rnx

diff --git a/Engine/Source/Programs/UnrealBuildTool/Platform/Windows/MicrosoftPlatformSDK.Versions.cs b/Engine/Source/Programs/UnrealBuildTool/Platform/Windows/MicrosoftPlatformSDK.Versions.cs
--- a/Engine/Source/Programs/UnrealBuildTool/Platform/Windows/MicrosoftPlatformSDK.Versions.cs
+++ b/Engine/Source/Programs/UnrealBuildTool/Platform/Windows/MicrosoftPlatformSDK.Versions.cs
@@ -21,12 +21,9 @@ namespace UnrealBuildTool
 		static readonly VersionNumberRange[] PreferredClangVersions =
 		{
 			VersionNumberRange.Parse("13.0.0", "13.999"),
-			VersionNumberRange.Parse("12.0.0", "12.999"),
-			VersionNumberRange.Parse("11.0.0", "11.999"),
-			VersionNumberRange.Parse("10.0.0", "10.999")
 		};
 
-		static readonly VersionNumber MinimumClangVersion = new VersionNumber(10, 0, 0);
+		static readonly VersionNumber MinimumClangVersion = new VersionNumber(13, 0, 0);
 
 		/// <summary>
 		/// Ranges of tested compiler toolchains to be used, in order of preference. If multiple toolchains in a range are present, the latest version will be preferred.
@@ -39,7 +36,7 @@ namespace UnrealBuildTool
 		};
 
 		/// <summary>
-		/// Tested compiler toolchanges that should not be allowed.
+		/// Tested compiler toolchains that should not be allowed.
 		/// </summary>
 		static readonly VersionNumberRange[] BannedVisualCppVersions = new VersionNumberRange[]
 		{
diff --git a/Engine/Source/Programs/UnrealBuildTool/Platform/Windows/VCToolChain.cs b/Engine/Source/Programs/UnrealBuildTool/Platform/Windows/VCToolChain.cs
--- a/Engine/Source/Programs/UnrealBuildTool/Platform/Windows/VCToolChain.cs
+++ b/Engine/Source/Programs/UnrealBuildTool/Platform/Windows/VCToolChain.cs
@@ -930,6 +930,10 @@ namespace UnrealBuildTool
 					Arguments.Add("-Wno-unused-but-set-variable");
 					Arguments.Add("-Wno-unused-but-set-parameter");
 				}
+				if (EnvVars.CompilerVersion >= new VersionNumber(14))
+				{
+					Arguments.Add("-Wno-bitwise-instead-of-logical");
+				}
 				Arguments.Add("-Wno-unused-local-typedefs");
 				Arguments.Add("-Wno-unused-function");
 				Arguments.Add("-Wno-unused-private-field");
diff --git a/Engine/Source/Programs/UnrealBuildTool/ToolChain/ClangToolChain.cs b/Engine/Source/Programs/UnrealBuildTool/ToolChain/ClangToolChain.cs
--- a/Engine/Source/Programs/UnrealBuildTool/ToolChain/ClangToolChain.cs
+++ b/Engine/Source/Programs/UnrealBuildTool/ToolChain/ClangToolChain.cs
@@ -502,6 +502,10 @@ namespace UnrealBuildTool
 				Arguments.Add("-Wno-unused-but-set-parameter");          // https://clang.llvm.org/docs/DiagnosticsReference.html#wunused-but-set-parameter				// new warning for clang 13
 				Arguments.Add("-Wno-ordered-compare-function-pointers"); // https://clang.llvm.org/docs/DiagnosticsReference.html#wordered-compare-function-pointers	// new warning for clang 13
 			}
+			if (CompilerVersionGreaterOrEqual(14, 0, 0))
+			{
+				Arguments.Add("-Wno-bitwise-instead-of-logical");       // https://clang.llvm.org/docs/DiagnosticsReference.html#wbitwise-instead-of-logical			// new warning for clang 14
+			}
 
 			Arguments.Add("-Wno-gnu-string-literal-operator-template"); // https://clang.llvm.org/docs/DiagnosticsReference.html#wgnu-string-literal-operator-template	// We use this feature to allow static FNames.
 			Arguments.Add("-Wno-inconsistent-missing-override");        // https://clang.llvm.org/docs/DiagnosticsReference.html#winconsistent-missing-override			// ?? no reason given
