Enhanced UBT to check all Android libs for libc++_shared.so dependency

diff --git a/Engine/Config/BaseEngine.ini b/Engine/Config/BaseEngine.ini
--- a/Engine/Config/BaseEngine.ini
+++ b/Engine/Config/BaseEngine.ini
@@ -2744,6 +2744,7 @@ bAllowPatchOBBFile=False
 bAllowOverflowOBBFiles=False
 bExtractNativeLibs=true
 bEnableDomStorage=False
+bDisableLibCppSharedDependencyValidation=True
 
 [/Script/AndroidPlatformEditor.AndroidSDKSettings]
 SDKAPILevel=latest
diff --git a/Engine/Source/Programs/UnrealBuildTool/Platform/Android/AndroidToolChain.cs b/Engine/Source/Programs/UnrealBuildTool/Platform/Android/AndroidToolChain.cs
--- a/Engine/Source/Programs/UnrealBuildTool/Platform/Android/AndroidToolChain.cs
+++ b/Engine/Source/Programs/UnrealBuildTool/Platform/Android/AndroidToolChain.cs
@@ -186,23 +186,23 @@ namespace UnrealBuildTool
 			string ArchitecturePathLinux = @"prebuilt/linux-x86_64";
 			string ExeExtension = ".exe";
 
-			if (Directory.Exists(Path.Combine(NDKPath, "toolchains", "llvm", ArchitecturePathWindows64)))
+			if (Directory.Exists(Path.Combine(NDKPath, @"toolchains/llvm", ArchitecturePathWindows64)))
 			{
 				Logger.LogDebug("        Found Windows 64 bit versions of toolchain");
 				ArchitecturePath = ArchitecturePathWindows64;
 			}
-			else if (Directory.Exists(Path.Combine(NDKPath, "toolchains", "llvm", ArchitecturePathWindows32)))
+			else if (Directory.Exists(Path.Combine(NDKPath, @"toolchains/llvm", ArchitecturePathWindows32)))
 			{
 				Logger.LogDebug("        Found Windows 32 bit versions of toolchain");
 				ArchitecturePath = ArchitecturePathWindows32;
 			}
-			else if (Directory.Exists(Path.Combine(NDKPath, "toolchains", "llvm", ArchitecturePathMac)))
+			else if (Directory.Exists(Path.Combine(NDKPath, @"toolchains/llvm", ArchitecturePathMac)))
 			{
 				Logger.LogDebug("        Found Mac versions of toolchain");
 				ArchitecturePath = ArchitecturePathMac;
 				ExeExtension = "";
 			}
-			else if (Directory.Exists(Path.Combine(NDKPath, "toolchains", "llvm", ArchitecturePathLinux)))
+			else if (Directory.Exists(Path.Combine(NDKPath, @"toolchains/llvm", ArchitecturePathLinux)))
 			{
 				Logger.LogDebug("        Found Linux versions of toolchain");
 				ArchitecturePath = ArchitecturePathLinux;
@@ -219,10 +219,10 @@ namespace UnrealBuildTool
 			SDK.TryConvertVersionToInt(NDKToolchainVersion, out NDKVersionInt);
 
 			// figure out clang version (will live in toolchains/llvm from NDK 21 forward
-			if (Directory.Exists(Path.Combine(NDKPath, "toolchains", "llvm")))
+			if (Directory.Exists(Path.Combine(NDKPath, @"toolchains/llvm")))
 			{
 				// look for version in AndroidVersion.txt (fail if not found)
-				string VersionFilename = Path.Combine(NDKPath, "toolchains", "llvm", ArchitecturePath, "AndroidVersion.txt");
+				string VersionFilename = Path.Combine(NDKPath, @"toolchains/llvm", ArchitecturePath, "AndroidVersion.txt");
 				if (!File.Exists(VersionFilename))
 				{
 					throw new BuildException("Cannot find supported Android toolchain");
@@ -271,8 +271,8 @@ namespace UnrealBuildTool
 			// NDK setup (enforce minimum API level)
 			int NDKApiLevel64Int = GetNdkApiLevelInt(MinimumNDKAPILevel);
 
-			string GCCToolchainPath = Path.Combine(NDKPath, "toolchains", "llvm", ArchitecturePath);
-			string SysrootPath = Path.Combine(NDKPath, "toolchains", "llvm", ArchitecturePath, "sysroot");
+			string GCCToolchainPath = Path.Combine(NDKPath, @"toolchains/llvm", ArchitecturePath);
+			string SysrootPath = Path.Combine(NDKPath, @"toolchains/llvm", ArchitecturePath, "sysroot");
 
 			// toolchain params (note: use ANDROID=1 same as we define it)
 			ToolchainLinkParamsArm64 = " --target=aarch64-none-linux-android" + NDKApiLevel64Int + " --gcc-toolchain=\"" + GCCToolchainPath + "\" --sysroot=\"" + SysrootPath + "\" -DANDROID=1";
@@ -287,6 +287,8 @@ namespace UnrealBuildTool
 				ToolchainParamsArm64 += " -D__ANDROID_API__=" + NDKApiLevel64Int;
 				ToolchainParamsx64 += " -D__ANDROID_API__=" + NDKApiLevel64Int;
 			}
+
+			ReadElfPath = Path.Combine(NDKPath, @"toolchains/llvm", ArchitecturePath, @"bin/llvm-readelf" + ExeExtension);
 		}
 
 		protected override ClangToolChainInfo GetToolChainInfo()
@@ -382,6 +384,13 @@ namespace UnrealBuildTool
 			return Ini.GetBool("/Script/AndroidRuntimeSettings.AndroidRuntimeSettings", "bDisableStackProtector", out bDisableStackProtector) && bDisableStackProtector;
 		}
 
+		private bool DisableLibCppSharedDependencyValidation()
+		{
+			ConfigHierarchy Ini = ConfigCache.ReadHierarchy(ConfigHierarchyType.Engine, DirectoryReference.FromFile(ProjectFile), UnrealTargetPlatform.Android);
+			bool bDisableLibCppSharedDependencyValidation = false;
+			return Ini.GetBool("/Script/AndroidRuntimeSettings.AndroidRuntimeSettings", "bDisableLibCppSharedDependencyValidation", out bDisableLibCppSharedDependencyValidation) && bDisableLibCppSharedDependencyValidation;
+		}
+
 		private string GetVersionScriptFilename(LinkEnvironment LinkEnvironment)
 		{
 			return Path.Combine(LinkEnvironment.IntermediateDirectory!.FullName, "ExportSymbols.ldscript");
@@ -1123,6 +1132,7 @@ namespace UnrealBuildTool
 		protected static string? ToolchainLinkParamsx64;
 		protected static string? ArPathArm64;
 		protected static string? ArPathx64;
+		protected static string? ReadElfPath;
 
 		static public string GetStripExecutablePath(string UnrealArch)
 		{
@@ -1965,6 +1975,26 @@ namespace UnrealBuildTool
 							{
 								AdditionalLibraries.Add(AbsoluteLibraryPath);
 							}
+
+							if (!DisableLibCppSharedDependencyValidation() && ReadElfPath != null)
+							{
+								string? Output = Utils.RunLocalProcessAndReturnStdOut(ReadElfPath, "--dynamic \"" + Library.FullName + "\"");
+								if (Output != null)
+								{
+									if (Output.Contains("libc++_shared.so"))
+									{
+										if (IsNewNDKModel())
+										{
+											throw new BuildException("Lib {0} depends on libc++_shared.so. There are known incompatibility issues when linking libc++_shared.so with Unreal Engine built with NDK22+." +
+												" Please rebuild your dependencies with static libc++!", Lib);
+										}
+										else
+										{
+											Log.TraceWarning("Lib {0} depends on libc++_shared.so. Unreal Engine is designed to be linked with libs that are built against static libc++ only. Please rebuild your dependencies with static libc++!", Lib);
+										}
+									}
+								}
+							}
 						}
 					}
 
diff --git a/Engine/Source/Runtime/Android/AndroidRuntimeSettings/Classes/AndroidRuntimeSettings.h b/Engine/Source/Runtime/Android/AndroidRuntimeSettings/Classes/AndroidRuntimeSettings.h
--- a/Engine/Source/Runtime/Android/AndroidRuntimeSettings/Classes/AndroidRuntimeSettings.h
+++ b/Engine/Source/Runtime/Android/AndroidRuntimeSettings/Classes/AndroidRuntimeSettings.h
@@ -395,6 +395,10 @@ public:
 	UPROPERTY(GlobalConfig, EditAnywhere, Category = AdvancedBuild, meta = (DisplayName = "Build with hidden symbol visibility in shipping config. [Experimental]"))
 	bool bBuildWithHiddenSymbolVisibility;
 
+	// Disable libc++_shared dependency validation in all .so files linked with libUnreal.so
+	UPROPERTY(GlobalConfig, EditAnywhere, Category = AdvancedBuild, meta = (DisplayName = "Disable libc++_shared dependency validation in all dependencies"))
+	bool bDisableLibCppSharedDependencyValidation;
+
 	// Always save .so file with symbols allowing use of addr2line on raw callstack addresses.
 	UPROPERTY(GlobalConfig, EditAnywhere, Category = AdvancedBuild, meta = (DisplayName = "Always save a copy of the libUnreal.so with symbols. [Experimental]"))
 	bool bSaveSymbols;
