Undo single changelist 21418118

diff --git a/Engine/Source/Programs/UnrealBuildTool/Platform/Windows/VCToolChain.cs b/Engine/Source/Programs/UnrealBuildTool/Platform/Windows/VCToolChain.cs
--- a/Engine/Source/Programs/UnrealBuildTool/Platform/Windows/VCToolChain.cs
+++ b/Engine/Source/Programs/UnrealBuildTool/Platform/Windows/VCToolChain.cs
@@ -417,7 +417,7 @@ namespace UnrealBuildTool
 				
 				Arguments.Add("/permissive-");
 				Arguments.Add("/Zc:strictStrings-"); // Have to disable strict const char* semantics due to Windows headers not being compliant.
-				if (EnvVars.CompilerVersion >= new VersionNumber(14, 32) && EnvVars.CompilerVersion < new VersionNumber(14, 33))
+				if (EnvVars.CompilerVersion >= new VersionNumber(14, 32))
 				{
 					Arguments.Add("/Zc:lambda-");
 				}
