Added a mode that prints out useful build graph info.

diff --git a/Engine/Source/Programs/UnrealBuildTool/Configuration/UEBuildModuleCPP.cs b/Engine/Source/Programs/UnrealBuildTool/Configuration/UEBuildModuleCPP.cs
--- a/Engine/Source/Programs/UnrealBuildTool/Configuration/UEBuildModuleCPP.cs
+++ b/Engine/Source/Programs/UnrealBuildTool/Configuration/UEBuildModuleCPP.cs
@@ -817,6 +817,8 @@ namespace UnrealBuildTool
 				Instance = new PrecompiledHeaderInstance(WrapperFile, CompileEnvironment, Output);
 				Template.Instances.Add(Instance);
 			}
+
+			Instance.Modules.Add(this);
 			return Instance;
 		}
 
diff --git a/Engine/Source/Programs/UnrealBuildTool/Modes/PrintBuildGraphInfo.cs b/Engine/Source/Programs/UnrealBuildTool/Modes/PrintBuildGraphInfo.cs
new file mode 100644
--- /dev/null
+++ b/Engine/Source/Programs/UnrealBuildTool/Modes/PrintBuildGraphInfo.cs
@@ -0,0 +1,167 @@
+// Copyright Epic Games, Inc. All Rights Reserved.
+
+using System;
+using System.Collections.Generic;
+using System.Linq;
+using System.Text;
+using System.Text.RegularExpressions;
+using EpicGames.Core;
+using UnrealBuildBase;
+using Microsoft.Extensions.Logging;
+
+namespace UnrealBuildTool
+{
+	/// <summary>
+	/// Generate a clang compile_commands file for a target
+	/// </summary>
+	[ToolMode("PrintBuildGraphInfo", ToolModeOptions.XmlConfig | ToolModeOptions.BuildPlatforms | ToolModeOptions.SingleInstance | ToolModeOptions.StartPrefetchingEngine | ToolModeOptions.ShowExecutionTime)]
+	class PrintBuildGraphInfo : ToolMode
+	{
+		/// <summary>
+		/// Execute the command
+		/// </summary>
+		/// <param name="Arguments">Command line arguments</param>
+		/// <returns>Exit code</returns>
+		/// <param name="Logger"></param>
+		public override int Execute(CommandLineArguments Arguments, ILogger Logger)
+		{
+			Arguments.ApplyTo(this);
+
+			// Create the build configuration object, and read the settings
+			BuildConfiguration BuildConfiguration = new BuildConfiguration();
+			XmlConfig.ApplyTo(BuildConfiguration);
+			Arguments.ApplyTo(BuildConfiguration);
+
+			// Parse all the target descriptors
+			List<TargetDescriptor> TargetDescriptors = TargetDescriptor.ParseCommandLine(Arguments, BuildConfiguration.bUsePrecompiled, BuildConfiguration.bSkipRulesCompile, BuildConfiguration.bForceRulesCompile, Logger);
+
+			using (ISourceFileWorkingSet WorkingSet = new EmptySourceFileWorkingSet())
+			{
+				// Find the compile commands for each file in the target
+				Dictionary<FileReference, string> FileToCommand = new();
+				foreach (TargetDescriptor TargetDescriptor in TargetDescriptors)
+				{
+					Logger.LogInformation("------------------------");
+					Logger.LogInformation("Target: {TargetName}", TargetDescriptor.Name);
+					Logger.LogInformation("------------------------");
+
+					// Create a makefile for the target
+					UEBuildTarget Target = UEBuildTarget.Create(TargetDescriptor, BuildConfiguration.bSkipRulesCompile, BuildConfiguration.bForceRulesCompile, BuildConfiguration.bUsePrecompiled, Logger);
+
+					List<UEBuildModule> NoUnityModules = new();
+					List<UEBuildModule> NotOptimizedModules = new();
+					List<UEBuildModule> NoPCHModules = new();
+					List<UEBuildModule> PrivatePCHModules = new();
+
+					// Figure out all the modules referenced by this target. This includes all the modules that are referenced, not just the ones compiled into binaries.
+					CppCompileEnvironment GlobalCompileEnvironment = Target.CreateCompileEnvironmentForProjectFiles(Logger);
+					HashSet<UEBuildModule> Modules = new HashSet<UEBuildModule>();
+					foreach (UEBuildBinary Binary in Target.Binaries)
+					{
+						CppCompileEnvironment BinaryCompileEnvironment = Binary.CreateBinaryCompileEnvironment(GlobalCompileEnvironment);
+						foreach (UEBuildModule Module in Binary.Modules)
+						{
+							if (Module is UEBuildModuleCPP ModuleCPP)
+							{
+								ModuleCPP.CreateCompileEnvironmentForIntellisense(Target.Rules, BinaryCompileEnvironment, Logger);
+							}
+							
+							Modules.Add(Module);
+						}
+					}
+
+					// Gather module info
+					foreach (UEBuildModule Module in Modules)
+					{
+						if (Module.Rules.PCHUsage == ModuleRules.PCHUsageMode.NoPCHs)
+						{
+							NoPCHModules.Add(Module);
+						}
+						else if (Module.Rules.PrivatePCHHeaderFile != null && (Module.Rules.PCHUsage == ModuleRules.PCHUsageMode.NoSharedPCHs || Module.Rules.PCHUsage == ModuleRules.PCHUsageMode.UseExplicitOrSharedPCHs))
+						{
+							PrivatePCHModules.Add(Module);
+						}
+
+						if (!Module.Rules.bUseUnity)
+						{
+							NoUnityModules.Add(Module);
+						}
+
+						if (Module.Rules.OptimizeCode == ModuleRules.CodeOptimization.Never)
+						{
+							NotOptimizedModules.Add(Module);
+						}
+					}
+
+					if (NoUnityModules.Any())
+					{
+						Logger.LogInformation("Modules not using unity files:");
+						NoUnityModules.SortBy(Module => Module.Name);
+						foreach (UEBuildModule Module in NoUnityModules)
+						{
+							Logger.LogInformation(" {ModuleName}", Module.Name);
+						}
+					}
+
+					if (NotOptimizedModules.Any())
+					{
+						Logger.LogInformation("Modules not optimized:");
+						NotOptimizedModules.SortBy(Module => Module.Name);
+						foreach (UEBuildModule Module in NotOptimizedModules)
+						{
+							Logger.LogInformation(" {ModuleName}", Module.Name);
+						}
+					}
+
+					Logger.LogInformation("PCH Usage:");
+
+					if (NoPCHModules.Any())
+					{
+						Logger.LogInformation(" Modules not using PCHs:");
+						NoPCHModules.SortBy(Module => Module.Name);
+						foreach (UEBuildModule Module in NoPCHModules)
+						{
+							Logger.LogInformation("  {ModuleName}", Module.Name);
+						}
+					}
+
+					if (PrivatePCHModules.Any())
+					{
+						Logger.LogInformation(" Modules using private PCHs:");
+						PrivatePCHModules.SortBy(Module => Module.Name);
+						foreach (UEBuildModule Module in PrivatePCHModules)
+						{
+							Logger.LogInformation("  {ModuleName}", Module.Name);
+						}
+					}
+
+					if (GlobalCompileEnvironment.SharedPCHs.Any())
+					{
+						Logger.LogInformation(" Shared PCHs:");
+						List<PrecompiledHeaderTemplate> SortedSharedPCHs = new(GlobalCompileEnvironment.SharedPCHs);
+						SortedSharedPCHs.SortBy(SharedPCH => SharedPCH.Module.Name);
+						foreach (PrecompiledHeaderTemplate Template in SortedSharedPCHs)
+						{
+							Logger.LogInformation("  Module {ModuleName} - {TimesUsed} Permutations:", Template.Module.Name, Template.Instances.Count);
+							List<PrecompiledHeaderInstance> SortedInstances = new(Template.Instances);
+							SortedInstances.SortBy(Instance => Instance.HeaderFile.Name);
+							foreach (PrecompiledHeaderInstance Instance in SortedInstances)
+							{
+								Logger.LogInformation("   {InstanceName} - Used by {TimesUsed} modules:", Instance.CompileEnvironment.PrecompiledHeaderIncludeFilename, Instance.Modules.Count);
+
+								List<UEBuildModuleCPP> SortedModules = new(Instance.Modules);
+								SortedModules.SortBy(Module => Module.Name);
+								foreach (UEBuildModuleCPP Module in SortedModules)
+								{
+									Logger.LogInformation("    {ModuleName}", Module.Name);
+								}
+							}
+						}
+					}
+				}
+			}
+
+			return 0;
+		}
+	}
+}
diff --git a/Engine/Source/Programs/UnrealBuildTool/System/PrecompiledHeaderInstance.cs b/Engine/Source/Programs/UnrealBuildTool/System/PrecompiledHeaderInstance.cs
--- a/Engine/Source/Programs/UnrealBuildTool/System/PrecompiledHeaderInstance.cs
+++ b/Engine/Source/Programs/UnrealBuildTool/System/PrecompiledHeaderInstance.cs
@@ -29,6 +29,11 @@ namespace UnrealBuildTool
 		/// </summary>
 		public CPPOutput Output;
 
+		/// <summary>
+		/// List of modules using this instance
+		/// </summary>
+		public List<UEBuildModuleCPP> Modules = new List<UEBuildModuleCPP>();
+
 		/// <summary>
 		/// Constructor
 		/// </summary>
