Cleaned up dublicated compiler flag in Android Toolchain

diff --git a/Engine/Source/Programs/UnrealBuildTool/Platform/Android/AndroidToolChain.cs b/Engine/Source/Programs/UnrealBuildTool/Platform/Android/AndroidToolChain.cs
--- a/Engine/Source/Programs/UnrealBuildTool/Platform/Android/AndroidToolChain.cs
+++ b/Engine/Source/Programs/UnrealBuildTool/Platform/Android/AndroidToolChain.cs
@@ -604,12 +604,6 @@ namespace UnrealBuildTool
 			Result += " -Wno-invalid-offsetof";         // needed to suppress warnings about using offsetof on non-POD types.
 			Result += " -Wno-logical-op-parentheses";   // needed for external headers we can't change
 
-			// silence warning introduced with NDK 24+
-			if (NDKVersionInt >= 240000)
-			{
-				Result += " -Wno-unused-but-set-variable";
-			}
-
 			if (BuildWithHiddenSymbolVisibility(CompileEnvironment))
 			{
 				// Result += " -fvisibility-ms-compat -fvisibility-inlines-hidden"; // This hides all symbols by default but exports all type info (vtable/rtti) for a non-monolithic setup
