Removed legacy Android Project Settings

diff --git a/Engine/Config/BaseEngine.ini b/Engine/Config/BaseEngine.ini
--- a/Engine/Config/BaseEngine.ini
+++ b/Engine/Config/BaseEngine.ini
@@ -2712,7 +2712,6 @@ StoreVersionOffsetX8664=0
 VersionDisplayName=1.0
 MinSDKVersion=26
 TargetSDKVersion=28
-bEnableGradle=true
 bEnableLint=false
 bShowLaunchImage=true
 bValidateTextureFormats=true
@@ -2724,7 +2723,7 @@ TextureFormatPriority_DXT=0.6
 TextureFormatPriority_ASTC=0.9
 bEnableNewKeyboard=True
 bAndroidVoiceEnabled=false
-bBuildWithHiddenSymbolVisibility=false
+bBuildWithHiddenSymbolVisibility=True
 bSaveSymbols=false
 bAllowControllers=True
 bAllowIMU=False
@@ -2737,7 +2736,6 @@ bBundleABISplit=True
 bBundleLanguageSplit=True
 bBundleDensitySplit=True
 bFullScreen=True
-bForceLDLinker=False
 bStripShaderReflection=True
 bForceSmallOBBFiles=False
 bAllowLargeOBBFiles=False
@@ -2745,7 +2743,11 @@ bAllowPatchOBBFile=False
 bAllowOverflowOBBFiles=False
 bExtractNativeLibs=true
 bEnableDomStorage=False
+bSkipLibCpp=True
 bDisableLibCppSharedDependencyValidation=True
+bDisableFunctionDataSections=False
+bEnableAdvancedBinaryCompression=False
+bDisableStackProtector=False
 
 [/Script/AndroidPlatformEditor.AndroidSDKSettings]
 SDKAPILevel=latest
diff --git a/Engine/Source/Developer/Android/AndroidPlatformEditor/Private/AndroidLicenseDialog.cpp b/Engine/Source/Developer/Android/AndroidPlatformEditor/Private/AndroidLicenseDialog.cpp
--- a/Engine/Source/Developer/Android/AndroidPlatformEditor/Private/AndroidLicenseDialog.cpp
+++ b/Engine/Source/Developer/Android/AndroidPlatformEditor/Private/AndroidLicenseDialog.cpp
@@ -267,9 +267,6 @@ FReply SAndroidLicenseDialog::OnAgree()
 
 FReply SAndroidLicenseDialog::OnCancel()
 {
-	// turn off Gradle checkbox
-	//GetMutableDefault<UAndroidRuntimeSettings>()->bEnableGradle = false;
-
 	TSharedRef<SWindow> ParentWindow = FSlateApplication::Get().FindWidgetWindow(AsShared()).ToSharedRef();
 	FSlateApplication::Get().RequestDestroyWindow(ParentWindow);
 	return FReply::Handled();
diff --git a/Engine/Source/Developer/Android/AndroidPlatformEditor/Private/AndroidTargetSettingsCustomization.cpp b/Engine/Source/Developer/Android/AndroidPlatformEditor/Private/AndroidTargetSettingsCustomization.cpp
--- a/Engine/Source/Developer/Android/AndroidPlatformEditor/Private/AndroidTargetSettingsCustomization.cpp
+++ b/Engine/Source/Developer/Android/AndroidPlatformEditor/Private/AndroidTargetSettingsCustomization.cpp
@@ -298,12 +298,6 @@ void FAndroidTargetSettingsCustomization::BuildAppManifestSection(IDetailLayoutB
 
 	// @todo android fat binary: Put back in when we expose those
 //	SETUP_SOURCEONLY_PROP(bSplitIntoSeparateApks, BuildCategory, LOCTEXT("SplitIntoSeparateAPKsToolTip", "If checked, CPU architectures and rendering types will be split into separate .apk files"));
-
-	// check for Gradle change
-	TSharedRef<IPropertyHandle> EnableGradleProperty = DetailLayout.GetProperty(GET_MEMBER_NAME_CHECKED(UAndroidRuntimeSettings, bEnableGradle));
-	EnableGradleProperty->MarkHiddenByCustomization();
-	//FSimpleDelegate EnableGradleChange = FSimpleDelegate::CreateSP(this, &FAndroidTargetSettingsCustomization::OnEnableGradleChange);
-	//EnableGradleProperty->SetOnPropertyValueChanged(EnableGradleChange);
 }
 
 bool FAndroidTargetSettingsCustomization::IsLicenseInvalid() const
@@ -368,49 +362,6 @@ FReply FAndroidTargetSettingsCustomization::OnAcceptSDKLicenseClicked()
 	return FReply::Handled();
 }
 
-void FAndroidTargetSettingsCustomization::OnEnableGradleChange()
-{
-	// only need to do this if enabling
-	if (!GetDefault<UAndroidRuntimeSettings>()->bEnableGradle)
-	{
-		return;
-	}
-
-	// only show if don't have a valid license
-	TSharedPtr<SAndroidLicenseDialog> LicenseDialog = SNew(SAndroidLicenseDialog);
-	if (!LicenseDialog->HasLicense())
-	{
-		FSimpleDelegate LicenseAcceptedCallback = FSimpleDelegate::CreateSP(this, &FAndroidTargetSettingsCustomization::OnLicenseAccepted);
-		LicenseDialog->SetLicenseAcceptedCallback(LicenseAcceptedCallback);
-
-		const FText AndroidLicenseWindowTitle = LOCTEXT("AndroidLicenseUnrealEditor", "Android SDK License");
-
-		TSharedPtr<SWindow> AndroidLicenseWindow =
-			SNew(SWindow)
-			.Title(AndroidLicenseWindowTitle)
-			.ClientSize(FVector2D(600.f, 700.f))
-			.HasCloseButton(false)
-			.SupportsMaximize(false)
-			.SupportsMinimize(false)
-			.SizingRule(ESizingRule::FixedSize)
-			[
-				LicenseDialog.ToSharedRef()
-			];
-
-		IMainFrameModule& MainFrame = FModuleManager::LoadModuleChecked<IMainFrameModule>("MainFrame");
-		TSharedPtr<SWindow> ParentWindow = MainFrame.GetParentWindow();
-
-		if (ParentWindow.IsValid())
-		{
-			FSlateApplication::Get().AddModalWindow(AndroidLicenseWindow.ToSharedRef(), ParentWindow.ToSharedRef());
-		}
-		else
-		{
-			FSlateApplication::Get().AddWindow(AndroidLicenseWindow.ToSharedRef());
-		}
-	}
-}
-
 void FAndroidTargetSettingsCustomization::BuildIconSection(IDetailLayoutBuilder& DetailLayout)
 {
 	// Icon category
diff --git a/Engine/Source/Developer/Android/AndroidPlatformEditor/Private/AndroidTargetSettingsCustomization.h b/Engine/Source/Developer/Android/AndroidPlatformEditor/Private/AndroidTargetSettingsCustomization.h
--- a/Engine/Source/Developer/Android/AndroidPlatformEditor/Private/AndroidTargetSettingsCustomization.h
+++ b/Engine/Source/Developer/Android/AndroidPlatformEditor/Private/AndroidTargetSettingsCustomization.h
@@ -44,9 +44,6 @@ private:
 	// Called when the app id is modified
 	void OnAppIDModified();
 
-	// Called when EnableGradle is modified
-	void OnEnableGradleChange();
-
 	// Called when License accepted
 	void OnLicenseAccepted();
 
diff --git a/Engine/Source/Developer/Android/AndroidTargetPlatform/Private/AndroidTargetPlatform.cpp b/Engine/Source/Developer/Android/AndroidTargetPlatform/Private/AndroidTargetPlatform.cpp
--- a/Engine/Source/Developer/Android/AndroidTargetPlatform/Private/AndroidTargetPlatform.cpp
+++ b/Engine/Source/Developer/Android/AndroidTargetPlatform/Private/AndroidTargetPlatform.cpp
@@ -371,18 +371,12 @@ int32 FAndroidTargetPlatform::CheckRequirements(bool bProjectHasCode, EBuildConf
 		bReadyToBuild |= ETargetPlatformReadyStatus::SDKNotFound;
 	}
 
-	bool bEnableGradle;
-	GConfig->GetBool(TEXT("/Script/AndroidRuntimeSettings.AndroidRuntimeSettings"), TEXT("bEnableGradle"), bEnableGradle, GEngineIni);
-
-	if (bEnableGradle)
+	// need to check license was accepted
+	if (!HasLicense())
 	{
-		// need to check license was accepted
-		if (!HasLicense())
-		{
-			OutTutorialPath.Empty();
-			CustomizedLogMessage = LOCTEXT("AndroidLicenseNotAcceptedMessageDetail", "SDK License must be accepted in the Android project settings to deploy your app to the device.");
-			bReadyToBuild |= ETargetPlatformReadyStatus::LicenseNotAccepted;
-		}
+		OutTutorialPath.Empty();
+		CustomizedLogMessage = LOCTEXT("AndroidLicenseNotAcceptedMessageDetail", "SDK License must be accepted in the Android project settings to deploy your app to the device.");
+		bReadyToBuild |= ETargetPlatformReadyStatus::LicenseNotAccepted;
 	}
 
 	return bReadyToBuild;
diff --git a/Engine/Source/Programs/AutomationTool/Android/AndroidPlatform.Automation.cs b/Engine/Source/Programs/AutomationTool/Android/AndroidPlatform.Automation.cs
--- a/Engine/Source/Programs/AutomationTool/Android/AndroidPlatform.Automation.cs
+++ b/Engine/Source/Programs/AutomationTool/Android/AndroidPlatform.Automation.cs
@@ -3413,11 +3413,7 @@ public class AndroidPlatform : Platform
 
 			int StoreVersionOffset = 0;
 			ConfigHierarchy Ini = ConfigCache.ReadHierarchy(ConfigHierarchyType.Engine, DirectoryReference.FromFile(SC.RawProjectPath), SC.StageTargetPlatform.PlatformType);
-			if (ApkName.Contains("-armv7-"))
-			{
-				Ini.GetInt32("/Script/AndroidRuntimeSettings.AndroidRuntimeSettings", "StoreVersionOffsetArmV7", out StoreVersionOffset);
-			}
-			else if (ApkName.Contains("-arm64-"))
+			if (ApkName.Contains("-arm64-"))
 			{
 				Ini.GetInt32("/Script/AndroidRuntimeSettings.AndroidRuntimeSettings", "StoreVersionOffsetArm64", out StoreVersionOffset);
 			}
diff --git a/Engine/Source/Programs/UnrealBuildTool/Platform/Android/AndroidExports.cs b/Engine/Source/Programs/UnrealBuildTool/Platform/Android/AndroidExports.cs
--- a/Engine/Source/Programs/UnrealBuildTool/Platform/Android/AndroidExports.cs
+++ b/Engine/Source/Programs/UnrealBuildTool/Platform/Android/AndroidExports.cs
@@ -90,7 +90,7 @@ namespace UnrealBuildTool
 		/// <returns></returns>
 		public static IAndroidToolChain CreateToolChain(FileReference ProjectFile)
 		{
-			return new AndroidToolChain(ProjectFile, false, null, null, Log.Logger);
+			return new AndroidToolChain(ProjectFile, null, null, Log.Logger);
 		}
 
 		/// <summary>
@@ -99,7 +99,7 @@ namespace UnrealBuildTool
 		/// <returns></returns>
 		public static IAndroidToolChain CreateTempToolChain()
 		{
-			return new AndroidToolChain(null, false, null, null, Log.Logger);
+			return new AndroidToolChain(null, null, null, Log.Logger);
 		}
 
 		/// <summary>
@@ -140,7 +140,7 @@ namespace UnrealBuildTool
 		/// <param name="Logger">Logger for output</param>
 		public static void StripSymbols(FileReference SourceFile, FileReference TargetFile, ILogger Logger)
 		{
-			AndroidToolChain ToolChain = new AndroidToolChain(null, false, null, null, Logger);
+			AndroidToolChain ToolChain = new AndroidToolChain(null, null, null, Logger);
 			ToolChain.StripSymbols(SourceFile, TargetFile, Logger);
 		}
 	}
diff --git a/Engine/Source/Programs/UnrealBuildTool/Platform/Android/AndroidToolChain.cs b/Engine/Source/Programs/UnrealBuildTool/Platform/Android/AndroidToolChain.cs
--- a/Engine/Source/Programs/UnrealBuildTool/Platform/Android/AndroidToolChain.cs
+++ b/Engine/Source/Programs/UnrealBuildTool/Platform/Android/AndroidToolChain.cs
@@ -56,7 +56,6 @@ namespace UnrealBuildTool
 		}
 
 		protected FileReference? ProjectFile;
-		private bool bUseLdGold;
 		private List<string> AdditionalArches;
 		private List<string> AdditionalGPUArches;
 		protected bool bExecuteCompilerThroughShell;
@@ -125,24 +124,21 @@ namespace UnrealBuildTool
 		[CommandLine("-Architectures=", ListSeparator = '+')]
 		public List<string> ArchitectureArg = new List<string>();
 
-		protected bool bEnableGcSections = true;
-
-		public AndroidToolChain(FileReference? InProjectFile, bool bInUseLdGold, IReadOnlyList<string>? InAdditionalArches, IReadOnlyList<string>? InAdditionalGPUArches, ILogger InLogger)
-			: this(InProjectFile, bInUseLdGold, InAdditionalArches, InAdditionalGPUArches, false, ClangToolChainOptions.None, InLogger)
+		public AndroidToolChain(FileReference? InProjectFile, IReadOnlyList<string>? InAdditionalArches, IReadOnlyList<string>? InAdditionalGPUArches, ILogger InLogger)
+			: this(InProjectFile, InAdditionalArches, InAdditionalGPUArches, false, ClangToolChainOptions.None, InLogger)
 		{
 		}
 
-		public AndroidToolChain(FileReference? InProjectFile, bool bInUseLdGold, IReadOnlyList<string>? InAdditionalArches, IReadOnlyList<string>? InAdditionalGPUArches, ClangToolChainOptions ToolchainOptions, ILogger InLogger)
-			: this(InProjectFile, bInUseLdGold, InAdditionalArches, InAdditionalGPUArches, false, ToolchainOptions, InLogger)
+		public AndroidToolChain(FileReference? InProjectFile, IReadOnlyList<string>? InAdditionalArches, IReadOnlyList<string>? InAdditionalGPUArches, ClangToolChainOptions ToolchainOptions, ILogger InLogger)
+			: this(InProjectFile, InAdditionalArches, InAdditionalGPUArches, false, ToolchainOptions, InLogger)
 		{
 		}
 
-		protected AndroidToolChain(FileReference? InProjectFile, bool bInUseLdGold, IReadOnlyList<string>? InAdditionalArches, IReadOnlyList<string>? InAdditionalGPUArches, bool bAllowMissingNDK, ClangToolChainOptions ToolchainOptions, ILogger InLogger)
+		protected AndroidToolChain(FileReference? InProjectFile, IReadOnlyList<string>? InAdditionalArches, IReadOnlyList<string>? InAdditionalGPUArches, bool bAllowMissingNDK, ClangToolChainOptions ToolchainOptions, ILogger InLogger)
 			: base(ToolchainOptions, InLogger)
 		{
 			Options = ToolchainOptions;
 			ProjectFile = InProjectFile;
-			bUseLdGold = bInUseLdGold;
 			AdditionalArches = new List<string>();
 			AdditionalGPUArches = new List<string>();
 
@@ -255,18 +251,9 @@ namespace UnrealBuildTool
 				AndroidClangBuild = "unknown";
 			}
 
-			if (ForceLDLinker())
-			{
-				// use ld before r21
-				ArPathArm64 = Utils.CollapseRelativeDirectories(Path.Combine(NDKPath, @"toolchains/aarch64-linux-android-4.9", ArchitecturePath, @"bin/aarch64-linux-android-ar" + ExeExtension));
-				ArPathx64 = Utils.CollapseRelativeDirectories(Path.Combine(NDKPath, @"toolchains/x86_64-4.9", ArchitecturePath, @"bin/x86_64-linux-android-ar" + ExeExtension));
-			}
-			else
-			{
-				// use lld for r21+
-				ArPathArm64 = Utils.CollapseRelativeDirectories(Path.Combine(NDKPath, @"toolchains/llvm", ArchitecturePath, @"bin/llvm-ar" + ExeExtension));
-				ArPathx64 = ArPathArm64;
-			}
+			// use lld for r21+
+			ArPathArm64 = Utils.CollapseRelativeDirectories(Path.Combine(NDKPath, @"toolchains/llvm", ArchitecturePath, @"bin/llvm-ar" + ExeExtension));
+			ArPathx64 = ArPathArm64;
 
 			// NDK setup (enforce minimum API level)
 			int NDKApiLevel64Int = GetNdkApiLevelInt(MinimumNDKAPILevel);
@@ -356,13 +343,6 @@ namespace UnrealBuildTool
 			return CompileEnvironment.Configuration == CppConfiguration.Shipping && (Ini.GetBool("/Script/AndroidRuntimeSettings.AndroidRuntimeSettings", "bBuildWithHiddenSymbolVisibility", out bBuild) && bBuild);
 		}
 
-		private bool ForceLDLinker()
-		{
-			ConfigHierarchy Ini = ConfigCache.ReadHierarchy(ConfigHierarchyType.Engine, DirectoryReference.FromFile(ProjectFile), UnrealTargetPlatform.Android);
-			bool bForceLDLinker = false;
-			return Ini.GetBool("/Script/AndroidRuntimeSettings.AndroidRuntimeSettings", "bForceLDLinker", out bForceLDLinker) && bForceLDLinker;
-		}
-
 		private bool DisableFunctionDataSections()
 		{
 			ConfigHierarchy Ini = ConfigCache.ReadHierarchy(ConfigHierarchyType.Engine, DirectoryReference.FromFile(ProjectFile), UnrealTargetPlatform.Android);
@@ -823,7 +803,7 @@ namespace UnrealBuildTool
 			Result += " -shared";
 			Result += " -Wl,-Bsymbolic";
 			Result += " -Wl,--no-undefined";
-			if (bEnableGcSections && !DisableFunctionDataSections())
+			if (!DisableFunctionDataSections())
 			{
 				Result += " -Wl,--gc-sections"; // Enable garbage collection of unused input sections. works best with -ffunction-sections, -fdata-sections
 			}
@@ -833,8 +813,6 @@ namespace UnrealBuildTool
 				Result += " -Wl,--strip-debug";
 			}
 
-			bool bUseLLD = !ForceLDLinker();
-			bool bAllowLdGold = true;
 			if (Architecture == "-x64")
 			{
 				Result += ToolchainLinkParamsx64;
@@ -844,40 +822,19 @@ namespace UnrealBuildTool
 			{
 				Result += ToolchainLinkParamsArm64;
 				Result += " -march=armv8-a";
-				bAllowLdGold = false;       // NDK issue 70838247
 			}
 
 			if (LinkEnvironment.Configuration == CppConfiguration.Shipping)
 			{
 				Result += " -Wl,--icf=all"; // Enables ICF (Identical Code Folding). [all, safe] safe == fold functions that can be proven not to have their address taken.
-				if (!bUseLLD)
-				{
-					Result += " -Wl,--icf-iterations=3";
-				}
-
 				Result += " -Wl,-O3";
 			}
 
-			if (bUseLLD)
-			{
-				Result += " -Wl,-no-pie";
+			Result += " -Wl,-no-pie";
+
+			// use lld as linker (requires llvm-strip)
+			Result += " -fuse-ld=lld";
 
-				// use lld as linker (requires llvm-strip)
-				Result += " -fuse-ld=lld";
-			}
-			else
-			{
-				if (bAllowLdGold && bUseLdGold)
-				{
-					// use ld.gold as linker (requires strip)
-					Result += " -fuse-ld=gold";
-				}
-				else
-				{
-					// use ld as linker (requires strip)
-					Result += " -fuse-ld=ld";
-				}
-			}
 
 			// make sure the DT_SONAME field is set properly (or we can a warning toast at startup on new Android)
 			Result += " -Wl,-soname,libUnreal.so";
diff --git a/Engine/Source/Programs/UnrealBuildTool/Platform/Android/UEBuildAndroid.cs b/Engine/Source/Programs/UnrealBuildTool/Platform/Android/UEBuildAndroid.cs
--- a/Engine/Source/Programs/UnrealBuildTool/Platform/Android/UEBuildAndroid.cs
+++ b/Engine/Source/Programs/UnrealBuildTool/Platform/Android/UEBuildAndroid.cs
@@ -383,7 +383,7 @@ namespace UnrealBuildTool
 			string NDKPath = Environment.GetEnvironmentVariable("NDKROOT")!;
 			NDKPath = NDKPath.Replace("\"", "");
 
-			AndroidToolChain ToolChain = new AndroidToolChain(Target.ProjectFile, false, Target.AndroidPlatform.Architectures, Target.AndroidPlatform.GPUArchitectures, Logger);
+			AndroidToolChain ToolChain = new AndroidToolChain(Target.ProjectFile, Target.AndroidPlatform.Architectures, Target.AndroidPlatform.GPUArchitectures, Logger);
 
 			// figure out the NDK version
 			string? NDKToolchainVersion = SDK.GetInstalledVersion();
@@ -530,16 +530,15 @@ namespace UnrealBuildTool
 
 		public override UEToolChain CreateToolChain(ReadOnlyTargetRules Target)
 		{
-			bool bUseLdGold = Target.bUseUnityBuild;
 			ClangToolChainOptions Options = CreateToolChainOptions(Target.AndroidPlatform.TargetRules);
-			return new AndroidToolChain(Target.ProjectFile, bUseLdGold, Target.AndroidPlatform.Architectures, Target.AndroidPlatform.GPUArchitectures, Options, Logger);
+			return new AndroidToolChain(Target.ProjectFile, Target.AndroidPlatform.Architectures, Target.AndroidPlatform.GPUArchitectures, Options, Logger);
 		}
 		public virtual UEToolChain CreateTempToolChainForProject(FileReference? ProjectFile)
 		{
 			AndroidTargetRules TargetRules = new AndroidTargetRules();
 			CommandLine.ParseArguments(Environment.GetCommandLineArgs(), TargetRules, Logger);
 			ClangToolChainOptions Options = CreateToolChainOptions(TargetRules);
-			return new AndroidToolChain(ProjectFile, true, null, null, Options, Logger);
+			return new AndroidToolChain(ProjectFile, null, null, Options, Logger);
 		}
 
 		/// <inheritdoc/>
diff --git a/Engine/Source/Programs/UnrealBuildTool/Platform/Android/UEDeployAndroid.cs b/Engine/Source/Programs/UnrealBuildTool/Platform/Android/UEDeployAndroid.cs
--- a/Engine/Source/Programs/UnrealBuildTool/Platform/Android/UEDeployAndroid.cs
+++ b/Engine/Source/Programs/UnrealBuildTool/Platform/Android/UEDeployAndroid.cs
@@ -953,7 +953,6 @@ namespace UnrealBuildTool
 		}
 
 		private int CachedStoreVersion = -1;
-		private int CachedStoreVersionOffsetArmV7 = 0;
 		private int CachedStoreVersionOffsetArm64 = 0;
 		private int CachedStoreVersionOffsetX8664= 0;
 
@@ -993,14 +992,12 @@ namespace UnrealBuildTool
 
 				CachedStoreVersion = StoreVersion;
 
-				Ini.GetInt32("/Script/AndroidRuntimeSettings.AndroidRuntimeSettings", "StoreVersionOffsetArmV7", out CachedStoreVersionOffsetArmV7);
 				Ini.GetInt32("/Script/AndroidRuntimeSettings.AndroidRuntimeSettings", "StoreVersionOffsetArm64", out CachedStoreVersionOffsetArm64);
 				Ini.GetInt32("/Script/AndroidRuntimeSettings.AndroidRuntimeSettings", "StoreVersionOffsetX8664", out CachedStoreVersionOffsetX8664);
 			}
 
 			switch (UnrealArch)
 			{
-				case "-armv7": return CachedStoreVersion + CachedStoreVersionOffsetArmV7;
 				case "-arm64": return CachedStoreVersion + CachedStoreVersionOffsetArm64;
 				case "-x64": return CachedStoreVersion + CachedStoreVersionOffsetX8664;
 			}
@@ -3350,13 +3347,6 @@ namespace UnrealBuildTool
 			return true;
 		}
 
-		private bool GradleEnabled()
-		{
-			ConfigHierarchy Ini = GetConfigCacheIni(ConfigHierarchyType.Engine);
-			bool bEnableGradle = false;
-			Ini.GetBool("/Script/AndroidRuntimeSettings.AndroidRuntimeSettings", "bEnableGradle", out bEnableGradle);
-			return bEnableGradle;
-		}
 		private bool BundleEnabled()
 		{
 			if (ForceAPKGeneration)
@@ -4008,11 +3998,6 @@ namespace UnrealBuildTool
 				return;
 			}
 
-			if (!GradleEnabled())
-			{
-				throw new BuildException("Support for building APK without Gradle is depreciated; please update your project Engine.ini.");
-			}
-
 			if (UPL!.GetLastError() != null)
 			{
 				throw new BuildException("Cannot make APK with UPL errors");
diff --git a/Engine/Source/Programs/UnrealBuildTool/ProjectFiles/VisualStudioCode/VSCodeProjectFileGenerator.cs b/Engine/Source/Programs/UnrealBuildTool/ProjectFiles/VisualStudioCode/VSCodeProjectFileGenerator.cs
--- a/Engine/Source/Programs/UnrealBuildTool/ProjectFiles/VisualStudioCode/VSCodeProjectFileGenerator.cs
+++ b/Engine/Source/Programs/UnrealBuildTool/ProjectFiles/VisualStudioCode/VSCodeProjectFileGenerator.cs
@@ -1176,26 +1176,16 @@ namespace UnrealBuildTool
 			// Get store version
 			int StoreVersion = 1;
 			int StoreVersionArm64 = 1;
-			int StoreVersionArmV7 = 1;
 			int StoreVersionOffsetArm64 = 0;
-			int StoreVersionOffsetArmV7 = 0;
 			Ini.GetInt32("/Script/AndroidRuntimeSettings.AndroidRuntimeSettings", "StoreVersion", out StoreVersion);
 			Ini.GetInt32("/Script/AndroidRuntimeSettings.AndroidRuntimeSettings", "StoreVersionOffsetArm64", out StoreVersionOffsetArm64);
-			Ini.GetInt32("/Script/AndroidRuntimeSettings.AndroidRuntimeSettings", "StoreVersionOffsetArmV7", out StoreVersionOffsetArmV7);
 			StoreVersionArm64 = StoreVersion + StoreVersionOffsetArm64;
-			StoreVersionArmV7 = StoreVersion + StoreVersionOffsetArmV7;
 
 			DirectoryReference SymbolPathArm64 = DirectoryReference.Combine(
 				BuildProduct.OutputFile.Directory,
 				Target.Name + "_Symbols_v" + StoreVersionArm64.ToString(),
 				Target.Name + "-arm64");
 
-			DirectoryReference SymbolPathArmV7 = DirectoryReference.Combine(
-				BuildProduct.OutputFile.Directory,
-				Target.Name + "_Symbols_v" + StoreVersionArmV7.ToString(),
-				Target.Name + "-armv7");
-
-
 			string LaunchTaskName = String.Format("{0} {1} {2} Deploy", Target.Name, BuildProduct.Platform, BuildProduct.Config);
 
 			List<string> ConfigTypes = new List<string>();
@@ -1230,7 +1220,6 @@ namespace UnrealBuildTool
 						{
 							OutFile.BeginArray("librarySearchPaths");
 							OutFile.AddUnnamedField("\\\"" + SymbolPathArm64.ToNormalizedPath() + "\\\"");
-							OutFile.AddUnnamedField("\\\"" + SymbolPathArmV7.ToNormalizedPath() + "\\\"");
 							OutFile.EndArray();
 
 							OutFile.BeginArray("lldbPreTargetCreateCommands");
diff --git a/Engine/Source/Runtime/Android/AndroidRuntimeSettings/Classes/AndroidRuntimeSettings.h b/Engine/Source/Runtime/Android/AndroidRuntimeSettings/Classes/AndroidRuntimeSettings.h
--- a/Engine/Source/Runtime/Android/AndroidRuntimeSettings/Classes/AndroidRuntimeSettings.h
+++ b/Engine/Source/Runtime/Android/AndroidRuntimeSettings/Classes/AndroidRuntimeSettings.h
@@ -193,10 +193,6 @@ public:
 	UPROPERTY(GlobalConfig, EditAnywhere, Category = "APK Packaging")
 	TEnumAsByte<EAndroidInstallLocation::Type> InstallLocation;
 
-	// Use Gradle instead of Ant for Java compiling and APK generation
-	UPROPERTY(GlobalConfig, EditAnywhere, Category = "APK Packaging", Meta = (DisplayName = "Enable Gradle instead of Ant"))
-	bool bEnableGradle;
-
 	// Enable -Xlint:unchecked and -Xlint:depreciation for Java compiling (Gradle only)
 	UPROPERTY(GlobalConfig, EditAnywhere, Category = "APK Packaging", Meta = (DisplayName = "Enable Lint depreciation checks"))
 	bool bEnableLint;
@@ -392,9 +388,13 @@ public:
 	bool bDetectVulkanByDefault;
 
 	// Build the shipping config with hidden visibility by default. Results in smaller .so file but will also removes symbols used to display callstack dumps.
-	UPROPERTY(GlobalConfig, EditAnywhere, Category = AdvancedBuild, meta = (DisplayName = "Build with hidden symbol visibility in shipping config. [Experimental]"))
+	UPROPERTY(GlobalConfig, EditAnywhere, Category = AdvancedBuild, meta = (DisplayName = "Build with hidden symbol visibility in shipping config."))
 	bool bBuildWithHiddenSymbolVisibility;
 
+	// Disables extra checks for buffer overflows, comes with perf improvement, but might make tracing stack corruptions in production harder. Note that _FORTIFY_SOURCE=2 is still enabled by the toolchain providing lightweight stack checks
+	UPROPERTY(GlobalConfig, EditAnywhere, Category = AdvancedBuild, meta = (DisplayName = "Disable extra checks for buffer overflows"))
+	bool bDisableStackProtector;
+
 	// Disable libc++_shared dependency validation in all .so files linked with libUnreal.so
 	UPROPERTY(GlobalConfig, EditAnywhere, Category = AdvancedBuild, meta = (DisplayName = "Disable libc++_shared dependency validation in all dependencies"))
 	bool bDisableLibCppSharedDependencyValidation;
@@ -403,10 +403,6 @@ public:
 	UPROPERTY(GlobalConfig, EditAnywhere, Category = AdvancedBuild, meta = (DisplayName = "Always save a copy of the libUnreal.so with symbols. [Experimental]"))
 	bool bSaveSymbols;
 
-	// Use legacy ld instead of new lld linker.
-	UPROPERTY(GlobalConfig, EditAnywhere, Category = AdvancedBuild, meta = (DisplayName = "Force linking to use ld instead of lld"))
-	bool bForceLDLinker;
-
 	// Strip shader reflection information under Android to avoid issues on older drivers
 	UPROPERTY(config, EditAnywhere, Category = AdvancedBuild, meta = (
 	DisplayName = "Strip shader reflection information",
