Fixed compilation

diff --git a/Engine/Source/Programs/UnrealBuildTool/Platform/Android/UEDeployAndroid.cs b/Engine/Source/Programs/UnrealBuildTool/Platform/Android/UEDeployAndroid.cs
--- a/Engine/Source/Programs/UnrealBuildTool/Platform/Android/UEDeployAndroid.cs
+++ b/Engine/Source/Programs/UnrealBuildTool/Platform/Android/UEDeployAndroid.cs
@@ -5442,7 +5442,7 @@ namespace UnrealBuildTool
 			AndroidTargetRules TargetRules = new AndroidTargetRules();
 			CommandLine.ParseArguments(Environment.GetCommandLineArgs(), TargetRules, Logger);
 			ClangToolChainOptions Options = AndroidPlatform.CreateToolChainOptions(TargetRules);
-			AndroidToolChain ToolChain = new AndroidToolChain(ProjectFile, false, null, null, Options, Logger);
+			AndroidToolChain ToolChain = new AndroidToolChain(ProjectFile, null, null, Options, Logger);
 
 			SavePackageInfo(ProjectName, ProjectDirectory.FullName, Type, bSkipGradleBuild);
 
