Removed Android NDK25 AGDE debugger workaround as AGDE canary v69 resolved this issue

diff --git a/Engine/Source/Programs/UnrealBuildTool/Platform/Android/AndroidProjectGenerator.cs b/Engine/Source/Programs/UnrealBuildTool/Platform/Android/AndroidProjectGenerator.cs
--- a/Engine/Source/Programs/UnrealBuildTool/Platform/Android/AndroidProjectGenerator.cs
+++ b/Engine/Source/Programs/UnrealBuildTool/Platform/Android/AndroidProjectGenerator.cs
@@ -21,11 +21,6 @@ namespace UnrealBuildTool
 		/// </summary>
 		private bool AGDEInstalled = false;
 
-		/// <summary>
-		/// Cached version of Android NDK toolchain
-		/// </summary>
-		private UInt64 NDKVersionInt;
-
 		public AndroidProjectGenerator(CommandLineArguments Arguments, ILogger Logger)
 			: base(Arguments, Logger)
 		{
@@ -36,12 +31,6 @@ namespace UnrealBuildTool
 				{
 					AGDEInstalled = false;
 				}
-				else
-				{
-					UEBuildPlatformSDK SDK = UEBuildPlatform.GetSDK(UnrealTargetPlatform.Android)!;
-					string? NDKToolchainVersion = SDK.GetInstalledVersion();
-					SDK.TryConvertVersionToInt(NDKToolchainVersion, out NDKVersionInt);
-				}
 			}
 		}
 
@@ -181,7 +170,6 @@ namespace UnrealBuildTool
 				string UserFileEntry = "<PropertyGroup " + InConditionString + ">\n";
 				UserFileEntry		+= "	<AndroidLldbStartupCommands>" +
 												"command script import \"" + Path.Combine(Unreal.EngineDirectory.FullName, "Extras", "LLDBDataFormatters", "UEDataFormatters_2ByteChars.py") + "\";" +
-												(NDKVersionInt >= 250000 ? "settings set target.source-map \"\" \"" + Path.Combine(Unreal.EngineDirectory.FullName, "Source") + "\";" : "") +
 												"$(AndroidLldbStartupCommands)" +
 											"</AndroidLldbStartupCommands>\n";
 				UserFileEntry		+= "</PropertyGroup>\n";
