Fix if conditional logic in TryGetValue for Enum TryParse

diff --git a/Engine/Source/Programs/UnrealBuildTool/System/ConfigHierarchy.cs b/Engine/Source/Programs/UnrealBuildTool/System/ConfigHierarchy.cs
--- a/Engine/Source/Programs/UnrealBuildTool/System/ConfigHierarchy.cs
+++ b/Engine/Source/Programs/UnrealBuildTool/System/ConfigHierarchy.cs
@@ -452,7 +452,7 @@ namespace UnrealBuildTool
 		public bool TryGetValue<T>(string SectionName, string KeyName, out T Value) where T : struct
 		{
 			string? Text;
-			if (TryGetValue(SectionName, KeyName, out Text))
+			if (!TryGetValue(SectionName, KeyName, out Text))
 			{
 				Value = default(T);
 				return false;
