[UnrealBuildTool]
* Changed "/FUNCTIONPADMIN" to "/FUNCTIONPADMIN:6" to get around problem with FUNCTIONPADMIN causing full link step all the time. Odd, but seems to work (6 is the default value for 64 bit targets)

diff --git a/Engine/Source/Programs/UnrealBuildTool/Platform/Windows/VCToolChain.cs b/Engine/Source/Programs/UnrealBuildTool/Platform/Windows/VCToolChain.cs
--- a/Engine/Source/Programs/UnrealBuildTool/Platform/Windows/VCToolChain.cs
+++ b/Engine/Source/Programs/UnrealBuildTool/Platform/Windows/VCToolChain.cs
@@ -1990,7 +1990,7 @@ namespace UnrealBuildTool
 			// Enable function level hot-patching
 			if(!bBuildImportLibraryOnly && Target.WindowsPlatform.bCreateHotpatchableImage)
 			{
-				Arguments.Add("/FUNCTIONPADMIN");
+				Arguments.Add("/FUNCTIONPADMIN:6"); // For some reason, not providing the number causes full linking to happen all the time
 			}
 
 			// For targets that are cross-referenced, we don't want to write a LIB file during the link step as that
