Deprecated bCompileForSize UBT Target option and replaced with OptimizationMode enumto fine tune a balance between speed and\or code size

diff --git a/Engine/Source/Programs/UnrealBuildTool/Configuration/TargetRules.cs b/Engine/Source/Programs/UnrealBuildTool/Configuration/TargetRules.cs
--- a/Engine/Source/Programs/UnrealBuildTool/Configuration/TargetRules.cs
+++ b/Engine/Source/Programs/UnrealBuildTool/Configuration/TargetRules.cs
@@ -231,6 +231,27 @@ namespace UnrealBuildTool
 		Shallow,
 	}
 
+	/// <summary>
+	/// Optimization mode for compiler settings
+	/// </summary>
+	public enum OptimizationMode
+	{
+		/// <summary>
+		/// Favor speed
+		/// </summary>
+		Speed,
+
+		/// <summary>
+		/// Favor minimal code size
+		/// </summary>
+		Size,
+
+		/// <summary>
+		/// Somewhere between Speed and Size
+		/// </summary>
+		SizeAndSpeed
+	}
+
 	/// <summary>
 	/// Utility class for EngineIncludeOrderVersion defines
 	/// </summary>
@@ -263,7 +284,7 @@ namespace UnrealBuildTool
 	}
 
 	/// <summary>
-	/// Attribute used to mark fields which much match between targets in the shared build environment
+	/// Attribute used to mark fields which must match between targets in the shared build environment
 	/// </summary>
 	[AttributeUsage(AttributeTargets.Field | AttributeTargets.Property, AllowMultiple = false)]
 	class RequiresUniqueBuildEnvironmentAttribute : Attribute
@@ -1115,9 +1136,16 @@ namespace UnrealBuildTool
 		/// True if we want to favor optimizing size over speed.
 		/// </summary>
 		[RequiresUniqueBuildEnvironment]
+		[Obsolete("Deprecated in UE5.1 - Please use OptimizationLevel instead.")]
 		[ConfigFile(ConfigHierarchyType.Engine, "/Script/BuildSettings.BuildSettings", "bCompileForSize")]
 		public bool bCompileForSize = false;
 
+		/// <summary>
+		/// Allows to fine tune optimizations level for speed and\or code size
+		/// </summary>
+		[RequiresUniqueBuildEnvironment]
+		public OptimizationMode OptimizationLevel = OptimizationMode.Speed;
+
 		/// <summary>
 		/// Whether to compile development automation tests.
 		/// </summary>
@@ -2862,11 +2890,17 @@ namespace UnrealBuildTool
 			get { return Inner.bCompileFreeType; }
 		}
 
+		[Obsolete("Deprecated in UE5.1 - Please use OptimizationLevel instead.")]
 		public bool bCompileForSize
 		{
 			get { return Inner.bCompileForSize; }
 		}
 
+		public OptimizationMode OptimizationLevel
+		{
+			get { return Inner.OptimizationLevel; }
+		}
+
 		public bool bRetainFramePointers
 		{
 			get { return Inner.bRetainFramePointers; }
diff --git a/Engine/Source/Programs/UnrealBuildTool/Configuration/UEBuildTarget.cs b/Engine/Source/Programs/UnrealBuildTool/Configuration/UEBuildTarget.cs
--- a/Engine/Source/Programs/UnrealBuildTool/Configuration/UEBuildTarget.cs
+++ b/Engine/Source/Programs/UnrealBuildTool/Configuration/UEBuildTarget.cs
@@ -3869,7 +3869,7 @@ namespace UnrealBuildTool
 			GlobalCompileEnvironment.bUndefinedIdentifierWarningsAsErrors = Rules.bUndefinedIdentifierErrors;
 			GlobalCompileEnvironment.bRetainFramePointers = Rules.bRetainFramePointers;
 			GlobalCompileEnvironment.bWarningsAsErrors = Rules.bWarningsAsErrors;
-			GlobalCompileEnvironment.bOptimizeForSize = Rules.bCompileForSize;
+			GlobalCompileEnvironment.OptimizationLevel = Rules.OptimizationLevel;
 			GlobalCompileEnvironment.bUseStaticCRT = Rules.bUseStaticCRT;
 			GlobalCompileEnvironment.bOmitFramePointers = Rules.bOmitFramePointers;
 			GlobalCompileEnvironment.bUsePDBFiles = Rules.bUsePDBFiles;
@@ -3894,7 +3894,6 @@ namespace UnrealBuildTool
 			GlobalLinkEnvironment.bUseDebugCRT = bUseDebugCRT;
 			GlobalLinkEnvironment.bUseStaticCRT = Rules.bUseStaticCRT;
 			GlobalLinkEnvironment.bIsBuildingConsoleApplication = Rules.bIsBuildingConsoleApplication;
-			GlobalLinkEnvironment.bOptimizeForSize = Rules.bCompileForSize;
 			GlobalLinkEnvironment.bOmitFramePointers = Rules.bOmitFramePointers;
 			GlobalLinkEnvironment.bSupportEditAndContinue = Rules.bSupportEditAndContinue;
 			GlobalLinkEnvironment.bCreateMapFile = Rules.bCreateMapFile;
diff --git a/Engine/Source/Programs/UnrealBuildTool/Platform/Android/AndroidToolChain.cs b/Engine/Source/Programs/UnrealBuildTool/Platform/Android/AndroidToolChain.cs
--- a/Engine/Source/Programs/UnrealBuildTool/Platform/Android/AndroidToolChain.cs
+++ b/Engine/Source/Programs/UnrealBuildTool/Platform/Android/AndroidToolChain.cs
@@ -659,10 +659,18 @@ namespace UnrealBuildTool
 			}
 			else
 			{
-				if (CompileEnvironment.bOptimizeForSize)
+				if (CompileEnvironment.OptimizationLevel == OptimizationMode.Size)
 				{
 					Result += " -Oz";
 				}
+				else if (CompileEnvironment.OptimizationLevel == OptimizationMode.SizeAndSpeed)
+				{
+					Result += " -Os";
+					if (Architecture == "-arm64")
+					{
+						Result += " -moutline";
+					}
+				}
 				else
 				{
 					Result += " -O3";
diff --git a/Engine/Source/Programs/UnrealBuildTool/Platform/Linux/LinuxToolChain.cs b/Engine/Source/Programs/UnrealBuildTool/Platform/Linux/LinuxToolChain.cs
--- a/Engine/Source/Programs/UnrealBuildTool/Platform/Linux/LinuxToolChain.cs
+++ b/Engine/Source/Programs/UnrealBuildTool/Platform/Linux/LinuxToolChain.cs
@@ -385,7 +385,22 @@ namespace UnrealBuildTool
 				}
 				else
 				{
-					Arguments.Add("-O3");
+					if (CompileEnvironment.OptimizationLevel == OptimizationMode.Size)
+					{
+						Arguments.Add("-Oz");
+					}
+					else if (CompileEnvironment.OptimizationLevel == OptimizationMode.SizeAndSpeed)
+					{
+						Arguments.Add("-Os");
+						if (Architecture.StartsWith("aarch64"))
+						{
+							Arguments.Add("-moutline");
+						}
+					}
+					else
+					{
+						Arguments.Add("-O3");
+					}
 				}
 			}
 
diff --git a/Engine/Source/Programs/UnrealBuildTool/Platform/Windows/VCToolChain.cs b/Engine/Source/Programs/UnrealBuildTool/Platform/Windows/VCToolChain.cs
--- a/Engine/Source/Programs/UnrealBuildTool/Platform/Windows/VCToolChain.cs
+++ b/Engine/Source/Programs/UnrealBuildTool/Platform/Windows/VCToolChain.cs
@@ -533,8 +533,15 @@ namespace UnrealBuildTool
 					// Maximum optimizations.
 					Arguments.Add("/Ox");
 
-					// Favor code speed.
-					Arguments.Add("/Ot");
+					if (CompileEnvironment.OptimizationLevel != OptimizationMode.Speed)
+					{
+						Arguments.Add("/Os");
+					}
+					else
+					{
+						// Favor code speed.
+						Arguments.Add("/Ot");
+					}
 
 					// Coalesce duplicate strings
 					Arguments.Add("/GF");
diff --git a/Engine/Source/Programs/UnrealBuildTool/System/CppCompileEnvironment.cs b/Engine/Source/Programs/UnrealBuildTool/System/CppCompileEnvironment.cs
--- a/Engine/Source/Programs/UnrealBuildTool/System/CppCompileEnvironment.cs
+++ b/Engine/Source/Programs/UnrealBuildTool/System/CppCompileEnvironment.cs
@@ -284,9 +284,9 @@ namespace UnrealBuildTool
 		public bool bOptimizeCode = false;
 
 		/// <summary>
-		/// Whether to optimize for minimal code size
+		/// Allows to fine tune optimizations level for speed and\or code size
 		/// </summary>
-		public bool bOptimizeForSize = false;
+		public OptimizationMode OptimizationLevel = OptimizationMode.Speed;
 
 		/// <summary>
 		/// True if debug info should be created.
@@ -527,7 +527,7 @@ namespace UnrealBuildTool
 			StaticAnalyzerDisabledCheckers = new HashSet<string>(Other.StaticAnalyzerDisabledCheckers);
 			StaticAnalyzerAdditionalCheckers = new HashSet<string>(Other.StaticAnalyzerAdditionalCheckers);
 			bOptimizeCode = Other.bOptimizeCode;
-			bOptimizeForSize = Other.bOptimizeForSize;
+			OptimizationLevel = Other.OptimizationLevel;
 			bCreateDebugInfo = Other.bCreateDebugInfo;
 			bIsBuildingLibrary = Other.bIsBuildingLibrary;
 			bIsBuildingDLL = Other.bIsBuildingDLL;
diff --git a/Engine/Source/Programs/UnrealBuildTool/System/LinkEnvironment.cs b/Engine/Source/Programs/UnrealBuildTool/System/LinkEnvironment.cs
--- a/Engine/Source/Programs/UnrealBuildTool/System/LinkEnvironment.cs
+++ b/Engine/Source/Programs/UnrealBuildTool/System/LinkEnvironment.cs
@@ -196,11 +196,6 @@ namespace UnrealBuildTool
 		/// </summary>
 		public int DefaultStackSizeCommit = 0;
 
-		/// <summary>
-		/// Whether to optimize for minimal code size
-		/// </summary>
-		public bool bOptimizeForSize = false;
-
 		/// <summary>
 		/// Whether to omit frame pointers or not. Disabling is useful for e.g. memory profiling on the PC
 		/// </summary>
@@ -366,7 +361,6 @@ namespace UnrealBuildTool
 			bIsBuildingDotNetAssembly = Other.bIsBuildingDotNetAssembly;
 			DefaultStackSize = Other.DefaultStackSize;
 			DefaultStackSizeCommit = Other.DefaultStackSizeCommit;
-			bOptimizeForSize = Other.bOptimizeForSize;
 			bOmitFramePointers = Other.bOmitFramePointers;
 			bSupportEditAndContinue = Other.bSupportEditAndContinue;
 			bUseIncrementalLinking = Other.bUseIncrementalLinking;
diff --git a/Engine/Source/Programs/UnrealBuildTool/ToolChain/AppleToolChain.cs b/Engine/Source/Programs/UnrealBuildTool/ToolChain/AppleToolChain.cs
--- a/Engine/Source/Programs/UnrealBuildTool/ToolChain/AppleToolChain.cs
+++ b/Engine/Source/Programs/UnrealBuildTool/ToolChain/AppleToolChain.cs
@@ -276,10 +276,19 @@ namespace UnrealBuildTool
 					Arguments.Add("-O1");
 					Arguments.Add("-g");
 				}
-				else if (CompileEnvironment.bOptimizeForSize)
+				else if (CompileEnvironment.OptimizationLevel == OptimizationMode.Size)
 				{
 					Arguments.Add("-Oz");
 				}
+				else if (CompileEnvironment.OptimizationLevel == OptimizationMode.SizeAndSpeed)
+				{
+					Arguments.Add("-Os");
+
+					if (CompileEnvironment.Architecture.StartsWith("aarch64"))
+					{
+						Arguments.Add("-moutline");
+					}
+				}
 				else
 				{
 					Arguments.Add("-O3");
diff --git a/Engine/Source/Programs/UnrealBuildTool/ToolChain/ISPCToolChain.cs b/Engine/Source/Programs/UnrealBuildTool/ToolChain/ISPCToolChain.cs
--- a/Engine/Source/Programs/UnrealBuildTool/ToolChain/ISPCToolChain.cs
+++ b/Engine/Source/Programs/UnrealBuildTool/ToolChain/ISPCToolChain.cs
@@ -575,7 +575,7 @@ namespace UnrealBuildTool
 			}
 			else
 			{
-				CommonArgs.Add("-O2");
+				CommonArgs.Add("-O3");
 			}
 			GlobalArguments.AddRange(CommonArgs);
 
