- More work on "modern xcode" mode (automatic signing works well, xcode can generate plists/entitlements, shared settings between Mac/IOS/TVOS). Needs the bUseModernXcode ini setting (see BaseEngine.ini) to activate
- Indexing improvements
- No longer asks to auto-create schemes, since we create all schemes needed
- Sped up the run-only project generation to skip a lot of cruft

diff --git a/Engine/Build/BatchFiles/Mac/XcodeBuild.sh b/Engine/Build/BatchFiles/Mac/XcodeBuild.sh
--- a/Engine/Build/BatchFiles/Mac/XcodeBuild.sh
+++ b/Engine/Build/BatchFiles/Mac/XcodeBuild.sh
@@ -4,6 +4,11 @@
 # (and can take the same arguments) but performs some interpretation of arguments that come from Xcode
 # Values for $ACTION: "" = building, "clean" = cleaning
 
+if [ ["$UE_SKIP_UBT"] == ["1"] ]; then
+	echo Skipping UBT per request
+	exit 0
+fi
+
 # Setup Environment
 source  Engine/Build/BatchFiles/Mac/SetupEnvironment.sh -dotnet Engine/Build/BatchFiles/Mac
 
@@ -111,7 +116,8 @@ elif [ $ACTION == "clean" ]; then
 fi
 
 echo Running dotnet Engine/Binaries/DotNET/UnrealBuildTool/UnrealBuildTool.dll $TARGET $PLATFORM $CONFIGURATION "$TRAILINGARGS" $UBT_ARCHFLAG $AdditionalFlags
-dotnet Engine/Binaries/DotNET/UnrealBuildTool/UnrealBuildTool.dll $TARGET $PLATFORM $CONFIGURATION "$TRAILINGARGS" $UBT_ARCHFLAG $AdditionalFlags
+# set an envvar to let UBT know that it's being run from xcode (envvar allows children to get the setting if needed)
+UE_BUILD_FROM_XCODE=1 dotnet Engine/Binaries/DotNET/UnrealBuildTool/UnrealBuildTool.dll $TARGET $PLATFORM $CONFIGURATION "$TRAILINGARGS" $UBT_ARCHFLAG $AdditionalFlags
 
 ExitCode=$?
 if [ $ExitCode -eq 254 ] || [ $ExitCode -eq 255 ] || [ $ExitCode -eq 2 ]; then
diff --git a/Engine/Config/BaseEngine.ini b/Engine/Config/BaseEngine.ini
--- a/Engine/Config/BaseEngine.ini
+++ b/Engine/Config/BaseEngine.ini
@@ -2675,7 +2675,7 @@ bEnableCloudKitSupport=False
 bRunAsCurrentUser=False
 IOSCloudKitSyncStrategy=None
 bGenerateCrashReportSymbols=false
-bAutomaticSigning=false
+bAutomaticSigning=true
 UseFastIntrinsics=False
 ForceFloats=True
 EnableMathOptimisations=True
@@ -2683,6 +2683,12 @@ MetalLanguageVersion=6
 bDisableMotionData=False
 bEnableAdvertisingIdentifier=True
 
+[XcodeConfiguration]
+bUseModernXcode=false
+bUseModernCodeSigning=true
+ModernSigningTeam=
+ModernBundleIdentifier=com.YourCompany.[PROJECT_NAME]
+
 ; These are the defaults for Android settings, and they need to be in the .ini since UBT reads the .ini settings, without instantiating the class
 [/Script/AndroidRuntimeSettings.AndroidRuntimeSettings]
 SDKAPILevelOverride=
diff --git a/Engine/Source/Programs/AutomationTool/IOS/IOSPlatform.Automation.cs b/Engine/Source/Programs/AutomationTool/IOS/IOSPlatform.Automation.cs
--- a/Engine/Source/Programs/AutomationTool/IOS/IOSPlatform.Automation.cs
+++ b/Engine/Source/Programs/AutomationTool/IOS/IOSPlatform.Automation.cs
@@ -1,4 +1,4 @@
-// Copyright Epic Games, Inc. All Rights Reserved.
+﻿// Copyright Epic Games, Inc. All Rights Reserved.
 using System;
 using System.Collections.Generic;
 using System.Linq;
@@ -1065,6 +1065,7 @@ public class IOSPlatform : Platform
 		PrintRunTime();
 	}
 
+#if true
 	private string EnsureXcodeProjectExists(FileReference RawProjectPath, string LocalRoot, string ShortProjectName, string ProjectRoot, bool IsCodeBasedProject, out bool bWasGenerated)
 	{
 		// first check for the .xcodeproj
@@ -1107,9 +1108,39 @@ public class IOSPlatform : Platform
 
 		return XcodeProj;
 	}
+#endif
 
 	private void CodeSign(string BaseDirectory, string GameName, FileReference RawProjectPath, UnrealTargetConfiguration TargetConfig, string LocalRoot, string ProjectName, string ProjectDirectory, bool IsCode, bool Distribution = false, string Provision = null, string Certificate = null, string Team = null, bool bAutomaticSigning = false, string SchemeName = null, string SchemeConfiguration = null)
 	{
+		bool bUseModernXcode = false;
+		if (OperatingSystem.IsMacOS())
+		{
+			ConfigHierarchy Ini = ConfigCache.ReadHierarchy(ConfigHierarchyType.Engine, RawProjectPath.Directory, UnrealTargetPlatform.IOS);
+			Ini.TryGetValue("XcodeConfiguration", "bUseModernXcode", out bUseModernXcode);
+		}
+
+		if (bUseModernXcode)
+		{
+			// we need an xcode project to continue
+			DirectoryReference GeneratedProject = null;
+			IOSExports.GenerateRunOnlyXcodeProject(RawProjectPath, TargetPlatformType, false, Logger, out GeneratedProject);
+
+			if (GeneratedProject == null || !DirectoryReference.Exists(GeneratedProject))
+			{
+				// something very bad happened
+				throw new AutomationException("iOS couldn't find the appropriate Xcode Project for " + RawProjectPath);
+			}
+
+			int ReturnCode = IOSExports.FinalizeAppWithXcode(GeneratedProject, TargetPlatformType, SchemeName ?? GameName, SchemeConfiguration ?? TargetConfig.ToString(), Provision, Certificate, Team, bAutomaticSigning, Distribution, bUseModernXcode, Logger);
+			if (ReturnCode != 0)
+			{
+				throw new AutomationException(ExitCode.Error_FailedToCodeSign, "CodeSign Failed");
+			}
+
+			return;
+		}
+// @todo make this use the new IOSExports functionality that modern mode is using, after testing it
+#if true
 		// check for the proper xcodeproject
 		bool bWasGenerated = false;
 		string XcodeProj = EnsureXcodeProjectExists(RawProjectPath, LocalRoot, ProjectName, ProjectDirectory, IsCode, out bWasGenerated);
@@ -1171,6 +1202,7 @@ public class IOSPlatform : Platform
 		{
 			throw new AutomationException(ExitCode.Error_FailedToCodeSign, "CodeSign Failed");
 		}
+#endif
 	}
 
 	private bool ShouldUseMaxIPACompression(ProjectParams Params)
@@ -2302,7 +2334,7 @@ public class IOSPlatform : Platform
 		}
 	}
 
-	#region Hooks
+#region Hooks
 
 	public override void PreBuildAgenda(UnrealBuild Build, UnrealBuild.BuildAgenda Agenda, ProjectParams Params)
 	{
@@ -2312,7 +2344,7 @@ public class IOSPlatform : Platform
 		}
 	}
 
-	#endregion
+#endregion
 
 	public override DirectoryReference GetProjectRootForStage(DirectoryReference RuntimeRoot, StagedDirectoryReference RelativeProjectRootForStage)
 	{
diff --git a/Engine/Source/Programs/UnrealBuildTool/Platform/IOS/IOSExports.cs b/Engine/Source/Programs/UnrealBuildTool/Platform/IOS/IOSExports.cs
--- a/Engine/Source/Programs/UnrealBuildTool/Platform/IOS/IOSExports.cs
+++ b/Engine/Source/Programs/UnrealBuildTool/Platform/IOS/IOSExports.cs
@@ -1,4 +1,4 @@
-// Copyright Epic Games, Inc. All Rights Reserved.
+﻿// Copyright Epic Games, Inc. All Rights Reserved.
 
 using System;
 using System.IO;
@@ -184,6 +184,174 @@
 		}
 
 		/// <summary>
+		/// Reads the per-platform .PackageVersionCounter file to get a version, and modifies it with updated value so that next build will get new version
+		/// </summary>
+		/// <param name="UProjectFile">Location of .uproject file (or null for the engine project)</param>
+		/// <param name="Platform">Which plaform to look up</param>
+		/// <returns></returns>
+		public static string GetAndUpdateVersionFile(FileReference? UProjectFile, UnrealTargetPlatform Platform)
+		{
+			FileReference RunningVersionFilename = UProjectFile == null ?
+				FileReference.Combine(Unreal.EngineDirectory, "Build", Platform.ToString(), "Engine.PackageVersionCounter") :
+				FileReference.Combine(UProjectFile.Directory, "Build", Platform.ToString(), $"{UProjectFile.GetFileNameWithoutAnyExtensions()}.PackageVersionCounter");
+
+			string CurrentVersion = "0.0";
+			if (FileReference.Exists(RunningVersionFilename))
+			{
+				CurrentVersion = FileReference.ReadAllText(RunningVersionFilename);
+			}
+
+			string[] VersionParts = CurrentVersion.Split('.');
+			int Major = int.Parse(VersionParts[0]);
+			int Minor = int.Parse(VersionParts[1]);
+
+			Minor++;
+
+			DirectoryReference.CreateDirectory(RunningVersionFilename.Directory);
+			FileReference.WriteAllText(RunningVersionFilename, $"{Major}.{Minor}");
+
+			return CurrentVersion;
+		}
+
+		/// <summary>
+		/// Genearate an run-only Xcode project, that is not meant to be used for anything else besides code-signing/running/etc of the native .app bundle
+		/// </summary>
+		/// <param name="UProjectFile">Location of .uproject file (or null for the engine project</param>
+		/// <param name="Platform">The platform to generate a project for</param>
+		/// <param name="bForDistribution">True if this is making a bild for uploading to app store</param>
+		/// <param name="Logger">Logging object</param>
+		/// <param name="GeneratedProjectFile">Returns the .xcworkspace that was made</param>
+		public static void GenerateRunOnlyXcodeProject(FileReference? UProjectFile, UnrealTargetPlatform Platform, bool bForDistribution, ILogger Logger, out DirectoryReference? GeneratedProjectFile)
+		{
+			List<string> Options = new()
+			{
+				$"-platforms={Platform}",
+				$"-{Platform}DeployOnly",
+				"-NoIntellisens",
+				"-IngnoreJunk",
+				bForDistribution ? "-distribution" : "-development",
+				"-IncludeTempTargets",
+				"-projectfileformat = XCode",
+				"-automated",
+			};
+
+			if (UProjectFile == null || UProjectFile.IsUnderDirectory(Unreal.EngineDirectory))
+			{
+				// @todo do we need these? where would the bundleid come from if there's no project?
+//				Options.Add("-bundleID=" + BundleID);
+//				Options.Add("-appname=" + AppName);
+				// @todo add an option to only add Engine target?
+			}
+			else
+			{
+				Options.Add($"-project=\"{UProjectFile.FullName}\"");
+				Options.Add("-game");
+			}
+
+			IOSToolChain.GenerateProjectFiles(UProjectFile, Options.ToArray(), Logger, out GeneratedProjectFile);
+		}
+
+		/// <summary>
+		/// Version of FinalizeAppWithXcode that is meant for modern xcode mode, where we assume all codesigning is setup already in the project, so nothing else is needed
+		/// </summary>
+		/// <param name="XcodeProject">The .xcworkspace file to build</param>
+		/// <param name="Platform">THe platform to make the .app for</param>
+		/// <param name="SchemeName">The name of the scheme (basically the target on the .xcworkspace)</param>
+		/// <param name="Configuration">Which configuration to make (Debug, etc)</param>
+		/// <param name="bForDistribution">True if this is making a bild for uploading to app store</param>
+		/// <param name="Logger">Logging object</param>
+		public static void FinalizeAppWithModernXcode(DirectoryReference XcodeProject, UnrealTargetPlatform Platform, string SchemeName, string Configuration, bool bForDistribution, ILogger Logger)
+		{
+			FinalizeAppWithXcode(XcodeProject, Platform, SchemeName, Configuration, null, null, null, false, bForDistribution, bUseModernXcode: true, Logger);
+		}
+
+		/// <summary>
+		/// Runs xcodebuild on a project (likely created with GenerateRunOnlyXcodeProject) to perform codesigning and creation of final .app
+		/// </summary>
+		/// <param name="XcodeProject">The .xcworkspace file to build</param>
+		/// <param name="Platform">THe platform to make the .app for</param>
+		/// <param name="SchemeName">The name of the scheme (basically the target on the .xcworkspace)</param>
+		/// <param name="Configuration">Which configuration to make (Debug, etc)</param>
+		/// <param name="Provision">An optional provision to codesign with (ignored in modern mod)</param>
+		/// <param name="Certificate">An optional certificate to codesign with (ignored in modern mode)</param>
+		/// <param name="Team">Optional Team to use when codesigning (ignored in modern mode)</param>
+		/// <param name="bAutomaticSigning">True if using automatic codesigning (where provision and certificate are not used) (ignored in modern mode)</param>
+		/// <param name="bForDistribution">True if this is making a bild for uploading to app store</param>
+		/// <param name="bUseModernXcode">True if the project was made with modern xcode mode</param>
+		/// <param name="Logger">Logging object</param>
+		/// <returns></returns>
+		public static int FinalizeAppWithXcode(DirectoryReference XcodeProject, UnrealTargetPlatform Platform, string SchemeName, string Configuration,
+			string? Provision, string? Certificate, string? Team, bool bAutomaticSigning, bool bForDistribution, bool bUseModernXcode, ILogger Logger)
+		{
+			List<string> Arguments = new()
+			{
+				"UBT_NO_POST_DEPLOY=true",
+				new IOSToolChainSettings(Logger).XcodeDeveloperDir + "usr/bin/xcodebuild",
+				"build",
+				$"-workspace \"{XcodeProject.FullName}\"",
+				$"-scheme \"{SchemeName}\"",
+				$"-configuration \"{Configuration}\"",
+				$"-destination generic/platform=" + (Platform == UnrealTargetPlatform.TVOS ? "tvOS" : "iOS"),
+				//$"-sdk {SDKName}",
+			};
+
+			if (bUseModernXcode)
+			{
+				Arguments.Add("-allowProvisioningUpdates");
+				// xcode gets confused it we _just_ wrote out entitlements while generating the temp project, and it thinks it was modified _during_ building
+				// but it wasn't, it was written before the build started
+				Arguments.Add("CODE_SIGN_ALLOW_ENTITLEMENTS_MODIFICATION=YES");
+			}
+			else
+			{
+				if (bAutomaticSigning)
+				{
+					Arguments.Add("CODE_SIGN_IDENTITY=" + (bForDistribution ? "\"iPhone Distribution\"" : "\"iPhone Developer\""));
+					Arguments.Add("CODE_SIGN_STYLE=\"Automatic\"");
+					Arguments.Add("-allowProvisioningUpdates");
+					Arguments.Add($"DEVELOPMENT_TEAM={Team}");
+				}
+				else
+				{
+					if (!string.IsNullOrEmpty(Certificate))
+					{
+						Arguments.Add($"CODE_SIGN_IDENTITY=\"{Certificate}\"");
+					}
+					else
+					{
+						Arguments.Add("CODE_SIGN_IDENTITY=" + (bForDistribution ? "\"iPhone Distribution\"" : "\"iPhone Developer\""));
+					}
+					if (!string.IsNullOrEmpty(Provision))
+					{
+						// read the provision to get the UUID
+						if (File.Exists(Environment.GetEnvironmentVariable("HOME") + "/Library/MobileDevice/Provisioning Profiles/" + Provision))
+						{
+							string UUID = "";
+							string AllText = File.ReadAllText(Environment.GetEnvironmentVariable("HOME") + "/Library/MobileDevice/Provisioning Profiles/" + Provision);
+							int idx = AllText.IndexOf("<key>UUID</key>");
+							if (idx > 0)
+							{
+								idx = AllText.IndexOf("<string>", idx);
+								if (idx > 0)
+								{
+									idx += "<string>".Length;
+									UUID = AllText.Substring(idx, AllText.IndexOf("</string>", idx) - idx);
+									Arguments.Add($"PROVISIONING_PROFILE_SPECIFIER={UUID}");
+
+									Logger.LogInformation("Extracted Provision UUID {0} from {1}", UUID, Provision);
+								}
+							}
+						}
+					}
+				}
+			}
+
+			int ReturnCode = Utils.RunLocalProcessAndLogOutput("/usr/bin/env", string.Join(" ", Arguments), Logger);
+//			DirectoryReference.Delete(XcodeProject, true);
+			return ReturnCode;
+		}
+
+		/// <summary>
 		/// 
 		/// </summary>
 		/// <param name="Platform"></param>
diff --git a/Engine/Source/Programs/UnrealBuildTool/Platform/IOS/IOSToolChain.cs b/Engine/Source/Programs/UnrealBuildTool/Platform/IOS/IOSToolChain.cs
--- a/Engine/Source/Programs/UnrealBuildTool/Platform/IOS/IOSToolChain.cs
+++ b/Engine/Source/Programs/UnrealBuildTool/Platform/IOS/IOSToolChain.cs
@@ -1,4 +1,4 @@
-// Copyright Epic Games, Inc. All Rights Reserved.
+﻿// Copyright Epic Games, Inc. All Rights Reserved.
 
 
 using System;
@@ -360,7 +360,7 @@ namespace UnrealBuildTool
 		{
 			// get the list of architectures to compile
 			string Archs =
-				UBTArchitecture == "-simulator" ? "i386" :
+				UBTArchitecture == "-simulator" ? "arm64" :
 				String.Join(",", (Configuration == CppConfiguration.Shipping) ? ProjectSettings.ShippingArchitectures : ProjectSettings.NonShippingArchitectures);
 
 			Log.TraceLogOnce("Compiling with these architectures: " + Archs);
@@ -677,7 +677,6 @@ namespace UnrealBuildTool
 				// This is not a shipping build so no need to delete the output file since symbols will not have been stripped from it.
 				LinkAction.CommandArguments = string.Format("-c \"\\\"{0}\\\" {1}\"", LinkerPath, LinkCommandArguments);
 			}
-
 			return OutputFile;
 		}
 
@@ -1233,7 +1232,14 @@ namespace UnrealBuildTool
 				OutputFiles.Add(StripCompleteFile);
 			}
 
-			if (!BinaryLinkEnvironment.bIsBuildingDLL)
+			bool bUseModernXcode = false;
+			if (OperatingSystem.IsMacOS())
+			{
+				ConfigHierarchy Ini = ConfigCache.ReadHierarchy(ConfigHierarchyType.Engine, DirectoryReference.FromFile(Target.ProjectFile), UnrealTargetPlatform.IOS);
+				Ini.TryGetValue("XcodeConfiguration", "bUseModernXcode", out bUseModernXcode);
+			}
+
+			if (!BinaryLinkEnvironment.bIsBuildingDLL && !bUseModernXcode)
 			{
 				// generate the asset catalog
 				bool bUserImagesExist = false;
@@ -1266,36 +1272,40 @@ namespace UnrealBuildTool
 				VersionNumber SdkVersion = VersionNumber.Parse(Settings.Value.IOSSDKVersion);
 
 				Dictionary<string, DirectoryReference> FrameworkNameToSourceDir = new Dictionary<string, DirectoryReference>();
-				FileReference StagedExecutablePath = GetStagedExecutablePath(Executable.Location, Target.Name);
-				DirectoryReference BundleDirectory = Target.bShouldCompileAsDLL ? Executable.Directory.Location : StagedExecutablePath.Directory;
-				foreach (UEBuildFramework Framework in BinaryLinkEnvironment.AdditionalFrameworks)
+
+				if (!bUseModernXcode)
 				{
-					if (Framework.FrameworkDirectory != null)
+					FileReference StagedExecutablePath = GetStagedExecutablePath(Executable.Location, Target.Name);
+					DirectoryReference BundleDirectory = Target.bShouldCompileAsDLL ? Executable.Directory.Location : StagedExecutablePath.Directory;
+					foreach (UEBuildFramework Framework in BinaryLinkEnvironment.AdditionalFrameworks)
 					{
-						if (!String.IsNullOrEmpty(Framework.CopyBundledAssets))
+						if (Framework.FrameworkDirectory != null)
 						{
-							// For now, this is hard coded, but we need to loop over all modules, and copy bundled assets that need it
-							DirectoryReference LocalSource = DirectoryReference.Combine(Framework.FrameworkDirectory, Framework.CopyBundledAssets);
-							string BundleName = Framework.CopyBundledAssets.Substring(Framework.CopyBundledAssets.LastIndexOf('/') + 1);
-							FrameworkNameToSourceDir[BundleName] = LocalSource;
-						}
-
-						if (Framework.bCopyFramework)
-						{
-							string FrameworkDir = Framework.FrameworkDirectory.FullName;
-							if (FrameworkDir.EndsWith(".framework"))
+							if (!String.IsNullOrEmpty(Framework.CopyBundledAssets))
 							{
-								FrameworkDir = Path.GetDirectoryName(FrameworkDir)!;
+								// For now, this is hard coded, but we need to loop over all modules, and copy bundled assets that need it
+								DirectoryReference LocalSource = DirectoryReference.Combine(Framework.FrameworkDirectory, Framework.CopyBundledAssets);
+								string BundleName = Framework.CopyBundledAssets.Substring(Framework.CopyBundledAssets.LastIndexOf('/') + 1);
+								FrameworkNameToSourceDir[BundleName] = LocalSource;
 							}
 
-							OutputFiles.Add(CopyBundleResource(new UEBuildBundleResource(new ModuleRules.BundleResource(Path.Combine(FrameworkDir, Framework.Name + ".framework"), "Frameworks")), Executable, BundleDirectory, Graph));
+							if (Framework.bCopyFramework)
+							{
+								string FrameworkDir = Framework.FrameworkDirectory.FullName;
+								if (FrameworkDir.EndsWith(".framework"))
+								{
+									FrameworkDir = Path.GetDirectoryName(FrameworkDir)!;
+								}
+
+								OutputFiles.Add(CopyBundleResource(new UEBuildBundleResource(new ModuleRules.BundleResource(Path.Combine(FrameworkDir, Framework.Name + ".framework"), "Frameworks")), Executable, BundleDirectory, Graph));
+							}
 						}
 					}
-				}
 
-				foreach (UEBuildBundleResource Resource in BinaryLinkEnvironment.AdditionalBundleResources)
-				{
-					OutputFiles.Add(CopyBundleResource(Resource, Executable, StagedExecutablePath.Directory, Graph));
+					foreach (UEBuildBundleResource Resource in BinaryLinkEnvironment.AdditionalBundleResources)
+					{
+						OutputFiles.Add(CopyBundleResource(Resource, Executable, StagedExecutablePath.Directory, Graph));
+					}
 				}
 
 				IOSPostBuildSyncTarget PostBuildSyncTarget = new IOSPostBuildSyncTarget(Target, Executable.Location, BinaryLinkEnvironment.IntermediateDirectory, UPLScripts, SdkVersion, FrameworkNameToSourceDir);
@@ -1424,7 +1434,7 @@ namespace UnrealBuildTool
 			}
 		}
 
-		internal static bool GenerateProjectFiles(FileReference? ProjectFile, string[] Arguments, ILogger Logger)
+		internal static bool GenerateProjectFiles(FileReference? ProjectFile, string[] Arguments, ILogger Logger, out DirectoryReference? XcodeProjectFile)
 		{
 			ProjectFileGenerator.bGenerateProjectFiles = true;
 			try
@@ -1436,12 +1446,20 @@ namespace UnrealBuildTool
 				PlatformProjectGenerators.RegisterPlatformProjectGenerator(UnrealTargetPlatform.TVOS, new TVOSProjectGenerator(CmdLine, Logger), Logger);
 
 				XcodeProjectFileGenerator Generator = new XcodeProjectFileGenerator(ProjectFile, CmdLine);
-				return Generator.GenerateProjectFiles(PlatformProjectGenerators, Arguments, Logger);
+				bool bSucces = Generator.GenerateProjectFiles(PlatformProjectGenerators, Arguments, Logger);
+				XcodeProjectFile = Generator.XCWorkspace;
+				return bSucces;
+			}
+			catch(Exception ex)
+			{
+				XcodeProjectFile = null;
+				Logger.LogError(ex.ToString());
 			}
 			finally
 			{
 				ProjectFileGenerator.bGenerateProjectFiles = false;
 			}
+			return false;
 		}
 
 		public static FileReference GetStagedExecutablePath(FileReference Executable, string TargetName)
@@ -1538,11 +1556,18 @@ namespace UnrealBuildTool
 
 		public static void PostBuildSync(IOSPostBuildSyncTarget Target, ILogger Logger)
 		{
+			// if xcode is building this, it will also do the Run stuff anyway, so no need to do it here as well
+			if (Environment.GetEnvironmentVariable("UE_BUILD_FROM_XCODE") == "1")
+			{
+				return;
+			}
+
 			ConfigHierarchy Ini = ConfigCache.ReadHierarchy(ConfigHierarchyType.Engine, GetActualProjectDirectory(Target.ProjectFile), UnrealTargetPlatform.IOS);
-			string? BundleID;
-			Ini.GetString("/Script/IOSRuntimeSettings.IOSRuntimeSettings", "BundleIdentifier", out BundleID!);
+			bool bUseModernXcode;
+			Ini.TryGetValue("XcodeConfiguration", "bUseModernXcode", out bUseModernXcode);
 
 			IOSProjectSettings ProjectSettings = ((IOSPlatform)UEBuildPlatform.GetBuildPlatform(Target.Platform)).ReadProjectSettings(Target.ProjectFile);
+			string? BundleID = ProjectSettings.BundleIdentifier;
 
 			bool bPerformFullAppCreation = true;
 			string PathToDsymZip = Target.OutputPath.FullName + ".dSYM.zip";
@@ -1596,6 +1621,43 @@ namespace UnrealBuildTool
 			FileReference.Copy(Target.OutputPath, StagedExecutablePath, true);
 			string RemoteShadowDirectoryMac = Target.OutputPath.Directory.FullName;
 
+			if (bUseModernXcode)
+			{
+				// generate a run-only project file for codesigning, etc
+				DirectoryReference? GeneratedProjectFile;
+				IOSExports.GenerateRunOnlyXcodeProject(Target.ProjectFile, Target.Platform, Target.bForDistribution, Logger, out GeneratedProjectFile);
+
+				// @todo - should we move this in to FinalizeAppWithXcode?
+				string ConfigName = Target.Configuration.ToString();
+				if (Target.TargetType != TargetType.Game && Target.TargetType != TargetType.Program)
+				{
+					ConfigName += " " + Target.TargetType.ToString();
+				}
+
+				string SchemeName;
+				if (AppName == "UnrealGame" || AppName == "UnrealClient")
+				{
+					if (Target.bBuildAsFramework)
+					{
+						SchemeName = "UnrealGame";
+					}
+					else
+					{
+						SchemeName = "UE5";
+					}
+				}
+				else
+				{
+					SchemeName = Target.ProjectFile!.GetFileNameWithoutExtension();
+				}
+
+				// run xcodebuild on the generated project to make the .app
+				IOSExports.FinalizeAppWithModernXcode(GeneratedProjectFile!, Target.Platform, SchemeName, ConfigName, Target.bForDistribution, Logger);
+
+				return;
+			}
+
+
 			if (Target.bCreateStubIPA || Target.bBuildAsFramework)
 			{
 				string FrameworkDerivedDataDir = Path.GetFullPath(Target.ProjectDirectory + "/Intermediate/" + (Target.Platform == UnrealTargetPlatform.IOS ? "IOS" : "TVOS") + "/DerivedData");
@@ -1781,16 +1843,17 @@ namespace UnrealBuildTool
 
 				if (!Target.bBuildAsFramework)
 				{
+					DirectoryReference? GeneratedProjectFile;
 				    if (AppName == "UnrealGame" || AppName == "UnrealClient" || Target.ProjectFile == null || Target.ProjectFile.IsUnderDirectory(Unreal.EngineDirectory))
 				    {
-					    GenerateProjectFiles(Target.ProjectFile, new string[] { "-platforms=" + (Target.Platform == UnrealTargetPlatform.IOS ? "IOS" : "TVOS"), "-NoIntellIsense", (Target.Platform == UnrealTargetPlatform.IOS ? "-iosdeployonly" : "-tvosdeployonly"), "-ignorejunk", (Target.bForDistribution ? "-distribution" : "-development"), "-bundleID=" + BundleID, "-includetemptargets", "-appname=" + AppName }, Logger);
+					    GenerateProjectFiles(Target.ProjectFile, new string[] { "-platforms=" + (Target.Platform == UnrealTargetPlatform.IOS ? "IOS" : "TVOS"), "-NoIntellIsense", (Target.Platform == UnrealTargetPlatform.IOS ? "-iosdeployonly" : "-tvosdeployonly"), "-ignorejunk", (Target.bForDistribution ? "-distribution" : "-development"), "-bundleID=" + BundleID, "-includetemptargets", "-appname=" + AppName }, Logger, out GeneratedProjectFile);
 				    }
 				    else
 				    {
-					    GenerateProjectFiles(Target.ProjectFile, new string[] { "-platforms=" + (Target.Platform == UnrealTargetPlatform.IOS ? "IOS" : "TVOS"), "-NoIntellIsense", (Target.Platform == UnrealTargetPlatform.IOS ? "-iosdeployonly" : "-tvosdeployonly"), "-ignorejunk", (Target.bForDistribution ? "-distribution" : "-development"), String.Format("-project={0}", Target.ProjectFile), "-game", "-bundleID=" + BundleID, "-includetemptargets" }, Logger);
+					    GenerateProjectFiles(Target.ProjectFile, new string[] { "-platforms=" + (Target.Platform == UnrealTargetPlatform.IOS ? "IOS" : "TVOS"), "-NoIntellIsense", (Target.Platform == UnrealTargetPlatform.IOS ? "-iosdeployonly" : "-tvosdeployonly"), "-ignorejunk", (Target.bForDistribution ? "-distribution" : "-development"), String.Format("-project={0}", Target.ProjectFile), "-game", "-bundleID=" + BundleID, "-includetemptargets" }, Logger, out GeneratedProjectFile);
 				    }
 				    // Make sure it exists
-				    if (!DirectoryReference.Exists(XcodeWorkspaceDir!))
+				    if (!DirectoryReference.Exists(XcodeWorkspaceDir!) || (GeneratedProjectFile != null && GeneratedProjectFile.ParentDirectory != XcodeWorkspaceDir))
 				    {
 					    throw new BuildException("Unable to create stub IPA; Xcode workspace not found at {0}", XcodeWorkspaceDir);
 				    }
diff --git a/Engine/Source/Programs/UnrealBuildTool/Platform/IOS/UEBuildIOS.cs b/Engine/Source/Programs/UnrealBuildTool/Platform/IOS/UEBuildIOS.cs
--- a/Engine/Source/Programs/UnrealBuildTool/Platform/IOS/UEBuildIOS.cs
+++ b/Engine/Source/Programs/UnrealBuildTool/Platform/IOS/UEBuildIOS.cs
@@ -1,4 +1,4 @@
-// Copyright Epic Games, Inc. All Rights Reserved.
+﻿// Copyright Epic Games, Inc. All Rights Reserved.
 
 using System;
 using System.Collections.Generic;
@@ -344,6 +344,18 @@ namespace UnrealBuildTool
 		[ConfigFile(ConfigHierarchyType.Game, "/Script/UnrealEd.ProjectPackagingSettings", "ForDistribution")]
 		public readonly bool bForDistribution = false;
 
+		/// <summary>
+		/// override for the app's display name if different from the project name
+		/// </summary>
+		[ConfigFile(ConfigHierarchyType.Game, "/Script/UnrealEd.ProjectPackagingSettings", "BundleName")]
+		public readonly string BundleName = "";
+
+		/// <summary>
+		/// longer display name than BundleName if needed
+		/// </summary>
+		[ConfigFile(ConfigHierarchyType.Game, "/Script/UnrealEd.ProjectPackagingSettings", "BundleDisplayName")]
+		public readonly string BundleDisplayName = "";
+
 		/// <summary>
 		/// Returns a list of all the non-shipping architectures which are supported
 		/// </summary>
diff --git a/Engine/Source/Programs/UnrealBuildTool/ProjectFiles/Xcode/XcodeProject.cs b/Engine/Source/Programs/UnrealBuildTool/ProjectFiles/Xcode/XcodeProject.cs
--- a/Engine/Source/Programs/UnrealBuildTool/ProjectFiles/Xcode/XcodeProject.cs
+++ b/Engine/Source/Programs/UnrealBuildTool/ProjectFiles/Xcode/XcodeProject.cs
@@ -107,33 +107,45 @@ namespace UnrealBuildTool.XcodeProjectXcconfig
 	{
 		// build settings that cause uniqueness
 		public string Definitions = "";
+		// @todo can we actually use this effectively with indexing other than fotced include?
 		public FileReference? PCHFile = null;
 		public bool bEnableRTTI = false;
 
 		// union of settings for all modules
-		public HashSet<string> APIDefines = new();
+		public HashSet<string> AllDefines = new() { "__INTELLISENSE__", "MONOLITHIC_BUILD=1" };
 		public HashSet<DirectoryReference> SystemIncludePaths = new();
 		public HashSet<DirectoryReference> UserIncludePaths = new();
 
 		public List<XcodeSourceFile> Files = new();
 		public List<UEBuildModuleCPP> Modules = new List<UEBuildModuleCPP>();
 
-		public void WriteXcconfigSettings(StringBuilder Content)
+		public FileReference ResponseFile;
+
+		public UnrealBatchedFiles(UnrealData UnrealData, int Index)
+		{
+			ResponseFile = FileReference.Combine(UnrealData.XcodeProjectFileLocation.ParentDirectory!, "ResponseFiles", $"{UnrealData.ProductName}{Index}.response");
+		}
+
+		public void GenerateResponseFile()
 		{
-			Content.WriteLine($"GCC_PREPROCESSOR_DEFINITIONS = __INTELLISENSE__ MONOLITHIC_BUILD=1 {Definitions} {string.Join(" ", APIDefines)}");
+			StringBuilder ResponseFileContents = new();
+			ResponseFileContents.Append("-I");
+			ResponseFileContents.AppendJoin(" -I", SystemIncludePaths.Select(x => x.FullName.Contains(' ') ? $"\"{x.FullName}\"" : x.FullName));
+			ResponseFileContents.Append(" -I");
+			ResponseFileContents.AppendJoin(" -I", UserIncludePaths.Select(x => x.FullName.Contains(' ') ? $"\"{x.FullName}\"" : x.FullName));
+
+			ResponseFileContents.Append(" -D");
+			ResponseFileContents.AppendJoin(" -D", AllDefines);
+
 			if (PCHFile != null)
 			{
-				Content.WriteLine($"GCC_PREFIX_HEADER = {PCHFile.FullName}");
+				ResponseFileContents.Append($" -include {PCHFile.FullName}");
 			}
-			Content.WriteLine($"GCC_ENABLE_CPP_RTTI = " + (bEnableRTTI ? "YES" : "NO"));
 
-			string SearchPaths = string.Join(" ", SystemIncludePaths.Select(x => x.FullName.Contains(' ') ? $"\"{x.FullName}\"" : x.FullName));
-			SearchPaths = SearchPaths.Replace("/x86_64/", "/$(ARCHS)/");
-			Content.WriteLine($"HEADER_SEARCH_PATHS = {SearchPaths}");
+			ResponseFileContents.Append(bEnableRTTI ? " -fno-rtti" : " -frtti");
 
-			SearchPaths = string.Join(" ", UserIncludePaths.Select(x => x.FullName.Contains(' ') ? $"\"{x.FullName}\"" : x.FullName));
-			SearchPaths = SearchPaths.Replace("/x86_64/", "/$(ARCHS)/");
-			Content.WriteLine($"USER_HEADER_SEARCH_PATHS = {SearchPaths}");
+			DirectoryReference.CreateDirectory(ResponseFile.Directory);
+			FileReference.WriteAllText(ResponseFile, ResponseFileContents.ToString());
 		}
 	}
 
@@ -152,7 +164,12 @@ namespace UnrealBuildTool.XcodeProjectXcconfig
 		public List<UnrealBatchedFiles> BatchedFiles = new();
  
 		public FileReference? UProjectFileLocation = null;
-		public FileReference? XcodeProjectFileLocation = null;
+		public DirectoryReference XcodeProjectFileLocation;
+		
+		public FileReference? InfoPlistLocation = null;
+		public bool bInfoPlistWasPremade = false;
+		public FileReference? EntitlementsLocation = null;
+		public bool bEntitlementsWasPremade = false;
 
 		// settings read from project configs
 		public IOSProjectSettings? IOSProjectSettings;
@@ -161,7 +178,8 @@ namespace UnrealBuildTool.XcodeProjectXcconfig
 		public TVOSProjectSettings? TVOSProjectSettings;
 		public TVOSProvisioningData? TVOSProvisioningData;
 
-		public string ProductName="";
+		public string ProductName;
+
 
 		/// <summary>
 		///  Used to mark the project for distribution (some platforms require this)
@@ -190,13 +208,37 @@ namespace UnrealBuildTool.XcodeProjectXcconfig
 
 		private XcodeProjectFile? ProjectFile;
 
+		public UnrealData(FileReference XcodeProjectFileLocation)
+		{
+			// the .xcodeproj is actually a directory
+			this.XcodeProjectFileLocation = new DirectoryReference(XcodeProjectFileLocation.FullName);
+			ProductName = XcodeProjectFileLocation.GetFileNameWithoutAnyExtensions();
+		}
+
+		public FileReference? FindUProjectFileLocation(XcodeProjectFile ProjectFile)
+		{
+			// find a uproject file (UE5 target won't have one)
+			foreach (Project Target in ProjectFile.ProjectTargets)
+			{
+				if (Target.UnrealProjectFilePath != null)
+				{
+					UProjectFileLocation = Target.UnrealProjectFilePath;
+					break;
+				}
+			}
+
+			// now that we have a UProject file (or not), update the FileCollection RootDirectory to point to it
+			ProjectFile.FileCollection.RootDirectory = UProjectFileLocation == null ? (Unreal.EngineDirectory) : UProjectFileLocation.Directory;
+
+			return UProjectFileLocation;
+		}
+
 		public bool Initialize(XcodeProjectFile ProjectFile, List<UnrealTargetPlatform> Platforms, List<UnrealTargetConfiguration> Configurations, ILogger Logger)
 		{
 			this.ProjectFile = ProjectFile;
 			this.Logger = Logger;
-			ProductName = ProjectFile.ProjectFilePath.GetFileNameWithoutAnyExtensions();
-			XcodeProjectFileLocation = ProjectFile.ProjectFilePath;
 
+			FindUProjectFileLocation(ProjectFile);
 
 			// Figure out all the desired configurations on the unreal side
 			AllConfigs = GetSupportedBuildConfigs(Platforms, Configurations, Logger);
@@ -207,16 +249,6 @@ namespace UnrealBuildTool.XcodeProjectXcconfig
 			}
 
 
-			// find a uproject file (UE5 target won't have one)
-			foreach (Project Target in ProjectFile.ProjectTargets)
-			{
-				if (Target.UnrealProjectFilePath != null)
-				{
-					UProjectFileLocation = Target.UnrealProjectFilePath;
-					break;
-				}
-			}
-
 			// this project makes an app bundle (.app directory instead of a raw executable or dylib) if none of the fings make a non-appbundle
 			bIsAppBundle = !AllConfigs.Any(x => x.ProjectTarget!.TargetRules!.bIsBuildingConsoleApplication || x.ProjectTarget.TargetRules.bShouldCompileAsDLL);
 			bHasEditorConfiguration = AllConfigs.Any(x => x.ProjectTarget!.TargetRules!.Type == TargetType.Editor);
@@ -226,8 +258,11 @@ namespace UnrealBuildTool.XcodeProjectXcconfig
 			{
 				IOSPlatform IOSPlatform = ((IOSPlatform)UEBuildPlatform.GetBuildPlatform(UnrealTargetPlatform.IOS));
 				IOSProjectSettings = IOSPlatform.ReadProjectSettings(UProjectFileLocation);
-				IOSProvisioningData = IOSPlatform.ReadProvisioningData(IOSProjectSettings, bForDistribution);
 				bUseAutomaticSigning |= IOSProjectSettings.bAutomaticSigning;
+				if (!bUseAutomaticSigning)
+				{
+					IOSProvisioningData = IOSPlatform.ReadProvisioningData(IOSProjectSettings, bForDistribution);
+				}
 			}
 
 
@@ -235,8 +270,11 @@ namespace UnrealBuildTool.XcodeProjectXcconfig
 			{
 				TVOSPlatform TVOSPlatform = ((TVOSPlatform)UEBuildPlatform.GetBuildPlatform(UnrealTargetPlatform.TVOS));
 				TVOSProjectSettings = TVOSPlatform.ReadProjectSettings(UProjectFileLocation);
-				TVOSProvisioningData = TVOSPlatform.ReadProvisioningData(TVOSProjectSettings, bForDistribution);
 				bUseAutomaticSigning |= TVOSProjectSettings.bAutomaticSigning;
+				if (!bUseAutomaticSigning)
+				{
+					TVOSProvisioningData = TVOSPlatform.ReadProvisioningData(TVOSProjectSettings, bForDistribution);
+				}
 			}
 
 			return true;
@@ -257,7 +295,12 @@ namespace UnrealBuildTool.XcodeProjectXcconfig
 			{
 				// if there are no ForceInclude files, then that means it's a module that forces the includes to come from a generated PCH file
 				// and so we will use this for definitions and uniqueness
-				PCHFile = CompileEnvironment.PrecompiledHeaderIncludeFilename;
+				if (CompileEnvironment.PrecompiledHeaderIncludeFilename != null)
+				{
+					PCHFile = FileReference.Combine(XcodeProjectFileLocation.ParentDirectory!, "PCHFiles", CompileEnvironment.PrecompiledHeaderIncludeFilename.GetFileName());
+					DirectoryReference.CreateDirectory(PCHFile.Directory);
+					FileReference.Copy(CompileEnvironment.PrecompiledHeaderIncludeFilename, PCHFile, true);
+				}
 			}
 			else
 			{
@@ -322,8 +365,7 @@ namespace UnrealBuildTool.XcodeProjectXcconfig
 			}
 			if (FileBatch == null)
 			{
-				//FileBatch = new UnrealBatchedFiles(UnrealData, SubTargets.Count + 1);
-				FileBatch = new UnrealBatchedFiles();
+				FileBatch = new UnrealBatchedFiles(this, BatchedFiles.Count + 1);
 				BatchedFiles.Add(FileBatch);
 				FileBatch.Definitions = DefinesString;
 				FileBatch.PCHFile = PCHFile;
@@ -331,35 +373,12 @@ namespace UnrealBuildTool.XcodeProjectXcconfig
 			}
 
 			FileBatch.Modules.Add(Module);
-			FileBatch.APIDefines.UnionWith(APIDefines);
+			FileBatch.AllDefines.UnionWith(Defines);
+			FileBatch.AllDefines.UnionWith(APIDefines);
 			FileBatch.SystemIncludePaths.UnionWith(CompileEnvironment.SystemIncludePaths);
 			FileBatch.UserIncludePaths.UnionWith(CompileEnvironment.UserIncludePaths);
 		}
 
-		public void FilterBuildableFiles(XcodeFileCollection FileCollection)
-		{
-			foreach (XcodeSourceFile File in FileCollection.BuildableFiles)
-			{
-				AddFileToBatch(File);
-			}
-		}
-
-		internal void AddFileToBatch(XcodeSourceFile File)
-		{
-			foreach (UnrealBatchedFiles Batch in BatchedFiles)
-			{
-				foreach (UEBuildModuleCPP Module in Batch.Modules)
-				{
-					if (Module.ContainsFile(File.Reference))
-					{
-						Batch.Files.Add(File);
-						return;
-					}
-				}
-			}
-		}
-
-
 		private List<UnrealBuildConfig> GetSupportedBuildConfigs(List<UnrealTargetPlatform> Platforms, List<UnrealTargetConfiguration> Configurations, ILogger Logger)
 		{
 			List<UnrealBuildConfig> BuildConfigs = new List<UnrealBuildConfig>();
@@ -791,7 +810,7 @@ namespace UnrealBuildTool.XcodeProjectXcconfig
 
 		public void CreateXcconfigFile(XcodeProject Project, string Name)
 		{
-			Xcconfig = new XcconfigFile(Project.UnrealData.XcodeProjectFileLocation!.Directory, Name);
+			Xcconfig = new XcconfigFile(Project.UnrealData.XcodeProjectFileLocation.ParentDirectory!, Name);
 			Project.FileCollection.AddFileReference(Xcconfig.Guid, $"Xcconfigs/{Xcconfig.FileRef.GetFileName()}", "test.xcconfig", "\"<group>\"", "Xcconfigs");
 		}
 
@@ -859,15 +878,42 @@ namespace UnrealBuildTool.XcodeProjectXcconfig
 		}
 	}
 
-	class XcodeBuildPhase : XcodeProjectNode
+	abstract class XcodeBuildPhase : XcodeProjectNode
 	{
 		public string Name;
+		public string IsAType;
 		public string Guid = XcodeProjectFileGenerator.MakeXcodeGuid();
-		private List<XcodeSourceFile> Items = new();
+		protected List<XcodeSourceFile> Items = new();
 
-		public XcodeBuildPhase(string Name)
+		public XcodeBuildPhase(string Name, string IsAType)
 		{
 			this.Name = Name;
+			this.IsAType = IsAType;
+		}
+
+		public override void Write(StringBuilder Content)
+		{
+			Content.WriteLine($"/* Begin {IsAType} section */");
+			Content.WriteLine(2, $"{Guid} = {{");
+			Content.WriteLine(3,	$"isa = {IsAType};");
+			Content.WriteLine(3,	"buildActionMask = 2147483647;");
+			Content.WriteLine(3,	"files = (");
+			foreach (XcodeSourceFile File in Items)
+			{
+				Content.WriteLine(4,		$"{File.FileGuid} /* {File.Reference.GetFileName()} in {Name} */,");
+			}
+			Content.WriteLine(3, ");");
+			Content.WriteLine(3, "runOnlyForDeploymentPostprocessing = 0;");
+			Content.WriteLine(2, "};");
+			Content.WriteLine($"/* End {IsAType} section */");
+		}
+	}
+
+	class XcodeSourcesBuildPhase : XcodeBuildPhase
+	{
+		public XcodeSourcesBuildPhase()
+			: base("Sources", "PBXSourcesBuildPhase")
+		{
 		}
 
 		public void AddFile(XcodeSourceFile File)
@@ -893,6 +939,27 @@ namespace UnrealBuildTool.XcodeProjectXcconfig
 		}
 	}
 
+	class XcodeResourcesBuildPhase : XcodeBuildPhase
+	{
+		private XcodeFileCollection FileCollection;
+
+		public XcodeResourcesBuildPhase(XcodeFileCollection FileCollection)
+			: base("Resources", "PBXResourcesBuildPhase")
+		{
+			this.FileCollection = FileCollection;
+		}
+
+		public void AddResource(FileReference Resource)
+		{
+			XcodeSourceFile ResourceSource = new XcodeSourceFile(Resource, null);
+			FileCollection.ProcessFile(ResourceSource, true);
+
+			Items.Add(ResourceSource);
+		}
+	}
+
+
+
 	class XcodeBuildConfig : XcodeProjectNode
 	{
 		public string Guid = XcodeProjectFileGenerator.MakeXcodeGuid();
@@ -909,7 +976,7 @@ namespace UnrealBuildTool.XcodeProjectXcconfig
 			Content.WriteLine(3,	"isa = XCBuildConfiguration;");
 			if (Xcconfig != null)
 			{
-				Content.WriteLine(3, $"baseConfigurationReference = {Xcconfig.Guid} /* {Xcconfig.Name}.xcconfig */;");
+				Content.WriteLine(3,	$"baseConfigurationReference = {Xcconfig.Guid} /* {Xcconfig.Name}.xcconfig */;");
 			}
 			Content.WriteLine(3,	"buildSettings = {");
 			Content.WriteLine(3,	"};");
@@ -1091,7 +1158,7 @@ namespace UnrealBuildTool.XcodeProjectXcconfig
 	class XcodeRunTarget : XcodeTarget
 	{
 		private string ProductGuid = XcodeProjectFileGenerator.MakeXcodeGuid();
-		private string ProductName;
+		private UnrealData UnrealData;
 
 		public XcodeRunTarget(XcodeProject Project)
 			: base(Project.UnrealData.bIsAppBundle ? XcodeTarget.Type.Run_App : XcodeTarget.Type.Run_Tool, Project.UnrealData)
@@ -1100,8 +1167,20 @@ namespace UnrealBuildTool.XcodeProjectXcconfig
 			References.Add(BuildConfigList);
 
 			// add the Product item to the project to be visible in left pane
-			ProductName = Project.UnrealData.ProductName;
-			Project.FileCollection.AddFileReference(ProductGuid, ProductName, Project.UnrealData.bIsAppBundle ? "wrapper.application" : "\"compiled.mach-o.executable\"", "BUILT_PRODUCTS_DIR", "Products");
+			UnrealData = Project.UnrealData;
+			Project.FileCollection.AddFileReference(ProductGuid, UnrealData.ProductName, Project.UnrealData.bIsAppBundle ? "wrapper.application" : "\"compiled.mach-o.executable\"", "BUILT_PRODUCTS_DIR", "Products");
+
+			// look for Assets
+			DirectoryReference Assets = DirectoryReference.Combine(Project.FileCollection.RootDirectory, "Build/IOS/Resources/Assets.xcassets");
+			if (!DirectoryReference.Exists(Assets))
+			{
+				Assets = DirectoryReference.Combine(Unreal.EngineDirectory, "Build/IOS/Resources/Assets.xcassets");
+			}
+			XcodeResourcesBuildPhase ResourcesBuildPhase = new XcodeResourcesBuildPhase(Project.FileCollection);
+			BuildPhases.Add(ResourcesBuildPhase);
+			References.Add(ResourcesBuildPhase);
+
+			ResourcesBuildPhase.AddResource(new FileReference(Assets.FullName));
 
 
 			// if we are run-only, then skip some stuff
@@ -1124,13 +1203,13 @@ namespace UnrealBuildTool.XcodeProjectXcconfig
 		protected override void WriteExtraTargetProperties(StringBuilder Content)
 		{
 			Content.WriteLine($"\t\t\tproductReference = {ProductGuid};");
-			Content.WriteLine($"\t\t\tproductName = \"{ProductName}\";");
+			Content.WriteLine($"\t\t\tproductName = \"{UnrealData.ProductName}\";");
 		}
 
 		public override void WriteXcconfigFile()
 		{
-			// also write to the Xcconfig file
-			Xcconfig!.AppendLine("GENERATE_INFOPLIST_FILE[sdk=macosx*] = YES // Xcode 14 code-signing fix for development");
+			// allow Xcode to generate the final plist file from our input, some INFOPLIST settings and other settings 
+			Xcconfig!.AppendLine("GENERATE_INFOPLIST_FILE = YES");
 
 			// #jira UE-143619: Pre Monterey macOS requires this option for a packaged app to run on iOS15 due to new code signature format. Could be removed once Monterey is miniuS.
 			Xcconfig.AppendLine("OTHER_CODE_SIGN_FLAGS = --generate-entitlement-der");
@@ -1138,6 +1217,10 @@ namespace UnrealBuildTool.XcodeProjectXcconfig
 			Xcconfig.AppendLine($"MACOSX_DEPLOYMENT_TARGET = {MacToolChain.Settings.MacOSVersion};");
 			Xcconfig.AppendLine("INFOPLIST_OUTPUT_FORMAT = xml");
 			Xcconfig.AppendLine("COMBINE_HIDPI_IMAGES = YES");
+			Xcconfig.AppendLine($"PRODUCT_BUNDLE_IDENTIFIER = {UnrealData.BundleIdentifier}");
+			Xcconfig.AppendLine($"ASSETCATALOG_COMPILER_APPICON_NAME = AppIcon");
+			
+
 
 			Xcconfig.Write();
 		}
@@ -1155,15 +1238,15 @@ namespace UnrealBuildTool.XcodeProjectXcconfig
 
 	class XcodeIndexTarget : XcodeTarget
 	{
-		private UnrealBatchedFiles? SoloTargetBatch;
+		private UnrealData UnrealData;
 
 		// just take the Project since it has everything we need, and is needed when adding target dependencies
 		public XcodeIndexTarget(XcodeProject Project)
 			: base(XcodeTarget.Type.Index, Project.UnrealData)
 		{
-			// do we have an editor? if so, use this Development Editor as the config to index, otherwise, use Development
-			string IndexConfig = Project.UnrealData.AllConfigs.Any(x => x.DisplayName == "Development Editor") ? "Development Editor" : "Development";
-			BuildConfigList = new XcodeBuildConfigList(Name, Project.UnrealData.AllConfigs.Where(x => x.DisplayName == IndexConfig).ToList());
+			UnrealData = Project.UnrealData;
+
+			BuildConfigList = new XcodeBuildConfigList(Name, UnrealData.AllConfigs);
 			References.Add(BuildConfigList);
 
 			CreateXcconfigFile(Project, Name);
@@ -1171,30 +1254,17 @@ namespace UnrealBuildTool.XcodeProjectXcconfig
 			BuildConfigList!.BuildConfigs.ForEach(x => x.Xcconfig = Xcconfig);
 
 
-			// now look to see if all files are buildable with one batch, in which case we don't need subtargets, or if we have multiple batches, then 
-			// make a subtarget for each batch of unique build settings
-			if (Project.UnrealData.BatchedFiles.Count == 1)
-			{
-				SoloTargetBatch = Project.UnrealData.BatchedFiles[0];
-
-				XcodeBuildPhase BuildPhase = new XcodeBuildPhase("Sources");
-				BuildPhases.Add(BuildPhase);
-				References.Add(BuildPhase);
+			// add all of the files to be natively compiled by this target
+			XcodeSourcesBuildPhase SourcesBuildPhase = new XcodeSourcesBuildPhase();
+			BuildPhases.Add(SourcesBuildPhase);
+			References.Add(SourcesBuildPhase);
 
-				foreach (XcodeSourceFile File in Project.FileCollection.BuildableFiles)
-				{
-					BuildPhase.AddFile(File);
-				}
-			}
-			else
+			foreach (KeyValuePair<XcodeSourceFile, FileReference?> Pair in Project.FileCollection.BuildableFilesToResponseFile)
 			{
-				// find the config infos this target is set to use, and copy those in to the subtargets
-				List<UnrealBuildConfig> Configs = BuildConfigList.BuildConfigs.Select(x => x.Info).ToList();
-
-				for (int Index = 0; Index < Project.UnrealData.BatchedFiles.Count; Index++)
+				// only add files that found a moduleto be part of (since we can't build without the build settings that come from a module)
+				if (Pair.Value != null)
 				{
-					XcodeIndexSubTarget SubTarget = new XcodeIndexSubTarget(Project, Index, Configs);
-					AddDependency(SubTarget, Project);
+					SourcesBuildPhase.AddFile(Pair.Key);
 				}
 			}
 		}
@@ -1207,60 +1277,8 @@ namespace UnrealBuildTool.XcodeProjectXcconfig
 			Xcconfig.AppendLine("GCC_WARN_CHECK_SWITCH_STATEMENTS = NO");
 			Xcconfig.AppendLine("GCC_PRECOMPILE_PREFIX_HEADER = YES");
 			Xcconfig.AppendLine("GCC_OPTIMIZATION_LEVEL = 0");
-
-			// if we have folded a file batch into this main target, write out the settings into this Xcconfig
-			if (SoloTargetBatch != null)
-			{
-				// ask the single sub target to write out the contents of our Xcconfig file since we are folding them up into this target
-				SoloTargetBatch.WriteXcconfigSettings(Xcconfig.Text);
-			}
-
-			Xcconfig.Write();
-		}
-	}
-
-	class XcodeIndexSubTarget : XcodeTarget
-	{
-		public string TargetGuid = XcodeProjectFileGenerator.MakeXcodeGuid();
-		public string BuildPhaseGuid = XcodeProjectFileGenerator.MakeXcodeGuid();
-		public string ConfigListGuid = XcodeProjectFileGenerator.MakeXcodeGuid();
-
-		private UnrealBatchedFiles BatchedFiles;
-
-		public XcodeIndexSubTarget(XcodeProject Project, int SubTargetIndex, List<UnrealBuildConfig> Configs)
-			: base(Type.Index, Project.UnrealData, $"SubIndex{SubTargetIndex}")
-		{
-			// get the settings and list of files for this subtarget
-			BatchedFiles = Project.UnrealData.BatchedFiles[SubTargetIndex];
-
-			// set the buildconfig list in the SubTarget
-			BuildConfigList = new XcodeBuildConfigList(Name, Configs);
-			References.Add(BuildConfigList);
-
-			CreateXcconfigFile(Project, Name);
-			// hook up each buildconfig to this Xcconfig 
-			BuildConfigList!.BuildConfigs.ForEach(x => x.Xcconfig = Xcconfig);
-
-			// create the build phase for this target for all files in the batch
-			XcodeBuildPhase BuildPhase = new XcodeBuildPhase("Sources");
-			BuildPhases.Add(BuildPhase);
-			References.Add(BuildPhase);
-
-			// and add each file to it
-			foreach (XcodeSourceFile File in BatchedFiles.Files)
-			{
-				BuildPhases[0].AddFile(File);
-			}
-		}
-
-		public override void WriteXcconfigFile()
-		{
-			// include the share Index config
-			Xcconfig!.AppendLine($"#include \"{Name.Split('_')[0]}_Index.xcconfig\"");
-
-			// then right the build settings from the modules we are batched together
-			BatchedFiles.WriteXcconfigSettings(Xcconfig.Text);
-
+			Xcconfig.AppendLine($"PRODUCT_NAME = {Name}");
+			Xcconfig.AppendLine("CONFIGURATION_BUILD_DIR = build");
 			Xcconfig.Write();
 		}
 	}
@@ -1350,113 +1368,212 @@ namespace UnrealBuildTool.XcodeProjectXcconfig
 			Content.WriteLine("/* End PBXProject section */");
 		}
 
-
-		public override void WriteXcconfigFile()
+		private List<string> GetSupportedOrientations(ConfigHierarchy Ini)
 		{
-			string UEDir = XcodeFileCollection.ConvertPath(Path.GetFullPath(Directory.GetCurrentDirectory() + "../../.."));
-			string SupportedIOSArchitectures = string.Join(" ", UnrealData.SupportedIOSArchitectures);
-			string? IOSRunTimeVersion = null;
-			string? TVOSRunTimeVersion = null;
-
-			// shortcut for mac only
-			string? IOSRunTimeDevices = null;
-			string? TVOSRunTimeDevices = null;
-
-			bool bAutomaticSigning = false;
-			string? UUID_IOS = "";
-			string? UUID_TVOS = "";
-			string? TEAM_IOS = "";
-			string? TEAM_TVOS = "";
-			string? IOS_CERT = "iPhone Developer";
-			string? TVOS_CERT = "iPhone Developer";
-			string IOS_BUNDLE = "";
-			string TVOS_BUNDLE = "";
+			List<string> Orientations = new();
 
-			if (UnrealData.IOSProjectSettings != null)
+			bool bSupported = true;
+			if (Ini.TryGetValue("/Script/IOSRuntimeSettings.IOSRuntimeSettings", "bSupportsPortraitOrientation", out bSupported) && bSupported)
 			{
-				IOSRunTimeVersion = UnrealData.IOSProjectSettings.RuntimeVersion;
-				IOSRunTimeDevices = UnrealData.IOSProjectSettings.RuntimeDevices;
-				bAutomaticSigning = UnrealData.IOSProjectSettings.bAutomaticSigning;
-				if (!bAutomaticSigning)
-				{
-					UUID_IOS = UnrealData.IOSProvisioningData!.MobileProvisionUUID;
-					IOS_CERT = UnrealData.IOSProvisioningData!.SigningCertificate;
-				}
-				TEAM_IOS = UnrealData.IOSProvisioningData!.TeamUUID;
-				IOS_BUNDLE = UnrealData.IOSProjectSettings.BundleIdentifier;
+				Orientations.Add("UIInterfaceOrientationPortrait");
+			}
+			if (Ini.TryGetValue("/Script/IOSRuntimeSettings.IOSRuntimeSettings", "bSupportsUpsideDownOrientation", out bSupported) && bSupported)
+			{
+				Orientations.Add("UIInterfaceOrientationPortraitUpsideDown");
 			}
 
-			if (UnrealData.TVOSProjectSettings != null)
+			string? PreferredLandscapeOrientation;
+			Ini.TryGetValue("/Script/IOSRuntimeSettings.IOSRuntimeSettings", "PreferredLandscapeOrientation", out PreferredLandscapeOrientation);
+			bool bSupportsLandscapeLeft = false;
+			Ini.TryGetValue("/Script/IOSRuntimeSettings.IOSRuntimeSettings", "bSupportsLandscapeLeftOrientation", out bSupportsLandscapeLeft);
+			bool bSupportsLandscapeRight = false;
+			Ini.TryGetValue("/Script/IOSRuntimeSettings.IOSRuntimeSettings", "bSupportsLandscapeRightOrientation", out bSupportsLandscapeRight);
+
+			if (bSupportsLandscapeLeft && PreferredLandscapeOrientation == "LandscapeLeft")
 			{
-				TVOSRunTimeVersion = UnrealData.TVOSProjectSettings.RuntimeVersion;
-				TVOSRunTimeDevices = UnrealData.TVOSProjectSettings.RuntimeDevices;
-				if (!bAutomaticSigning)
-				{
-					UUID_TVOS = UnrealData.TVOSProvisioningData!.MobileProvisionUUID;
-					TVOS_CERT = UnrealData.TVOSProvisioningData!.SigningCertificate;
-				}
-				TEAM_TVOS = UnrealData.TVOSProvisioningData!.TeamUUID;
-				TVOS_BUNDLE = UnrealData.TVOSProjectSettings.BundleIdentifier;
+				Orientations.Add("UIInterfaceOrientationLandscapeLeft");
+			}
+			if (bSupportsLandscapeRight)
+			{
+				Orientations.Add("UIInterfaceOrientationLandscapeRight");
+			}
+			if (bSupportsLandscapeLeft && PreferredLandscapeOrientation != "LandscapeLeft")
+			{
+				Orientations.Add("UIInterfaceOrientationLandscapeLeft");
 			}
 
-			//#jira UE-50382 Xcode Address Sanitizer feature does not work on iOS
-			// address sanitizer dylib loader depends on the SDKROOT parameter. For macosx or default (missing, translated as macosx), the path is incorrect for iphone/appletv
-			Xcconfig!.AppendLine($"SDKROOT[sdk=macosx*] = macosx");
-			Xcconfig.AppendLine($"SDKROOT[sdk=iphone*] = iphoneos");
-			Xcconfig.AppendLine($"SDKROOT[sdk=appletvos*] = appletvos");
-			Xcconfig.AppendLine("USE_HEADERMAP = NO");
-			Xcconfig.AppendLine("ONLY_ACTIVE_ARCH = YES");
+			return Orientations;
+		}
 
-			if (bAutomaticSigning && UnrealData.bWriteCodeSigningSettings)
+
+		public override void WriteXcconfigFile()
+		{
+			if (Xcconfig == null)
 			{
-				Xcconfig.AppendLine("CODE_SIGN_STYLE = Automatic");
+				return;
 			}
 
-			if (IOSRunTimeVersion != null)
+			string UEDir = XcodeFileCollection.ConvertPath(Path.GetFullPath(Directory.GetCurrentDirectory() + "../../.."));
+
+			ConfigHierarchy SharedPlatformIni = ConfigCache.ReadHierarchy(ConfigHierarchyType.Engine, UnrealData.UProjectFileLocation?.Directory, UnrealTargetPlatform.Mac);
+
+			// read some settings from new config settings
+			bool bAutomaticSigning = true;
+			string? SigningTeam;
+			SharedPlatformIni.TryGetValue("XcodeConfiguration", "bUseModernCodeSigning", out bAutomaticSigning);
+			SharedPlatformIni.TryGetValue("XcodeConfiguration", "ModernSigningTeam", out SigningTeam);
+
+			List<string> SupportedDevices = new();
+
+			StringBuilder Content = new();
+			string ProjectName = UnrealData.ProductName;
+
+			// figure out the directory that contains the Binaries and Intermediate directories
+			string ProjectRootDir = (UnrealData.UProjectFileLocation == null) ? $"{UEDir}/Engine" : UnrealData.UProjectFileLocation.Directory.FullName;
+
+
+			// Info.plist generation settings
+
+			////////////////////////
+			// IOS
+			////////////////////////
+
+			if (UnrealData.IOSProjectSettings != null)
 			{
-				Xcconfig.AppendLine($"VALID_ARCHS[sdk=iphone*] = {SupportedIOSArchitectures}");
-				Xcconfig.AppendLine($"IPHONEOS_DEPLOYMENT_TARGET = {IOSRunTimeVersion}");
-				Xcconfig.AppendLine($"TARGETED_DEVICE_FAMILY[sdk=iphone*] = {IOSRunTimeDevices}");
-				if (UnrealData.bWriteCodeSigningSettings)
+				ConfigHierarchy Ini = ConfigCache.ReadHierarchy(ConfigHierarchyType.Engine, UnrealData.UProjectFileLocation?.Directory, UnrealTargetPlatform.IOS);
+
+				Xcconfig.AppendLine("// IOS Settings");
+
+				SupportedDevices.Add(UnrealData.IOSProjectSettings.RuntimeDevices);
+
+				Xcconfig.AppendLine("SDKROOT[sdk=iphoneos*] = iphoneos");
+				Xcconfig.AppendLine($"CONFIGURATION_BUILD_DIR[sdk=iphoneos*] = {ProjectRootDir}/Binaries/IOS/Payload");
+				Xcconfig.AppendLine($"IPHONEOS_DEPLOYMENT_TARGET = {UnrealData.IOSProjectSettings.RuntimeVersion}");
+				Xcconfig.AppendLine($"CURRENT_PROJECT_VERSION[sdk=iphoneos*] = {IOSExports.GetAndUpdateVersionFile(UnrealData.UProjectFileLocation, UnrealTargetPlatform.IOS)}");
+
+				// short version string
+				string BundleShortVersion;
+				Ini.GetString("/Script/IOSRuntimeSettings.IOSRuntimeSettings", "VersionInfo", out BundleShortVersion);
+				Xcconfig.AppendLine($"MARKETING_VERSION[sdk=iphoneos*] = {BundleShortVersion}");
+				
+
+				if (UnrealData.IOSProvisioningData != null && UnrealData.bWriteCodeSigningSettings)
 				{
-					if (!string.IsNullOrEmpty(TEAM_IOS))
+					if (!string.IsNullOrEmpty(UnrealData.IOSProvisioningData.MobileProvisionUUID))
 					{
-						Xcconfig.AppendLine($"DEVELOPMENT_TEAM[sdk=iphone*] = {TEAM_IOS}");
-					}
-					Xcconfig.AppendLine($"CODE_SIGN_IDENTITY[sdk=iphone*] = {IOS_CERT}");
-					if (!bAutomaticSigning && !string.IsNullOrEmpty(UUID_IOS))
-					{
-						Xcconfig.AppendLine($"PROVISIONING_PROFILE_SPECIFIER[sdk=iphone*] = {UUID_IOS}");
+						Xcconfig.AppendLine($"PROVISIONING_PROFILE_SPECIFIER[sdk=iphoneos*] = {UnrealData.IOSProvisioningData.MobileProvisionUUID}");
 					}
+					string IOSCert = UnrealData.IOSProvisioningData.SigningCertificate ?? (UnrealData.bForDistribution ? "Apple Distribution" : "Apple Developer");
+					Xcconfig.AppendLine($"CODE_SIGN_IDENTITY[sdk=iphoneos*] = {IOSCert}");
+				}
+
+
+				// Info plist generation settings
+				Xcconfig.AppendLine("INFOPLIST_KEY_UIStatusBarHidden = YES");
+				Xcconfig.AppendLine("INFOPLIST_KEY_UIRequiresFullScreen = YES");
+
+				List<string> SupportedOrientations = GetSupportedOrientations(Ini);
+				Content.WriteLine(3, $"INFOPLIST_KEY_UISupportedInterfaceOrientations = \"{string.Join(" ", SupportedOrientations)}\"");
+
+				if (UnrealData.IOSProjectSettings.BundleName != "" && UnrealData.IOSProjectSettings.BundleName != "[PROJECT_NAME]")
+				{
+					string DisplayName = UnrealData.IOSProjectSettings.BundleDisplayName.Replace("[PROJECT_NAME]", ProjectName).Replace("_", "");
+					Content.WriteLine(3, $"INFOPLIST_KEY_CFBundleName = {DisplayName}");
 				}
-				if (UnrealData.UProjectFileLocation != null)
+
+				if (UnrealData.IOSProjectSettings.BundleDisplayName != "" && UnrealData.IOSProjectSettings.BundleDisplayName != "[PROJECT_NAME]")
 				{
-					Xcconfig.AppendLine($"PRODUCT_BUNDLE_IDENTIFIER[sdk=iphone*] = {IOS_BUNDLE}");
+					string DisplayName = UnrealData.IOSProjectSettings.BundleDisplayName.Replace("[PROJECT_NAME]", ProjectName).Replace("_", "");
+					Content.WriteLine(3, $"INFOPLIST_KEY_CFBundleDisplayName = {DisplayName}");
 				}
+
 			}
-			if (TVOSRunTimeVersion != null)
+
+			////////////////////////
+			// TVOS
+			////////////////////////
+
+			if (UnrealData.TVOSProjectSettings != null)
 			{
-				Xcconfig.AppendLine($"VALID_ARCHS[sdk=appletvos*] = {SupportedIOSArchitectures}");
-				Xcconfig.AppendLine($"TVOS_DEPLOYMENT_TARGET = {TVOSRunTimeVersion}");
-				Xcconfig.AppendLine($"TARGETED_DEVICE_FAMILY[sdk=appletvos*] = {TVOSRunTimeDevices}");
-				if (UnrealData.bWriteCodeSigningSettings)
+				ConfigHierarchy Ini = ConfigCache.ReadHierarchy(ConfigHierarchyType.Engine, UnrealData.UProjectFileLocation?.Directory, UnrealTargetPlatform.TVOS);
+
+				Xcconfig.AppendLine("");
+				Xcconfig.AppendLine("// TVOS Settings");
+
+				SupportedDevices.Add(UnrealData.TVOSProjectSettings.RuntimeDevices);
+
+				Xcconfig.AppendLine("SDKROOT[sdk=appletvos*] = appletvos");
+				Xcconfig.AppendLine($"CONFIGURATION_BUILD_DIR[sdk=appletvos*] = {ProjectRootDir}/Binaries/TVOS/Payload");
+				Xcconfig.AppendLine($"TVOS_DEPLOYMENT_TARGET = {UnrealData.TVOSProjectSettings.RuntimeVersion}");
+				Xcconfig.AppendLine($"CURRENT_PROJECT_VERSION[sdk=appletvos*] = {IOSExports.GetAndUpdateVersionFile(UnrealData.UProjectFileLocation, UnrealTargetPlatform.TVOS)}");
+
+				// short version string
+				string BundleShortVersion;
+				Ini.GetString("/Script/IOSRuntimeSettings.IOSRuntimeSettings", "VersionInfo", out BundleShortVersion);
+				Xcconfig.AppendLine($"MARKETING_VERSION[sdk=appletvos*] = {BundleShortVersion}");
+
+				if (UnrealData.TVOSProvisioningData != null && UnrealData.bWriteCodeSigningSettings)
 				{
-					if (!string.IsNullOrEmpty(TEAM_TVOS))
+					if (!string.IsNullOrEmpty(UnrealData.TVOSProvisioningData.MobileProvisionUUID))
 					{
-						Xcconfig.AppendLine($"DEVELOPMENT_TEAM[sdk=iphone*] = {TEAM_TVOS}");
-					}
-					Xcconfig.AppendLine($"CODE_SIGN_IDENTITY[sdk=appletvos*] = {TVOS_CERT}");
-					if (!bAutomaticSigning && !string.IsNullOrEmpty(UUID_TVOS))
-					{
-						Xcconfig.AppendLine($"PROVISIONING_PROFILE_SPECIFIER[sdk=appletvos*] = {UUID_TVOS}");
+						Xcconfig.AppendLine($"PROVISIONING_PROFILE_SPECIFIER[sdk=iphoneos*] = {UnrealData.TVOSProvisioningData.MobileProvisionUUID}");
 					}
+					string TVOSCert = UnrealData.TVOSProvisioningData.SigningCertificate ?? (UnrealData.bForDistribution ? "Apple Distribution" : "Apple Developer");
+					Xcconfig.AppendLine($"CODE_SIGN_IDENTITY[sdk=iphoneos*] = {TVOSCert}");
 				}
-				if (UnrealData.UProjectFileLocation != null)
+
+
+				// Info plist generation settings
+				Xcconfig.AppendLine("INFOPLIST_KEY_UIStatusBarHidden = YES");
+				Xcconfig.AppendLine("INFOPLIST_KEY_UIRequiresFullScreen = YES");
+			}
+
+			////////////////////////
+			// MAC
+			////////////////////////
+
+			Xcconfig.AppendLine("");
+			Xcconfig.AppendLine("// Mac Settings");
+			Xcconfig.AppendLine($"SDKROOT[sdk=macosx*] = macosx");
+			Xcconfig.AppendLine($"CONFIGURATION_BUILD_DIR[sdk=macosx*] = {ProjectRootDir}/Binaries/Mac");
+
+
+			////////////////////////
+			// IOS / TVOS
+			////////////////////////
+
+
+			////////////////////////
+			// MAC / IOS / TVOS
+			////////////////////////
+
+			Xcconfig.AppendLine("");
+			Xcconfig.AppendLine("// Cross-platform settings, and defaults for anything not specified above (Xcode uses first-seen value, not last-seen)");
+			Xcconfig.AppendLine("USE_HEADERMAP = NO");
+			Xcconfig.AppendLine("ONLY_ACTIVE_ARCH = YES");
+			Xcconfig.AppendLine($"TARGETED_DEVICE_FAMILY = {string.Join(",", SupportedDevices)}");
+			
+
+			if (UnrealData.bWriteCodeSigningSettings)
+			{
+				Xcconfig.AppendLine("CODE_SIGN_STYLE = " + (bAutomaticSigning ? "Automatic" : "Manual"));
+				if (!string.IsNullOrEmpty(SigningTeam))
 				{
-					Xcconfig.AppendLine($"PRODUCT_BUNDLE_IDENTIFIER[sdk=appletvos*] = {TVOS_BUNDLE}");
+					Xcconfig.AppendLine($"DEVELOPMENT_TEAM = {SigningTeam}");
 				}
 			}
 
+			if (UnrealData.EntitlementsLocation != null)
+			{
+				string RelativeLoc = UnrealData.EntitlementsLocation.MakeRelativeTo(DirectoryReference.Combine(UnrealData.XcodeProjectFileLocation.ParentDirectory!));
+				Xcconfig.AppendLine($"CODE_SIGN_ENTITLEMENTS = {RelativeLoc}");
+			}
+			if (UnrealData.InfoPlistLocation != null)
+			{
+				string RelativeLoc = UnrealData.InfoPlistLocation.MakeRelativeTo(DirectoryReference.Combine(UnrealData.XcodeProjectFileLocation.ParentDirectory!));
+				Xcconfig.AppendLine($"INFOPLIST_FILE = {RelativeLoc}");
+			}
+
 			Xcconfig.Write();
 
 			// Now for each config write out the specific settings
@@ -1483,8 +1600,6 @@ namespace UnrealBuildTool.XcodeProjectXcconfig
 				string MacExecutableFileName = MacExecutablePath.GetFileName();
 				string SupportedMacArchitectures = string.Join(" ", UnrealData.GetSupportedMacArchitectures(Config, UnrealData.UProjectFileLocation));
 
-				// figure out the directory that contains the Binaries and Intermediate directories
-				string ProjectRootDir = (bIsUnrealGame || bIsUnrealClient || GamePath == null) ? $"{UEDir}/Engine" : GamePath;
 
 				// UnrealClient re-uses UnrealGame
 				string PlistTargetName = bIsUnrealClient ? "UnrealGame" : Config.BuildTarget;
@@ -1492,9 +1607,6 @@ namespace UnrealBuildTool.XcodeProjectXcconfig
 				string IOSInfoPlistPath = $"{ProjectRootDir}/Intermediate/IOS/{PlistTargetName}-Info.plist";
 				string TVOSInfoPlistPath = $"{ProjectRootDir}/Intermediate/TVOS/{PlistTargetName}-Info.plist";
 
-				// @todo is really necessary to only set entitlements in the non-UnrealGame, project case? what about mac entitlements?
-				bool bSetEntitlements = !bIsUnrealGame && !bIsUnrealClient && UnrealData.UProjectFileLocation != null;
-
 				string SupportedPlatforms = "";
 				if (Config.ProjectTarget!.TargetRules != null)
 				{
@@ -1525,7 +1637,7 @@ namespace UnrealBuildTool.XcodeProjectXcconfig
 
 				ConfigXcconfig.AppendLine($"UE_BUILD_TARGET_NAME = {Config.BuildTarget}");
 				ConfigXcconfig.AppendLine($"UE_BUILD_TARGET_CONFIG = {Config.BuildConfig}");
-
+				
 				// list the supported platforms, which will narrow down what other settings are active
 				ConfigXcconfig.AppendLine($"SUPPORTED_PLATFORMS = {SupportedPlatforms}");
 
@@ -1536,93 +1648,79 @@ namespace UnrealBuildTool.XcodeProjectXcconfig
 					ConfigXcconfig.AppendLine("// Mac setup");
 					ConfigXcconfig.AppendLine($"VALID_ARCHS[sdk=macosx*] = {SupportedMacArchitectures}");
 					ConfigXcconfig.AppendLine($"PRODUCT_NAME[sdk=macosx*] = {MacExecutableFileName}");
-					ConfigXcconfig.AppendLine($"CONFIGURATION_BUILD_DIR[sdk=macosx*] = {MacExecutableDir}");
-					ConfigXcconfig.AppendLine($"INFOPLIST_FILE[sdk=macosx*] = {MacInfoPlistPath}");
 				}
 
 				if (Config.bSupportsIOS)
 				{
 					ConfigXcconfig.AppendLine("");
 					ConfigXcconfig.AppendLine("// IOS setup");
-					ConfigXcconfig.AppendLine($"PRODUCT_NAME[sdk=iphone*] = {Config.BuildTarget}"); // @todo: change to Path.GetFileName(Config.IOSExecutablePath) when we stop using payload
-					ConfigXcconfig.AppendLine($"CONFIGURATION_BUILD_DIR[sdk=iphone*] = {ProjectRootDir}/Binaries/IOS/Payload");
-					ConfigXcconfig.AppendLine($"INFOPLIST_FILE[sdk=iphone*] = {IOSInfoPlistPath}");
-					if (bSetEntitlements)
-					{
-						ConfigXcconfig.AppendLine($"CODE_SIGN_ENTITLEMENTS[sdk=iphone*] = {ProjectRootDir}/Intermediate/IOS/{PlistTargetName}.entitlements");
-					}
+					ConfigXcconfig.AppendLine($"PRODUCT_NAME[sdk=iphoneos*] = {Config.BuildTarget}"); // @todo: change to Path.GetFileName(Config.IOSExecutablePath) when we stop using payload
 				}
 				if (Config.bSupportsTVOS)
 				{
 					ConfigXcconfig.AppendLine("");
 					ConfigXcconfig.AppendLine("// TVOS setup");
 					ConfigXcconfig.AppendLine($"PRODUCT_NAME[sdk=appletvos*] = {Config.BuildTarget}"); // @todo: change to Path.GetFileName(Config.IOSExecutablePath) when we stop using payload
-					ConfigXcconfig.AppendLine($"CONFIGURATION_BUILD_DIR[sdk=appletvos*] = {ProjectRootDir}/Binaries/TVOS/Payload");
-					ConfigXcconfig.AppendLine($"INFOPLIST_FILE[sdk=appletvos*] = {TVOSInfoPlistPath}");
-					if (bSetEntitlements)
-					{
-						ConfigXcconfig.AppendLine($"CODE_SIGN_ENTITLEMENTS[sdk=appletvos*] = {ProjectRootDir}/Intermediate/TVOS/{PlistTargetName}.entitlements");
-					}
 				}
 
-
-				// @todo can we remove this and use GENERATE_INFOPLIST_FILE when plist doesn't exist?
-
-				// Prepare a temp Info.plist file so Xcode has some basic info about the target immediately after opening the project.
-				// This is needed for the target to pass the settings validation before code signing. UBT will overwrite this plist file later, with proper contents.
-				if (BuildHostPlatform.Current.Platform == UnrealTargetPlatform.Mac)
-				{
-					bool bCreateMacInfoPlist = !File.Exists(MacInfoPlistPath);
-					bool bCreateIOSInfoPlist = !File.Exists(IOSInfoPlistPath) && IOSRunTimeVersion != null;
-					bool bCreateTVOSInfoPlist = !File.Exists(TVOSInfoPlistPath) && TVOSRunTimeVersion != null;
-					if (bCreateMacInfoPlist || bCreateIOSInfoPlist || bCreateTVOSInfoPlist)
-					{
-						DirectoryReference? ProjectPath = GameDir;
-						DirectoryReference EngineDir = DirectoryReference.Combine(new DirectoryReference(UEDir), "Engine");
-						string GameName = Config.BuildTarget;
-						bool bIsClient = false;
-						if (ProjectPath == null)
-						{
-							ProjectPath = EngineDir;
-						}
-						if (bIsUnrealGame)
-						{
-							ProjectPath = EngineDir;
-							GameName = "UnrealGame";
-							bIsClient = (UnrealData.AppName == "UnrealClient");
-						}
-
-						if (bCreateMacInfoPlist)
-						{
-							Directory.CreateDirectory(Path.GetDirectoryName(MacInfoPlistPath)!);
-							UEDeployMac.GeneratePList(ProjectPath.FullName, bIsUnrealGame, GameName, Config.BuildTarget, EngineDir.FullName, MacExecutableFileName);
-						}
-						if (bCreateIOSInfoPlist)
-						{
-							// get the receipt
-							FileReference ReceiptFilename;
-							if (bIsUnrealGame)
-							{
-								ReceiptFilename = TargetReceipt.GetDefaultPath(Unreal.EngineDirectory, "UnrealGame", UnrealTargetPlatform.IOS, Config.BuildConfig, "");
-							}
-							else
-							{
-								ReceiptFilename = TargetReceipt.GetDefaultPath(ProjectPath, GameName, UnrealTargetPlatform.IOS, Config.BuildConfig, "");
-							}
-							Directory.CreateDirectory(Path.GetDirectoryName(IOSInfoPlistPath)!);
-							bool bSupportPortrait, bSupportLandscape;
-							TargetReceipt? Receipt;
-							TargetReceipt.TryRead(ReceiptFilename, out Receipt);
-							bool bBuildAsFramework = UEDeployIOS.GetCompileAsDll(Receipt);
-							UEDeployIOS.GenerateIOSPList(UnrealData.UProjectFileLocation, Config.BuildConfig, ProjectPath.FullName, bIsUnrealGame, GameName, bIsClient, Config.BuildTarget, EngineDir.FullName, ProjectPath + "/Binaries/IOS/Payload", null, UnrealData.BundleIdentifier, bBuildAsFramework, UnrealData.Logger!, out bSupportPortrait, out bSupportLandscape);
-						}
-						if (bCreateTVOSInfoPlist)
-						{
-							Directory.CreateDirectory(Path.GetDirectoryName(TVOSInfoPlistPath)!);
-							UEDeployTVOS.GenerateTVOSPList(ProjectPath.FullName, bIsUnrealGame, GameName, bIsClient, Config.BuildTarget, EngineDir.FullName, ProjectPath + "/Binaries/TVOS/Payload", null, UnrealData.BundleIdentifier, UnrealData.Logger!);
-						}
-					}
-				}
+// @todo this is a reminder of plist settings we are not handling yet
+#if false
+				//// Prepare a temp Info.plist file so Xcode has some basic info about the target immediately after opening the project.
+				//// This is needed for the target to pass the settings validation before code signing. UBT will overwrite this plist file later, with proper contents.
+				//if (BuildHostPlatform.Current.Platform == UnrealTargetPlatform.Mac)
+				//{
+				//	bool bCreateMacInfoPlist = !File.Exists(MacInfoPlistPath);
+				//	bool bCreateIOSInfoPlist = !File.Exists(IOSInfoPlistPath) && IOSRunTimeVersion != null;
+				//	bool bCreateTVOSInfoPlist = !File.Exists(TVOSInfoPlistPath) && TVOSRunTimeVersion != null;
+				//	if (bCreateMacInfoPlist || bCreateIOSInfoPlist || bCreateTVOSInfoPlist)
+				//	{
+				//		DirectoryReference? ProjectPath = GameDir;
+				//		DirectoryReference EngineDir = DirectoryReference.Combine(new DirectoryReference(UEDir), "Engine");
+				//		string GameName = Config.BuildTarget;
+				//		bool bIsClient = false;
+				//		if (ProjectPath == null)
+				//		{
+				//			ProjectPath = EngineDir;
+				//		}
+				//		if (bIsUnrealGame)
+				//		{
+				//			ProjectPath = EngineDir;
+				//			GameName = "UnrealGame";
+				//			bIsClient = (UnrealData.AppName == "UnrealClient");
+				//		}
+
+				//		if (bCreateMacInfoPlist)
+				//		{
+				//			Directory.CreateDirectory(Path.GetDirectoryName(MacInfoPlistPath)!);
+				//			UEDeployMac.GeneratePList(ProjectPath.FullName, bIsUnrealGame, GameName, Config.BuildTarget, EngineDir.FullName, MacExecutableFileName);
+				//		}
+				//		if (bCreateIOSInfoPlist)
+				//		{
+				//			// get the receipt
+				//			FileReference ReceiptFilename;
+				//			if (bIsUnrealGame)
+				//			{
+				//				ReceiptFilename = TargetReceipt.GetDefaultPath(Unreal.EngineDirectory, "UnrealGame", UnrealTargetPlatform.IOS, Config.BuildConfig, "");
+				//			}
+				//			else
+				//			{
+				//				ReceiptFilename = TargetReceipt.GetDefaultPath(ProjectPath, GameName, UnrealTargetPlatform.IOS, Config.BuildConfig, "");
+				//			}
+				//			Directory.CreateDirectory(Path.GetDirectoryName(IOSInfoPlistPath)!);
+				//			bool bSupportPortrait, bSupportLandscape;
+				//			TargetReceipt? Receipt;
+				//			TargetReceipt.TryRead(ReceiptFilename, out Receipt);
+				//			bool bBuildAsFramework = UEDeployIOS.GetCompileAsDll(Receipt);
+				//			UEDeployIOS.GenerateIOSPList(UnrealData.UProjectFileLocation, Config.BuildConfig, ProjectPath.FullName, bIsUnrealGame, GameName, bIsClient, Config.BuildTarget, EngineDir.FullName, ProjectPath + "/Binaries/IOS/Payload", null, UnrealData.BundleIdentifier, bBuildAsFramework, UnrealData.Logger!, out bSupportPortrait, out bSupportLandscape);
+				//		}
+				//		if (bCreateTVOSInfoPlist)
+				//		{
+				//			Directory.CreateDirectory(Path.GetDirectoryName(TVOSInfoPlistPath)!);
+				//			UEDeployTVOS.GenerateTVOSPList(ProjectPath.FullName, bIsUnrealGame, GameName, bIsClient, Config.BuildTarget, EngineDir.FullName, ProjectPath + "/Binaries/TVOS/Payload", null, UnrealData.BundleIdentifier, UnrealData.Logger!);
+				//		}
+				//	}
+				//}
+#endif
 
 				ConfigXcconfig.Write();
 			}
@@ -1643,11 +1741,15 @@ namespace UnrealBuildTool.XcodeProjectXcconfig
 		// we need relativ paths, so this is the directory that paths are relative to
 		private DirectoryReference ProjectDirectory;
 
+		// the location of the uproject file, or the engine directory for non-uproject projects
+		public DirectoryReference RootDirectory;
+
 		public string MainGroupGuid = "";
 		//public string ProductRefGroupGuid = XcodeProjectFileGenerator.MakeXcodeGuid();
 
 		internal Dictionary<string, XcodeFileGroup> Groups = new();
-		internal List<XcodeSourceFile> BuildableFiles = new();
+		internal Dictionary<XcodeSourceFile, FileReference?> BuildableFilesToResponseFile = new();
+		internal List<XcodeSourceFile> BuildableResourceFiles = new();
 		internal List<XcodeSourceFile> AllFiles = new();
 
 		private List<ManualFileReference> ManualFiles = new();
@@ -1657,35 +1759,44 @@ namespace UnrealBuildTool.XcodeProjectXcconfig
 		public XcodeFileCollection(XcodeProjectFile ProjectFile)
 		{
 			ProjectDirectory = ProjectFile.ProjectFilePath.Directory;
+			RootDirectory = ProjectFile.UnrealData.UProjectFileLocation == null ? (Unreal.EngineDirectory) : ProjectFile.UnrealData.UProjectFileLocation.Directory;
 		}
 
 		public void ProcessFile(XcodeSourceFile File, bool bIsForBuild)
 		{
-			//string FileName = File.Reference.GetFileName();
-			//string FileExtension = Path.GetExtension(FileName);
-			//string FilePath = File.Reference.MakeRelativeTo(ProjectFile.Directory);
-			//string FilePathMac = Utils.CleanDirectorySeparators(FilePath, '/');
-			//string FileGuid = File.FileGuid;
-			//string FileRefGuid = File.FileRefGuid;
-
-			// this is the main entry into the project for the file tree on the left of Xcode
-			//			AddFileReference(FileRefGuid, FilePathMac, GetFileType(FileExtension));
-
-			AllFiles.Add(File);
-
-			// remember all buildable files - we may further filter this later when looking for SubTargets - the module the file is
-			// in may not be built, which means it cannot be indexed
-			if (bIsForBuild && ShouldIncludeFileInBuildPhaseSection(File))
+			// remember all buildable files
+			if (bIsForBuild)
 			{
-				BuildableFiles.Add(File);
+				string FileExtension = File.Reference.GetExtension();
+
+				if (IsSourceCode(FileExtension))
+				{
+					// look if it contains any of the exluded names, and is so, don't include it
+					if (!ExcludedFolders.Any(x => File.Reference.ContainsName(x, 0)))
+					{
+						// will fill in the response file later
+						BuildableFilesToResponseFile[File] = null;
+					}
+				}
+				else if (IsResourceFile(FileExtension))
+				{
+					BuildableResourceFiles.Add(File);
+				}
 			}
+			
 
 			// group the files by path
-			XcodeFileGroup? Group = FindGroupByAbsolutePath(File.Reference.Directory.FullName);
+			XcodeFileGroup? Group = FindGroupByAbsolutePath(File.Reference.Directory);
 			if (Group != null)
 			{
+				AllFiles.Add(File);
 				Group.Files.Add(File);
 			}
+			else
+			{
+				string GroupName = File.Reference.IsUnderDirectory(Unreal.EngineDirectory) ? "EngineReferences" : "ExternalReferences";
+				AddFileReference(File.FileRefGuid, File.Reference.FullName, GetFileType(File.Reference.GetExtension()), "<absolute>", "EngineReferences");
+			}
 		}
 
 		/// <summary>
@@ -1704,7 +1815,7 @@ namespace UnrealBuildTool.XcodeProjectXcconfig
 				Guid = Guid,
 				RelativePath = RelativePath,
 				FileType = FileType,
-				SourceTree = SourceTree,
+				SourceTree = SourceTree.StartsWith("<") ? $"\"{SourceTree}\"" : SourceTree,
 				GroupName = GroupName,
 			});
 
@@ -1719,9 +1830,18 @@ namespace UnrealBuildTool.XcodeProjectXcconfig
 		public void Write(StringBuilder Content)
 		{
 			Content.WriteLine("/* Begin PBXBuildFile section */");
-			foreach (XcodeSourceFile File in BuildableFiles)
+			foreach (KeyValuePair<XcodeSourceFile, FileReference?> Pair in BuildableFilesToResponseFile)
 			{
-				Content.WriteLine($"\t\t{File.FileGuid} = {{isa = PBXBuildFile; fileRef = {File.FileRefGuid}; }}; /* {File.Reference.GetFileName()} */");
+				string CompileFlags = "";
+				if (Pair.Value != null)
+				{
+					CompileFlags = $" settings = {{COMPILER_FLAGS = \"@{Pair.Value.FullName}\"; }};";
+				}
+				Content.WriteLine($"\t\t{Pair.Key.FileGuid} = {{isa = PBXBuildFile; fileRef = {Pair.Key.FileRefGuid};{CompileFlags} }}; /* {Pair.Key.Reference.GetFileName()} */");
+			}
+			foreach (XcodeSourceFile Resource in BuildableResourceFiles)
+			{
+				Content.WriteLine($"\t\t{Resource.FileGuid} = {{isa = PBXBuildFile; fileRef = {Resource.FileRefGuid}; }}; /* {Resource.Reference.GetFileName()} */");
 			}
 			Content.WriteLine("/* End PBXBuildFile section */");
 			Content.WriteLine();
@@ -1782,27 +1902,6 @@ namespace UnrealBuildTool.XcodeProjectXcconfig
 
 			}
 
-			// Products group
-			//			Content.WriteLine($"\t\t\t\t{XcodeProject.TargetAppGuid} /* {XcodeProject.TargetName} */,");
-
-			// @todo extensions
-			//foreach (UnrealExtensionInfo EI in AllExtensions)
-			//{
-			//	Content.WriteLine(string.Format("\t\t\t\t{0} /* {1} */,{2}", EI.ProductGuid, EI.Name, ProjectFileGenerator.NewLine));
-			//}
-			// Xcconfigs group
-			//Content.WriteLine($"\t\t{XcconfigsRefGroupGuid} /* XcconfigsRefGroupGuid */ = {{");
-			//	Content.WriteLine("\t\t\tisa = PBXGroup;");
-			//	Content.WriteLine("\t\t\tchildren = (");
-			//		foreach (XcconfigFile Xcconfig in XcodeProject.AllXcconfigFiles)
-			//		{
-			//			Content.WriteLine($"\t\t\t\t{Xcconfig.Guid} /* {Xcconfig.Filename} */,{ProjectFileGenerator.NewLine}");
-			//		}
-			//	Content.WriteLine("\t\t\t);");
-			//	Content.WriteLine("\t\t\tname = Xcconfigs;");
-			//	Content.WriteLine("\t\t\tsourceTree = \"<group>\";");
-			//Content.WriteLine("\t\t};");
-
 			Content.WriteLine("/* End PBXGroup section */");
 		}
 
@@ -1918,6 +2017,8 @@ namespace UnrealBuildTool.XcodeProjectXcconfig
 					return "image.png";
 				case ".icns":
 					return "image.icns";
+				case ".xcassets":
+					return "folder.assetcatalog";
 				default:
 					return "file.text";
 			}
@@ -1931,26 +2032,26 @@ namespace UnrealBuildTool.XcodeProjectXcconfig
 			return Extension == ".c" || Extension == ".cc" || Extension == ".cpp" || Extension == ".m" || Extension == ".mm";
 		}
 
-		static string[] ExcludedFolders = PlatformExports.GetExcludedFolderNames(UnrealTargetPlatform.Mac);
-		internal static bool ShouldIncludeFileInBuildPhaseSection(XcodeSourceFile SourceFile)
+		/// <summary>
+		/// Returns true if Extension is a known extension for files containing source code
+		/// </summary>
+		internal static bool IsResourceFile(string Extension)
 		{
-			string FileExtension = SourceFile.Reference.GetExtension();
-
-			if (IsSourceCode(FileExtension))
-			{
-				// look if it contains any of the exluded names, and is so, don't include it
-				return !ExcludedFolders.Any(x => SourceFile.Reference.ContainsName(x, 0));
-			}
-
-			return false;
+			return Extension == ".xcassets";
 		}
 
+		/// <summary>
+		/// Cache the folders excluded by Mac
+		/// </summary>
+		static string[] ExcludedFolders = PlatformExports.GetExcludedFolderNames(UnrealTargetPlatform.Mac);
+
 		/// <summary>
 		/// Returns a project navigator group to which the file should belong based on its path.
 		/// Creates a group tree if it doesn't exist yet.
 		/// </summary>
-		internal XcodeFileGroup? FindGroupByAbsolutePath(string AbsolutePath)
+		internal XcodeFileGroup? FindGroupByAbsolutePath(DirectoryReference DirectoryPath)
 		{
+			string AbsolutePath = DirectoryPath.FullName;
 			string[] Parts = AbsolutePath.Split(Path.DirectorySeparatorChar);
 			string CurrentPath = "/";
 			Dictionary<string, XcodeFileGroup> CurrentSubGroups = Groups;
@@ -1969,7 +2070,14 @@ namespace UnrealBuildTool.XcodeProjectXcconfig
 				XcodeFileGroup CurrentGroup;
 				if (!CurrentSubGroups.ContainsKey(CurrentPath))
 				{
-					CurrentGroup = new XcodeFileGroup(Path.GetFileName(CurrentPath), CurrentPath, CurrentPath.EndsWith(".xcassets"));
+					// we only want files under the project/engine dir, we have to make external references to other files, otherwise the grouping gets all broken
+					if (!DirectoryPath.IsUnderDirectory(RootDirectory))
+					{
+						return null;
+					}
+
+
+					CurrentGroup = new XcodeFileGroup(Path.GetFileName(CurrentPath), CurrentPath, GetFileType(Path.GetExtension(CurrentPath)).StartsWith("folder"));
 					if (Groups.Count == 0)
 					{
 						MainGroupGuid = CurrentGroup.GroupGuid;
@@ -2018,7 +2126,7 @@ namespace UnrealBuildTool.XcodeProjectXcconfig
 		public XcodeProjectFile(FileReference InitFilePath, DirectoryReference BaseDir, bool bIsForDistribution, string BundleID, string AppName)
 			: base(InitFilePath, BaseDir)
 		{
-			UnrealData = new UnrealData();
+			UnrealData = new UnrealData(InitFilePath);
 
 			UnrealData.bForDistribution = bIsForDistribution;
 			UnrealData.BundleIdentifier = BundleID;
@@ -2040,6 +2148,9 @@ namespace UnrealBuildTool.XcodeProjectXcconfig
 		/// </summary>
 		public XcodeFileCollection FileCollection;
 
+		private XcodeProjectLegacy.XcodeProjectFile? LegacyProjectFile = null;
+		private bool bHasCheckedForLegacy = false;
+
 		/// <summary>
 		/// Allocates a generator-specific source file object
 		/// </summary>
@@ -2055,8 +2166,38 @@ namespace UnrealBuildTool.XcodeProjectXcconfig
 			return new XcodeSourceFile(InitFilePath, InitProjectSubFolder);
 		}
 
+		private void ConditionalCreateLegacyProject()
+		{
+			if (!bHasCheckedForLegacy)
+			{
+				bHasCheckedForLegacy = true;
+				if (ProjectTargets.Count == 0)
+				{
+					throw new BuildException("Expected to have targets before AddModule is called");
+				}
+				FileReference? UProjectFileLocation = UnrealData.FindUProjectFileLocation(this);
+				ConfigHierarchy Ini = ConfigCache.ReadHierarchy(ConfigHierarchyType.Engine, UProjectFileLocation?.Directory, UnrealTargetPlatform.Mac);
+				bool bUseModernXcode;
+				if (!(Ini.TryGetValue("XcodeConfiguration", "bUseModernXcode", out bUseModernXcode) && bUseModernXcode))
+				{
+					LegacyProjectFile = new XcodeProjectLegacy.XcodeProjectFile(ProjectFilePath, BaseDir, UnrealData.bForDistribution, UnrealData.BundleIdentifier, UnrealData.AppName);
+					LegacyProjectFile.ProjectTargets.AddRange(ProjectTargets);
+					LegacyProjectFile.SourceFiles.AddRange(SourceFiles);
+					LegacyProjectFile.IsGeneratedProject = IsGeneratedProject;
+				}
+			}
+		}
+
 		public override void AddModule(UEBuildModuleCPP Module, CppCompileEnvironment CompileEnvironment)
 		{
+			ConditionalCreateLegacyProject();
+
+			if (LegacyProjectFile != null)
+			{
+				LegacyProjectFile.AddModule(Module, CompileEnvironment);
+				return;
+			}
+
 			UnrealData.AddModule(Module, CompileEnvironment);
 		}
 
@@ -2076,25 +2217,68 @@ namespace UnrealBuildTool.XcodeProjectXcconfig
 			// cache the main group
 			FileCollection.MainGroupGuid = XcodeFileCollection.GetRootGroupGuid(FileCollection.Groups);
 
-			UnrealData.FilterBuildableFiles(FileCollection);
-		}
+			// filter each file into the appropriate batch
+			foreach (XcodeSourceFile File in FileCollection.BuildableFilesToResponseFile.Keys)
+			{
+				AddFileToBatch(File, FileCollection);
+			}
 
+			// write out the response files for each batch now that everything is done
+			foreach (UnrealBatchedFiles Batch in UnrealData.BatchedFiles)
+			{
+				Batch.GenerateResponseFile();
+			}
+		}
 
+		private void AddFileToBatch(XcodeSourceFile File, XcodeFileCollection FileCollection)
+		{
+			foreach (UnrealBatchedFiles Batch in UnrealData.BatchedFiles)
+			{
+				foreach (UEBuildModuleCPP Module in Batch.Modules)
+				{
+					if (Module.ContainsFile(File.Reference))
+					{
+						Batch.Files.Add(File);
+						FileCollection.BuildableFilesToResponseFile[File] = Batch.ResponseFile;
+						return;
+					}
+				}
+			}
+		}
 
 
 		/// Implements Project interface
 		public override bool WriteProjectFile(List<UnrealTargetPlatform> InPlatforms, List<UnrealTargetConfiguration> InConfigurations, PlatformProjectGeneratorCollection PlatformProjectGenerators, ILogger Logger)
 		{
+			ConditionalCreateLegacyProject();
+
+			if (LegacyProjectFile != null)
+			{
+				return LegacyProjectFile.WriteProjectFile(InPlatforms, InConfigurations, PlatformProjectGenerators, Logger);
+			}
+
 			if (UnrealData.Initialize(this, InPlatforms, InConfigurations, Logger) == false)
 			{
 				// if we failed to initialize, we silently return to move on (it's not an error, it's a project with nothing to do)
 				return true;
 			}
 
+			if (string.IsNullOrEmpty(UnrealData.BundleIdentifier))
+			{
+				ConfigHierarchy Ini = ConfigCache.ReadHierarchy(ConfigHierarchyType.Engine, UnrealData.UProjectFileLocation?.Directory, UnrealTargetPlatform.Mac);
+				string? BundleId;
+				if (Ini.TryGetValue("XcodeConfiguration", "ModernBundleIdentifier", out BundleId))
+				{
+					// adapted from IOSProjectSettings
+					string ProjectName = (UnrealData.UProjectFileLocation != null) ? UnrealData.ProductName : "UnrealGame";
+					UnrealData.BundleIdentifier = BundleId.Replace("[PROJECT_NAME]", ProjectName).Replace("_", "");
+				}	
+			}
+
 			// look for an existing project to use as a template (if none found, create one from scratch)
 			DirectoryReference BuildDirLocation = UnrealData.UProjectFileLocation == null ? Unreal.EngineDirectory : UnrealData.UProjectFileLocation.Directory;
-			string ExistingProjectName = UnrealData.XcodeProjectFileLocation!.GetFileNameWithoutAnyExtensions();
-			FileReference TemplateProject = FileReference.Combine(BuildDirLocation, "Build", "IOS", ExistingProjectName + ".xcodeproj", "project.pbxproj");
+			string ExistingProjectName = UnrealData.ProductName;
+			FileReference TemplateProject = FileReference.Combine(BuildDirLocation, $"Build/IOS/{UnrealData.ProductName}.xcodeproj/project.pbxproj");
 
 			UnrealData.bIsMergingProjects = FileReference.Exists(TemplateProject);
 			UnrealData.bWriteCodeSigningSettings = !UnrealData.bIsMergingProjects;
@@ -2102,16 +2286,10 @@ namespace UnrealBuildTool.XcodeProjectXcconfig
 
 			// turn all UE files into internal representation
 			ProcessSourceFiles();
-
+			
 			// now create the xcodeproject elements (project -> target -> buildconfigs, etc)
 			RootProject = new XcodeProject(UnrealData, FileCollection);
 
-
-			if (TemplateProject.FullName.Contains("ThirdPerson"))
-			{
-				Console.WriteLine("foo");
-			}
-
 			bool bSuccess;
 			if (FileReference.Exists(TemplateProject))
 			{
@@ -2119,6 +2297,9 @@ namespace UnrealBuildTool.XcodeProjectXcconfig
 			}
 			else
 			{
+				// write metadata now so we can add them to the FileCollection
+				ConditionalWriteMetadataFiles(UnrealTargetPlatform.IOS);
+
 				FileReference PBXProjFilePath = ProjectFilePath + "/project.pbxproj";
 
 				StringBuilder Content = new StringBuilder();
@@ -2146,8 +2327,6 @@ namespace UnrealBuildTool.XcodeProjectXcconfig
 			}
 
 
-
-
 			bool bNeedScheme = UnrealData.CanBuildProjectLocally(this, Logger);
 			if (bNeedScheme)
 			{
@@ -2174,17 +2353,341 @@ namespace UnrealBuildTool.XcodeProjectXcconfig
 			return bSuccess;
 		}
 
+		private void ConditionalWriteMetadataFiles(UnrealTargetPlatform Platform)
+		{
+			ConfigHierarchy Ini = ConfigCache.ReadHierarchy(ConfigHierarchyType.Engine, UnrealData.UProjectFileLocation?.Directory, Platform);
 
-		private string Plist(string Command)
+			/////////////////////
+			// Entitlements
+			/////////////////////
+
+			FileReference? ExistingEntitlements = UnrealData.UProjectFileLocation == null ? null :
+				FileReference.Combine(UnrealData.UProjectFileLocation.Directory, $"Build/Xcode/{UnrealData.ProductName}.entitlements");
+
+			if (ExistingEntitlements != null && FileReference.Exists(ExistingEntitlements))
+			{
+				UnrealData.EntitlementsLocation = ExistingEntitlements;
+				UnrealData.bEntitlementsWasPremade = true;
+				string RelativeLoc = UnrealData.EntitlementsLocation.MakeRelativeTo(UnrealData.XcodeProjectFileLocation.ParentDirectory!);
+				FileCollection.AddFileReference(XcodeProjectFileGenerator.MakeXcodeGuid(), RelativeLoc, "text.plist", "\"<group>\"", "Metadata");
+			}
+			else
+			{
+				UnrealData.EntitlementsLocation = FileReference.Combine(UnrealData.XcodeProjectFileLocation.ParentDirectory!, "Metadata", $"{UnrealData.ProductName}.entitlements");
+				FileCollection.AddFileReference(XcodeProjectFileGenerator.MakeXcodeGuid(), $"Metadata/{UnrealData.EntitlementsLocation.GetFileName()}", "text.plist.entitlements", "\"<group>\"", "Metadata");
+
+				string Filename = UnrealData.EntitlementsLocation.FullName;
+				DirectoryReference.CreateDirectory(UnrealData.EntitlementsLocation.Directory);
+
+				// reset file to start (or create it)
+				Plist("Clear", Filename);
+
+				bool bValue = false;
+				// remote push notifications
+				if (Ini.GetBool("/Script/IOSRuntimeSettings.IOSRuntimeSettings", "bEnableRemoteNotificationsSupport", out bValue) && bValue)
+				{
+					// Xcode will replace this with production when using a distro cert
+					Plist("Add aps-environment string development", Filename);
+				}
+
+				// Sign in with Apple
+				if (Ini.GetBool("/Script/IOSRuntimeSettings.IOSRuntimeSettings", "bEnableSignInWithAppleSupport", out bValue) && bValue)
+				{
+					Plist("Add com.apple.developer.applesignin string Default", Filename);
+				}
+
+				// Add Multi-user support for tvOS
+				if (Ini.GetBool("/Script/IOSRuntimeSettings.IOSRuntimeSettings", "bRunAsCurrentUser", out bValue) && bValue)
+				{
+					Plist("Add com.apple.developer.user-management array", Filename);
+					Plist("Add com.apple.developer.user-management:0 string runs-as-current-user", Filename);
+				}
+			}
+
+
+			/////////////////////
+			// Info.plist
+			/////////////////////
+
+			FileReference? ExistingInfoPlist = UnrealData.UProjectFileLocation == null ? null :
+				FileReference.Combine(UnrealData.UProjectFileLocation.Directory, $"Build/Xcode/{UnrealData.ProductName}-Info.plist");
+			if (ExistingInfoPlist != null && FileReference.Exists(ExistingInfoPlist))
+			{
+				UnrealData.InfoPlistLocation = ExistingInfoPlist;
+				UnrealData.bInfoPlistWasPremade = true;
+				string RelativeLoc = UnrealData.InfoPlistLocation.MakeRelativeTo(UnrealData.XcodeProjectFileLocation.ParentDirectory!);
+				FileCollection.AddFileReference(XcodeProjectFileGenerator.MakeXcodeGuid(), RelativeLoc, "text.plist", "\"<group>\"", "Metadata");
+			}
+			else
+			{
+				UnrealData.InfoPlistLocation = FileReference.Combine(UnrealData.XcodeProjectFileLocation.ParentDirectory!, "Metadata", $"{UnrealData.ProductName}-Info.plist");
+				FileCollection.AddFileReference(XcodeProjectFileGenerator.MakeXcodeGuid(), $"Metadata/{UnrealData.InfoPlistLocation.GetFileName()}", "text.plist", "\"<group>\"", "Metadata");
+
+
+				string Filename = UnrealData.InfoPlistLocation.FullName;
+
+				// reset file to start (or create it)
+				Plist("Clear", Filename);
+
+				////////////////////////
+				// PLIST MAC / IOS
+				////////////////////////
+
+
+				////////////////////////
+				// PLIST MAC
+				////////////////////////
+
+				// disable non-exempt encryption
+				Plist("Add :ITSAppUsesNonExemptEncryption bool false", Filename);
+
+				////////////////////////
+				// PLIST IOS
+				////////////////////////
+
+				if (UnrealData.IOSProjectSettings != null)
+				{
+					bool HACK_bSupportsFilesApp = true;
+					if (UnrealData.IOSProjectSettings.bFileSharingEnabled)
+					{
+						Plist("Add :UIFileSharingEnabled bool true", Filename);
+					}
+					if (HACK_bSupportsFilesApp)
+					{
+						Plist("Add :LSSupportsOpeningDocumentsInPlace bool true", Filename);
+					}
+					Plist("Add :UIViewControllerBasedStatusBarAppearance bool false", Filename);
+				}
+			}
+
+// @todo this is a reminder of plist settings we are not handling yet
+#if false
+			Text.AppendLine("\t\t\t<key>CFBundleURLName</key>");
+			Text.AppendLine("\t\t\t<string>com.Epic.Unreal</string>");
+			Text.AppendLine("\t\t\t<key>CFBundleURLSchemes</key>");
+			Text.AppendLine("\t\t\t<array>");
+			Text.AppendLine(string.Format("\t\t\t\t<string>{0}</string>", bIsUnrealGame ? "UnrealGame" : GameName));
+			if (bEnableGoogleSupport)
+			{
+				Text.AppendLine(string.Format("\t\t\t\t<string>{0}</string>", GoogleReversedClientId));
+			}
+			Text.AppendLine("\t\t\t</array>");
+			Text.AppendLine("\t\t</dict>");
+			Text.AppendLine("\t</array>");
+
+
+			Text.AppendLine("\t<key>CFBundleDevelopmentRegion</key>");
+			Text.AppendLine("\t<string>English</string>");
+
+			Text.AppendLine("\t<key>CFBundleExecutable</key>");
+			string BundleExecutable = bIsUnrealGame ?
+				(bIsClient ? "UnrealClient" : "UnrealGame") :
+				(bIsClient ? GameName + "Client" : GameName);
+			Text.AppendLine(string.Format("\t<string>{0}</string>", BundleExecutable));
+
+			Text.AppendLine("\t<key>CFBundleVersion</key>");
+			Text.AppendLine(string.Format("\t<string>{0}</string>", VersionUtilities.UpdateBundleVersion(OldPListData, InEngineDir)));
+			Text.AppendLine("\t<key>CFBundleShortVersionString</key>");
+			Text.AppendLine(string.Format("\t<string>{0}</string>", BundleShortVersion));
+
+			Text.AppendLine("\t<key>UIRequiredDeviceCapabilities</key>");
+			Text.AppendLine("\t<array>");
+			foreach(string Line in RequiredCaps.Split("\r\n".ToCharArray()))
+			{
+				if (!string.IsNullOrWhiteSpace(Line))
+				{
+					Text.AppendLine(Line);
+				}
+			}
+			Text.AppendLine("\t</array>");
+
+			Text.AppendLine("\t<key>CFBundleIcons</key>");
+			Text.AppendLine("\t<dict>");
+			Text.AppendLine("\t\t<key>CFBundlePrimaryIcon</key>");
+			Text.AppendLine("\t\t<dict>");
+			Text.AppendLine("\t\t\t<key>CFBundleIconFiles</key>");
+			Text.AppendLine("\t\t\t<array>");
+			Text.AppendLine("\t\t\t\t<string>AppIcon60x60</string>");
+			Text.AppendLine("\t\t\t</array>");
+			Text.AppendLine("\t\t\t<key>CFBundleIconName</key>");
+			Text.AppendLine("\t\t\t<string>AppIcon</string>");
+			Text.AppendLine("\t\t\t<key>UIPrerenderedIcon</key>");
+			Text.AppendLine("\t\t\t<true/>");
+			Text.AppendLine("\t\t</dict>");
+			Text.AppendLine("\t</dict>");
+			Text.AppendLine("\t<key>CFBundleIcons~ipad</key>");
+			Text.AppendLine("\t<dict>");
+			Text.AppendLine("\t\t<key>CFBundlePrimaryIcon</key>");
+			Text.AppendLine("\t\t<dict>");
+			Text.AppendLine("\t\t\t<key>CFBundleIconFiles</key>");
+			Text.AppendLine("\t\t\t<array>");
+			Text.AppendLine("\t\t\t\t<string>AppIcon60x60</string>");
+			Text.AppendLine("\t\t\t\t<string>AppIcon76x76</string>");
+			Text.AppendLine("\t\t\t</array>");
+			Text.AppendLine("\t\t\t<key>CFBundleIconName</key>");
+			Text.AppendLine("\t\t\t<string>AppIcon</string>");
+			Text.AppendLine("\t\t\t<key>UIPrerenderedIcon</key>");
+			Text.AppendLine("\t\t\t<true/>");
+			Text.AppendLine("\t\t</dict>");
+			Text.AppendLine("\t</dict>");
+			Text.AppendLine("\t<key>UILaunchStoryboardName</key>");
+			Text.AppendLine("\t<string>LaunchScreen</string>");
+
+			if (File.Exists(DirectoryReference.FromFile(ProjectFile) + "/Build/IOS/Resources/Interface/LaunchScreen.storyboard") && VersionUtilities.bCustomLaunchscreenStoryboard)
+			{
+				string LaunchStoryboard = DirectoryReference.FromFile(ProjectFile) + "/Build/IOS/Resources/Interface/LaunchScreen.storyboard";
+
+				if (BuildHostPlatform.Current.Platform == UnrealTargetPlatform.Mac)
+				{
+					string outputStoryboard = LaunchStoryboard + "c";
+					string argsStoryboard = "--compile " + outputStoryboard + " " + LaunchStoryboard;
+					string stdOutLaunchScreen = Utils.RunLocalProcessAndReturnStdOut("ibtool", argsStoryboard, Logger);
+
+					Logger.LogInformation("LaunchScreen Storyboard compilation results : {Results}", stdOutLaunchScreen);
+				}
+				else
+				{
+					Logger.LogWarning("Custom Launchscreen compilation storyboard only compatible on Mac for now");
+				}
+			}
+
+
+			// add location services descriptions if used
+			if (!string.IsNullOrWhiteSpace(LocationAlwaysUsageDescription))
+			{
+				Text.AppendLine("\t<key>NSLocationAlwaysAndWhenInUseUsageDescription</key>");
+				Text.AppendLine(string.Format("\t<string>{0}</string>", LocationAlwaysUsageDescription));
+			}
+			if (!string.IsNullOrWhiteSpace(LocationWhenInUseDescription))
+			{
+				Text.AppendLine("\t<key>NSLocationWhenInUseUsageDescription</key>");
+				Text.AppendLine(string.Format("\t<string>{0}</string>", LocationWhenInUseDescription));
+			}
+			// disable HTTPS requirement
+			if (bDisableHTTPS)
+			{
+				Text.AppendLine("\t<key>NSAppTransportSecurity</key>");
+				Text.AppendLine("\t\t<dict>");
+				Text.AppendLine("\t\t\t<key>NSAllowsArbitraryLoads</key><true/>");
+				Text.AppendLine("\t\t</dict>");
+			}
+
+			if (bEnableFacebookSupport)
+			{
+				Text.AppendLine("\t<key>FacebookAppID</key>");
+				Text.AppendLine(string.Format("\t<string>{0}</string>", FacebookAppID));
+				Text.AppendLine("\t<key>FacebookDisplayName</key>");
+				Text.AppendLine(string.Format("\t<string>{0}</string>", FacebookDisplayName));
+				
+				if (!bEnableAutomaticLogging)
+				{
+					Text.AppendLine("<key>FacebookAutoLogAppEventsEnabled</key><false/>");
+				}
+				
+				if (!bEnableAdvertisingId)
+				{
+					Text.AppendLine("<key>FacebookAdvertiserIDCollectionEnabled</key><false/>");
+				}
+
+				Text.AppendLine("\t<key>LSApplicationQueriesSchemes</key>");
+				Text.AppendLine("\t<array>");
+				Text.AppendLine("\t\t<string>fbapi</string>");
+				Text.AppendLine("\t\t<string>fb-messenger-api</string>");
+				Text.AppendLine("\t\t<string>fb-messenger-share-api</string>");
+				Text.AppendLine("\t\t<string>fbauth2</string>");
+				Text.AppendLine("\t\t<string>fbshareextension</string>");
+				Text.AppendLine("\t</array>");
+			}
+
+			if (!string.IsNullOrEmpty(ExtraData))
+			{
+				ExtraData = ExtraData.Replace("\\n", "\n");
+				foreach(string Line in ExtraData.Split("\r\n".ToCharArray()))
+				{
+					if (!string.IsNullOrWhiteSpace(Line))
+					{
+						Text.AppendLine("\t" + Line);
+					}
+				}
+			}
+
+			// Add remote-notifications as background mode
+			if (bRemoteNotificationsSupported || bBackgroundFetch)
+			{
+				Text.AppendLine("\t<key>UIBackgroundModes</key>");
+				Text.AppendLine("\t<array>");
+                if (bRemoteNotificationsSupported)
+                {
+				    Text.AppendLine("\t\t<string>remote-notification</string>");
+                }
+                if (bBackgroundFetch)
+                {
+                    Text.AppendLine("\t\t<string>fetch</string>");
+                }
+				Text.AppendLine("\t</array>");
+			}
+
+			// write the iCloud container identifier, if present in the old file
+			if (!string.IsNullOrEmpty(OldPListData))
+			{
+				int index = OldPListData.IndexOf("ICloudContainerIdentifier");
+				if (index > 0)
+				{
+					index = OldPListData.IndexOf("<string>", index) + 8;
+					int length = OldPListData.IndexOf("</string>", index) - index;
+					string ICloudContainerIdentifier = OldPListData.Substring(index, length);
+					Text.AppendLine("\t<key>ICloudContainerIdentifier</key>");
+					Text.AppendLine(string.Format("\t<string>{0}</string>", ICloudContainerIdentifier));
+				}
+			}
+
+			Text.AppendLine("</dict>");
+			Text.AppendLine("</plist>");
+
+			// Create the intermediate directory if needed
+			if (!Directory.Exists(IntermediateDirectory))
+			{
+				Directory.CreateDirectory(IntermediateDirectory);
+			}
+
+			if (UPL != null)
+			{
+				// Allow UPL to modify the plist here
+				XDocument XDoc;
+				try
+				{
+					XDoc = XDocument.Parse(Text.ToString());
+				}
+				catch (Exception e)
+				{
+					throw new BuildException("plist is invalid {0}\n{1}", e, Text.ToString());
+				}
+
+				XDoc.DocumentType!.InternalSubset = "";
+				UPL.ProcessPluginNode("None", "iosPListUpdates", "", ref XDoc);
+				string result = XDoc.Declaration?.ToString() + "\n" + XDoc.ToString().Replace("<!DOCTYPE plist PUBLIC \"-//Apple//DTD PLIST 1.0//EN\" \"http://www.apple.com/DTDs/PropertyList-1.0.dtd\"[]>", "<!DOCTYPE plist PUBLIC \"-//Apple//DTD PLIST 1.0//EN\" \"http://www.apple.com/DTDs/PropertyList-1.0.dtd\">");
+				File.WriteAllText(PListFile, result);
+
+				Text = new StringBuilder(result);
+			}
+#endif
+		}
+
+		private string Plist(string Command, string PlistFile)
 		{
 			Command = Command.Replace("\"", "\\\"");
-			string Output = Utils.RunLocalProcessAndReturnStdOut("/usr/libexec/PlistBuddy", $"-c \"{Command}\" \"{ProjectFilePath.FullName}/project.pbxproj\"");
-			
-			//Console.WriteLine($"{Command} ==> {Output}");
+			string Output = Utils.RunLocalProcessAndReturnStdOut("/usr/libexec/PlistBuddy", $"-c \"{Command}\" \"{PlistFile}\"");
+
+			Console.WriteLine($"{Command} ==> {Output}");
 
 			return Output;
 		}
 
+		private string Plist(string Command)
+		{
+			return Plist(Command, ProjectFilePath.FullName + "/project.pbxproj");
+		}
+
 		private void PlistSetAdd(string Entry, string Value, string Type="string")
 		{
 			string AddOutput = Plist($"Add {Entry} {Type} {Value}");
@@ -2424,7 +2927,7 @@ namespace UnrealBuildTool.XcodeProjectXcconfig
 			return true;
 		}
 
-		#region Schemes
+#region Schemes
 
 		private FileReference GetUserSchemeManagementFilePath()
 		{
@@ -2481,8 +2984,6 @@ namespace UnrealBuildTool.XcodeProjectXcconfig
 			Content.WriteLine("            <BuildableReference");
 			Content.WriteLine("               BuildableIdentifier = \"primary\"");
 			Content.WriteLine("               BlueprintIdentifier = \"" + TargetGuid + "\"");
-			Content.WriteLine("               BuildableName = \"" + TargetName + ".app\"");
-			Content.WriteLine("               BlueprintName = \"" + TargetName + "\"");
 			Content.WriteLine("               ReferencedContainer = \"container:" + TargetName + ".xcodeproj\">");
 			Content.WriteLine("            </BuildableReference>");
 			Content.WriteLine("         </BuildActionEntry>");
@@ -2499,8 +3000,6 @@ namespace UnrealBuildTool.XcodeProjectXcconfig
 			Content.WriteLine("            <BuildableReference");
 			Content.WriteLine("               BuildableIdentifier = \"primary\"");
 			Content.WriteLine("               BlueprintIdentifier = \"" + TargetGuid + "\"");
-			Content.WriteLine("               BuildableName = \"" + TargetName + ".app\"");
-			Content.WriteLine("               BlueprintName = \"" + TargetName + "\"");
 			Content.WriteLine("               ReferencedContainer = \"container:" + TargetName + ".xcodeproj\">");
 			Content.WriteLine("            </BuildableReference>");
 			Content.WriteLine("      </MacroExpansion>");
@@ -2522,8 +3021,6 @@ namespace UnrealBuildTool.XcodeProjectXcconfig
 			Content.WriteLine("            <BuildableReference");
 			Content.WriteLine("               BuildableIdentifier = \"primary\"");
 			Content.WriteLine("               BlueprintIdentifier = \"" + TargetGuid + "\"");
-			Content.WriteLine("               BuildableName = \"" + TargetName + ".app\"");
-			Content.WriteLine("               BlueprintName = \"" + TargetName + "\"");
 			Content.WriteLine("               ReferencedContainer = \"container:" + TargetName + ".xcodeproj\">");
 			Content.WriteLine("            </BuildableReference>");
 			Content.WriteLine("      </BuildableProductRunnable>");
diff --git a/Engine/Source/Programs/UnrealBuildTool/ProjectFiles/Xcode/XcodeProjectFileGenerator.cs b/Engine/Source/Programs/UnrealBuildTool/ProjectFiles/Xcode/XcodeProjectFileGenerator.cs
--- a/Engine/Source/Programs/UnrealBuildTool/ProjectFiles/Xcode/XcodeProjectFileGenerator.cs
+++ b/Engine/Source/Programs/UnrealBuildTool/ProjectFiles/Xcode/XcodeProjectFileGenerator.cs
@@ -15,6 +15,8 @@ namespace UnrealBuildTool
 	/// </summary>
 	class XcodeProjectFileGenerator : ProjectFileGenerator
 	{
+		public DirectoryReference? XCWorkspace;
+
 		// always seed the random number the same, so multiple runs of the generator will generate the same project
 		static Random Rand = new Random(0);
 
@@ -33,11 +35,6 @@ namespace UnrealBuildTool
 		/// </summary>
 		string AppName = "";
 
-		/// <summary>
-		/// Whether or not to use the new .xcconfig file 
-		/// </summary>
-		private bool bUseXcconfig = false;
-
 		public XcodeProjectFileGenerator(FileReference? InOnlyGameProject, CommandLineArguments CommandLine)
 			: base(InOnlyGameProject)
 		{
@@ -54,11 +51,6 @@ namespace UnrealBuildTool
 			{
 				AppName = CommandLine.GetString("-appname=");
 			}
-
-			if (CommandLine.HasOption("-xcconfig"))
-			{
-				bUseXcconfig = true;
-			}
 		}
 
 		/// <summary>
@@ -120,14 +112,11 @@ namespace UnrealBuildTool
 		/// <returns>The newly allocated project file object</returns>
 		protected override ProjectFile AllocateProjectFile(FileReference InitFilePath, DirectoryReference BaseDir)
 		{
-			if (bUseXcconfig)
-			{
-				return new XcodeProjectXcconfig.XcodeProjectFile(InitFilePath, BaseDir, bForDistribution, BundleIdentifier, AppName);
-			}
-			else
-			{
-				return new XcodeProjectLegacy.XcodeProjectFile(InitFilePath, BaseDir, bForDistribution, BundleIdentifier, AppName);
-			}
+			// this may internally (later) make a Legacy project object if the unreal project wants old behavior
+			// unfortunately, we can't read the project configs now because we don't have enough information to 
+			// find the .uproject file for that would make this project (we could change the high level to pass it 
+			// down but it would touch all project generators - not worth it if we end up removing the legacy)
+			return new XcodeProjectXcconfig.XcodeProjectFile(InitFilePath, BaseDir, bForDistribution, BundleIdentifier, AppName);
 		}
 
 		private bool WriteWorkspaceSettingsFile(string Path, ILogger Logger)
@@ -137,9 +126,10 @@ namespace UnrealBuildTool
 			WorkspaceSettingsContent.Append("<!DOCTYPE plist PUBLIC \"-//Apple//DTD PLIST 1.0//EN\" \"http://www.apple.com/DTDs/PropertyList-1.0.dtd\">" + ProjectFileGenerator.NewLine);
 			WorkspaceSettingsContent.Append("<plist version=\"1.0\">" + ProjectFileGenerator.NewLine);
 			WorkspaceSettingsContent.Append("<dict>" + ProjectFileGenerator.NewLine);
+			// @todo when we move to xcode 14, we remove these next 4 keys
 			WorkspaceSettingsContent.Append("\t<key>BuildSystemType</key>" + ProjectFileGenerator.NewLine);
 			WorkspaceSettingsContent.Append("\t<string>Original</string>" + ProjectFileGenerator.NewLine);
-            WorkspaceSettingsContent.Append("\t<key>BuildLocationStyle</key>" + ProjectFileGenerator.NewLine);
+			WorkspaceSettingsContent.Append("\t<key>BuildLocationStyle</key>" + ProjectFileGenerator.NewLine);
 			WorkspaceSettingsContent.Append("\t<string>UseTargetSettings</string>" + ProjectFileGenerator.NewLine);
 			WorkspaceSettingsContent.Append("\t<key>CustomBuildLocationType</key>" + ProjectFileGenerator.NewLine);
 			WorkspaceSettingsContent.Append("\t<string>RelativeToDerivedData</string>" + ProjectFileGenerator.NewLine);
@@ -165,10 +155,12 @@ namespace UnrealBuildTool
 			WorkspaceSettingsContent.Append("<!DOCTYPE plist PUBLIC \"-//Apple//DTD PLIST 1.0//EN\" \"http://www.apple.com/DTDs/PropertyList-1.0.dtd\">" + ProjectFileGenerator.NewLine);
 			WorkspaceSettingsContent.Append("<plist version=\"1.0\">" + ProjectFileGenerator.NewLine);
 			WorkspaceSettingsContent.Append("<dict>" + ProjectFileGenerator.NewLine);
-			WorkspaceSettingsContent.Append("\t<key>DisableBuildSystemDeprecationWarning</key>" + ProjectFileGenerator.NewLine);
-			WorkspaceSettingsContent.Append("\t<true/>" + ProjectFileGenerator.NewLine);
-            WorkspaceSettingsContent.Append("\t<key>DisableBuildSystemDeprecationDiagnostic</key>" + ProjectFileGenerator.NewLine);
-            WorkspaceSettingsContent.Append("\t<true/>" + ProjectFileGenerator.NewLine);
+			//WorkspaceSettingsContent.Append("\t<key>DisableBuildSystemDeprecationWarning</key>" + ProjectFileGenerator.NewLine);
+			//WorkspaceSettingsContent.Append("\t<true/>" + ProjectFileGenerator.NewLine);
+			//WorkspaceSettingsContent.Append("\t<key>DisableBuildSystemDeprecationDiagnostic</key>" + ProjectFileGenerator.NewLine);
+			//WorkspaceSettingsContent.Append("\t<true/>" + ProjectFileGenerator.NewLine);
+			WorkspaceSettingsContent.Append("\t<key>IDEWorkspaceSharedSettings_AutocreateContextsIfNeeded</key>" + ProjectFileGenerator.NewLine);
+			WorkspaceSettingsContent.Append("\t<false/>" + ProjectFileGenerator.NewLine);
 			WorkspaceSettingsContent.Append("</dict>" + ProjectFileGenerator.NewLine);
 			WorkspaceSettingsContent.Append("</plist>" + ProjectFileGenerator.NewLine);
 			return WriteFileIfChanged(Path, WorkspaceSettingsContent.ToString(), Logger, new UTF8Encoding());
@@ -262,6 +254,9 @@ namespace UnrealBuildTool
 				bSuccess = WriteWorkspaceSettingsFile(WorkspaceSettingsFilePath, Logger);
 				string WorkspaceSharedSettingsFilePath = PrimaryProjectPath + "/" + ProjectName + ".xcworkspace/xcshareddata/WorkspaceSettings.xcsettings";
 				bSuccess = WriteWorkspaceSharedSettingsFile(WorkspaceSharedSettingsFilePath, Logger);
+
+				// cache the location of the workspace, for users of this to know where the final workspace is
+				XCWorkspace = new FileReference(WorkspaceDataFilePath).Directory;
 			}
 
 
@@ -313,15 +308,25 @@ namespace UnrealBuildTool
 				}
 			}
 
-			if (bGeneratingGameProjectFiles)
+			if (bGeneratingRunIOSProject || bGeneratingRunTVOSProject)
 			{
-				if (bGeneratingRunIOSProject || bGeneratingRunTVOSProject || UnrealBuildBase.Unreal.IsEngineInstalled())
+				bIncludeEnginePrograms = false;
+				//bIncludeEngineSource = false;
+				bIncludeTemplateFiles = false;
+				bIncludeConfigFiles = false;
+				bIncludeDocumentation = false;
+				bIncludeShaderSource = false;
+
+				// generate just the engine project
+				if (OnlyGameProject == null)
+				{
+					bIncludeEngineSource = true;
+				}
+				// generate just the game project
+				else
 				{
-					// an Engine target is required in order to be able to get Xcode to sign blueprint projects
-					// always include the engine target for installed builds.
-					bIncludeEnginePrograms = true;
+					bGeneratingGameProjectFiles = true;
 				}
-				bIncludeEngineSource = true;
 			}
 		}
 	}
diff --git a/Engine/Source/Programs/UnrealBuildTool/ProjectFiles/Xcode/XcodeProjectLegacy.cs b/Engine/Source/Programs/UnrealBuildTool/ProjectFiles/Xcode/XcodeProjectLegacy.cs
--- a/Engine/Source/Programs/UnrealBuildTool/ProjectFiles/Xcode/XcodeProjectLegacy.cs
+++ b/Engine/Source/Programs/UnrealBuildTool/ProjectFiles/Xcode/XcodeProjectLegacy.cs
@@ -11,49 +11,10 @@ using System.Text;
 using EpicGames.Core;
 using UnrealBuildBase;
 using Microsoft.Extensions.Logging;
+using UnrealBuildTool.XcodeProjectXcconfig;
 
 namespace UnrealBuildTool.XcodeProjectLegacy
 {
-	/// <summary>
-	/// Info needed to make a file a member of specific group
-	/// </summary>
-	class XcodeSourceFile : ProjectFile.SourceFile
-	{
-		/// <summary>
-		/// Constructor
-		/// </summary>
-		public XcodeSourceFile(FileReference InitFilePath, DirectoryReference? InitRelativeBaseFolder)
-			: base(InitFilePath, InitRelativeBaseFolder)
-		{
-			FileGuid = XcodeProjectFileGenerator.MakeXcodeGuid();
-			FileRefGuid = XcodeProjectFileGenerator.MakeXcodeGuid();
-		}
-
-		/// <summary>
-		/// File Guid for use in Xcode project
-		/// </summary>
-		public string FileGuid
-		{
-			get;
-			private set;
-		}
-
-		public void ReplaceGuids(string NewFileGuid, string NewFileRefGuid)
-		{
-			FileGuid = NewFileGuid;
-			FileRefGuid = NewFileRefGuid;
-		}
-
-		/// <summary>
-		/// File reference Guid for use in Xcode project
-		/// </summary>
-		public string FileRefGuid
-		{
-			get;
-			private set;
-		}
-	}
-
 	/// <summary>
 	/// Represents a group of files shown in Xcode's project navigator as a folder
 	/// </summary>
