- Enabled MDL/USD importers now that they have arm64/m1 support!

diff --git a/Engine/Source/Programs/UnrealBuildTool/Platform/Mac/UEBuildMac.cs b/Engine/Source/Programs/UnrealBuildTool/Platform/Mac/UEBuildMac.cs
--- a/Engine/Source/Programs/UnrealBuildTool/Platform/Mac/UEBuildMac.cs
+++ b/Engine/Source/Programs/UnrealBuildTool/Platform/Mac/UEBuildMac.cs
@@ -151,11 +151,6 @@ namespace UnrealBuildTool
 					"MetaHumans",
 					"NeuralMorphModel",
 
-					// MDL
-					"MDLImporter",
-					"USDImporter",
-					"USDMultiUser",
-
 					// EOSSDK
 					"EOSManager",
 					"EOSShared",
