Rewrote the GenerateClangDatabase mode in UBT to work with the clang tools like IWYU.

diff --git a/Engine/Source/Programs/UnrealBuildTool/Configuration/UEBuildTarget.cs b/Engine/Source/Programs/UnrealBuildTool/Configuration/UEBuildTarget.cs
--- a/Engine/Source/Programs/UnrealBuildTool/Configuration/UEBuildTarget.cs
+++ b/Engine/Source/Programs/UnrealBuildTool/Configuration/UEBuildTarget.cs
@@ -2405,7 +2405,7 @@ namespace UnrealBuildTool
 		/// Creates a toolchain for the current target. May be overridden by the target rules.
 		/// </summary>
 		/// <returns>New toolchain instance</returns>
-		private UEToolChain CreateToolchain(UnrealTargetPlatform Platform)
+		public UEToolChain CreateToolchain(UnrealTargetPlatform Platform)
 		{
 			if (Rules.ToolChainName == null)
 			{
diff --git a/Engine/Source/Programs/UnrealBuildTool/Modes/GenerateClangDatabase.cs b/Engine/Source/Programs/UnrealBuildTool/Modes/GenerateClangDatabase.cs
--- a/Engine/Source/Programs/UnrealBuildTool/Modes/GenerateClangDatabase.cs
+++ b/Engine/Source/Programs/UnrealBuildTool/Modes/GenerateClangDatabase.cs
@@ -4,7 +4,7 @@ using System;
 using System.Collections.Generic;
 using System.Linq;
 using System.Text;
-using System.Threading.Tasks;
+using System.Text.RegularExpressions;
 using EpicGames.Core;
 using UnrealBuildBase;
 using Microsoft.Extensions.Logging;
@@ -17,12 +17,88 @@ namespace UnrealBuildTool
 	[ToolMode("GenerateClangDatabase", ToolModeOptions.XmlConfig | ToolModeOptions.BuildPlatforms | ToolModeOptions.SingleInstance | ToolModeOptions.StartPrefetchingEngine | ToolModeOptions.ShowExecutionTime)]
 	class GenerateClangDatabase : ToolMode
 	{
+		static Regex LineRegex = new Regex(@"^(\-include|\-I|\/I|\/imsvc|\-isystem|\/FI)\s*(.*)");
+
 		/// <summary>
 		/// Set of filters for files to include in the database. Relative to the root directory, or to the project file.
 		/// </summary>
 		[CommandLine("-Filter=")]
 		List<string> FilterRules = new List<string>();
 
+		/// <summary>
+		/// Execute any actions which result in code generation (eg. ISPC compilation)
+		/// </summary>
+		[CommandLine("-ExecCodeGenActions")]
+		public bool bExecCodeGenActions = false;
+
+		/// <summary>
+		/// This ActionGraphBuilder captures the build output from a UEBuildModuleCPP so it can be consumed later.
+		/// </summary>
+		private class CaptureActionGraphBuilder : IActionGraphBuilder
+		{
+			private readonly ILogger Logger;
+			public List<IExternalAction> CapturedActions = new List<IExternalAction>();
+			public List<Tuple<FileItem, IEnumerable<string>>> CapturedTextFiles = new List<Tuple<FileItem, IEnumerable<string>>>();
+
+			/// <summary>
+			/// Constructor
+			/// </summary>
+			/// <param name="InLogger"></param>
+			public CaptureActionGraphBuilder(ILogger InLogger)
+			{
+				Logger = InLogger;
+			}
+
+			/// <inheritdoc/>
+			public void AddAction(IExternalAction Action)
+			{
+				CapturedActions.Add(Action);
+			}
+
+			/// <inheritdoc/>
+			public void CreateIntermediateTextFile(FileItem FileItem, string Contents)
+			{
+				Utils.WriteFileIfChanged(FileItem, Contents, Logger);
+			}
+
+			/// <inheritdoc/>
+			public void CreateIntermediateTextFile(FileItem FileItem, IEnumerable<string> ContentLines)
+			{
+				Utils.WriteFileIfChanged(FileItem, ContentLines, Logger);
+				CapturedTextFiles.Add(new Tuple<FileItem, IEnumerable<string>>(FileItem, ContentLines));
+			}
+
+			/// <inheritdoc/>
+			public void AddSourceDir(DirectoryItem SourceDir)
+			{
+			}
+
+			/// <inheritdoc/>
+			public void AddSourceFiles(DirectoryItem SourceDir, FileItem[] SourceFiles)
+			{
+			}
+
+			/// <inheritdoc/>
+			public void AddFileToWorkingSet(FileItem File)
+			{
+			}
+
+			/// <inheritdoc/>
+			public void AddCandidateForWorkingSet(FileItem File)
+			{
+			}
+
+			/// <inheritdoc/>
+			public void AddDiagnostic(string Message)
+			{
+			}
+
+			/// <inheritdoc/>
+			public void SetOutputItemsForModule(string ModuleName, FileItem[] OutputItems)
+			{
+			}
+		}
+
 		/// <summary>
 		/// Execute the command
 		/// </summary>
@@ -67,11 +143,27 @@ namespace UnrealBuildTool
 
 					// Create a makefile for the target
 					UEBuildTarget Target = UEBuildTarget.Create(TargetDescriptor, BuildConfiguration.bSkipRulesCompile, BuildConfiguration.bForceRulesCompile, BuildConfiguration.bUsePrecompiled, Logger);
+					UEToolChain TargetToolChain = Target.CreateToolchain(Target.Platform);
 
-					// Find the location of the compiler
-					FileReference ClangPath = FindClangCompiler(Target, Logger);
+					// Execute code generation actions
+					if (bExecCodeGenActions)
+					{
+						// Create the makefile
+						TargetMakefile Makefile = Target.Build(BuildConfiguration, WorkingSet, TargetDescriptor, Logger);
+						List<LinkedAction> Actions = Makefile.Actions.ConvertAll(x => new LinkedAction(x, TargetDescriptor));
+						ActionGraph.Link(Actions);
 
-					bool IsWindowsClang = ClangPath.GetFileName().Equals("clang-cl.exe", StringComparison.OrdinalIgnoreCase);
+						// Filter all the actions to execute
+						HashSet<FileItem> PrerequisiteItems = new HashSet<FileItem>(Makefile.Actions.SelectMany(x => x.ProducedItems).Where(x => x.HasExtension(".h") || x.HasExtension(".cpp")));
+						List<LinkedAction> PrerequisiteActions = ActionGraph.GatherPrerequisiteActions(Actions, PrerequisiteItems);
+
+						// Execute these actions
+						if (PrerequisiteActions.Count > 0)
+						{
+							Logger.LogInformation("Executing actions that produce source files...");
+							ActionGraph.ExecuteActions(BuildConfiguration, PrerequisiteActions, new List<TargetDescriptor> { TargetDescriptor }, Logger);
+						}
+					}
 
 					// Create all the binaries and modules
 					CppCompileEnvironment GlobalCompileEnvironment = Target.CreateCompileEnvironmentForProjectFiles(Logger);
@@ -82,93 +174,62 @@ namespace UnrealBuildTool
 						{
 							if (!Module.Rules.bUsePrecompiled)
 							{
+								// Gather all the files we care about
 								UEBuildModuleCPP.InputFileCollection InputFileCollection = Module.FindInputFiles(Target.Platform, new Dictionary<DirectoryItem, FileItem[]>());
+								List<FileItem> InputFiles = new List<FileItem>();
+								InputFiles.AddRange(InputFileCollection.CPPFiles);
+								InputFiles.AddRange(InputFileCollection.CCFiles);
 
-								CppCompileEnvironment ModuleCompileEnvironment = Module.CreateModuleCompileEnvironment(Target.Rules, BinaryCompileEnvironment);
-
-								StringBuilder CommandBuilder = new StringBuilder();
-								StringBuilder CppCommandBuilder = new StringBuilder();
-								StringBuilder CCommandBuilder = new StringBuilder();
-								CommandBuilder.AppendFormat("\"{0}\"", ClangPath.FullName);
-
-								switch (ModuleCompileEnvironment.CppStandard)
-								{
-									case CppStandardVersion.Cpp14:
-										CppCommandBuilder.AppendFormat(IsWindowsClang ? " /std:c++14" : " -std=c++14");
-										break;
-									case CppStandardVersion.Cpp17:
-										CppCommandBuilder.AppendFormat(IsWindowsClang ? " /std:c++17" : " -std=c++17");
-										break;
-									case CppStandardVersion.Cpp20:
-									case CppStandardVersion.Latest:
-										CppCommandBuilder.AppendFormat(IsWindowsClang ? " /std:c++latest" : " -std=c++20");
-										break;
-									default:
-										throw new BuildException($"Unsupported C++ standard type set: {ModuleCompileEnvironment.CppStandard}");
-								}
-
-								if (ModuleCompileEnvironment.bEnableCoroutines)
+								var fileList = new List<FileReference>();
+								foreach (FileItem InputFile in InputFiles)
 								{
-									CppCommandBuilder.AppendFormat(" -fcoroutines-ts");
-									if (!ModuleCompileEnvironment.bEnableExceptions)
+									if (FileFilter == null || FileFilter.Matches(InputFile.Location.MakeRelativeTo(Unreal.RootDirectory)))
 									{
-										CppCommandBuilder.AppendFormat(" -Wno-coroutine-missing-unhandled-exception");
+										var fileRef = new FileReference(InputFile.AbsolutePath);
+										fileList.Add(fileRef);
 									}
 								}
 
-								switch (ModuleCompileEnvironment.CStandard)
+								if (fileList.Any())
 								{
-									case CStandardVersion.Default:
-										break;
-									case CStandardVersion.C89:
-										CCommandBuilder.AppendFormat(IsWindowsClang ? "" : " -std=c89");
-										break;
-									case CStandardVersion.C99:
-										CCommandBuilder.AppendFormat(IsWindowsClang ? "" : " -std=c99");
-										break;
-									case CStandardVersion.C11:
-										CCommandBuilder.AppendFormat(IsWindowsClang ? " /std:c11" : " -std=c11");
-										break;
-									case CStandardVersion.C17:
-										CCommandBuilder.AppendFormat(IsWindowsClang ? " /std:c17" : " -std=c17");
-										break;
-									case CStandardVersion.Latest:
-										CCommandBuilder.AppendFormat(IsWindowsClang ? " /std:c17" : " -std=c2x");
-										break;
-									default:
-										throw new BuildException($"Unsupported C standard type set: {ModuleCompileEnvironment.CStandard}");
-								}
+									CaptureActionGraphBuilder ActionGraphBuilder = new CaptureActionGraphBuilder(Logger);
 
-								foreach (FileItem ForceIncludeFile in ModuleCompileEnvironment.ForceIncludeFiles)
-								{
-									CommandBuilder.AppendFormat(" -include \"{0}\"", ForceIncludeFile.FullName);
-								}
-								foreach (string Definition in ModuleCompileEnvironment.Definitions)
-								{
-									CommandBuilder.AppendFormat(" -D\"{0}\"", Definition.Replace("\"", "\\\""));
-								}
-								foreach (DirectoryReference IncludePath in ModuleCompileEnvironment.UserIncludePaths)
-								{
-									CommandBuilder.AppendFormat(" -I\"{0}\"", IncludePath);
-								}
-								foreach (DirectoryReference IncludePath in ModuleCompileEnvironment.SystemIncludePaths)
-								{
-									CommandBuilder.AppendFormat(" -I\"{0}\"", IncludePath);
-								}
+									Module.Compile(Target.Rules, TargetToolChain, BinaryCompileEnvironment, fileList, WorkingSet, ActionGraphBuilder, Logger);
 
-								foreach (FileItem InputFile in InputFileCollection.CPPFiles)
-								{
-									if (FileFilter == null || FileFilter.Matches(InputFile.Location.MakeRelativeTo(Unreal.RootDirectory)))
-									{
-										FileToCommand[InputFile.Location] = String.Format("{0} {1} \"{2}\"", CommandBuilder, CppCommandBuilder, InputFile.FullName);
-									}
-								}
+									IEnumerable<IExternalAction> CompileActions = ActionGraphBuilder.CapturedActions.Where(Action => Action.ActionType == ActionType.Compile);
 
-								foreach (FileItem InputFile in InputFileCollection.CFiles)
-								{
-									if (FileFilter == null || FileFilter.Matches(InputFile.Location.MakeRelativeTo(Unreal.RootDirectory)))
+									if (CompileActions.Any())
 									{
-										FileToCommand[InputFile.Location] = String.Format("{0} {1} \"{2}\"", CommandBuilder, CCommandBuilder, InputFile.FullName);
+										// convert any rsp files
+										Dictionary<string, string> UpdatedResFiles = new Dictionary<string, string>();
+										foreach (Tuple<FileItem, IEnumerable<string>> FileAndContents in ActionGraphBuilder.CapturedTextFiles)
+										{
+											if (FileAndContents.Item1.AbsolutePath.EndsWith(".rsp") || FileAndContents.Item1.AbsolutePath.EndsWith(".response"))
+											{
+												string NewResPath = ConvertResponseFile(FileAndContents.Item1, FileAndContents.Item2, Logger);
+												UpdatedResFiles[FileAndContents.Item1.AbsolutePath] = NewResPath;
+											}
+										}
+
+										foreach (IExternalAction Action in CompileActions)
+										{
+											// Create the command
+											StringBuilder CommandBuilder = new StringBuilder();
+											string CommandArguments = Action.CommandArguments.Replace(".rsp", ".rsp.gcd").Replace(".response", ".response.gcd");
+											CommandBuilder.AppendFormat("\"{0}\" {1}", Action.CommandPath, CommandArguments);
+
+											foreach (string ExtraArgument in GetExtraPlatformArguments(TargetToolChain))
+											{
+												CommandBuilder.AppendFormat(" {0}", ExtraArgument);
+											}
+
+											// find source file
+											var SourceFile = Action.PrerequisiteItems.FirstOrDefault(fi => fi.HasExtension(".cpp") || fi.HasExtension(".c") || fi.HasExtension(".c"));
+											if (SourceFile != null)
+											{
+												FileToCommand[SourceFile.Location] = CommandBuilder.ToString();
+											}
+										}
 									}
 								}
 							}
@@ -197,45 +258,63 @@ namespace UnrealBuildTool
 			return 0;
 		}
 
-		/// <summary>
-		/// Searches for the Clang compiler for the given platform.
-		/// </summary>
-		/// <param name="Target">The build platform to use to search for the Clang compiler.</param>
-		/// <param name="Logger">Logger for output</param>
-		/// <returns>The path to the Clang compiler.</returns>
-		private static FileReference FindClangCompiler(UEBuildTarget Target, ILogger Logger)
+		private IEnumerable<string> GetExtraPlatformArguments(UEToolChain TargetToolChain)
 		{
-			UnrealTargetPlatform HostPlatform = BuildHostPlatform.Current.Platform;
+			IList<string> ExtraPlatformArguments = new List<string>();
 
-			if (OperatingSystem.IsWindows())
+			ClangToolChain? ClangToolChain = TargetToolChain as ClangToolChain;
+			if (ClangToolChain != null)
 			{
-				VCEnvironment Environment = VCEnvironment.Create(WindowsCompiler.Clang, Target.Platform,
-					Target.Rules.WindowsPlatform.Architecture, null,
-					Target.Rules.WindowsPlatform.WindowsSdkVersion, null,
-					Logger);
-
-				return Environment.CompilerPath;
+				ClangToolChain.AddExtraToolArguments(ExtraPlatformArguments);
 			}
-			else if (OperatingSystem.IsLinux())
+
+			return ExtraPlatformArguments;
+		}
+
+		static private string ConvertResponseFile(FileItem OriginalFileItem, IEnumerable<string> FileContents, ILogger Logger)
+		{
+			List<string> NewFileContents = new List<string>(FileContents);
+			FileItem NewFileItem = FileItem.GetItemByFileReference(new FileReference(OriginalFileItem.AbsolutePath + ".gcd"));
+
+			for (int i = 0; i < NewFileContents.Count; i++)
 			{
-				string? Clang = LinuxCommon.WhichClang(Logger);
+				string Line = NewFileContents[i].TrimStart();
 
-				if (Clang != null)
+				// Ignore empty strings
+				if (String.IsNullOrEmpty(Line)) 
 				{
-					return FileReference.FromString(Clang);
+					continue;
 				}
-			}
-			else if (OperatingSystem.IsMacOS())
-			{
-				MacToolChainSettings Settings = new MacToolChainSettings(false, Logger);
-				DirectoryReference? ToolchainDir = DirectoryReference.FromString(Settings.ToolchainDir);
-
-				if (ToolchainDir != null)
+				// The file that is going to compile
+				else if (!Line.StartsWith("-") && !Line.StartsWith("/"))
 				{
-					return FileReference.Combine(ToolchainDir, "clang++");
+					var OldPath = Line.Replace("\"", "");
+					var FileReference = new FileReference(OldPath);
+					Line = Line.Replace(OldPath, FileReference.FullName.Replace("\\", "/"));
 				}
+				// Arguments
+				else
+				{
+					Match LineMatch = LineRegex.Match(Line);
+					if (LineMatch.Success)
+					{
+						var OldPath = LineMatch.Groups[2].Value.Replace("\"", "");
+						if (OldPath[^1] == '\\' || OldPath[^1] == '/')
+						{
+							OldPath = OldPath.Remove(OldPath.Length - 1, 1);
+						}
+
+						var FileReference = new FileReference(OldPath);
+						Line = Line.Replace(OldPath, FileReference.FullName.Replace("\\", "/"));
+					}
+				}
+
+				NewFileContents[i] = Line;
 			}
-			return FileReference.FromString("clang++");
+
+			Utils.WriteFileIfChanged(NewFileItem, NewFileContents, Logger);
+
+			return NewFileItem.AbsolutePath;
 		}
 	}
 }
diff --git a/Engine/Source/Programs/UnrealBuildTool/Platform/Linux/LinuxToolChain.cs b/Engine/Source/Programs/UnrealBuildTool/Platform/Linux/LinuxToolChain.cs
--- a/Engine/Source/Programs/UnrealBuildTool/Platform/Linux/LinuxToolChain.cs
+++ b/Engine/Source/Programs/UnrealBuildTool/Platform/Linux/LinuxToolChain.cs
@@ -1580,5 +1580,16 @@ namespace UnrealBuildTool
 			StartInfo.CreateNoWindow = true;
 			Utils.RunLocalProcessAndLogOutput(StartInfo, Logger);
 		}
+
+		public override void AddExtraToolArguments(IList<string> ExtraArguments)
+		{
+			// We explicitly include the clang include directories so tools like IWYU can run outside of the clang directory.
+			// More info: [REDACTED_PROFILE_URL]/include-what-you-use#how-to-install
+			string? InternalSdkPath = UEBuildPlatform.GetSDK(UnrealTargetPlatform.Linux)!.GetInternalSDKPath();
+			if (InternalSdkPath != null)
+			{
+				ExtraArguments.Add(String.Format("-isystem\"{0}\"", System.IO.Path.Combine(InternalSdkPath, "lib/clang/" + Info.ClangVersion + "/include/").Replace("\\", "/")));
+			}
+		}
 	}
 }
diff --git a/Engine/Source/Programs/UnrealBuildTool/ToolChain/ClangToolChain.cs b/Engine/Source/Programs/UnrealBuildTool/ToolChain/ClangToolChain.cs
--- a/Engine/Source/Programs/UnrealBuildTool/ToolChain/ClangToolChain.cs
+++ b/Engine/Source/Programs/UnrealBuildTool/ToolChain/ClangToolChain.cs
@@ -1049,5 +1049,13 @@ namespace UnrealBuildTool
 
 			return Result;
 		}
+
+		/// <summary>
+		/// Used by other tools to get the extra arguments to run vanilla clang for a particular platform.
+		/// </summary>
+		/// <param name="ExtraArguments">List of extra arguments to add to.</param>
+		public virtual void AddExtraToolArguments(IList<string> ExtraArguments)
+		{
+		}
 	}
 }
\ No newline at end of file
