UnrealBuildTool: Option to remove '_of_x' from cpp unity filenames

#rnx

diff --git a/Engine/Source/Programs/UnrealBuildTool/Configuration/TargetRules.cs b/Engine/Source/Programs/UnrealBuildTool/Configuration/TargetRules.cs
--- a/Engine/Source/Programs/UnrealBuildTool/Configuration/TargetRules.cs
+++ b/Engine/Source/Programs/UnrealBuildTool/Configuration/TargetRules.cs
@@ -1423,6 +1423,14 @@ namespace UnrealBuildTool
 		[XmlConfigFile(Category = "BuildConfiguration")]
 		public bool bStressTestUnity = false;
 
+		/// <summary>
+		/// Whether to add additional information to the unity files, such as '_of_X' in the file name.
+		/// </summary>
+		[CommandLine("-DetailedUnityFiles", Value = "true")]
+		[CommandLine("-NoDetailedUnityFiles", Value = "false")]
+		[XmlConfigFile(Category = "BuildConfiguration")]
+		public bool bDetailedUnityFiles = true;
+
 		/// <summary>
 		/// Whether to force debug info to be generated.
 		/// </summary>
@@ -3070,6 +3078,11 @@ namespace UnrealBuildTool
 			get { return Inner.bStressTestUnity; }
 		}
 
+		public bool bDetailedUnityFiles
+		{
+			get { return Inner.bDetailedUnityFiles; }
+		}
+
 		public bool bDisableDebugInfo
 		{
 			get { return Inner.bDisableDebugInfo; }
diff --git a/Engine/Source/Programs/UnrealBuildTool/System/Unity.cs b/Engine/Source/Programs/UnrealBuildTool/System/Unity.cs
--- a/Engine/Source/Programs/UnrealBuildTool/System/Unity.cs
+++ b/Engine/Source/Programs/UnrealBuildTool/System/Unity.cs
@@ -258,7 +258,14 @@ namespace UnrealBuildTool
 				string UnityCPPFileName;
 				if (AllUnityFiles.Count > 1)
 				{
-					UnityCPPFileName = string.Format("{0}{1}.{2}_of_{3}.cpp", ModulePrefix, BaseName, CurrentUnityFileCount, AllUnityFiles.Count);
+					if (Target.bDetailedUnityFiles)
+					{
+						UnityCPPFileName = string.Format("{0}{1}.{2}_of_{3}.cpp", ModulePrefix, BaseName, CurrentUnityFileCount, AllUnityFiles.Count);
+					}
+					else
+					{
+						UnityCPPFileName = string.Format("{0}{1}.{2}.cpp", ModulePrefix, BaseName, CurrentUnityFileCount);
+					}
 				}
 				else
 				{
