UnrealBuildTool: Determinisitic build configuration update to manage via TargetRules and ModuleRules

#rnx

diff --git a/Engine/Source/Programs/UnrealBuildTool/Configuration/ModuleRules.cs b/Engine/Source/Programs/UnrealBuildTool/Configuration/ModuleRules.cs
--- a/Engine/Source/Programs/UnrealBuildTool/Configuration/ModuleRules.cs
+++ b/Engine/Source/Programs/UnrealBuildTool/Configuration/ModuleRules.cs
@@ -745,6 +745,11 @@ namespace UnrealBuildTool
 		}
 		private EngineIncludeOrderVersion? IncludeOrderVersionPrivate;
 
+		/// <summary>
+		/// Set flags for determinstic compiles (experimental, may not be fully supported). Deterministic linking is controlled via TargetRules.
+		/// </summary>
+		public bool bDeterministic = false;
+
 		/// <summary>
 		/// Use run time type information
 		/// </summary>
diff --git a/Engine/Source/Programs/UnrealBuildTool/Configuration/TargetRules.cs b/Engine/Source/Programs/UnrealBuildTool/Configuration/TargetRules.cs
--- a/Engine/Source/Programs/UnrealBuildTool/Configuration/TargetRules.cs
+++ b/Engine/Source/Programs/UnrealBuildTool/Configuration/TargetRules.cs
@@ -1470,6 +1470,23 @@ namespace UnrealBuildTool
 		[XmlConfigFile(Category = "BuildConfiguration")]
 		public bool bUsePCHFiles = true;
 
+		/// <summary>
+		/// Set flags require for determinstic linking (experimental, may not be fully supported).
+		/// Deterministic compiling is controlled via ModuleRules.
+		/// </summary>
+		[XmlConfigFile(Category = "BuildConfiguration")]
+		public bool bDeterministic = false;
+
+		/// <summary>
+		/// Force set flags require for determinstic compiling and linking (experimental, may not be fully supported).
+		/// This setting is only recommended for testing, instead:
+		/// * Set bDeterministic on a per module basis in ModuleRules to control deterministic compiling.
+		/// * Set bDeterministic on a per target basis in TargetRules to control deterministic linking.
+		/// </summary>
+		[CommandLine("-Deterministic")]
+		[XmlConfigFile(Category = "BuildConfiguration")]
+		public bool bForceDeterministic = false;
+
 		/// <summary>
 		/// Whether to just preprocess source files for this target, and skip compilation
 		/// </summary>
@@ -3108,6 +3125,16 @@ namespace UnrealBuildTool
 			get { return Inner.bUsePCHFiles; }
 		}
 
+		public bool bDeterministic
+		{
+			get { return Inner.bDeterministic; }
+		}
+
+		public bool bForceDeterministic
+		{
+			get { return Inner.bForceDeterministic; }
+		}
+
 		public bool bPreprocessOnly
 		{
 			get { return Inner.bPreprocessOnly; }
diff --git a/Engine/Source/Programs/UnrealBuildTool/Configuration/UEBuildModuleCPP.cs b/Engine/Source/Programs/UnrealBuildTool/Configuration/UEBuildModuleCPP.cs
--- a/Engine/Source/Programs/UnrealBuildTool/Configuration/UEBuildModuleCPP.cs
+++ b/Engine/Source/Programs/UnrealBuildTool/Configuration/UEBuildModuleCPP.cs
@@ -1490,6 +1490,7 @@ namespace UnrealBuildTool
 			Result.StaticAnalyzerAdditionalCheckers = Rules.StaticAnalyzerAdditionalCheckers;
 			Result.bEnableUndefinedIdentifierWarnings = Rules.bEnableUndefinedIdentifierWarnings;
 			Result.IncludeOrderVersion = Rules.IncludeOrderVersion;
+			Result.bDeterministic |= Rules.bDeterministic;
 
 			// If the module overrides the C++ language version, override it on the compile environment
 			if (Rules.CppStandard != CppStandardVersion.Default)
diff --git a/Engine/Source/Programs/UnrealBuildTool/Configuration/UEBuildTarget.cs b/Engine/Source/Programs/UnrealBuildTool/Configuration/UEBuildTarget.cs
--- a/Engine/Source/Programs/UnrealBuildTool/Configuration/UEBuildTarget.cs
+++ b/Engine/Source/Programs/UnrealBuildTool/Configuration/UEBuildTarget.cs
@@ -3889,6 +3889,7 @@ namespace UnrealBuildTool
 			GlobalCompileEnvironment.CppStandard = Rules.CppStandard;
 			GlobalCompileEnvironment.CStandard = Rules.CStandard;
 			GlobalCompileEnvironment.AdditionalArguments = Rules.AdditionalCompilerArguments ?? String.Empty;
+			GlobalCompileEnvironment.bDeterministic = Rules.bForceDeterministic;
 			GlobalCompileEnvironment.CrashDiagnosticDirectory = Rules.CrashDiagnosticDirectory;
 
 			GlobalLinkEnvironment.bUseDebugCRT = bUseDebugCRT;
@@ -3908,6 +3909,7 @@ namespace UnrealBuildTool
 			GlobalLinkEnvironment.bPGOProfile = Rules.bPGOProfile;
 			GlobalLinkEnvironment.bUseIncrementalLinking = Rules.bUseIncrementalLinking;
 			GlobalLinkEnvironment.bUseFastPDBLinking = Rules.bUseFastPDBLinking ?? false;
+			GlobalLinkEnvironment.bDeterministic = Rules.bDeterministic || Rules.bForceDeterministic;
 			GlobalLinkEnvironment.bPrintTimingInfo = Rules.bPrintToolChainTimingInfo;
 			GlobalLinkEnvironment.AdditionalArguments = Rules.AdditionalLinkerArguments ?? String.Empty;
 
diff --git a/Engine/Source/Programs/UnrealBuildTool/Platform/Windows/UEBuildWindows.cs b/Engine/Source/Programs/UnrealBuildTool/Platform/Windows/UEBuildWindows.cs
--- a/Engine/Source/Programs/UnrealBuildTool/Platform/Windows/UEBuildWindows.cs
+++ b/Engine/Source/Programs/UnrealBuildTool/Platform/Windows/UEBuildWindows.cs
@@ -368,13 +368,6 @@ namespace UnrealBuildTool
 		[XmlConfigFile(Category = "WindowsPlatform")]
 		public bool bShowIncludes = false;
 
-		/// <summary>
-		/// Set flags require for determinstic compiles (experimental)
-		/// </summary>
-		[CommandLine("-Deterministic")]
-		[XmlConfigFile(Category = "WindowsPlatform")]
-		public bool bDeterministic = false;
-
 		/// <summary>
 		/// Bundle a working version of dbghelp.dll with the application, and use this to generate minidumps. This works around a bug with the Windows 10 Fall Creators Update (1709)
 		/// where rich PE headers larger than a certain size would result in corrupt minidumps.
@@ -668,11 +661,6 @@ namespace UnrealBuildTool
 			get { return Inner.bShowIncludes; }
 		}
 
-		public bool bDeterministic
-		{
-			get { return Inner.bDeterministic; }
-		}
-
 		public string GetVisualStudioCompilerVersionName()
 		{
 			return Inner.GetVisualStudioCompilerVersionName();
diff --git a/Engine/Source/Programs/UnrealBuildTool/Platform/Windows/VCToolChain.cs b/Engine/Source/Programs/UnrealBuildTool/Platform/Windows/VCToolChain.cs
--- a/Engine/Source/Programs/UnrealBuildTool/Platform/Windows/VCToolChain.cs
+++ b/Engine/Source/Programs/UnrealBuildTool/Platform/Windows/VCToolChain.cs
@@ -466,12 +466,16 @@ namespace UnrealBuildTool
 			}
 
 			// Experimental deterministic compile support
-			if (Target.WindowsPlatform.bDeterministic)
+			if (CompileEnvironment.bDeterministic)
 			{
 				if (Target.WindowsPlatform.Compiler.IsMSVC())
 				{
 					Arguments.Add("/experimental:deterministic");
 				}
+				else if (Target.WindowsPlatform.Compiler.IsClang())
+				{
+					Arguments.Add("/Brepro");
+				}
 			}
 
 			if (Target.WindowsPlatform.Compiler.IsMSVC())
@@ -925,6 +929,12 @@ namespace UnrealBuildTool
 					}
 				}
 
+				// Warn if __DATE__ or __TIME__ are used as they prevent reproducible builds
+				if (CompileEnvironment.bDeterministic)
+				{
+					Arguments.Add("-Wdate-time -Wno-error=date-time");
+				}
+
 				// This is disabled because clang explicitly warns about changing pack alignment in a header and not
 				// restoring it afterwards, which is something we do with the Pre/PostWindowsApi.h headers.
 				Arguments.Add("-Wno-pragma-pack");
@@ -1114,13 +1124,17 @@ namespace UnrealBuildTool
 			// Don't embed the full PDB path; we want to be able to move binaries elsewhere. They will always be side by side.
 			Arguments.Add("/PDBALTPATH:%_PDB%");
 
-			// Experimental deterministic compile support
-			if (Target.WindowsPlatform.bDeterministic)
+			// Experimental deterministic link support
+			if (LinkEnvironment.bDeterministic)
 			{
 				if (Target.WindowsPlatform.Compiler.IsMSVC())
 				{
 					Arguments.Add("/experimental:deterministic");
 				}
+				else if (Target.WindowsPlatform.Compiler.IsClang())
+				{
+					Arguments.Add("/Brepro");
+				}
 			}
 
 			if (Target.WindowsPlatform.Compiler.IsMSVC() && !LinkEnvironment.bPGOOptimize && !LinkEnvironment.bPGOProfile)
diff --git a/Engine/Source/Programs/UnrealBuildTool/System/CppCompileEnvironment.cs b/Engine/Source/Programs/UnrealBuildTool/System/CppCompileEnvironment.cs
--- a/Engine/Source/Programs/UnrealBuildTool/System/CppCompileEnvironment.cs
+++ b/Engine/Source/Programs/UnrealBuildTool/System/CppCompileEnvironment.cs
@@ -471,6 +471,11 @@ namespace UnrealBuildTool
 		/// </summary>
 		public EngineIncludeOrderVersion IncludeOrderVersion = EngineIncludeOrderVersion.Latest;
 
+		/// <summary>
+		/// Set flags for determinstic compiles (experimental).
+		/// </summary>
+		public bool bDeterministic = false;
+
 		/// <summary>
 		/// Directory where to put crash report files for platforms that support it
 		/// </summary>
@@ -563,6 +568,7 @@ namespace UnrealBuildTool
 			CStandard = Other.CStandard;
 			bEnableCoroutines = Other.bEnableCoroutines;
 			IncludeOrderVersion = Other.IncludeOrderVersion;
+			bDeterministic = Other.bDeterministic;
 			CrashDiagnosticDirectory = Other.CrashDiagnosticDirectory;
 		}
 	}
diff --git a/Engine/Source/Programs/UnrealBuildTool/System/LinkEnvironment.cs b/Engine/Source/Programs/UnrealBuildTool/System/LinkEnvironment.cs
--- a/Engine/Source/Programs/UnrealBuildTool/System/LinkEnvironment.cs
+++ b/Engine/Source/Programs/UnrealBuildTool/System/LinkEnvironment.cs
@@ -256,6 +256,11 @@ namespace UnrealBuildTool
 		/// </summary>
 		public bool bIgnoreUnresolvedSymbols;
 
+		/// <summary>
+		/// Set flags for determinstic compiles (experimental).
+		/// </summary>
+		public bool bDeterministic = false;
+
 		/// <summary>
 		/// Whether to log detailed timing information
 		/// </summary>
@@ -373,6 +378,7 @@ namespace UnrealBuildTool
 			bUsePDBFiles = Other.bUsePDBFiles;
 			bUseFastPDBLinking = Other.bUseFastPDBLinking;
 			bIgnoreUnresolvedSymbols = Other.bIgnoreUnresolvedSymbols;
+			bDeterministic = Other.bDeterministic;
 			bPrintTimingInfo = Other.bPrintTimingInfo;
 			PackagePath = Other.PackagePath;
 			CrashDiagnosticDirectory = Other.CrashDiagnosticDirectory;
diff --git a/Engine/Source/Programs/UnrealBuildTool/ToolChain/ClangToolChain.cs b/Engine/Source/Programs/UnrealBuildTool/ToolChain/ClangToolChain.cs
--- a/Engine/Source/Programs/UnrealBuildTool/ToolChain/ClangToolChain.cs
+++ b/Engine/Source/Programs/UnrealBuildTool/ToolChain/ClangToolChain.cs
@@ -550,6 +550,12 @@ namespace UnrealBuildTool
 				Arguments.Add("-Werror=deprecated-declarations");
 			}
 
+			// Warn if __DATE__ or __TIME__ are used as they prevent reproducible builds
+			if (CompileEnvironment.bDeterministic)
+			{
+				Arguments.Add("-Wdate-time -Wno-error=date-time"); // https://clang.llvm.org/docs/DiagnosticsReference.html#wdate-time
+			}
+
 			// https://clang.llvm.org/docs/DiagnosticsReference.html#wshadow
 			if (CompileEnvironment.ShadowVariableWarningLevel != WarningLevel.Off)
 			{
