- Updated abseil libs for arm64 Mac (fat libs), from francisco.vicente
- Re-enabled the onnx plugins for arm64

diff --git a/Engine/Source/Programs/UnrealBuildTool/Platform/Mac/UEBuildMac.cs b/Engine/Source/Programs/UnrealBuildTool/Platform/Mac/UEBuildMac.cs
--- a/Engine/Source/Programs/UnrealBuildTool/Platform/Mac/UEBuildMac.cs
+++ b/Engine/Source/Programs/UnrealBuildTool/Platform/Mac/UEBuildMac.cs
@@ -142,14 +142,14 @@ namespace UnrealBuildTool
 				Target.DisablePlugins.AddRange(new string[]
 				{
 					// onnx:
-					"NeuralNetworkInference",
-					"MLAdapter",
-					"LegacyVertexDeltaModel",
-					"MLDeformerFramework",
-					"VertexDeltaModel",
-					"MetaHumanAnimator",
-					"MetaHumans",
-					"NeuralMorphModel",
+					//"NeuralNetworkInference",
+					//"MLAdapter",
+					//"LegacyVertexDeltaModel",
+					//"MLDeformerFramework",
+					//"VertexDeltaModel",
+					//"MetaHumanAnimator",
+					//"MetaHumans",
+					//"NeuralMorphModel",
 
 					// EOSSDK
 					"EOSManager",
