- Deleted pointless commented lines

diff --git a/Engine/Source/Programs/UnrealBuildTool/Platform/Mac/UEBuildMac.cs b/Engine/Source/Programs/UnrealBuildTool/Platform/Mac/UEBuildMac.cs
--- a/Engine/Source/Programs/UnrealBuildTool/Platform/Mac/UEBuildMac.cs
+++ b/Engine/Source/Programs/UnrealBuildTool/Platform/Mac/UEBuildMac.cs
@@ -141,16 +141,6 @@ namespace UnrealBuildTool
 			{
 				Target.DisablePlugins.AddRange(new string[]
 				{
-					// onnx:
-					//"NeuralNetworkInference",
-					//"MLAdapter",
-					//"LegacyVertexDeltaModel",
-					//"MLDeformerFramework",
-					//"VertexDeltaModel",
-					//"MetaHumanAnimator",
-					//"MetaHumans",
-					//"NeuralMorphModel",
-
 					// EOSSDK
 					"EOSManager",
 					"EOSShared",
