- Restored some behavior I had left commented out by mistake

diff --git a/Engine/Source/Programs/UnrealBuildTool/ProjectFiles/Xcode/XcodeProjectFileGenerator.cs b/Engine/Source/Programs/UnrealBuildTool/ProjectFiles/Xcode/XcodeProjectFileGenerator.cs
--- a/Engine/Source/Programs/UnrealBuildTool/ProjectFiles/Xcode/XcodeProjectFileGenerator.cs
+++ b/Engine/Source/Programs/UnrealBuildTool/ProjectFiles/Xcode/XcodeProjectFileGenerator.cs
@@ -155,10 +155,11 @@ namespace UnrealBuildTool
 			WorkspaceSettingsContent.Append("<!DOCTYPE plist PUBLIC \"-//Apple//DTD PLIST 1.0//EN\" \"http://www.apple.com/DTDs/PropertyList-1.0.dtd\">" + ProjectFileGenerator.NewLine);
 			WorkspaceSettingsContent.Append("<plist version=\"1.0\">" + ProjectFileGenerator.NewLine);
 			WorkspaceSettingsContent.Append("<dict>" + ProjectFileGenerator.NewLine);
-			//WorkspaceSettingsContent.Append("\t<key>DisableBuildSystemDeprecationWarning</key>" + ProjectFileGenerator.NewLine);
-			//WorkspaceSettingsContent.Append("\t<true/>" + ProjectFileGenerator.NewLine);
-			//WorkspaceSettingsContent.Append("\t<key>DisableBuildSystemDeprecationDiagnostic</key>" + ProjectFileGenerator.NewLine);
-			//WorkspaceSettingsContent.Append("\t<true/>" + ProjectFileGenerator.NewLine);
+			// @todo when we move to xcode 14, we remove these next 2 keys
+			WorkspaceSettingsContent.Append("\t<key>DisableBuildSystemDeprecationWarning</key>" + ProjectFileGenerator.NewLine);
+			WorkspaceSettingsContent.Append("\t<true/>" + ProjectFileGenerator.NewLine);
+			WorkspaceSettingsContent.Append("\t<key>DisableBuildSystemDeprecationDiagnostic</key>" + ProjectFileGenerator.NewLine);
+			WorkspaceSettingsContent.Append("\t<true/>" + ProjectFileGenerator.NewLine);
 			WorkspaceSettingsContent.Append("\t<key>IDEWorkspaceSharedSettings_AutocreateContextsIfNeeded</key>" + ProjectFileGenerator.NewLine);
 			WorkspaceSettingsContent.Append("\t<false/>" + ProjectFileGenerator.NewLine);
 			WorkspaceSettingsContent.Append("</dict>" + ProjectFileGenerator.NewLine);
