UnrealBuildTool: Disallow .NET projects in the generated VS2019 solution and warn when generating

#rnx

diff --git a/Engine/Source/Programs/UnrealBuildTool/ProjectFiles/VisualStudio/VCProjectFileGenerator.cs b/Engine/Source/Programs/UnrealBuildTool/ProjectFiles/VisualStudio/VCProjectFileGenerator.cs
--- a/Engine/Source/Programs/UnrealBuildTool/ProjectFiles/VisualStudio/VCProjectFileGenerator.cs
+++ b/Engine/Source/Programs/UnrealBuildTool/ProjectFiles/VisualStudio/VCProjectFileGenerator.cs
@@ -319,6 +319,13 @@ namespace UnrealBuildTool
 					}
 				}
 			}
+
+			if (bIncludeDotNetPrograms && Settings.ProjectFileFormat == VCProjectFileFormat.VisualStudio2019)
+			{
+				Logger.LogWarning("Visual Studio 2019 does not support .NET 6.0 C# projects, these projects will not be added to the generated solution.");
+				Logger.LogWarning("Please generate the Visual Studio 2022 solution if .NET 6.0 C# project support is required.");
+				bIncludeDotNetPrograms = false;
+			}
 		}
 
 
