- Fixed a bug with stub IPAs
- Re-disabled automatic code signing

diff --git a/Engine/Config/BaseEngine.ini b/Engine/Config/BaseEngine.ini
--- a/Engine/Config/BaseEngine.ini
+++ b/Engine/Config/BaseEngine.ini
@@ -2656,7 +2656,7 @@ bEnableCloudKitSupport=False
 bRunAsCurrentUser=False
 IOSCloudKitSyncStrategy=None
 bGenerateCrashReportSymbols=false
-bAutomaticSigning=true
+bAutomaticSigning=false
 UseFastIntrinsics=False
 ForceFloats=True
 EnableMathOptimisations=True
diff --git a/Engine/Source/Programs/UnrealBuildTool/Platform/IOS/IOSToolChain.cs b/Engine/Source/Programs/UnrealBuildTool/Platform/IOS/IOSToolChain.cs
--- a/Engine/Source/Programs/UnrealBuildTool/Platform/IOS/IOSToolChain.cs
+++ b/Engine/Source/Programs/UnrealBuildTool/Platform/IOS/IOSToolChain.cs
@@ -1843,6 +1843,7 @@ namespace UnrealBuildTool
 
 				if (!Target.bBuildAsFramework)
 				{
+					Logger.LogInformatioj("Generating project...");
 					DirectoryReference? GeneratedProjectFile;
 				    if (AppName == "UnrealGame" || AppName == "UnrealClient" || Target.ProjectFile == null || Target.ProjectFile.IsUnderDirectory(Unreal.EngineDirectory))
 				    {
@@ -1852,8 +1853,9 @@ namespace UnrealBuildTool
 				    {
 					    GenerateProjectFiles(Target.ProjectFile, new string[] { "-platforms=" + (Target.Platform == UnrealTargetPlatform.IOS ? "IOS" : "TVOS"), "-NoIntellIsense", (Target.Platform == UnrealTargetPlatform.IOS ? "-iosdeployonly" : "-tvosdeployonly"), "-ignorejunk", (Target.bForDistribution ? "-distribution" : "-development"), String.Format("-project={0}", Target.ProjectFile), "-game", "-bundleID=" + BundleID, "-includetemptargets" }, Logger, out GeneratedProjectFile);
 				    }
-				    // Make sure it exists
-				    if (!DirectoryReference.Exists(XcodeWorkspaceDir!) || (GeneratedProjectFile != null && GeneratedProjectFile.ParentDirectory != XcodeWorkspaceDir))
+					Logger.LogInformatioj("Generated {0}!", GeneratedProjectFile?.FullName ?? "NONE");
+					// Make sure it exists
+					if (!DirectoryReference.Exists(XcodeWorkspaceDir!) || (GeneratedProjectFile != null && GeneratedProjectFile != XcodeWorkspaceDir))
 				    {
 					    throw new BuildException("Unable to create stub IPA; Xcode workspace not found at {0}", XcodeWorkspaceDir);
 				    }
