Fix issue with ios logger to unblock builds

diff --git a/Engine/Source/Programs/UnrealBuildTool/Platform/IOS/IOSToolChain.cs b/Engine/Source/Programs/UnrealBuildTool/Platform/IOS/IOSToolChain.cs
--- a/Engine/Source/Programs/UnrealBuildTool/Platform/IOS/IOSToolChain.cs
+++ b/Engine/Source/Programs/UnrealBuildTool/Platform/IOS/IOSToolChain.cs
@@ -1843,7 +1843,7 @@ namespace UnrealBuildTool
 
 				if (!Target.bBuildAsFramework)
 				{
-					Logger.LogInformatioj("Generating project...");
+					Logger.LogInformation("Generating project...");
 					DirectoryReference? GeneratedProjectFile;
 				    if (AppName == "UnrealGame" || AppName == "UnrealClient" || Target.ProjectFile == null || Target.ProjectFile.IsUnderDirectory(Unreal.EngineDirectory))
 				    {
