Another fix for build in ios toolchain
#rnx

diff --git a/Engine/Source/Programs/UnrealBuildTool/Platform/IOS/IOSToolChain.cs b/Engine/Source/Programs/UnrealBuildTool/Platform/IOS/IOSToolChain.cs
--- a/Engine/Source/Programs/UnrealBuildTool/Platform/IOS/IOSToolChain.cs
+++ b/Engine/Source/Programs/UnrealBuildTool/Platform/IOS/IOSToolChain.cs
@@ -1853,7 +1853,7 @@ namespace UnrealBuildTool
 				    {
 					    GenerateProjectFiles(Target.ProjectFile, new string[] { "-platforms=" + (Target.Platform == UnrealTargetPlatform.IOS ? "IOS" : "TVOS"), "-NoIntellIsense", (Target.Platform == UnrealTargetPlatform.IOS ? "-iosdeployonly" : "-tvosdeployonly"), "-ignorejunk", (Target.bForDistribution ? "-distribution" : "-development"), String.Format("-project={0}", Target.ProjectFile), "-game", "-bundleID=" + BundleID, "-includetemptargets" }, Logger, out GeneratedProjectFile);
 				    }
-					Logger.LogInformatioj("Generated {0}!", GeneratedProjectFile?.FullName ?? "NONE");
+					Logger.LogInformation("Generated {0}!", GeneratedProjectFile?.FullName ?? "NONE");
 					// Make sure it exists
 					if (!DirectoryReference.Exists(XcodeWorkspaceDir!) || (GeneratedProjectFile != null && GeneratedProjectFile != XcodeWorkspaceDir))
 				    {
