UnrealBuildTool: Change VS2019 .NET project generation warning to information

#rnx

diff --git a/Engine/Source/Programs/UnrealBuildTool/ProjectFiles/VisualStudio/VCProjectFileGenerator.cs b/Engine/Source/Programs/UnrealBuildTool/ProjectFiles/VisualStudio/VCProjectFileGenerator.cs
--- a/Engine/Source/Programs/UnrealBuildTool/ProjectFiles/VisualStudio/VCProjectFileGenerator.cs
+++ b/Engine/Source/Programs/UnrealBuildTool/ProjectFiles/VisualStudio/VCProjectFileGenerator.cs
@@ -322,8 +322,8 @@ namespace UnrealBuildTool
 
 			if (bIncludeDotNetPrograms && Settings.ProjectFileFormat == VCProjectFileFormat.VisualStudio2019)
 			{
-				Logger.LogWarning("Visual Studio 2019 does not support .NET 6.0 C# projects, these projects will not be added to the generated solution.");
-				Logger.LogWarning("Please generate the Visual Studio 2022 solution if .NET 6.0 C# project support is required.");
+				Logger.LogInformation("Visual Studio 2019 does not support .NET 6.0 C# projects, these projects will not be added to the generated solution.");
+				Logger.LogInformation("Please generate the Visual Studio 2022 solution if .NET 6.0 C# project support is required.");
 				bIncludeDotNetPrograms = false;
 			}
 		}
