UnrealBuildTool: Generate .vsconfig that recommends Visual Studio components to be installed, this will prompt the user to run the installer in the Solution Explorer

PR #9185 (inspired by)

diff --git a/Engine/Source/Programs/UnrealBuildTool/Platform/Windows/MicrosoftPlatformSDK.Versions.cs b/Engine/Source/Programs/UnrealBuildTool/Platform/Windows/MicrosoftPlatformSDK.Versions.cs
--- a/Engine/Source/Programs/UnrealBuildTool/Platform/Windows/MicrosoftPlatformSDK.Versions.cs
+++ b/Engine/Source/Programs/UnrealBuildTool/Platform/Windows/MicrosoftPlatformSDK.Versions.cs
@@ -1,6 +1,7 @@
 // Copyright Epic Games, Inc. All Rights Reserved.
 
 using EpicGames.Core;
+using System.Collections.Generic;
 using System.Linq;
 
 namespace UnrealBuildTool
@@ -55,6 +56,62 @@ namespace UnrealBuildTool
 
 		static readonly VersionNumber MinimumIntelOneApiVersion = new VersionNumber(2021, 0, 0);
 
+		/// <summary>
+		/// The default set of components that should be suggested to be installed for Visual Studio 2019 or 2022.
+		/// This or the 2019\2022 specific components should be updated if the referred visual cpp version changes
+		/// </summary>
+		static readonly string[] VisualStudioSuggestedComponents = new string[]
+		{
+			"Microsoft.VisualStudio.Workload.CoreEditor",
+			"Microsoft.VisualStudio.Workload.Universal",
+			"Microsoft.VisualStudio.Workload.ManagedDesktop",
+			"Microsoft.VisualStudio.Workload.NativeDesktop",
+			"Microsoft.VisualStudio.Workload.NativeGame",
+			"Microsoft.VisualStudio.Component.VC.Tools.x86.x64",
+			"Microsoft.VisualStudio.Component.VC.Tools.ARM64",
+			"Microsoft.VisualStudio.Component.Windows10SDK",
+		};
+
+		/// <summary>
+		/// Additional set of components that should be suggested to be installed for Visual Studio 2019.
+		/// </summary>
+		static readonly string[] VisualStudio2019SuggestedComponents = new string[]
+		{
+		};
+
+		/// <summary>
+		/// Additional set of components that should be suggested to be installed for Visual Studio 2022.
+		/// </summary>
+		static readonly string[] VisualStudio2022SuggestedComponents = new string[]
+		{
+			"Microsoft.VisualStudio.Workload.NativeCrossPlat",
+			"Microsoft.VisualStudio.Component.VC.14.32.17.2.ARM64",
+			"Microsoft.VisualStudio.Component.VC.14.32.17.2.x86.x64"
+		};
+
+		/// <summary>
+		/// Returns the list of suggested of components that should be suggested to be installed for Visual Studio.
+		/// Used to generate a .vsconfig file which will prompt Visual Studio to ask the user to install these components.
+		/// </summary>
+		public static IEnumerable<string> GetVisualStudioSuggestedComponents(VCProjectFileFormat Format)
+		{
+			SortedSet<string> Components = new SortedSet<string>();
+			Components.UnionWith(VisualStudioSuggestedComponents);
+
+			switch (Format)
+			{
+				case VCProjectFileFormat.VisualStudio2019:
+					Components.UnionWith(VisualStudio2019SuggestedComponents);
+					break;
+				case VCProjectFileFormat.VisualStudio2022:
+					Components.UnionWith(VisualStudio2022SuggestedComponents);
+					break;
+				default:
+					throw new BuildException("Unsupported Visual Studio version");
+			}
+
+			return Components;
+		}
 
 		public override string GetMainVersion()
 		{
diff --git a/Engine/Source/Programs/UnrealBuildTool/ProjectFiles/VisualStudio/VCProjectFileGenerator.cs b/Engine/Source/Programs/UnrealBuildTool/ProjectFiles/VisualStudio/VCProjectFileGenerator.cs
--- a/Engine/Source/Programs/UnrealBuildTool/ProjectFiles/VisualStudio/VCProjectFileGenerator.cs
+++ b/Engine/Source/Programs/UnrealBuildTool/ProjectFiles/VisualStudio/VCProjectFileGenerator.cs
@@ -56,6 +56,12 @@ namespace UnrealBuildTool
 		[XmlConfigFile(Category = "BuildConfiguration")]
 		public bool bWriteSolutionOptionFile = true;
 
+		/// <summary>
+		/// Whether to write a .vsconfig file next to the sln to suggest components to install.
+		/// </summary>
+		[XmlConfigFile(Category = "BuildConfiguration")]
+		public bool bVsConfigFile = true;
+
 		/// <summary>
 		/// Forces UBT to be built in debug configuration, regardless of the solution configuration
 		/// </summary>
@@ -173,6 +179,11 @@ namespace UnrealBuildTool
 			{
 				FileReference.Delete(PrimaryProjDeleteFilename);
 			}
+			PrimaryProjDeleteFilename = FileReference.Combine(InPrimaryProjectDirectory, ".vsconfig");
+			if (FileReference.Exists(PrimaryProjDeleteFilename))
+			{
+				FileReference.Delete(PrimaryProjDeleteFilename);
+			}
 
 			// Delete the project files folder
 			if (DirectoryReference.Exists(InIntermediateProjectFilesDirectory))
@@ -888,6 +899,23 @@ namespace UnrealBuildTool
 				}
 			}
 
+			if (bSuccess && Settings.bVsConfigFile && OperatingSystem.IsWindows())
+			{
+				StringBuilder VsConfigFileContent = new StringBuilder();
+
+				VsConfigFileContent.AppendLine("{");
+				VsConfigFileContent.AppendLine("  \"version\": \"1.0\",");
+				VsConfigFileContent.AppendLine("  \"components\": [");
+				IEnumerable<string> Components = MicrosoftPlatformSDK.GetVisualStudioSuggestedComponents(Settings.ProjectFileFormat);
+				string ComponentsString = string.Join($",{Environment.NewLine}    ", Components.Select(x => $"\"{x}\""));
+				VsConfigFileContent.AppendLine($"    {ComponentsString}");
+				VsConfigFileContent.AppendLine("  ]");
+				VsConfigFileContent.AppendLine("}");
+
+				FileReference VsConfigFileName = FileReference.Combine(PrimaryProjectPath, ".vsconfig");
+				bSuccess = WriteFileIfChanged(VsConfigFileName.FullName, VsConfigFileContent.ToString(), Logger);
+			}
+
 			return bSuccess;
 		}
 
