UnrealBuildTool: Fix checking VCProjectFileGenerator.Version when selecting default compiler, the wrong class and property name were being used

#rnx

diff --git a/Engine/Source/Programs/UnrealBuildTool/Platform/Windows/UEBuildWindows.cs b/Engine/Source/Programs/UnrealBuildTool/Platform/Windows/UEBuildWindows.cs
--- a/Engine/Source/Programs/UnrealBuildTool/Platform/Windows/UEBuildWindows.cs
+++ b/Engine/Source/Programs/UnrealBuildTool/Platform/Windows/UEBuildWindows.cs
@@ -949,9 +949,10 @@ namespace UnrealBuildTool
 					}
 				} 
 			}
+
 			// Also check the default format for the Visual Studio project generator
 			object? ProjectFormatObject;
-			if (XmlConfig.TryGetValue(typeof(VCProjectFileGenerator), "Version", out ProjectFormatObject))
+			if (XmlConfig.TryGetValue(typeof(VCProjectFileSettings), "ProjectFileFormat", out ProjectFormatObject))
 			{
 				VCProjectFileFormat ProjectFormat = (VCProjectFileFormat)ProjectFormatObject;
 				if (ProjectFormat == VCProjectFileFormat.VisualStudio2019)
