UnrealBuildTool: Allow generating visual studio solutions even if no valid toolchains are installed, as the user can now open the solution to install those needed components.

#rnx

diff --git a/Engine/Source/Programs/UnrealBuildTool/Platform/Windows/MicrosoftPlatformSDK.Versions.cs b/Engine/Source/Programs/UnrealBuildTool/Platform/Windows/MicrosoftPlatformSDK.Versions.cs
--- a/Engine/Source/Programs/UnrealBuildTool/Platform/Windows/MicrosoftPlatformSDK.Versions.cs
+++ b/Engine/Source/Programs/UnrealBuildTool/Platform/Windows/MicrosoftPlatformSDK.Versions.cs
@@ -86,7 +86,9 @@ namespace UnrealBuildTool
 		{
 			"Microsoft.VisualStudio.Workload.NativeCrossPlat",
 			"Microsoft.VisualStudio.Component.VC.14.32.17.2.ARM64",
-			"Microsoft.VisualStudio.Component.VC.14.32.17.2.x86.x64"
+			"Microsoft.VisualStudio.Component.VC.14.32.17.2.x86.x64",
+			"Microsoft.Net.Component.4.5.TargetingPack",
+			"Microsoft.Net.Component.4.6.2.TargetingPack",
 		};
 
 		/// <summary>
@@ -131,5 +133,12 @@ namespace UnrealBuildTool
 			MinVersion = PreferredWindowsSdkVersions.Min()?.ToString();
 			MaxVersion = null;
 		}
+
+		/// <summary>
+		/// Whether toolchain errors should be ignored. Enable to ignore banned toolchains when generating projects,
+		/// as components such as the recommended toolchain can be installed by opening the generated solution via the .vsconfig file.
+		/// If enabled the error will be downgraded to a warning.
+		/// </summary>
+		public static bool IgnoreToolchainErrors = false;
 	}
 }
diff --git a/Engine/Source/Programs/UnrealBuildTool/Platform/Windows/MicrosoftPlatformSDK.cs b/Engine/Source/Programs/UnrealBuildTool/Platform/Windows/MicrosoftPlatformSDK.cs
--- a/Engine/Source/Programs/UnrealBuildTool/Platform/Windows/MicrosoftPlatformSDK.cs
+++ b/Engine/Source/Programs/UnrealBuildTool/Platform/Windows/MicrosoftPlatformSDK.cs
@@ -592,7 +592,12 @@ namespace UnrealBuildTool
 
 			if (ToolChain?.Error != null)
 			{
-				throw new BuildException(ToolChain.Error);
+				if (!IgnoreToolchainErrors)
+				{
+					throw new BuildException(ToolChain.Error);
+				}
+				// If errors are ignored, log the warning instead
+				Log.TraceWarningOnce(ToolChain.Error);
 			}
 
 			return ToolChain;
@@ -893,7 +898,7 @@ namespace UnrealBuildTool
 			VersionNumberRange? Banned = BannedVisualCppVersions.FirstOrDefault(x => x.Contains(Version));
 			if (Banned != null)
 			{
-				Error = $"UnrealBuildTool has banned the MSVC {Banned} toolchains due to compiler issues. Please install a different toolchain such as {PreferredVisualCppVersions.Select(x => x.Min).Max()} from the Visual Studio installer.";
+				Error = $"UnrealBuildTool has banned the MSVC {Banned} toolchains due to compiler issues. Please install a different toolchain such as {PreferredVisualCppVersions.Select(x => x.Min).Max()} by opening the generated solution and installing recommended components or from the Visual Studio installer.";
 			}
 
 			Logger.LogDebug("Found Visual Studio toolchain: {ToolChainDir} (Family={Family}, FamilyRank={FamilyRank}, Version={Version}, Is64Bit={Is64Bit}, Preview={Preview}, Architecture={Arch}, Error={Error}, Redist={RedistDir})", ToolChainDir, Family, FamilyRank, Version, Is64Bit, bPreview, WindowsArchitecture.x64.ToString(), Error != null, RedistDir);
diff --git a/Engine/Source/Programs/UnrealBuildTool/ProjectFiles/VisualStudio/VCProjectFileGenerator.cs b/Engine/Source/Programs/UnrealBuildTool/ProjectFiles/VisualStudio/VCProjectFileGenerator.cs
--- a/Engine/Source/Programs/UnrealBuildTool/ProjectFiles/VisualStudio/VCProjectFileGenerator.cs
+++ b/Engine/Source/Programs/UnrealBuildTool/ProjectFiles/VisualStudio/VCProjectFileGenerator.cs
@@ -133,6 +133,9 @@ namespace UnrealBuildTool
 			{
 				BuildToolOverride = "-2022";
 			}
+
+			// Allow generating the solution even if the only installed toolchain is banned.			
+			MicrosoftPlatformSDK.IgnoreToolchainErrors = true;
 		}
 
 		public override string[] GetTargetArguments(string[] Arguments)
