Added more debug info to help find modules that are set to use shared PCHs but use private pch headers. This can result in compile issues if PCHs are disabled.

diff --git a/Engine/Source/Programs/UnrealBuildTool/Modes/PrintBuildGraphInfo.cs b/Engine/Source/Programs/UnrealBuildTool/Modes/PrintBuildGraphInfo.cs
--- a/Engine/Source/Programs/UnrealBuildTool/Modes/PrintBuildGraphInfo.cs
+++ b/Engine/Source/Programs/UnrealBuildTool/Modes/PrintBuildGraphInfo.cs
@@ -52,6 +52,7 @@ namespace UnrealBuildTool
 					List<UEBuildModule> NotOptimizedModules = new();
 					List<UEBuildModule> NoPCHModules = new();
 					List<UEBuildModule> PrivatePCHModules = new();
+					List<UEBuildModule> SharedPCHWithPrivateHeaderModules = new();
 
 					// Figure out all the modules referenced by this target. This includes all the modules that are referenced, not just the ones compiled into binaries.
 					CppCompileEnvironment GlobalCompileEnvironment = Target.CreateCompileEnvironmentForProjectFiles(Logger);
@@ -81,6 +82,10 @@ namespace UnrealBuildTool
 						{
 							PrivatePCHModules.Add(Module);
 						}
+						else if (Module.Rules.PrivatePCHHeaderFile != null && Module.Rules.PCHUsage == ModuleRules.PCHUsageMode.UseSharedPCHs)
+						{
+							SharedPCHWithPrivateHeaderModules.Add(Module);
+						}
 
 						if (!Module.Rules.bUseUnity)
 						{
@@ -113,6 +118,17 @@ namespace UnrealBuildTool
 						}
 					}
 
+					if (SharedPCHWithPrivateHeaderModules.Any())
+					{
+						Logger.LogInformation("Modules using a shared PCH and private PCH header:");
+						Logger.LogInformation(" Note these might not compile properly when PCHs are disabled.");
+						SharedPCHWithPrivateHeaderModules.SortBy(Module => Module.Name);
+						foreach (UEBuildModule Module in SharedPCHWithPrivateHeaderModules)
+						{
+							Logger.LogInformation("  {ModuleName}", Module.Name);
+						}
+					}
+
 					Logger.LogInformation("PCH Usage:");
 
 					if (NoPCHModules.Any())
