UnrealBuildTool: If a specific requested toolchain is not available, error when it is not found and list out available toolchains to assist the user

#rnx

diff --git a/Engine/Source/Programs/UnrealBuildTool/Platform/Windows/MicrosoftPlatformSDK.cs b/Engine/Source/Programs/UnrealBuildTool/Platform/Windows/MicrosoftPlatformSDK.cs
--- a/Engine/Source/Programs/UnrealBuildTool/Platform/Windows/MicrosoftPlatformSDK.cs
+++ b/Engine/Source/Programs/UnrealBuildTool/Platform/Windows/MicrosoftPlatformSDK.cs
@@ -519,7 +519,8 @@ namespace UnrealBuildTool
 				ToolChain = SelectToolChain(ToolChains, x => x.ThenByDescending(x => x.Version), Architecture);
 				if (ToolChain == null)
 				{
-					throw new BuildException("Unable to find C++ toolchain for {0} {1}", Compiler, Architecture.ToString());
+					DumpToolChains(ToolChains, x => x.ThenByDescending(x => x.Version), Architecture, Logger);
+					throw new BuildException("Unable to find valid C++ toolchain for {0} {1}", Compiler, Architecture.ToString());
 				}
 			}
 			else if (String.Compare(CompilerVersion, "Preview", StringComparison.InvariantCultureIgnoreCase) == 0)
@@ -527,15 +528,17 @@ namespace UnrealBuildTool
 				ToolChain = SelectToolChain(ToolChains, x => x.ThenByDescending(x => x.IsPreview), Architecture);
 				if (ToolChain == null || !ToolChain.IsPreview)
 				{
-					throw new BuildException("Unable to find preview toolchain for {0} {1}", Compiler, Architecture.ToString());
+					DumpToolChains(ToolChains, x => x.ThenByDescending(x => x.IsPreview), Architecture, Logger);
+					throw new BuildException("Unable to find valid preview toolchain for {0} {1}", Compiler, Architecture.ToString());
 				}
 			}
 			else if (VersionNumber.TryParse(CompilerVersion, out VersionNumber? ToolChainVersion))
 			{
 				ToolChain = SelectToolChain(ToolChains, x => x.ThenByDescending(x => x.Version == ToolChainVersion).ThenByDescending(x => x.Family == ToolChainVersion), Architecture);
-				if (ToolChain == null)
+				if (ToolChain == null || !(ToolChain.Version == ToolChainVersion || ToolChain.Family == ToolChainVersion))
 				{
-					throw new BuildException("Unable to find {0} toolchain for {1} {2}", ToolChainVersion, Compiler, Architecture.ToString());
+					DumpToolChains(ToolChains, x => x.ThenByDescending(x => x.Version == ToolChainVersion).ThenByDescending(x => x.Family == ToolChainVersion), Architecture, Logger);
+					throw new BuildException("Unable to find valid {0} toolchain for {1} {2}", ToolChainVersion, Compiler, Architecture.ToString());
 				}
 			}
 			else
@@ -603,6 +606,37 @@ namespace UnrealBuildTool
 			return ToolChain;
 		}
 
+		/// <summary>
+		/// Dump all available toolchain info, combining a custom preference with a default sort order
+		/// </summary>
+		/// <param name="ToolChains"></param>
+		/// <param name="Preference">Ordering function</param>
+		/// <param name="Architecture">Architecture that must be supported</param>
+		/// <param name="Logger">The ILogger interface to write to</param>
+		static void DumpToolChains(IEnumerable<ToolChainInstallation> ToolChains, Func<IOrderedEnumerable<ToolChainInstallation>, IOrderedEnumerable<ToolChainInstallation>> Preference, WindowsArchitecture Architecture, ILogger Logger)
+		{
+			var SortedToolChains = Preference(ToolChains.Where(x => x.Architecture == Architecture)
+				.OrderByDescending(x => x.Error == null))
+				.ThenByDescending(x => x.Is64Bit)
+				.ThenBy(x => x.IsPreview)
+				.ThenBy(x => x.FamilyRank)
+				.ThenByDescending(x => x.IsAutoSdk)
+				.ThenByDescending(x => x.Version);
+
+			if (SortedToolChains.Count() > 0)
+			{
+				Logger.LogInformation("Available {Architecture} toolchains ({Count}):", Architecture, SortedToolChains.Count());
+				foreach (ToolChainInstallation ToolChain in SortedToolChains)
+				{
+					Logger.LogInformation(" * {ToolChainDir}\n    (Family={Family}, FamilyRank={FamilyRank}, Version={Version}, Is64Bit={Is64Bit}, Preview={Preview}, Architecture={Arch}, Error={Error})", ToolChain.BaseDir, ToolChain.Family, ToolChain.FamilyRank, ToolChain.Version, ToolChain.Is64Bit, ToolChain.IsPreview, ToolChain.Architecture, ToolChain.Error != null);
+				}
+			}
+			else
+			{
+				Logger.LogInformation("No available {Architecture} toolchains found", Architecture);
+			}
+		}
+
 		static List<ToolChainInstallation> FindToolChainInstallations(WindowsCompiler Compiler, ILogger Logger)
 		{
 			List<ToolChainInstallation>? ToolChains;
