UnrealBuildTool: Set VS2022 as default compiler and generated solution if all available platforms are supported
Add messaging when generating if the default solution must be downgraded to VS2019 due to a platform that does not support VS2022
Unban [REDACTED_PHONE]

diff --git a/Engine/Source/Programs/UnrealBuildTool/Configuration/UEBuildPlatform.cs b/Engine/Source/Programs/UnrealBuildTool/Configuration/UEBuildPlatform.cs
--- a/Engine/Source/Programs/UnrealBuildTool/Configuration/UEBuildPlatform.cs
+++ b/Engine/Source/Programs/UnrealBuildTool/Configuration/UEBuildPlatform.cs
@@ -189,6 +189,14 @@ namespace UnrealBuildTool
 			return VCProjectFileFormat.Default;
 		}
 
+		/// <summary>
+		/// The version required to support Visual Studio
+		/// </summary>
+		public virtual Version GetVersionRequiredForVisualStudio(VCProjectFileFormat Format)
+		{
+			return new Version();
+		}
+
 		/// <summary>
 		/// Gets all the registered platforms
 		/// </summary>
diff --git a/Engine/Source/Programs/UnrealBuildTool/Platform/Windows/MicrosoftPlatformSDK.Versions.cs b/Engine/Source/Programs/UnrealBuildTool/Platform/Windows/MicrosoftPlatformSDK.Versions.cs
--- a/Engine/Source/Programs/UnrealBuildTool/Platform/Windows/MicrosoftPlatformSDK.Versions.cs
+++ b/Engine/Source/Programs/UnrealBuildTool/Platform/Windows/MicrosoftPlatformSDK.Versions.cs
@@ -32,8 +32,8 @@ namespace UnrealBuildTool
 		/// </summary>
 		static readonly VersionNumberRange[] PreferredVisualCppVersions = new VersionNumberRange[]
 		{
-			VersionNumberRange.Parse("14.29.30133", "14.29.99999"), // VS2019 16.11.x
 			VersionNumberRange.Parse("14.32.31326", "14.32.99999"), // VS2022 17.2.x
+			VersionNumberRange.Parse("14.29.30133", "14.29.99999"), // VS2019 16.11.x
 		};
 
 		/// <summary>
@@ -41,7 +41,6 @@ namespace UnrealBuildTool
 		/// </summary>
 		static readonly VersionNumberRange[] BannedVisualCppVersions = new VersionNumberRange[]
 		{
-			VersionNumberRange.Parse("14.33.31629", "14.33.99999"), // VS2022 17.3.x: https://developercommunity.visualstudio.com/t/clexe-143331629-Access-violation-IN/10122382
 		};
 
 		static readonly VersionNumber MinimumVisualCppVersion = new VersionNumber(14, 29, 30133);
@@ -58,7 +57,7 @@ namespace UnrealBuildTool
 
 		/// <summary>
 		/// The default set of components that should be suggested to be installed for Visual Studio 2019 or 2022.
-		/// This or the 2019\2022 specific components should be updated if the referred visual cpp version changes
+		/// This or the 2019\2022 specific components should be updated if the preferred visual cpp version changes
 		/// </summary>
 		static readonly string[] VisualStudioSuggestedComponents = new string[]
 		{
@@ -87,7 +86,6 @@ namespace UnrealBuildTool
 			"Microsoft.VisualStudio.Workload.NativeCrossPlat",
 			"Microsoft.VisualStudio.Component.VC.14.32.17.2.ARM64",
 			"Microsoft.VisualStudio.Component.VC.14.32.17.2.x86.x64",
-			"Microsoft.Net.Component.4.5.TargetingPack",
 			"Microsoft.Net.Component.4.6.2.TargetingPack",
 		};
 
diff --git a/Engine/Source/Programs/UnrealBuildTool/Platform/Windows/UEBuildWindows.cs b/Engine/Source/Programs/UnrealBuildTool/Platform/Windows/UEBuildWindows.cs
--- a/Engine/Source/Programs/UnrealBuildTool/Platform/Windows/UEBuildWindows.cs
+++ b/Engine/Source/Programs/UnrealBuildTool/Platform/Windows/UEBuildWindows.cs
@@ -843,14 +843,14 @@ namespace UnrealBuildTool
 			// Set the compiler version if necessary
 			if (Target.WindowsPlatform.Compiler == WindowsCompiler.Default)
 			{
-				if (Target.StaticAnalyzer == StaticAnalyzer.PVSStudio && HasCompiler(WindowsCompiler.VisualStudio2019, Target.WindowsPlatform.Architecture, Logger))
-				{
-					Target.WindowsPlatform.Compiler = WindowsCompiler.VisualStudio2019;
-				}
 				if (Target.StaticAnalyzer == StaticAnalyzer.PVSStudio && HasCompiler(WindowsCompiler.VisualStudio2022, Target.WindowsPlatform.Architecture, Logger))
 				{
 					Target.WindowsPlatform.Compiler = WindowsCompiler.VisualStudio2022;
 				}
+				else if (Target.StaticAnalyzer == StaticAnalyzer.PVSStudio && HasCompiler(WindowsCompiler.VisualStudio2019, Target.WindowsPlatform.Architecture, Logger))
+				{
+					Target.WindowsPlatform.Compiler = WindowsCompiler.VisualStudio2019;
+				}
 				else
 				{
 					Target.WindowsPlatform.Compiler = GetDefaultCompiler(Target.ProjectFile, Target.WindowsPlatform.Architecture, Logger);
@@ -979,34 +979,34 @@ namespace UnrealBuildTool
 				}
 			}
 
-			// Second, default based on what's installed, test for 2019 first
-			if (MicrosoftPlatformSDK.HasValidCompiler(WindowsCompiler.VisualStudio2019, Architecture, Logger))
+			// Second, default based on what's installed, test for 2022 first
+			if (MicrosoftPlatformSDK.HasValidCompiler(WindowsCompiler.VisualStudio2022, Architecture, Logger))
 			{
-				return WindowsCompiler.VisualStudio2019;
+				return WindowsCompiler.VisualStudio2022;
 			}
-			else if (MicrosoftPlatformSDK.HasValidCompiler(WindowsCompiler.VisualStudio2022, Architecture, Logger))
+			else if (MicrosoftPlatformSDK.HasValidCompiler(WindowsCompiler.VisualStudio2019, Architecture, Logger))
 			{
-				return WindowsCompiler.VisualStudio2022;
+				return WindowsCompiler.VisualStudio2019;
 			}
 
 			// If we do have a Visual Studio installation, but we're missing just the C++ parts, warn about that.
-			if (TryGetVSInstallDirs(WindowsCompiler.VisualStudio2019, Logger) != null)
-			{
-				string ToolSetWarning = Architecture == WindowsArchitecture.x64 ?
-					"MSVC v142 - VS 2019 C++ x64/x86 build tools (Latest)" :
-					"MSVC v142 - VS 2019 C++ ARM64 build tools (Latest)";
-				Logger.LogWarning("Visual Studio 2019 is installed, but is missing the C++ toolchain. Please verify that the \"{Component}\" component is selected in the Visual Studio 2019 installation options.", ToolSetWarning);
-			}
-			else if (TryGetVSInstallDirs(WindowsCompiler.VisualStudio2022, Logger) != null)
+			if (TryGetVSInstallDirs(WindowsCompiler.VisualStudio2022, Logger) != null)
 			{
 				string ToolSetWarning = Architecture == WindowsArchitecture.x64 ?
 					"MSVC v143 - VS 2022 C++ x64/x86 build tools (Latest)" :
 					"MSVC v143 - VS 2022 C++ ARM64 build tools (Latest)";
 				Logger.LogWarning("Visual Studio 2022 is installed, but is missing the C++ toolchain. Please verify that the \"{Component}\" component is selected in the Visual Studio 2022 installation options.", ToolSetWarning);
 			}
+			else if (TryGetVSInstallDirs(WindowsCompiler.VisualStudio2019, Logger) != null)
+			{
+				string ToolSetWarning = Architecture == WindowsArchitecture.x64 ?
+					"MSVC v142 - VS 2019 C++ x64/x86 build tools (Latest)" :
+					"MSVC v142 - VS 2019 C++ ARM64 build tools (Latest)";
+				Logger.LogWarning("Visual Studio 2019 is installed, but is missing the C++ toolchain. Please verify that the \"{Component}\" component is selected in the Visual Studio 2019 installation options.", ToolSetWarning);
+			}
 			else
 			{
-				Logger.LogWarning("No Visual C++ installation was found. Please download and install Visual Studio 2019 or 2022 with C++ components.");
+				Logger.LogWarning("No Visual C++ installation was found. Please download and install Visual Studio 2022 or 2019 with C++ components.");
 			}
 
 			// Finally, default to VS2019 anyway
diff --git a/Engine/Source/Programs/UnrealBuildTool/Platform/Windows/VCEnvironment.cs b/Engine/Source/Programs/UnrealBuildTool/Platform/Windows/VCEnvironment.cs
--- a/Engine/Source/Programs/UnrealBuildTool/Platform/Windows/VCEnvironment.cs
+++ b/Engine/Source/Programs/UnrealBuildTool/Platform/Windows/VCEnvironment.cs
@@ -456,17 +456,17 @@ namespace UnrealBuildTool
 			DirectoryReference? SelectedToolChainDir;
 			if(Compiler.IsClang())
 			{
-				if (WindowsPlatform.TryGetToolChainDir(WindowsCompiler.VisualStudio2019, null, Architecture, Logger, out SelectedToolChainVersion, out SelectedToolChainDir, out SelectedRedistDir))
+				if (WindowsPlatform.TryGetToolChainDir(WindowsCompiler.VisualStudio2022, null, Architecture, Logger, out SelectedToolChainVersion, out SelectedToolChainDir, out SelectedRedistDir))
 				{
-					ToolChain = WindowsCompiler.VisualStudio2019;
+					ToolChain = WindowsCompiler.VisualStudio2022;
 				}
-				else if (WindowsPlatform.TryGetToolChainDir(WindowsCompiler.VisualStudio2022, null, Architecture, Logger, out SelectedToolChainVersion, out SelectedToolChainDir, out SelectedRedistDir))
+				else if (WindowsPlatform.TryGetToolChainDir(WindowsCompiler.VisualStudio2019, null, Architecture, Logger, out SelectedToolChainVersion, out SelectedToolChainDir, out SelectedRedistDir))
 				{
-					ToolChain = WindowsCompiler.VisualStudio2022;
+					ToolChain = WindowsCompiler.VisualStudio2019;
 				}
 				else
 				{
-					throw new BuildException("{0} or {1} must be installed in order to build this target.", WindowsPlatform.GetCompilerName(WindowsCompiler.VisualStudio2019), WindowsPlatform.GetCompilerName(WindowsCompiler.VisualStudio2022));
+					throw new BuildException("{0} or {1} must be installed in order to build this target.", WindowsPlatform.GetCompilerName(WindowsCompiler.VisualStudio2022), WindowsPlatform.GetCompilerName(WindowsCompiler.VisualStudio2019));
 				}
 			}
 			else
diff --git a/Engine/Source/Programs/UnrealBuildTool/ProjectFiles/VisualStudio/VCProjectFileGenerator.cs b/Engine/Source/Programs/UnrealBuildTool/ProjectFiles/VisualStudio/VCProjectFileGenerator.cs
--- a/Engine/Source/Programs/UnrealBuildTool/ProjectFiles/VisualStudio/VCProjectFileGenerator.cs
+++ b/Engine/Source/Programs/UnrealBuildTool/ProjectFiles/VisualStudio/VCProjectFileGenerator.cs
@@ -268,19 +268,19 @@ namespace UnrealBuildTool
 			base.SetupSupportedPlatformsAndConfigurations(IncludeAllPlatforms, Logger, out SupportedPlatformNames);
 
 			// If we have a non-default setting for visual studio, check the compiler exists. If not, revert to the default.
-			if (Settings.ProjectFileFormat == VCProjectFileFormat.VisualStudio2019)
+			if (Settings.ProjectFileFormat == VCProjectFileFormat.VisualStudio2022)
 			{
-				if (!WindowsPlatform.HasCompiler(WindowsCompiler.VisualStudio2019, WindowsArchitecture.x64, Logger))
+				if (!WindowsPlatform.HasCompiler(WindowsCompiler.VisualStudio2022, WindowsArchitecture.x64, Logger))
 				{
-					Logger.LogWarning("Visual Studio C++ 2019 installation not found - ignoring preferred project file format.");
+					Logger.LogWarning("Visual Studio C++ 2022 installation not found - ignoring preferred project file format.");
 					Settings.ProjectFileFormat = VCProjectFileFormat.Default;
 				}
 			}
-			else if (Settings.ProjectFileFormat == VCProjectFileFormat.VisualStudio2022)
+			else if (Settings.ProjectFileFormat == VCProjectFileFormat.VisualStudio2019)
 			{
-				if (!WindowsPlatform.HasCompiler(WindowsCompiler.VisualStudio2022, WindowsArchitecture.x64, Logger))
+				if (!WindowsPlatform.HasCompiler(WindowsCompiler.VisualStudio2019, WindowsArchitecture.x64, Logger))
 				{
-					Logger.LogWarning("Visual Studio C++ 2022 installation not found - ignoring preferred project file format.");
+					Logger.LogWarning("Visual Studio C++ 2019 installation not found - ignoring preferred project file format.");
 					Settings.ProjectFileFormat = VCProjectFileFormat.Default;
 				}
 			}
@@ -294,19 +294,12 @@ namespace UnrealBuildTool
 				List<VisualStudioInstallation> Installations = MicrosoftPlatformSDK.FindVisualStudioInstallations(Logger).Where(x => WindowsPlatform.HasCompiler(x.Compiler, WindowsArchitecture.x64, Logger)).ToList();
 
 				// Get the corresponding project file format
-				VCProjectFileFormat Format = VCProjectFileFormat.VisualStudio2019;
-				foreach (VisualStudioInstallation Installation in Installations)
+				bool Vs2019Available = Installations.Any(x => x.Compiler == WindowsCompiler.VisualStudio2019);
+				bool Vs2022Available = Installations.Any(x => x.Compiler == WindowsCompiler.VisualStudio2022);
+				VCProjectFileFormat Format = VCProjectFileFormat.VisualStudio2022;
+				if (!Vs2022Available && Vs2019Available)
 				{
-					if (Installation.Compiler == WindowsCompiler.VisualStudio2022)
-					{
-						Format = VCProjectFileFormat.VisualStudio2022;
-						break;
-					}
-					else if (Installation.Compiler == WindowsCompiler.VisualStudio2019)
-					{
-						Format = VCProjectFileFormat.VisualStudio2019;
-						break;
-					}
+					Format = VCProjectFileFormat.VisualStudio2019;
 				}
 				Settings.ProjectFileFormat = Format;
 
@@ -326,7 +319,24 @@ namespace UnrealBuildTool
 								// Reduce the Visual Studio version to the max supported by each platform we plan to include.
 								if (Settings.ProjectFileFormat == VCProjectFileFormat.Default || ProposedFormat < Settings.ProjectFileFormat)
 								{
-									Settings.ProjectFileFormat = ProposedFormat;
+									Logger.LogInformation("Available {SupportedPlatform} SDK does not support Visual Studio 2022.", SupportedPlatform);
+									Version Version = BuildPlatform.GetVersionRequiredForVisualStudio(VCProjectFileFormat.VisualStudio2022);
+									if (Version > new Version())
+									{
+										Logger.LogInformation("Please update {SupportedPlatform} SDK to {Version} if Visual Studio 2022 support is desired.", SupportedPlatform, Version);
+									}
+									if (!Vs2019Available)
+									{
+										Logger.LogInformation("Generated solution cannot be downgraded to Visual Studio 2019 as it is not installed. Please install Visual Studio 2019 if {SupportedPlatform} SDK support is required.", SupportedPlatform);
+									}
+									else
+									{
+										Settings.ProjectFileFormat = ProposedFormat;
+										Logger.LogInformation("Downgrading generated solution to Visual Studio 2019.");
+										Logger.LogInformation("To force Visual Studio 2022 solutions to always be generated add the following to BuildConfiguration.xml:");
+										Logger.LogInformation("  <VCProjectFileGenerator>\r\n    <Version>VisualStudio2022</Version>\r\n  </VCProjectFileGenerator>");
+									}
+									Logger.LogInformation(string.Empty);
 								}
 							}
 						}
