Force compiler back to VS2019 for erroring project while compile error is investigated

#rnx

diff --git a/Engine/Source/Programs/UnrealBuildTool/Platform/Windows/UEBuildWindows.cs b/Engine/Source/Programs/UnrealBuildTool/Platform/Windows/UEBuildWindows.cs
--- a/Engine/Source/Programs/UnrealBuildTool/Platform/Windows/UEBuildWindows.cs
+++ b/Engine/Source/Programs/UnrealBuildTool/Platform/Windows/UEBuildWindows.cs
@@ -157,7 +157,7 @@ namespace UnrealBuildTool
 		/// <summary>
 		/// Version of the compiler toolchain to use on Windows platform. A value of "default" will be changed to a specific version at UBT start up.
 		/// </summary>
-		[ConfigFile(ConfigHierarchyType.Engine, "/Script/WindowsTargetPlatform.WindowsTargetSettings", "CompilerVersion")]
+		[ConfigFile(ConfigHierarchyType.Engine, "/Script/WindowsTargetPlatform.WindowsTargetSettings", "Compiler")]
 		[XmlConfigFile(Category = "WindowsPlatform")]
 		[CommandLine("-2019", Value = nameof(WindowsCompiler.VisualStudio2019))]
 		[CommandLine("-2022", Value = nameof(WindowsCompiler.VisualStudio2022))]
@@ -179,6 +179,7 @@ namespace UnrealBuildTool
 		/// the string "Preview" to select the newest available preview version. By default, and if it is available, we use the toolchain version indicated by
 		/// WindowsPlatform.DefaultToolChainVersion (otherwise, we use the latest version).
 		/// </summary>
+		[ConfigFile(ConfigHierarchyType.Engine, "/Script/WindowsTargetPlatform.WindowsTargetSettings", "CompilerVersion")]
 		[XmlConfigFile(Category = "WindowsPlatform")]
 		[CommandLine("-CompilerVersion")]
 		public string? CompilerVersion = null;
