Crypto - only write Crypto.json if the settings has been changed for better iterative cooks

#9265

diff --git a/Engine/Source/Programs/UnrealBuildTool/System/EncryptionAndSigning.cs b/Engine/Source/Programs/UnrealBuildTool/System/EncryptionAndSigning.cs
--- a/Engine/Source/Programs/UnrealBuildTool/System/EncryptionAndSigning.cs
+++ b/Engine/Source/Programs/UnrealBuildTool/System/EncryptionAndSigning.cs
@@ -186,7 +186,7 @@ namespace UnrealBuildTool
 			public void Save(FileReference InFile)
 			{
 				DirectoryReference.CreateDirectory(InFile.Directory);
-				FileReference.WriteAllText(InFile, JsonSerializer.Serialize(this));
+				FileReference.WriteAllTextIfDifferent(InFile, JsonSerializer.Serialize(this));
 			}
 		}
 
