UnrealBuildTool: Unban [REDACTED_PHONE] toolchain

#rnx

diff --git a/Engine/Source/Programs/UnrealBuildTool/Platform/Windows/MicrosoftPlatformSDK.Versions.cs b/Engine/Source/Programs/UnrealBuildTool/Platform/Windows/MicrosoftPlatformSDK.Versions.cs
--- a/Engine/Source/Programs/UnrealBuildTool/Platform/Windows/MicrosoftPlatformSDK.Versions.cs
+++ b/Engine/Source/Programs/UnrealBuildTool/Platform/Windows/MicrosoftPlatformSDK.Versions.cs
@@ -33,7 +33,7 @@ namespace UnrealBuildTool
 		static readonly VersionNumberRange[] PreferredVisualCppVersions = new VersionNumberRange[]
 		{
 			VersionNumberRange.Parse("14.29.30133", "14.29.99999"), // VS2019 16.11.x
-			VersionNumberRange.Parse("14.32.31326", "14.32.99999"), // VS2022 17.2.x
+			VersionNumberRange.Parse("14.33.31629", "14.33.99999"), // VS2022 17.3.x
 		};
 
 		/// <summary>
@@ -41,7 +41,6 @@ namespace UnrealBuildTool
 		/// </summary>
 		static readonly VersionNumberRange[] BannedVisualCppVersions = new VersionNumberRange[]
 		{
-			VersionNumberRange.Parse("14.33.31629", "14.33.99999"), // VS2022 17.3.x: https://developercommunity.visualstudio.com/t/clexe-143331629-Access-violation-IN/10122382
 		};
 
 		static readonly VersionNumber MinimumVisualCppVersion = new VersionNumber(14, 29, 30133);
@@ -58,7 +57,7 @@ namespace UnrealBuildTool
 
 		/// <summary>
 		/// The default set of components that should be suggested to be installed for Visual Studio 2019 or 2022.
-		/// This or the 2019\2022 specific components should be updated if the referred visual cpp version changes
+		/// This or the 2019\2022 specific components should be updated if the preferred visual cpp version changes
 		/// </summary>
 		static readonly string[] VisualStudioSuggestedComponents = new string[]
 		{
@@ -85,9 +84,8 @@ namespace UnrealBuildTool
 		static readonly string[] VisualStudio2022SuggestedComponents = new string[]
 		{
 			"Microsoft.VisualStudio.Workload.NativeCrossPlat",
-			"Microsoft.VisualStudio.Component.VC.14.32.17.2.ARM64",
-			"Microsoft.VisualStudio.Component.VC.14.32.17.2.x86.x64",
-			"Microsoft.Net.Component.4.5.TargetingPack",
+			"Microsoft.VisualStudio.Component.VC.14.33.17.3.ARM64",
+			"Microsoft.VisualStudio.Component.VC.14.33.17.3.x86.x64",
 			"Microsoft.Net.Component.4.6.2.TargetingPack",
 		};
 
