- Upgrade Mac binary to EOS-SDK-21689671-v1.15.2.1

diff --git a/Engine/Source/Programs/UnrealBuildTool/Platform/Mac/UEBuildMac.cs b/Engine/Source/Programs/UnrealBuildTool/Platform/Mac/UEBuildMac.cs
--- a/Engine/Source/Programs/UnrealBuildTool/Platform/Mac/UEBuildMac.cs
+++ b/Engine/Source/Programs/UnrealBuildTool/Platform/Mac/UEBuildMac.cs
@@ -141,16 +141,6 @@ namespace UnrealBuildTool
 			{
 				Target.DisablePlugins.AddRange(new string[]
 				{
-					// EOSSDK
-					"EOSManager",
-					"EOSShared",
-					"EOSVoiceChat",
-					"EOSReservedHooks",
-					"OnlineServicesEOS",
-					"OnlineServicesEOSGS",
-					"OnlineSubsystemEOS",
-					"SocketSubsystemEOS",
-
 					// Melange
 					"DatasmithC4DImporter",
 				});
