- Enabling DatasmithC4DImporter plugin now that it has support for Arm64 Mac

diff --git a/Engine/Source/Programs/UnrealBuildTool/Platform/Mac/UEBuildMac.cs b/Engine/Source/Programs/UnrealBuildTool/Platform/Mac/UEBuildMac.cs
--- a/Engine/Source/Programs/UnrealBuildTool/Platform/Mac/UEBuildMac.cs
+++ b/Engine/Source/Programs/UnrealBuildTool/Platform/Mac/UEBuildMac.cs
@@ -141,8 +141,7 @@ namespace UnrealBuildTool
 			{
 				Target.DisablePlugins.AddRange(new string[]
 				{
-					// Melange
-					"DatasmithC4DImporter",
+					// Currently none need to be disabled, but add names of plugins here that are incompatible with arm64
 				});
 			}
 
