UnrealBuildTool: Suppress VS2022 static analysis warning for 5.1

#rnx

diff --git a/Engine/Source/Programs/UnrealBuildTool/Platform/Windows/VCToolChain.cs b/Engine/Source/Programs/UnrealBuildTool/Platform/Windows/VCToolChain.cs
--- a/Engine/Source/Programs/UnrealBuildTool/Platform/Windows/VCToolChain.cs
+++ b/Engine/Source/Programs/UnrealBuildTool/Platform/Windows/VCToolChain.cs
@@ -331,6 +331,12 @@ namespace UnrealBuildTool
 
 				// Don't bother generating code, only analyze code (may report fewer warnings though.)
 				//Arguments.Add("/analyze:only");
+
+				// Re-evalulate new analysis warnings at a later time
+				if (EnvVars.CompilerVersion >= new VersionNumber(14, 32))
+				{
+					Arguments.Add("/wd6031"); // return value ignored: called-function could return unexpected value
+				}
 			}
 
 			// Prevents the compiler from displaying its logo for each invocation.
