UnrealBuildTool: Support -SingleFile compile for .mm and .rc files

#rnx

diff --git a/Engine/Source/Programs/UnrealBuildTool/Configuration/UEBuildModuleCPP.cs b/Engine/Source/Programs/UnrealBuildTool/Configuration/UEBuildModuleCPP.cs
--- a/Engine/Source/Programs/UnrealBuildTool/Configuration/UEBuildModuleCPP.cs
+++ b/Engine/Source/Programs/UnrealBuildTool/Configuration/UEBuildModuleCPP.cs
@@ -425,9 +425,15 @@ namespace UnrealBuildTool
 				InputFiles.CPPFiles.RemoveAll(x => !SpecificFilesToCompile.Contains(x.Location));
 				InputFiles.CCFiles.RemoveAll(x => !SpecificFilesToCompile.Contains(x.Location));
 				InputFiles.CFiles.RemoveAll(x => !SpecificFilesToCompile.Contains(x.Location));
-
-				if (InputFiles.CPPFiles.Count == 0 && InputFiles.CCFiles.Count == 0 && InputFiles.CFiles.Count == 0 &&
-					!SpecificFilesToCompile.Any(x => ContainsFile(x)))
+				InputFiles.MMFiles.RemoveAll(x => !SpecificFilesToCompile.Contains(x.Location));
+				InputFiles.RCFiles.RemoveAll(x => !SpecificFilesToCompile.Contains(x.Location));
+
+				if (InputFiles.CPPFiles.Count == 0
+					&& InputFiles.CCFiles.Count == 0
+					&& InputFiles.CFiles.Count == 0
+					&& InputFiles.MMFiles.Count == 0
+					&& InputFiles.RCFiles.Count == 0
+					&& !SpecificFilesToCompile.Any(x => ContainsFile(x)))
 				{
 					return new List<FileItem>();
 				}
