UnrealBuildTool: Add missing PrerequisiteItems when copying generated ISPC header

#rnx

diff --git a/Engine/Source/Programs/UnrealBuildTool/ToolChain/ISPCToolChain.cs b/Engine/Source/Programs/UnrealBuildTool/ToolChain/ISPCToolChain.cs
--- a/Engine/Source/Programs/UnrealBuildTool/ToolChain/ISPCToolChain.cs
+++ b/Engine/Source/Programs/UnrealBuildTool/ToolChain/ISPCToolChain.cs
@@ -434,7 +434,7 @@ namespace UnrealBuildTool
 			foreach (FileItem ISPCFile in InputFiles)
 			{
 				Action CompileAction = Graph.CreateAction(ActionType.Compile);
-				CompileAction.CommandDescription = "Compile";
+				CompileAction.CommandDescription = "Generate Header";
 				CompileAction.WorkingDirectory = Unreal.EngineSourceDirectory;
 				CompileAction.CommandPath = new FileReference(GetISPCHostCompilerPath(BuildHostPlatform.Current.Platform));
 				CompileAction.StatusDescription = Path.GetFileName(ISPCFile.AbsolutePath);
@@ -505,6 +505,7 @@ namespace UnrealBuildTool
 					CopyAction.CommandArguments = $"-c 'cp -f \"\"{SourceFile}\"\" \"\"{TargetFile}\"'";
 				}
 				CopyAction.WorkingDirectory = Unreal.EngineSourceDirectory;
+				CopyAction.PrerequisiteItems.Add(ISPCFile);
 				CopyAction.PrerequisiteItems.Add(SourceFileItem);
 				CopyAction.ProducedItems.Add(TargetFileItem);
 				CopyAction.StatusDescription = TargetFileItem.Location.GetFileName();
