Suppress tidy checks for apple clang toolchains when running clang analyzer

#rnx

diff --git a/Engine/Source/Programs/UnrealBuildTool/ToolChain/AppleToolChain.cs b/Engine/Source/Programs/UnrealBuildTool/ToolChain/AppleToolChain.cs
--- a/Engine/Source/Programs/UnrealBuildTool/ToolChain/AppleToolChain.cs
+++ b/Engine/Source/Programs/UnrealBuildTool/ToolChain/AppleToolChain.cs
@@ -313,6 +313,15 @@ namespace UnrealBuildTool
 			}
 		}
 
+		/// <inheritdoc/>
+		protected override void GetCompileArguments_Analyze(CppCompileEnvironment CompileEnvironment, List<string> Arguments)
+		{
+			base.GetCompileArguments_Analyze(CompileEnvironment, Arguments);
+
+			// Disable all clang tidy checks
+			Arguments.Add($"-Xclang -analyzer-tidy-checker=-*");
+		}
+
 		/// <inheritdoc/>
 		protected override void GetCompileArguments_Global(CppCompileEnvironment CompileEnvironment, List<string> Arguments)
 		{
