UnrealBuildTool: Update clang pgo profile to pass -fprofile-generate instead of -fprofile-instr-generate

#rnx

diff --git a/Engine/Source/Programs/UnrealBuildTool/ToolChain/ClangToolChain.cs b/Engine/Source/Programs/UnrealBuildTool/ToolChain/ClangToolChain.cs
--- a/Engine/Source/Programs/UnrealBuildTool/ToolChain/ClangToolChain.cs
+++ b/Engine/Source/Programs/UnrealBuildTool/ToolChain/ClangToolChain.cs
@@ -648,7 +648,7 @@ namespace UnrealBuildTool
 			{
 				// Always enable LTO when generating PGO profile data.
 				Log.TraceInformationOnce("Enabling Profile Guided Instrumentation (PGI). Linking will take a while.");
-				Arguments.Add("-fprofile-instr-generate");
+				Arguments.Add("-fprofile-generate");
 			}
 
 			if (!CompileEnvironment.bUseInlining)
