Add an option to include the CPP/WinRT language projection headers. Set TargetWindowsVersion to configured value when using Windows 10 SDK.

diff --git a/Engine/Platforms/Hololens/Source/Programs/UnrealBuildTool/UEBuildHoloLens.cs b/Engine/Platforms/Hololens/Source/Programs/UnrealBuildTool/UEBuildHoloLens.cs
--- a/Engine/Platforms/Hololens/Source/Programs/UnrealBuildTool/UEBuildHoloLens.cs
+++ b/Engine/Platforms/Hololens/Source/Programs/UnrealBuildTool/UEBuildHoloLens.cs
@@ -214,7 +214,7 @@ namespace UnrealBuildTool
 
 		protected override VCEnvironment CreateVCEnvironment(TargetRules Target)
 		{
-			VCEnvironmentParameters Params = new VCEnvironmentParameters(Target.WindowsPlatform.Compiler, Platform, Target.WindowsPlatform.Architecture, Target.WindowsPlatform.CompilerVersion, Target.WindowsPlatform.WindowsSdkVersion, null, Logger);
+			VCEnvironmentParameters Params = new VCEnvironmentParameters(Target.WindowsPlatform.Compiler, Platform, Target.WindowsPlatform.Architecture, Target.WindowsPlatform.CompilerVersion, Target.WindowsPlatform.WindowsSdkVersion, null, Target.WindowsPlatform.bUseCPPWinRT, Logger);
 			return new HoloLensEnvironment(Params, Logger);
 		}
 
diff --git a/Engine/Source/Programs/UnrealBuildTool/Executors/Experimental/FASTBuild.cs b/Engine/Source/Programs/UnrealBuildTool/Executors/Experimental/FASTBuild.cs
--- a/Engine/Source/Programs/UnrealBuildTool/Executors/Experimental/FASTBuild.cs
+++ b/Engine/Source/Programs/UnrealBuildTool/Executors/Experimental/FASTBuild.cs
@@ -767,7 +767,7 @@ namespace UnrealBuildTool
 				// it probably means we are building for another platform.
                 if(BuildType == FBBuildType.Windows)
                 {
-					VCEnv = VCEnvironment.Create(WindowsPlatform.GetDefaultCompiler(null, WindowsArchitecture.x64, Logger), UnrealTargetPlatform.Win64, WindowsArchitecture.x64, null, null, null, Logger);
+					VCEnv = VCEnvironment.Create(WindowsPlatform.GetDefaultCompiler(null, WindowsArchitecture.x64, Logger), UnrealTargetPlatform.Win64, WindowsArchitecture.x64, null, null, null, false, Logger);
 				}
             }
 			catch (Exception)
diff --git a/Engine/Source/Programs/UnrealBuildTool/Platform/Windows/UEBuildWindows.cs b/Engine/Source/Programs/UnrealBuildTool/Platform/Windows/UEBuildWindows.cs
--- a/Engine/Source/Programs/UnrealBuildTool/Platform/Windows/UEBuildWindows.cs
+++ b/Engine/Source/Programs/UnrealBuildTool/Platform/Windows/UEBuildWindows.cs
@@ -208,6 +208,12 @@ namespace UnrealBuildTool
 		[ConfigFile(ConfigHierarchyType.Engine, "/Script/WindowsTargetPlatform.WindowsTargetSettings", "bUseWindowsSDK10")]
 		public bool bUseWindowsSDK10 = false;
 
+		/// <summary>
+		/// Enable building with the C++/WinRT language projection
+		/// </summary>
+		[ConfigFile(ConfigHierarchyType.Engine, "/Script/WindowsTargetPlatform.WindowsTargetSettings", "bUseCPPWinRT")]
+		public bool bUseCPPWinRT = false;
+
 		/// <summary>
 		/// Enables runtime ray tracing support.
 		/// </summary>
@@ -545,6 +551,11 @@ namespace UnrealBuildTool
 			get { return Inner.bUseWindowsSDK10; }
 		}
 
+		public bool bUseCPPWinRT
+		{
+			get { return Inner.bUseCPPWinRT; }
+		}
+
 		public bool bEnableRayTracing
 		{
 			get { return Inner.bEnableRayTracing; }
@@ -808,7 +819,7 @@ namespace UnrealBuildTool
 		[SupportedOSPlatform("windows")]
 		protected virtual VCEnvironment CreateVCEnvironment(TargetRules Target)
 		{
-			return VCEnvironment.Create(Target.WindowsPlatform.Compiler, Platform, Target.WindowsPlatform.Architecture, Target.WindowsPlatform.CompilerVersion, Target.WindowsPlatform.WindowsSdkVersion, null, Logger);
+			return VCEnvironment.Create(Target.WindowsPlatform.Compiler, Platform, Target.WindowsPlatform.Architecture, Target.WindowsPlatform.CompilerVersion, Target.WindowsPlatform.WindowsSdkVersion, null, Target.WindowsPlatform.bUseCPPWinRT, Logger);
 		}
 
 		/// <summary>
@@ -1293,17 +1304,9 @@ namespace UnrealBuildTool
 			CompileEnvironment.Definitions.Add("NDIS_MINIPORT_MAJOR_VERSION=0");
 
 			CompileEnvironment.Definitions.Add("WIN32=1");
-			if (Target.WindowsPlatform.bUseWindowsSDK10)
-			{
-				CompileEnvironment.Definitions.Add(String.Format("_WIN32_WINNT=0x{0:X4}", 0x0602));
-				CompileEnvironment.Definitions.Add(String.Format("WINVER=0x{0:X4}", 0x0602));
 
-			}
-			else
-			{
-				CompileEnvironment.Definitions.Add(String.Format("_WIN32_WINNT=0x{0:X4}", Target.WindowsPlatform.TargetWindowsVersion));
-				CompileEnvironment.Definitions.Add(String.Format("WINVER=0x{0:X4}", Target.WindowsPlatform.TargetWindowsVersion));
-			}
+			CompileEnvironment.Definitions.Add(String.Format("_WIN32_WINNT=0x{0:X4}", Target.WindowsPlatform.TargetWindowsVersion));
+			CompileEnvironment.Definitions.Add(String.Format("WINVER=0x{0:X4}", Target.WindowsPlatform.TargetWindowsVersion));
 			
 			CompileEnvironment.Definitions.Add("PLATFORM_WINDOWS=1");
 			CompileEnvironment.Definitions.Add("PLATFORM_MICROSOFT=1");
diff --git a/Engine/Source/Programs/UnrealBuildTool/Platform/Windows/VCEnvironment.cs b/Engine/Source/Programs/UnrealBuildTool/Platform/Windows/VCEnvironment.cs
--- a/Engine/Source/Programs/UnrealBuildTool/Platform/Windows/VCEnvironment.cs
+++ b/Engine/Source/Programs/UnrealBuildTool/Platform/Windows/VCEnvironment.cs
@@ -65,6 +65,11 @@ namespace UnrealBuildTool
 		/// </summary>
 		public readonly VersionNumber WindowsSdkVersion;
 
+		/// <summary>
+		/// Use the CPP/WinRT language projection
+		/// </summary>
+		public readonly bool bUseCPPWinRT;
+
 		/// <summary>
 		/// The path to the linker for linking executables
 		/// </summary>
@@ -118,6 +123,7 @@ namespace UnrealBuildTool
 			this.WindowsSdkDir = Params.WindowsSdkDir;
 			this.WindowsSdkVersion = Params.WindowsSdkVersion;
 			this.RedistDir = Params.RedistDir;
+			this.bUseCPPWinRT = Params.bUseCPPWinRT;
 
 			// Get the standard VC paths
 			DirectoryReference VCToolPath = GetVCToolPath(ToolChain, ToolChainDir, Architecture);
@@ -335,6 +341,10 @@ namespace UnrealBuildTool
 				IncludePaths.Add(DirectoryReference.Combine(IncludeRootDir, "shared"));
 				IncludePaths.Add(DirectoryReference.Combine(IncludeRootDir, "um"));
 				IncludePaths.Add(DirectoryReference.Combine(IncludeRootDir, "winrt"));
+				if (bUseCPPWinRT)
+				{
+					IncludePaths.Add(DirectoryReference.Combine(IncludeRootDir, "cppwinrt"));
+				}
 
 				DirectoryReference LibraryRootDir = DirectoryReference.Combine(WindowsSdkDir, "lib", WindowsSdkVersion.ToString());
 				LibraryPaths.Add(DirectoryReference.Combine(LibraryRootDir, "ucrt", ArchFolder));
@@ -369,12 +379,13 @@ namespace UnrealBuildTool
 		/// <param name="CompilerVersion">The specific toolchain version to use</param>
 		/// <param name="WindowsSdkVersion">Version of the Windows SDK to use</param>
 		/// <param name="SuppliedSdkDirectoryForVersion">If specified, this is the SDK directory to use, otherwise, attempt to look up via registry. If specified, the WindowsSdkVersion is used directly</param>
+		/// <param name="bUseCPPWinRT">Include the CPP/WinRT language projection</param>
 		/// <param name="Logger">Logger for output</param>
 		/// <returns>New environment object with paths for the given settings</returns>
 		[SupportedOSPlatform("windows")]
-		public static VCEnvironment Create(WindowsCompiler Compiler, UnrealTargetPlatform Platform, WindowsArchitecture Architecture, string? CompilerVersion, string? WindowsSdkVersion, string? SuppliedSdkDirectoryForVersion, ILogger Logger)
+		public static VCEnvironment Create(WindowsCompiler Compiler, UnrealTargetPlatform Platform, WindowsArchitecture Architecture, string? CompilerVersion, string? WindowsSdkVersion, string? SuppliedSdkDirectoryForVersion, bool bUseCPPWinRT, ILogger Logger)
 		{
-			return Create( new VCEnvironmentParameters(Compiler, Platform, Architecture, CompilerVersion, WindowsSdkVersion, SuppliedSdkDirectoryForVersion, Logger), Logger );
+			return Create( new VCEnvironmentParameters(Compiler, Platform, Architecture, CompilerVersion, WindowsSdkVersion, SuppliedSdkDirectoryForVersion, bUseCPPWinRT, Logger), Logger );
 		}
 
 		/// <summary>
@@ -427,6 +438,9 @@ namespace UnrealBuildTool
 		/// <summary>Optional directory for redistributable items (DLLs etc)</summary>
 		public DirectoryReference? RedistDir;
 
+		/// <summary>Include the CPP/WinRT language projection</summary>
+		public bool bUseCPPWinRT;
+
 		/// <summary>
 		/// Creates VC environment construction parameters with the given settings
 		/// </summary>
@@ -436,10 +450,11 @@ namespace UnrealBuildTool
 		/// <param name="CompilerVersion">The specific toolchain version to use</param>
 		/// <param name="WindowsSdkVersion">Version of the Windows SDK to use</param>
 		/// <param name="SuppliedSdkDirectoryForVersion">If specified, this is the SDK directory to use, otherwise, attempt to look up via registry. If specified, the WindowsSdkVersion is used directly</param>
+		/// <param name="bUseCPPWinRT">Include the CPP/WinRT language projection</param>
 		/// <param name="Logger">Logger for output</param>
 		/// <returns>Creation parameters for VC environment</returns>
 		[SupportedOSPlatform("windows")]
-		public VCEnvironmentParameters (WindowsCompiler Compiler, UnrealTargetPlatform Platform, WindowsArchitecture Architecture, string? CompilerVersion, string? WindowsSdkVersion, string? SuppliedSdkDirectoryForVersion, ILogger Logger)
+		public VCEnvironmentParameters (WindowsCompiler Compiler, UnrealTargetPlatform Platform, WindowsArchitecture Architecture, string? CompilerVersion, string? WindowsSdkVersion, string? SuppliedSdkDirectoryForVersion, bool bUseCPPWinRT, ILogger Logger)
 		{
 			// Get the compiler version info
 			VersionNumber? SelectedCompilerVersion;
@@ -509,6 +524,7 @@ namespace UnrealBuildTool
 			this.WindowsSdkDir = SelectedWindowsSdkDir;
 			this.WindowsSdkVersion = SelectedWindowsSdkVersion;
 			this.RedistDir = SelectedRedistDir;
+			this.bUseCPPWinRT = bUseCPPWinRT;
 		}
 	}
 }
diff --git a/Engine/Source/Programs/UnrealBuildTool/Platform/Windows/VCToolChain.cs b/Engine/Source/Programs/UnrealBuildTool/Platform/Windows/VCToolChain.cs
--- a/Engine/Source/Programs/UnrealBuildTool/Platform/Windows/VCToolChain.cs
+++ b/Engine/Source/Programs/UnrealBuildTool/Platform/Windows/VCToolChain.cs
@@ -2323,7 +2323,7 @@ namespace UnrealBuildTool
 		public static string GetVCIncludePaths(UnrealTargetPlatform Platform, WindowsCompiler Compiler, string? CompilerVersion, ILogger Logger)
 		{
 			// Make sure we've got the environment variables set up for this target
-			VCEnvironment EnvVars = VCEnvironment.Create(Compiler, Platform, WindowsArchitecture.x64, CompilerVersion, null, null, Logger);
+			VCEnvironment EnvVars = VCEnvironment.Create(Compiler, Platform, WindowsArchitecture.x64, CompilerVersion, null, null, false, Logger);
 
 			// Also add any include paths from the INCLUDE environment variable.  MSVC is not necessarily running with an environment that
 			// matches what UBT extracted from the vcvars*.bat using SetEnvironmentVariablesFromBatchFile().  We'll use the variables we
diff --git a/Engine/Source/Programs/UnrealBuildTool/ProjectFiles/VisualStudioCode/VSCodeProjectFileGenerator.cs b/Engine/Source/Programs/UnrealBuildTool/ProjectFiles/VisualStudioCode/VSCodeProjectFileGenerator.cs
--- a/Engine/Source/Programs/UnrealBuildTool/ProjectFiles/VisualStudioCode/VSCodeProjectFileGenerator.cs
+++ b/Engine/Source/Programs/UnrealBuildTool/ProjectFiles/VisualStudioCode/VSCodeProjectFileGenerator.cs
@@ -328,7 +328,7 @@ namespace UnrealBuildTool
 			DirectoryReference? SysRootPath = null;
 			if (OperatingSystem.IsWindows())
 			{
-				VCEnvironment Environment = VCEnvironment.Create(WindowsPlatform.GetDefaultCompiler(null, Target.Rules.WindowsPlatform.Architecture, Logger), Target.Platform, Target.Rules.WindowsPlatform.Architecture, null, Target.Rules.WindowsPlatform.WindowsSdkVersion, null, Logger);
+				VCEnvironment Environment = VCEnvironment.Create(WindowsPlatform.GetDefaultCompiler(null, Target.Rules.WindowsPlatform.Architecture, Logger), Target.Platform, Target.Rules.WindowsPlatform.Architecture, null, Target.Rules.WindowsPlatform.WindowsSdkVersion, null, Target.Rules.WindowsPlatform.bUseCPPWinRT, Logger);
 				CompilerPath = FileReference.FromString(Environment.CompilerPath.FullName);
 				UsingClang = false;
 			}
