UnrealBuildTool: Rename clang trace summary columns as the data is not miliseconds

#rnx

diff --git a/Engine/Source/Programs/UnrealBuildTool/Modes/AggregateClangTimingInfo.cs b/Engine/Source/Programs/UnrealBuildTool/Modes/AggregateClangTimingInfo.cs
--- a/Engine/Source/Programs/UnrealBuildTool/Modes/AggregateClangTimingInfo.cs
+++ b/Engine/Source/Programs/UnrealBuildTool/Modes/AggregateClangTimingInfo.cs
@@ -43,21 +43,21 @@ namespace UnrealBuildTool
 			public string Module => SourceFile.Directory.GetDirectoryName();
 			public string Name => SourceFile.GetFileName();
 
-			public long TotalExecuteCompilerMs { get; init; }
+			public long TotalExecuteCompiler { get; init; }
 
 			// Subset of execute compiler
-			public long TotalFrontendMs { get; init; }
-			public long TotalBackendMs { get; init; }
+			public long TotalFrontend { get; init; }
+			public long TotalBackend { get; init; }
 
 			// Subset of frontend
-			public long TotalSourceMs { get; init; }
-			public long TotalInstantiateFunctionMs { get; init; }
-			public long TotalCodeGenFunctionMs { get; init; }
+			public long TotalSource { get; init; }
+			public long TotalInstantiateFunction { get; init; }
+			public long TotalCodeGenFunction { get; init; }
 
 			// Subset of backend
-			public long TotalModuleToFunctionPassAdaptorMs { get; init; }
-			public long TotalModuleInlinerWrapperPassMs { get; init; }
-			public long TotalOptModuleMs { get; init; }
+			public long TotalModuleToFunctionPassAdaptor { get; init; }
+			public long TotalModuleInlinerWrapperPass { get; init; }
+			public long TotalOptModule { get; init; }
 
 			// Frontend entry counts
 			public long SourceEntries { get; init; }
@@ -71,22 +71,22 @@ namespace UnrealBuildTool
 			public TraceData(FileReference inputFile, ClangTrace? trace)
 			{
 				SourceFile = inputFile;
-
-				TotalExecuteCompilerMs = trace?.traceEvents?.FindLast(x => string.Equals(x.name, "Total ExecuteCompiler"))?.dur ?? 0;
+				
+				TotalExecuteCompiler = trace?.traceEvents?.FindLast(x => string.Equals(x.name, "Total ExecuteCompiler"))?.dur ?? 0;
 
 				// Subset of execute compiler
-				TotalFrontendMs = trace?.traceEvents?.FindLast(x => string.Equals(x.name, "Total Frontend"))?.dur ?? 0;
-				TotalBackendMs = trace?.traceEvents?.FindLast(x => string.Equals(x.name, "Total Backend"))?.dur ?? 0;
+				TotalFrontend = trace?.traceEvents?.FindLast(x => string.Equals(x.name, "Total Frontend"))?.dur ?? 0;
+				TotalBackend = trace?.traceEvents?.FindLast(x => string.Equals(x.name, "Total Backend"))?.dur ?? 0;
 
 				// Subset of frontend
-				TotalSourceMs = trace?.traceEvents?.FindLast(x => string.Equals(x.name, "Total Source"))?.dur ?? 0;
-				TotalInstantiateFunctionMs = trace?.traceEvents?.FindLast(x => string.Equals(x.name, "Total InstantiateFunction"))?.dur ?? 0;
-				TotalCodeGenFunctionMs = trace?.traceEvents?.FindLast(x => string.Equals(x.name, "Total CodeGen Function"))?.dur ?? 0;
+				TotalSource = trace?.traceEvents?.FindLast(x => string.Equals(x.name, "Total Source"))?.dur ?? 0;
+				TotalInstantiateFunction = trace?.traceEvents?.FindLast(x => string.Equals(x.name, "Total InstantiateFunction"))?.dur ?? 0;
+				TotalCodeGenFunction = trace?.traceEvents?.FindLast(x => string.Equals(x.name, "Total CodeGen Function"))?.dur ?? 0;
 
 				// Subset of backend
-				TotalModuleToFunctionPassAdaptorMs = trace?.traceEvents?.FindLast(x => string.Equals(x.name, "Total ModuleToFunctionPassAdaptor"))?.dur ?? 0;
-				TotalModuleInlinerWrapperPassMs = trace?.traceEvents?.FindLast(x => string.Equals(x.name, "Total ModuleInlinerWrapperPass"))?.dur ?? 0;
-				TotalOptModuleMs = trace?.traceEvents?.FindLast(x => string.Equals(x.name, "Total OptModule"))?.dur ?? 0;
+				TotalModuleToFunctionPassAdaptor = trace?.traceEvents?.FindLast(x => string.Equals(x.name, "Total ModuleToFunctionPassAdaptor"))?.dur ?? 0;
+				TotalModuleInlinerWrapperPass = trace?.traceEvents?.FindLast(x => string.Equals(x.name, "Total ModuleInlinerWrapperPass"))?.dur ?? 0;
+				TotalOptModule = trace?.traceEvents?.FindLast(x => string.Equals(x.name, "Total OptModule"))?.dur ?? 0;
 
 				// Frontend entry counts
 				SourceEntries = trace?.traceEvents?.Where(x => string.Equals(x.name, "Source")).LongCount() ?? 0;
@@ -102,21 +102,21 @@ namespace UnrealBuildTool
 				"Module",
 				"Name",
 
-				"TotalExecuteCompilerMs",
+				"TotalExecuteCompiler",
 
 				// Subset of execute compiler
-				"TotalFrontendMs",
-				"TotalBackendMs",
+				"TotalFrontend",
+				"TotalBackend",
 
 				// Subset of frontend
-				"TotalSourceMs",
-				"TotalInstantiateFunctionMs",
-				"TotalCodeGenFunctionMs",
+				"TotalSource",
+				"TotalInstantiateFunction",
+				"TotalCodeGenFunction",
 
 				// Subset of backend
-				"TotalModuleToFunctionPassAdaptorMs",
-				"TotalModuleInlinerWrapperPassMs",
-				"TotalOptModuleMs",
+				"TotalModuleToFunctionPassAdaptor",
+				"TotalModuleInlinerWrapperPass",
+				"TotalOptModule",
 
 				// Frontend entry counts
 				"SourceEntries",
@@ -221,7 +221,7 @@ namespace UnrealBuildTool
 				string TempFilePath = Path.Join(Path.GetTempPath(), HeadersFile.GetFileName() + ".tmp");
 				using (StreamWriter Writer = new StreamWriter(TempFilePath))
 				{
-					Writer.WriteLine("Source,Count,TotalMs,MinMs,MaxMs,AverageMs");
+					Writer.WriteLine("Source,Count,Total,Min,Max,Average");
 					foreach (var Data in Sources.OrderBy(x => x.Key.FullName).Where(x => x.Key.HasExtension(".h")))
 					{
 						Writer.WriteLine($"{Data.Key.FullName},{Data.Value.Count},{Data.Value.Sum()},{Data.Value.Min()},{Data.Value.Max()},{Data.Value.Average()}");
