UnrealBuildTool: Archive clang traces when gathering -Timing summaries

#rnx

diff --git a/Engine/Source/Programs/UnrealBuildTool/Platform/Windows/VCToolChain.cs b/Engine/Source/Programs/UnrealBuildTool/Platform/Windows/VCToolChain.cs
--- a/Engine/Source/Programs/UnrealBuildTool/Platform/Windows/VCToolChain.cs
+++ b/Engine/Source/Programs/UnrealBuildTool/Platform/Windows/VCToolChain.cs
@@ -1606,14 +1606,14 @@ namespace UnrealBuildTool
 
 						FileItem AggregateOutputFile = FileItem.GetItemByFileReference(FileReference.Combine(Makefile.ProjectIntermediateDirectory, $"{Target.Name}.trace.csv"));
 						FileItem HeadersOutputFile = FileItem.GetItemByFileReference(FileReference.Combine(Makefile.ProjectIntermediateDirectory, $"{Target.Name}.headers.csv"));
-						List<string> ActionArgs = new List<string>()
+						List<string> AggregateActionArgs = new List<string>()
 						{
-							$"-AggregateFile={AggregateOutputFile.FullName}",
 							$"-ManifestFile={ManifestFile.FullName}",
+							$"-AggregateFile={AggregateOutputFile.FullName}",
 							$"-HeadersFile={HeadersOutputFile.FullName}",
 						};
 
-						Action AggregateTimingInfoAction = MakefileBuilder.CreateRecursiveAction<AggregateClangTimingInfo>(ActionType.ParseTimingInfo, string.Join(" ", ActionArgs));
+						Action AggregateTimingInfoAction = MakefileBuilder.CreateRecursiveAction<AggregateClangTimingInfo>(ActionType.ParseTimingInfo, string.Join(" ", AggregateActionArgs));
 						AggregateTimingInfoAction.WorkingDirectory = Unreal.EngineSourceDirectory;
 						AggregateTimingInfoAction.StatusDescription = $"Aggregating {TimingJsonFiles.Count} Timing File(s)";
 						AggregateTimingInfoAction.bCanExecuteRemotely = false;
@@ -1622,8 +1622,24 @@ namespace UnrealBuildTool
 
 						AggregateTimingInfoAction.ProducedItems.Add(AggregateOutputFile);
 						AggregateTimingInfoAction.ProducedItems.Add(HeadersOutputFile);
-						Makefile.OutputItems.Add(AggregateOutputFile);
-						Makefile.OutputItems.Add(HeadersOutputFile);
+						Makefile.OutputItems.AddRange(AggregateTimingInfoAction.ProducedItems);
+
+						FileItem ArchiveOutputFile = FileItem.GetItemByFileReference(FileReference.Combine(Makefile.ProjectIntermediateDirectory, $"{Target.Name}.traces.zip"));
+						List<string> ArchiveActionArgs = new List<string>()
+						{
+							$"-ManifestFile={ManifestFile.FullName}",
+							$"-ArchiveFile={ArchiveOutputFile.FullName}",
+						};
+
+						Action ArchiveTimingInfoAction = MakefileBuilder.CreateRecursiveAction<AggregateClangTimingInfo>(ActionType.ParseTimingInfo, string.Join(" ", ArchiveActionArgs));
+						ArchiveTimingInfoAction.WorkingDirectory = Unreal.EngineSourceDirectory;
+						ArchiveTimingInfoAction.StatusDescription = $"Archiving {TimingJsonFiles.Count} Timing File(s)";
+						ArchiveTimingInfoAction.bCanExecuteRemotely = false;
+						ArchiveTimingInfoAction.bCanExecuteRemotelyWithSNDBS = false;
+						ArchiveTimingInfoAction.PrerequisiteItems.AddRange(TimingJsonFiles);
+
+						ArchiveTimingInfoAction.ProducedItems.Add(ArchiveOutputFile);
+						Makefile.OutputItems.AddRange(ArchiveTimingInfoAction.ProducedItems);
 					}
 				}
 			}
diff --git a/Engine/Source/Programs/UnrealBuildTool/ToolChain/ClangToolChain.cs b/Engine/Source/Programs/UnrealBuildTool/ToolChain/ClangToolChain.cs
--- a/Engine/Source/Programs/UnrealBuildTool/ToolChain/ClangToolChain.cs
+++ b/Engine/Source/Programs/UnrealBuildTool/ToolChain/ClangToolChain.cs
@@ -236,14 +236,14 @@ namespace UnrealBuildTool
 
 					FileItem AggregateOutputFile = FileItem.GetItemByFileReference(FileReference.Combine(Makefile.ProjectIntermediateDirectory, $"{Target.Name}.trace.csv"));
 					FileItem HeadersOutputFile = FileItem.GetItemByFileReference(FileReference.Combine(Makefile.ProjectIntermediateDirectory, $"{Target.Name}.headers.csv"));
-					List<string> ActionArgs = new List<string>()
+					List<string> AggregateActionArgs = new List<string>()
 					{
-						$"-AggregateFile={AggregateOutputFile.FullName}",
 						$"-ManifestFile={ManifestFile.FullName}",
+						$"-AggregateFile={AggregateOutputFile.FullName}",
 						$"-HeadersFile={HeadersOutputFile.FullName}",
 					};
 
-					Action AggregateTimingInfoAction = MakefileBuilder.CreateRecursiveAction<AggregateClangTimingInfo>(ActionType.ParseTimingInfo, string.Join(" ", ActionArgs));
+					Action AggregateTimingInfoAction = MakefileBuilder.CreateRecursiveAction<AggregateClangTimingInfo>(ActionType.ParseTimingInfo, string.Join(" ", AggregateActionArgs));
 					AggregateTimingInfoAction.WorkingDirectory = Unreal.EngineSourceDirectory;
 					AggregateTimingInfoAction.StatusDescription = $"Aggregating {TimingJsonFiles.Count} Timing File(s)";
 					AggregateTimingInfoAction.bCanExecuteRemotely = false;
@@ -252,8 +252,24 @@ namespace UnrealBuildTool
 
 					AggregateTimingInfoAction.ProducedItems.Add(AggregateOutputFile);
 					AggregateTimingInfoAction.ProducedItems.Add(HeadersOutputFile);
-					Makefile.OutputItems.Add(AggregateOutputFile);
-					Makefile.OutputItems.Add(HeadersOutputFile);
+					Makefile.OutputItems.AddRange(AggregateTimingInfoAction.ProducedItems);
+
+					FileItem ArchiveOutputFile = FileItem.GetItemByFileReference(FileReference.Combine(Makefile.ProjectIntermediateDirectory, $"{Target.Name}.traces.zip"));
+					List<string> ArchiveActionArgs = new List<string>()
+					{
+						$"-ManifestFile={ManifestFile.FullName}",
+						$"-ArchiveFile={ArchiveOutputFile.FullName}",
+					};
+
+					Action ArchiveTimingInfoAction = MakefileBuilder.CreateRecursiveAction<AggregateClangTimingInfo>(ActionType.ParseTimingInfo, string.Join(" ", ArchiveActionArgs));
+					ArchiveTimingInfoAction.WorkingDirectory = Unreal.EngineSourceDirectory;
+					ArchiveTimingInfoAction.StatusDescription = $"Archiving {TimingJsonFiles.Count} Timing File(s)";
+					ArchiveTimingInfoAction.bCanExecuteRemotely = false;
+					ArchiveTimingInfoAction.bCanExecuteRemotelyWithSNDBS = false;
+					ArchiveTimingInfoAction.PrerequisiteItems.AddRange(TimingJsonFiles);
+
+					ArchiveTimingInfoAction.ProducedItems.Add(ArchiveOutputFile);
+					Makefile.OutputItems.AddRange(ArchiveTimingInfoAction.ProducedItems);
 				}
 			}
 		}
