UnrealBuildTool: Add .inl to trace header include list

#rnx

diff --git a/Engine/Source/Programs/UnrealBuildTool/Modes/AggregateClangTimingInfo.cs b/Engine/Source/Programs/UnrealBuildTool/Modes/AggregateClangTimingInfo.cs
--- a/Engine/Source/Programs/UnrealBuildTool/Modes/AggregateClangTimingInfo.cs
+++ b/Engine/Source/Programs/UnrealBuildTool/Modes/AggregateClangTimingInfo.cs
@@ -222,7 +222,7 @@ namespace UnrealBuildTool
 				using (StreamWriter Writer = new StreamWriter(TempFilePath))
 				{
 					Writer.WriteLine("Source,Count,Total,Min,Max,Average");
-					foreach (var Data in Sources.OrderBy(x => x.Key.FullName).Where(x => x.Key.HasExtension(".h")))
+					foreach (var Data in Sources.OrderBy(x => x.Key.FullName).Where(x => x.Key.HasExtension(".h") || x.Key.HasExtension(".inl")))
 					{
 						Writer.WriteLine($"{Data.Key.FullName},{Data.Value.Count},{Data.Value.Sum()},{Data.Value.Min()},{Data.Value.Max()},{Data.Value.Average()}");
 					}
@@ -246,7 +246,6 @@ namespace UnrealBuildTool
 					{
 						ZipArchive.CreateEntryFromFile_CrossPlatform(HeadersFile.FullName, HeadersFile.GetFileName(), CompressionLevel.Optimal);
 					}
-					List<string> Metadata = new List<string>();
 					foreach (FileReference SourceFile in SourceFiles)
 					{
 						FileReference JsonFile = new FileReference($"{SourceFile.FullName}.json");
