UnrealBuildTool: Generate CompileScore data when compiling with clang traces if the compile score extractor is available in /Engine/Extras/ThirdPartyNotUE/CompileScore

diff --git a/Engine/Source/Programs/UnrealBuildTool/Platform/Windows/VCToolChain.cs b/Engine/Source/Programs/UnrealBuildTool/Platform/Windows/VCToolChain.cs
--- a/Engine/Source/Programs/UnrealBuildTool/Platform/Windows/VCToolChain.cs
+++ b/Engine/Source/Programs/UnrealBuildTool/Platform/Windows/VCToolChain.cs
@@ -1640,6 +1640,28 @@ namespace UnrealBuildTool
 
 						ArchiveTimingInfoAction.ProducedItems.Add(ArchiveOutputFile);
 						Makefile.OutputItems.AddRange(ArchiveTimingInfoAction.ProducedItems);
+
+						// Extract CompileScore data from traces
+						FileReference ScoreDataExtractor = FileReference.Combine(Unreal.RootDirectory, "Engine", "Extras", "ThirdPartyNotUE", "CompileScore", "ScoreDataExtractor.exe");
+						if (RuntimePlatform.IsWindows && FileReference.Exists(ScoreDataExtractor))
+						{
+							FileItem CompileScoreOutput = FileItem.GetItemByFileReference(FileReference.Combine(Makefile.ProjectIntermediateDirectory, $"{Target.Name}.scor"));
+
+							Action CompileScoreExtractorAction = MakefileBuilder.CreateAction(ActionType.ParseTimingInfo);
+							CompileScoreExtractorAction.WorkingDirectory = Unreal.EngineSourceDirectory;
+							CompileScoreExtractorAction.StatusDescription = $"Extracting CompileScore";
+							CompileScoreExtractorAction.bCanExecuteRemotely = false;
+							CompileScoreExtractorAction.bCanExecuteRemotelyWithSNDBS = false;
+							CompileScoreExtractorAction.PrerequisiteItems.AddRange(TimingJsonFiles);
+							CompileScoreExtractorAction.CommandPath = ScoreDataExtractor;
+							CompileScoreExtractorAction.CommandArguments = $"-clang -verbosity 0 -timelinepack 1000000 -extract -i \"{NormalizeCommandLinePath(Makefile.ProjectIntermediateDirectory)}\" -o \"{NormalizeCommandLinePath(CompileScoreOutput)}\"";
+
+							CompileScoreExtractorAction.ProducedItems.Add(CompileScoreOutput);
+							CompileScoreExtractorAction.ProducedItems.Add(FileItem.GetItemByFileReference(FileReference.Combine(Makefile.ProjectIntermediateDirectory, $"{Target.Name}.scor.gbl")));
+							CompileScoreExtractorAction.ProducedItems.Add(FileItem.GetItemByFileReference(FileReference.Combine(Makefile.ProjectIntermediateDirectory, $"{Target.Name}.scor.incl")));
+							CompileScoreExtractorAction.ProducedItems.Add(FileItem.GetItemByFileReference(FileReference.Combine(Makefile.ProjectIntermediateDirectory, $"{Target.Name}.scor.t0000")));
+							Makefile.OutputItems.AddRange(CompileScoreExtractorAction.ProducedItems);
+						}
 					}
 				}
 			}
diff --git a/Engine/Source/Programs/UnrealBuildTool/ToolChain/ClangToolChain.cs b/Engine/Source/Programs/UnrealBuildTool/ToolChain/ClangToolChain.cs
--- a/Engine/Source/Programs/UnrealBuildTool/ToolChain/ClangToolChain.cs
+++ b/Engine/Source/Programs/UnrealBuildTool/ToolChain/ClangToolChain.cs
@@ -270,6 +270,28 @@ namespace UnrealBuildTool
 
 					ArchiveTimingInfoAction.ProducedItems.Add(ArchiveOutputFile);
 					Makefile.OutputItems.AddRange(ArchiveTimingInfoAction.ProducedItems);
+
+					// Extract CompileScore data from traces
+					FileReference ScoreDataExtractor = FileReference.Combine(Unreal.RootDirectory, "Engine", "Extras", "ThirdPartyNotUE", "CompileScore", "ScoreDataExtractor.exe");
+					if (RuntimePlatform.IsWindows && FileReference.Exists(ScoreDataExtractor))
+					{
+						FileItem CompileScoreOutput = FileItem.GetItemByFileReference(FileReference.Combine(Makefile.ProjectIntermediateDirectory, $"{Target.Name}.scor"));
+
+						Action CompileScoreExtractorAction = MakefileBuilder.CreateAction(ActionType.ParseTimingInfo);
+						CompileScoreExtractorAction.WorkingDirectory = Unreal.EngineSourceDirectory;
+						CompileScoreExtractorAction.StatusDescription = $"Extracting CompileScore";
+						CompileScoreExtractorAction.bCanExecuteRemotely = false;
+						CompileScoreExtractorAction.bCanExecuteRemotelyWithSNDBS = false;
+						CompileScoreExtractorAction.PrerequisiteItems.AddRange(TimingJsonFiles);
+						CompileScoreExtractorAction.CommandPath = ScoreDataExtractor;
+						CompileScoreExtractorAction.CommandArguments = $"-clang -verbosity 0 -timelinepack 1000000 -extract -i \"{NormalizeCommandLinePath(Makefile.ProjectIntermediateDirectory)}\" -o \"{NormalizeCommandLinePath(CompileScoreOutput)}\"";
+
+						CompileScoreExtractorAction.ProducedItems.Add(CompileScoreOutput);
+						CompileScoreExtractorAction.ProducedItems.Add(FileItem.GetItemByFileReference(FileReference.Combine(Makefile.ProjectIntermediateDirectory, $"{Target.Name}.scor.gbl")));
+						CompileScoreExtractorAction.ProducedItems.Add(FileItem.GetItemByFileReference(FileReference.Combine(Makefile.ProjectIntermediateDirectory, $"{Target.Name}.scor.incl")));
+						CompileScoreExtractorAction.ProducedItems.Add(FileItem.GetItemByFileReference(FileReference.Combine(Makefile.ProjectIntermediateDirectory, $"{Target.Name}.scor.t0000")));
+						Makefile.OutputItems.AddRange(CompileScoreExtractorAction.ProducedItems);
+					}
 				}
 			}
 		}
