UnrealBuildTool: Option to consider logical cores when running local builds

#rnx

diff --git a/Engine/Source/Programs/UnrealBuildTool/Configuration/BuildConfiguration.cs b/Engine/Source/Programs/UnrealBuildTool/Configuration/BuildConfiguration.cs
--- a/Engine/Source/Programs/UnrealBuildTool/Configuration/BuildConfiguration.cs
+++ b/Engine/Source/Programs/UnrealBuildTool/Configuration/BuildConfiguration.cs
@@ -103,6 +103,13 @@ namespace UnrealBuildTool
 		[CommandLine("-MaxParallelActions")]
 		public int MaxParallelActions = 0;
 
+		/// <summary>
+		/// Consider logical cores when determining how many total cpu cores are available.
+		/// </summary>
+		[XmlConfigFile(Name = "bAllCores")]
+		[CommandLine("-AllCores")]
+		public bool bAllCores = false;
+
 		/// <summary>
 		/// If true, force header regeneration. Intended for the build machine.
 		/// </summary>
diff --git a/Engine/Source/Programs/UnrealBuildTool/Executors/Experimental/FASTBuild.cs b/Engine/Source/Programs/UnrealBuildTool/Executors/Experimental/FASTBuild.cs
--- a/Engine/Source/Programs/UnrealBuildTool/Executors/Experimental/FASTBuild.cs
+++ b/Engine/Source/Programs/UnrealBuildTool/Executors/Experimental/FASTBuild.cs
@@ -157,11 +157,11 @@ namespace UnrealBuildTool
 		/// <summary>
 		/// Constructor
 		/// </summary>
-		public FASTBuild(int MaxLocalActions, bool bCompactOutput, ILogger Logger)
+		public FASTBuild(int MaxLocalActions, bool bAllCores, bool bCompactOutput, ILogger Logger)
 		{
 			XmlConfig.ApplyTo(this);
 
-			this.LocalExecutor = new ParallelExecutor(MaxLocalActions, bCompactOutput, Logger);
+			this.LocalExecutor = new ParallelExecutor(MaxLocalActions, bAllCores, bCompactOutput, Logger);
 		}
 
 		public override string Name
diff --git a/Engine/Source/Programs/UnrealBuildTool/Executors/Experimental/HordeExecutor.cs b/Engine/Source/Programs/UnrealBuildTool/Executors/Experimental/HordeExecutor.cs
--- a/Engine/Source/Programs/UnrealBuildTool/Executors/Experimental/HordeExecutor.cs
+++ b/Engine/Source/Programs/UnrealBuildTool/Executors/Experimental/HordeExecutor.cs
@@ -279,9 +279,10 @@ namespace UnrealBuildTool
 		/// Constructor
 		/// </summary>
 		/// <param name="MaxLocalActions">How many local actions to execute in parallel</param>
+		/// <param name="bAllCores">Consider logical cores when determining how many total cpu cores are available</param>
 		/// <param name="bCompactOutput"></param>
 		/// <param name="Logger"></param>
-		public HordeExecutor(int MaxLocalActions, bool bCompactOutput, ILogger Logger) : base(MaxLocalActions, bCompactOutput, Logger)
+		public HordeExecutor(int MaxLocalActions, bool bAllCores, bool bCompactOutput, ILogger Logger) : base(MaxLocalActions, bAllCores, bCompactOutput, Logger)
 		{
 			XmlConfig.ApplyTo(this);
 		}
diff --git a/Engine/Source/Programs/UnrealBuildTool/Executors/HybridExecutor.cs b/Engine/Source/Programs/UnrealBuildTool/Executors/HybridExecutor.cs
--- a/Engine/Source/Programs/UnrealBuildTool/Executors/HybridExecutor.cs
+++ b/Engine/Source/Programs/UnrealBuildTool/Executors/HybridExecutor.cs
@@ -16,7 +16,6 @@ namespace UnrealBuildTool
 		/// <summary>
 		/// Maximum number of actions to execute locally.
 		/// </summary>
-		[XmlConfigFile(Category = "HybridExecutor")]
 		int MaxLocalActions;
 
 		/// <summary>
@@ -32,10 +31,10 @@ namespace UnrealBuildTool
 		/// <summary>
 		/// Constructor
 		/// </summary>
-		public HybridExecutor(List<TargetDescriptor> TargetDescriptors, int InMaxLocalActions, bool bCompactOutput, ILogger Logger, ActionExecutor? InLocalExecutor = null, ActionExecutor? InRemoteExecutor = null)
+		public HybridExecutor(List<TargetDescriptor> TargetDescriptors, int InMaxLocalActions, bool bAllCores, bool bCompactOutput, ILogger Logger, ActionExecutor? InLocalExecutor = null, ActionExecutor? InRemoteExecutor = null)
 		{
 			MaxLocalActions = InMaxLocalActions;
-			LocalExecutor = InLocalExecutor ?? new ParallelExecutor(MaxLocalActions, bCompactOutput, Logger);
+			LocalExecutor = InLocalExecutor ?? new ParallelExecutor(MaxLocalActions, bAllCores, bCompactOutput, Logger);
 			RemoteExecutor = InRemoteExecutor ?? (XGE.IsAvailable(Logger) ? (ActionExecutor)new XGE() : new SNDBS(TargetDescriptors));
 
 			XmlConfig.ApplyTo(this);
diff --git a/Engine/Source/Programs/UnrealBuildTool/Executors/ParallelExecutor.cs b/Engine/Source/Programs/UnrealBuildTool/Executors/ParallelExecutor.cs
--- a/Engine/Source/Programs/UnrealBuildTool/Executors/ParallelExecutor.cs
+++ b/Engine/Source/Programs/UnrealBuildTool/Executors/ParallelExecutor.cs
@@ -22,7 +22,10 @@ namespace UnrealBuildTool
 		/// Maximum processor count for local execution. 
 		/// </summary>
 		[XmlConfigFile]
-		private static int MaxProcessorCount = int.MaxValue;
+		[Obsolete("ParallelExecutor.MaxProcessorCount is deprecated. Please update xml to use BuildConfiguration.MaxParallelActions")]
+#pragma warning disable 0169
+		private static int MaxProcessorCount;
+#pragma warning restore 0169
 
 		/// <summary>
 		/// Processor count multiplier for local execution. Can be below 1 to reserve CPU for other tasks.
@@ -90,7 +93,7 @@ namespace UnrealBuildTool
 
 		private static readonly char[] LineEndingSplit = new char[] { '\n', '\r' };
 
-		public static int GetDefaultNumParallelProcesses(ILogger Logger)
+		public static int GetDefaultNumParallelProcesses(int MaxLocalActions, bool bAllCores, ILogger Logger)
 		{
 			double MemoryPerActionBytesComputed = Math.Max(MemoryPerActionBytes, MemoryPerActionBytesOverride);
 			if (MemoryPerActionBytesComputed > MemoryPerActionBytes)
@@ -98,29 +101,22 @@ namespace UnrealBuildTool
 				Logger.LogInformation("Overriding MemoryPerAction with target-defined value of {Memory} bytes", MemoryPerActionBytesComputed / 1024 / 1024 / 1024);
 			}
 
-			return Utils.GetMaxActionsToExecuteInParallel(MaxProcessorCount, ProcessorCountMultiplier, Convert.ToInt64(MemoryPerActionBytesComputed));
+			return Utils.GetMaxActionsToExecuteInParallel(MaxLocalActions, ProcessorCountMultiplier, bAllCores, Convert.ToInt64(MemoryPerActionBytesComputed));
 		}
 
 		/// <summary>
 		/// Constructor
 		/// </summary>
 		/// <param name="MaxLocalActions">How many actions to execute in parallel</param>
+		/// <param name="bAllCores">Consider logical cores when determining how many total cpu cores are available</param>
 		/// <param name="bCompactOutput">Should output be written in a compact fashion</param>
 		/// <param name="Logger">Logger for output</param>
-		public ParallelExecutor(int MaxLocalActions, bool bCompactOutput, ILogger Logger)
+		public ParallelExecutor(int MaxLocalActions, bool bAllCores, bool bCompactOutput, ILogger Logger)
 		{
 			XmlConfig.ApplyTo(this);
 
-			// if specified this caps how many processors we can use
-			if (MaxLocalActions > 0)
-			{
-				NumParallelProcesses = MaxLocalActions;
-			}
-			else
-			{
-				// Figure out how many processors to use
-				NumParallelProcesses = GetDefaultNumParallelProcesses(Logger);
-			}
+			// Figure out how many processors to use
+			NumParallelProcesses = GetDefaultNumParallelProcesses(MaxLocalActions, bAllCores, Logger);
 
 			this.bCompactOutput = bCompactOutput;
 		}
@@ -484,6 +480,6 @@ namespace UnrealBuildTool
 		/// <summary>
 		/// Maximum number of processes that should be used for execution
 		/// </summary>
-		public static int GetMaxParallelProcesses(ILogger Logger) => ParallelExecutor.GetDefaultNumParallelProcesses(Logger);
+		public static int GetMaxParallelProcesses(ILogger Logger) => ParallelExecutor.GetDefaultNumParallelProcesses(0, false, Logger);
 	}
 }
diff --git a/Engine/Source/Programs/UnrealBuildTool/System/ActionGraph.cs b/Engine/Source/Programs/UnrealBuildTool/System/ActionGraph.cs
--- a/Engine/Source/Programs/UnrealBuildTool/System/ActionGraph.cs
+++ b/Engine/Source/Programs/UnrealBuildTool/System/ActionGraph.cs
@@ -337,11 +337,11 @@ namespace UnrealBuildTool
 		/// </summary>
 		private static ActionExecutor SelectExecutor(BuildConfiguration BuildConfiguration, int ActionCount, List<TargetDescriptor> TargetDescriptors, ILogger Logger)
 		{
-			if (ActionCount > ParallelExecutor.GetDefaultNumParallelProcesses(Logger))
+			if (ActionCount > ParallelExecutor.GetDefaultNumParallelProcesses(BuildConfiguration.MaxParallelActions, BuildConfiguration.bAllCores, Logger))
 			{
 				if (BuildConfiguration.bAllowHybridExecutor && HybridExecutor.IsAvailable(Logger))
 				{
-					return new HybridExecutor(TargetDescriptors, BuildConfiguration.MaxParallelActions, BuildConfiguration.bCompactOutput, Logger);
+					return new HybridExecutor(TargetDescriptors, BuildConfiguration.MaxParallelActions, BuildConfiguration.bAllCores, BuildConfiguration.bCompactOutput, Logger);
 				}
 				else if (BuildConfiguration.bAllowXGE && XGE.IsAvailable(Logger) && ActionCount >= XGE.MinActions)
 				{
@@ -349,7 +349,7 @@ namespace UnrealBuildTool
 				}
 				else if (BuildConfiguration.bAllowFASTBuild && FASTBuild.IsAvailable(Logger))
 				{
-					return new FASTBuild(BuildConfiguration.MaxParallelActions, BuildConfiguration.bCompactOutput, Logger);
+					return new FASTBuild(BuildConfiguration.MaxParallelActions, BuildConfiguration.bAllCores, BuildConfiguration.bCompactOutput, Logger);
 				}
 				else if (BuildConfiguration.bAllowSNDBS && SNDBS.IsAvailable(Logger))
 				{
@@ -357,11 +357,11 @@ namespace UnrealBuildTool
 				}
 				else if (BuildConfiguration.bAllowHordeCompute && HordeExecutor.IsAvailable())
 				{
-					return new HordeExecutor(BuildConfiguration.MaxParallelActions, BuildConfiguration.bCompactOutput, Logger);
+					return new HordeExecutor(BuildConfiguration.MaxParallelActions, BuildConfiguration.bAllCores, BuildConfiguration.bCompactOutput, Logger);
 				}
 			}
 
-			return new ParallelExecutor(BuildConfiguration.MaxParallelActions, BuildConfiguration.bCompactOutput, Logger);
+			return new ParallelExecutor(BuildConfiguration.MaxParallelActions, BuildConfiguration.bAllCores, BuildConfiguration.bCompactOutput, Logger);
 		}
 
 		/// <summary>
diff --git a/Engine/Source/Programs/UnrealBuildTool/System/Utils.cs b/Engine/Source/Programs/UnrealBuildTool/System/Utils.cs
--- a/Engine/Source/Programs/UnrealBuildTool/System/Utils.cs
+++ b/Engine/Source/Programs/UnrealBuildTool/System/Utils.cs
@@ -1272,10 +1272,11 @@ namespace UnrealBuildTool
 		/// Determines the maximum number of actions to execute in parallel, taking into account the resources available on this machine.
 		/// </summary>
 		/// <param name="MaxProcessorCount">How many actions to execute in parallel. When 0 a default will be chosen based on system resources</param>
-		/// <param name="ProcessorCountMultiplier">Processor count multiplier for local execution. Can be below 1 to reserve CPU for other tasks.</param>
-		/// <param name="MemoryPerActionBytes"></param>
+		/// <param name="ProcessorCountMultiplier">Physical processor count multiplier for local execution. Can be below 1 to reserve CPU for other tasks.</param>
+		/// <param name="ConsiderLogicalCores">Consider logical cores when determing max actions to execute in parallel. Unused if ProcessorCountMultiplier is not 1.0.</param>
+		/// <param name="MemoryPerActionBytes">Limit max number of actions based on total system memory.</param>
 		/// <returns>Max number of actions to execute in parallel</returns>
-		public static int GetMaxActionsToExecuteInParallel(int MaxProcessorCount, double ProcessorCountMultiplier, long MemoryPerActionBytes)
+		public static int GetMaxActionsToExecuteInParallel(int MaxProcessorCount, double ProcessorCountMultiplier, bool ConsiderLogicalCores, long MemoryPerActionBytes)
 		{
 			// Get the number of logical processors
 			int NumLogicalCores = Utils.GetLogicalProcessorCount();
@@ -1298,6 +1299,11 @@ namespace UnrealBuildTool
 				MaxActionsToExecuteInParallel = (int)(NumPhysicalCores * ProcessorCountMultiplier);
 				Log.TraceInformationOnce($"  Requested {ProcessorCountMultiplier} process count multiplier: limiting max parallel actions to {MaxActionsToExecuteInParallel}");
 			}
+			else if (ConsiderLogicalCores && NumLogicalCores > NumPhysicalCores)
+			{
+				Log.TraceInformationOnce($"  Executing up to {NumLogicalCores} processes, one per logical core");
+				MaxActionsToExecuteInParallel = NumLogicalCores;
+			}
 			// kick off a task per physical core - evidence suggests that, in general, using more cores does not yield significantly better throughput
 			else
 			{
@@ -1320,7 +1326,7 @@ namespace UnrealBuildTool
 				}
 			}
 
-			if (MaxProcessorCount < MaxActionsToExecuteInParallel)
+			if (MaxProcessorCount > 0 && MaxProcessorCount < MaxActionsToExecuteInParallel)
 			{
 				MaxActionsToExecuteInParallel = Math.Max(1, Math.Min(MaxActionsToExecuteInParallel, MaxProcessorCount));
 				Log.TraceInformationOnce($"  Requested max {MaxProcessorCount} action(s): limiting max parallel actions to {MaxActionsToExecuteInParallel}");
