PR #9546: Improvements to PVS-Studio. PVSToolChain handles 3 additional settings from the Settings.xml file. New settings are added into .cfg file that is used by PVS-Studio.exe. (Contributed by [REDACTED_PERSON])

#rnx

diff --git a/Engine/Source/Programs/UnrealBuildTool/Platform/Windows/PVSToolChain.cs b/Engine/Source/Programs/UnrealBuildTool/Platform/Windows/PVSToolChain.cs
--- a/Engine/Source/Programs/UnrealBuildTool/Platform/Windows/PVSToolChain.cs
+++ b/Engine/Source/Programs/UnrealBuildTool/Platform/Windows/PVSToolChain.cs
@@ -48,6 +48,29 @@ namespace UnrealBuildTool
 		MISRA = 32,
 	}
 
+  	/// <summary>
+	/// Flags for the PVS analyzer timeout
+	/// </summary>
+	public enum AnalysisTimeoutFlags
+	{
+		/// <summary>
+		/// Analisys timeout for file 10 minutes (600 seconds)
+		/// </summary>
+		After_10_minutes = 600,
+		/// <summary>
+		/// Analisys timeout for file 30 minutes (1800 seconds)
+		/// </summary>
+		After_30_minutes = 1800,
+		/// <summary>
+		/// Analisys timeout for file 60 minutes (3600 seconds)
+		/// </summary>
+		After_60_minutes = 3600,
+		/// <summary>
+		/// Analisys timeout when not set (a lot of seconds)
+		/// </summary>
+		No_timeout = 999999
+	}
+
 	/// <summary>
 	/// Partial representation of PVS-Studio main settings file
 	/// </summary>
@@ -94,6 +117,21 @@ namespace UnrealBuildTool
 		/// </summary>
 		public bool DisableMISRAAnalysis;
 
+    		/// <summary>
+		/// File analysis timeout
+		/// </summary>
+		public AnalysisTimeoutFlags AnalysisTimeout;
+
+		/// <summary>
+		/// Disable analyzer Level 3 (Low) messages
+		/// </summary>
+		public bool NoNoise;
+
+		/// <summary>
+		/// Enable the display of analyzer rules exceptions which can be specified by comments and .pvsconfig files.
+		/// </summary>
+		public bool ReportDisabledRules;
+    
 		/// <summary>
 		/// Gets the analysis mode flags from the settings
 		/// </summary>
@@ -195,6 +233,99 @@ namespace UnrealBuildTool
 		/// Private storage for the mode flags
 		/// </summary>
 		PVSAnalysisModeFlags? ModePrivate;
+    
+    		/// <summary>
+		/// Override for the analysis timeoutFlag to use
+		/// </summary>
+		public AnalysisTimeoutFlags AnalysisTimeoutFlag
+		{
+			get
+			{
+				if (TimeoutPrivate.HasValue)
+				{
+					return TimeoutPrivate.Value;
+				}
+				else if (UseApplicationSettings && ApplicationSettings.Value != null)
+				{
+					return ApplicationSettings.Value.AnalysisTimeout;
+				}
+				else
+				{
+					return AnalysisTimeoutFlags.After_10_minutes;
+				}
+			}
+			set
+			{
+				TimeoutPrivate = value;
+			}
+		}
+
+		/// <summary>
+		/// Private storage for the analysis timeout
+		/// </summary>
+		AnalysisTimeoutFlags? TimeoutPrivate;
+
+		/// <summary>
+		/// Override for the disable Level 3 (Low) analyzer messages
+		/// </summary>
+		public bool EnableNoNoise
+		{
+			get
+			{
+				if (EnableNoNoisePrivate.HasValue)
+				{
+					return EnableNoNoisePrivate.Value;
+				}
+				else if (UseApplicationSettings && ApplicationSettings.Value != null)
+				{
+					return ApplicationSettings.Value.NoNoise;
+				}
+				else
+				{
+					return false;
+				}
+			}
+			set
+			{
+				EnableNoNoisePrivate = value;
+			}
+		}
+
+		/// <summary>
+		/// Private storage for the NoNoise analyzer setting
+		/// </summary>
+		bool? EnableNoNoisePrivate;
+
+		/// <summary>
+		/// Override for the enable the display of analyzer rules exceptions which can be specified by comments and .pvsconfig files.
+		/// </summary>
+		public bool EnableReportDisabledRules
+		{
+			get
+			{
+				if (EnableReportDisabledRulesPrivate.HasValue)
+				{
+					return EnableReportDisabledRulesPrivate.Value;
+				}
+				else if (UseApplicationSettings && ApplicationSettings.Value != null)
+				{
+					return ApplicationSettings.Value.ReportDisabledRules;
+				}
+				else
+				{
+					return false;
+				}
+			}
+			set
+			{
+				EnableReportDisabledRulesPrivate = value;
+			}
+		}
+
+		/// <summary>
+		/// Private storage for the ReportDisabledRules analyzer setting
+		/// </summary>
+		bool? EnableReportDisabledRulesPrivate;
 	}
 
 	/// <summary>
@@ -239,6 +370,30 @@ namespace UnrealBuildTool
 		{
 			get { return Inner.ModeFlags; }
 		}
+    
+		/// <summary>
+		/// Override for the analysis timeout to use
+		/// </summary>
+		public AnalysisTimeoutFlags AnalysisTimeoutFlag
+		{
+			get { return Inner.AnalysisTimeoutFlag; }
+		}
+
+		/// <summary>
+		/// Override NoNoise analysis setting to use
+		/// </summary>
+		public bool EnableNoNoise
+		{
+			get { return Inner.EnableNoNoise; }
+		}
+
+		/// <summary>
+		/// Override EnableReportDisabledRules analysis setting to use
+		/// </summary>
+		public bool EnableReportDisabledRules
+		{
+			get { return Inner.EnableReportDisabledRules; }
+		}
 	}
 
 	/// <summary>
@@ -467,6 +622,9 @@ namespace UnrealBuildTool
 		}
 
 		public static readonly VersionNumber CLVerWithCPP20Support = new VersionNumber(14, 23);
+    
+    		const string CPP_20 = "c++20";
+		const string CPP_17 = "c++17";
 
 		public static string GetLangStandForCfgFile(CppStandardVersion cppStandard, VersionNumber compilerVersion)
 		{
@@ -475,10 +633,13 @@ namespace UnrealBuildTool
 			switch (cppStandard)
 			{
 				case CppStandardVersion.Cpp17:
-					cppCfgStandard = "c++17";
+					cppCfgStandard = CPP_17;
+					break;
+        			case CppStandardVersion.Cpp20:
+					cppCfgStandard = CPP_20;
 					break;
 				case CppStandardVersion.Latest:
-					cppCfgStandard = VersionNumber.Compare(compilerVersion, CLVerWithCPP20Support) >= 0 ? "c++20" : "c++17";
+					cppCfgStandard = VersionNumber.Compare(compilerVersion, CLVerWithCPP20Support) >= 0 ? CPP_20 : CPP_17;
 					break;
 				default:
 					cppCfgStandard = "c++14";
@@ -488,6 +649,24 @@ namespace UnrealBuildTool
 			return cppCfgStandard;
 		}
 
+    		public static bool ShouldCompileAsC(String compilerCommandLine, String sourceFileName)
+		{
+			int CFlagLastPosition = Math.Max(Math.Max(compilerCommandLine.LastIndexOf("/TC "), compilerCommandLine.LastIndexOf("/Tc ")),
+											 Math.Max(compilerCommandLine.LastIndexOf("-TC "), compilerCommandLine.LastIndexOf("-Tc ")));
+
+			int CppFlagLastPosition = Math.Max(Math.Max(compilerCommandLine.LastIndexOf("/TP "), compilerCommandLine.LastIndexOf("/Tp ")),
+											   Math.Max(compilerCommandLine.LastIndexOf("-TP "), compilerCommandLine.LastIndexOf("-Tp ")));
+
+			bool compileAsCCode;
+			if (CFlagLastPosition == CppFlagLastPosition)
+				//ни один флаг, определяющий язык, не задан. Определяем по расширению файла
+				compileAsCCode = Path.GetExtension(sourceFileName).Equals(".c", StringComparison.InvariantCultureIgnoreCase);
+			else
+				compileAsCCode = CFlagLastPosition > CppFlagLastPosition;
+
+			return compileAsCCode;
+		}
+
 		public override CPPOutput CompileCPPFiles(CppCompileEnvironment CompileEnvironment, List<FileItem> InputFiles, DirectoryReference OutputDir, string ModuleName, IActionGraphBuilder Graph)
 		{
 			// Ignore generated files
@@ -555,7 +734,7 @@ namespace UnrealBuildTool
 						}
 					}
 				}
-				if (Platform == UnrealTargetPlatform.Win64)
+        		if (Platform == UnrealTargetPlatform.Win64)
 				{
 					ConfigFileContents.Append("platform=x64\n");
 				}
@@ -563,18 +742,45 @@ namespace UnrealBuildTool
 				{
 					throw new BuildException("PVS-Studio does not support this platform");
 				}
-				ConfigFileContents.Append("preprocessor=visualcpp\n");
-				ConfigFileContents.Append("language=C++\n");
+        		ConfigFileContents.Append("preprocessor=visualcpp\n");
+
+				bool shouldCompileAsC = ShouldCompileAsC(String.Join(" ", PreprocessAction.Arguments), SourceFileItem.AbsolutePath);
+				ConfigFileContents.AppendFormat("language={0}\n", shouldCompileAsC ? "C" : "C++");
+
 				ConfigFileContents.Append("skip-cl-exe=yes\n");
 
-				if(AnalyzerVersion.CompareTo(new Version("7.07")) >= 0)
+				WindowsCompiler WindowsCompiler = Target.WindowsPlatform.Compiler;
+				bool isVisualCppCompiler = WindowsCompiler == WindowsCompiler.VisualStudio2022 || WindowsCompiler == WindowsCompiler.VisualStudio2019;
+				if (AnalyzerVersion.CompareTo(new Version("7.07")) >= 0)
 				{
 					VersionNumber compilerVersion = Target.WindowsPlatform.Environment.CompilerVersion;
 					string languageStandardForCfg = GetLangStandForCfgFile(PreprocessCompileEnvironment.CppStandard, compilerVersion);
 
 					ConfigFileContents.AppendFormat("std={0}\n", languageStandardForCfg);
+          
+          			bool disableMsExtentinsFromArgs = PreprocessAction.Arguments.Any(arg => arg.Equals("/Za") || arg.Equals("-Za") || arg.Equals("/permissive-"));
+					bool disableMsExtentions = isVisualCppCompiler && (languageStandardForCfg == CPP_20 || disableMsExtentinsFromArgs);
+					ConfigFileContents.AppendFormat("disable-ms-extensions={0}\n", disableMsExtentions ? "yes" : "no");
+				}
+
+				if (isVisualCppCompiler && PreprocessAction.Arguments.Any(arg => arg.StartsWith("/await")))
+				{
+					ConfigFileContents.Append("msvc-await=yes\n");
 				}
 
+				if (Settings.EnableNoNoise)
+				{
+					ConfigFileContents.Append("no-noise=yes\n");
+				}
+
+				if (Settings.EnableReportDisabledRules)
+				{
+					ConfigFileContents.Append("report-disabled-rules=yes\n");
+				}
+
+        		int Timeout = (int)(Settings.AnalysisTimeoutFlag == AnalysisTimeoutFlags.No_timeout ? 0 : Settings.AnalysisTimeoutFlag);
+				ConfigFileContents.AppendFormat("timeout={0}\n", Timeout);
+
 				string BaseFileName = PreprocessedFileItem.Location.GetFileName();
 
 				FileReference ConfigFileLocation = FileReference.Combine(OutputDir, BaseFileName + ".cfg");
