Fixed an issue where people using writable workspaces would compile everything without using unity files because of the change to merge .gen.cpp files with module .cpp files in unity files.

diff --git a/Engine/Source/Programs/UnrealBuildTool/System/Unity.cs b/Engine/Source/Programs/UnrealBuildTool/System/Unity.cs
--- a/Engine/Source/Programs/UnrealBuildTool/System/Unity.cs
+++ b/Engine/Source/Programs/UnrealBuildTool/System/Unity.cs
@@ -185,7 +185,7 @@ namespace UnrealBuildTool
 			// Every single file in the module appears in the working set. Don't bother using adaptive unity for this module.
 			// Otherwise it would make full builds really slow.
 			GetAdaptiveFiles(Target, CPPFiles, HeaderFiles, CompileEnvironment, WorkingSet, BaseName, IntermediateDirectory, Graph, out NormalFiles, out AdaptiveFiles);
-			if (NormalFiles.Count == 0)
+			if (NormalFiles.Where(file => !file.HasExtension(".gen.cpp")).Count() == 0)
 			{
 				NormalFiles = CPPFiles;
 				AdaptiveFiles.RemoveAll(new HashSet<FileItem>(NormalFiles).Contains);
