- Fix xcode 14 for mac that needs a plist generated for it

diff --git a/Engine/Source/Programs/UnrealBuildTool/ProjectFiles/Xcode/XcodeProjectLegacy.cs b/Engine/Source/Programs/UnrealBuildTool/ProjectFiles/Xcode/XcodeProjectLegacy.cs
--- a/Engine/Source/Programs/UnrealBuildTool/ProjectFiles/Xcode/XcodeProjectLegacy.cs
+++ b/Engine/Source/Programs/UnrealBuildTool/ProjectFiles/Xcode/XcodeProjectLegacy.cs
@@ -1030,6 +1030,7 @@ namespace UnrealBuildTool.XcodeProjectLegacy
 				Content.Append("\t\t\t\tSUPPORTED_PLATFORMS = \"macosx\";" + ProjectFileGenerator.NewLine);
 				Content.Append("\t\t\t\tPRODUCT_NAME = \"" + MacExecutableFileName + "\";" + ProjectFileGenerator.NewLine);
 				Content.Append("\t\t\t\tCONFIGURATION_BUILD_DIR = \"" + MacExecutableDir + "\";" + ProjectFileGenerator.NewLine);
+				Content.Append("\t\t\t\tGENERATE_INFOPLIST_FILE = YES;" + ProjectFileGenerator.NewLine);
 			}
 			else
 			{
@@ -1139,6 +1140,7 @@ namespace UnrealBuildTool.XcodeProjectLegacy
 					Content.Append("\t\t\t\t\"PRODUCT_NAME[sdk=macosx*]\" = \"" + MacExecutableFileName + "\";" + ProjectFileGenerator.NewLine);
 					Content.Append("\t\t\t\t\"CONFIGURATION_BUILD_DIR[sdk=macosx*]\" = \"" + MacExecutableDir + "\";" + ProjectFileGenerator.NewLine);
 					Content.Append("\t\t\t\t\"SDKROOT[sdk=macosx]\" = macosx;" + ProjectFileGenerator.NewLine);
+					Content.Append("\t\t\t\t\"GENERATE_INFOPLIST_FILE[sdk=macosx] = YES;" + ProjectFileGenerator.NewLine);
 				}
 
 				if (bIsUnrealGame || bIsUnrealClient)
