UnrealBuildTool: Write trace files for android if enabled

#rnx

diff --git a/Engine/Source/Programs/UnrealBuildTool/Platform/Android/AndroidToolChain.cs b/Engine/Source/Programs/UnrealBuildTool/Platform/Android/AndroidToolChain.cs
--- a/Engine/Source/Programs/UnrealBuildTool/Platform/Android/AndroidToolChain.cs
+++ b/Engine/Source/Programs/UnrealBuildTool/Platform/Android/AndroidToolChain.cs
@@ -1337,8 +1337,25 @@ namespace UnrealBuildTool
 					// Add the source file path to the command-line.
 					FileArguments += string.Format(" \"{0}\"", SourceFile.AbsolutePath);
 
+					// Generate the timing info
+					if (CompileEnvironment.bPrintTimingInfo)
+					{
+						bool ShortArchName = !(CompileEnvironment.PrecompiledHeaderAction == PrecompiledHeaderAction.Create && !bIsPlainCFile);
+						FileItem TraceFile;
+						if (CompileEnvironment.PrecompiledHeaderAction == PrecompiledHeaderAction.Create && !bIsPlainCFile)
+						{
+							TraceFile = FileItem.GetItemByFileReference(FileReference.Combine(OutputDir, InlineArchName(Path.GetFileName(SourceFile.AbsolutePath), Arch) + ".json"));
+						}
+						else
+						{
+							TraceFile = FileItem.GetItemByFileReference(FileReference.Combine(OutputDir, InlineArchName(Path.GetFileName(SourceFile.AbsolutePath) + ".json", Arch, true)));
+						}
+						FileArguments += " -ftime-trace";
+						CompileAction.ProducedItems.Add(TraceFile);
+					}
+
 					// Generate the included header dependency list
-					if(CompileEnvironment.bGenerateDependenciesFile)
+					if (CompileEnvironment.bGenerateDependenciesFile)
 					{
 						FileItem DependencyListFile = FileItem.GetItemByFileReference(FileReference.Combine(OutputDir, InlineArchName(Path.GetFileName(SourceFile.AbsolutePath) + ".d", Arch, true)));
 						FileArguments += string.Format(" -MD -MF\"{0}\"", DependencyListFile.AbsolutePath.Replace('\\', '/'));
