UnrealBuildTool: Fix extension of clang-cl timing manifest

#rnx

diff --git a/Engine/Source/Programs/UnrealBuildTool/Platform/Windows/VCToolChain.cs b/Engine/Source/Programs/UnrealBuildTool/Platform/Windows/VCToolChain.cs
--- a/Engine/Source/Programs/UnrealBuildTool/Platform/Windows/VCToolChain.cs
+++ b/Engine/Source/Programs/UnrealBuildTool/Platform/Windows/VCToolChain.cs
@@ -1571,14 +1571,14 @@ namespace UnrealBuildTool
 				if (TimingJsonFiles.Count > 1)
 				{
 					// Generate the file manifest for the aggregator.
-					FileReference ManifestFile = FileReference.Combine(Makefile.ProjectIntermediateDirectory, $"{Target.Name}TimingManifest.txt");
-					if (!DirectoryReference.Exists(ManifestFile.Directory))
+					if (!DirectoryReference.Exists(Makefile.ProjectIntermediateDirectory))
 					{
-						DirectoryReference.CreateDirectory(ManifestFile.Directory);
+						DirectoryReference.CreateDirectory(Makefile.ProjectIntermediateDirectory);
 					}
 
 					if (Target.WindowsPlatform.Compiler.IsMSVC())
 					{
+						FileReference ManifestFile = FileReference.Combine(Makefile.ProjectIntermediateDirectory, $"{Target.Name}TimingManifest.txt");
 						File.WriteAllLines(ManifestFile.FullName, TimingJsonFiles.Select(f => f.AbsolutePath));
 
 						FileReference ExpectedCompileTimeFile = FileReference.FromString(Path.Combine(Makefile.ProjectIntermediateDirectory.FullName, $"{Target.Name}.json"));
@@ -1602,6 +1602,7 @@ namespace UnrealBuildTool
 					}
 					else if (Target.WindowsPlatform.Compiler.IsClang())
 					{
+						FileReference ManifestFile = FileReference.Combine(Makefile.ProjectIntermediateDirectory, $"{Target.Name}TimingManifest.csv");
 						File.WriteAllLines(ManifestFile.FullName, TimingJsonFiles.Select(f => f.FullName.Remove(f.FullName.Length - ".json".Length)));
 
 						FileItem AggregateOutputFile = FileItem.GetItemByFileReference(FileReference.Combine(Makefile.ProjectIntermediateDirectory, $"{Target.Name}.trace.csv"));
