- Fixed bad xcode project format, missing a "

diff --git a/Engine/Source/Programs/UnrealBuildTool/ProjectFiles/Xcode/XcodeProjectLegacy.cs b/Engine/Source/Programs/UnrealBuildTool/ProjectFiles/Xcode/XcodeProjectLegacy.cs
--- a/Engine/Source/Programs/UnrealBuildTool/ProjectFiles/Xcode/XcodeProjectLegacy.cs
+++ b/Engine/Source/Programs/UnrealBuildTool/ProjectFiles/Xcode/XcodeProjectLegacy.cs
@@ -1140,7 +1140,7 @@ namespace UnrealBuildTool.XcodeProjectLegacy
 					Content.Append("\t\t\t\t\"PRODUCT_NAME[sdk=macosx*]\" = \"" + MacExecutableFileName + "\";" + ProjectFileGenerator.NewLine);
 					Content.Append("\t\t\t\t\"CONFIGURATION_BUILD_DIR[sdk=macosx*]\" = \"" + MacExecutableDir + "\";" + ProjectFileGenerator.NewLine);
 					Content.Append("\t\t\t\t\"SDKROOT[sdk=macosx]\" = macosx;" + ProjectFileGenerator.NewLine);
-					Content.Append("\t\t\t\t\"GENERATE_INFOPLIST_FILE[sdk=macosx] = YES;" + ProjectFileGenerator.NewLine);
+					Content.Append("\t\t\t\t\"GENERATE_INFOPLIST_FILE[sdk=macosx*]\" = YES;" + ProjectFileGenerator.NewLine);
 				}
 
 				if (bIsUnrealGame || bIsUnrealClient)
