- Fixing BuildTest project that is under a differently named directory - the uprojectdirs file couldn't be found because the UE_ENGINE_DIRECTORY was wrong

diff --git a/Engine/Source/Programs/UnrealBuildTool/Configuration/UEBuildTarget.cs b/Engine/Source/Programs/UnrealBuildTool/Configuration/UEBuildTarget.cs
--- a/Engine/Source/Programs/UnrealBuildTool/Configuration/UEBuildTarget.cs
+++ b/Engine/Source/Programs/UnrealBuildTool/Configuration/UEBuildTarget.cs
@@ -1,4 +1,4 @@
-// Copyright Epic Games, Inc. All Rights Reserved.
+﻿// Copyright Epic Games, Inc. All Rights Reserved.
 
 using System;
 using System.Collections.Generic;
@@ -4257,7 +4257,8 @@ namespace UnrealBuildTool
 					{
 						OutputFilePath = OutputFilePath.Substring(0, OutputFilePath.LastIndexOf(".app/Contents/MacOS") + 4);
 					}
-					string EnginePath = Utils.CleanDirectorySeparators(Unreal.EngineDirectory.MakeRelativeTo(ExecutableBinary.OutputFilePath.Directory), '/');
+
+					string EnginePath = Utils.CleanDirectorySeparators(Unreal.EngineDirectory.MakeRelativeTo(new DirectoryReference(OutputFilePath).ParentDirectory!), '/');
 					if (EnginePath.EndsWith("/") == false)
 					{
 						EnginePath += "/";
