Added warnings when C++ UHT is used.  These warnings include the reasons why C++ UHT is being used.  The warning is currently disabled but informational messages will still be generated.

diff --git a/Engine/Source/Programs/UnrealBuildTool/Configuration/BuildConfiguration.cs b/Engine/Source/Programs/UnrealBuildTool/Configuration/BuildConfiguration.cs
--- a/Engine/Source/Programs/UnrealBuildTool/Configuration/BuildConfiguration.cs
+++ b/Engine/Source/Programs/UnrealBuildTool/Configuration/BuildConfiguration.cs
@@ -145,11 +145,17 @@ namespace UnrealBuildTool
 		public bool bForceDebugUnrealHeaderTool = false;
 
 		/// <summary>
-		/// If true, use UHT internal to UBT
+		/// If true, use C# UHT internal to UBT
 		/// </summary>
 		[XmlConfigFile(Category = "UEBuildConfiguration")]
 		public bool bUseBuiltInUnrealHeaderTool = true;
 
+		/// <summary>
+		/// If true, generate warnings when C++ UHT is used
+		/// </summary>
+		[XmlConfigFile(Category = "UEBuildConfiguration")]
+		public bool bWarnOnCppUnrealHeaderTool = false;
+
 		/// <summary>
 		/// Whether to skip compiling rules assemblies and just assume they are valid
 		/// </summary>
diff --git a/Engine/Source/Programs/UnrealBuildTool/Configuration/UEBuildTarget.cs b/Engine/Source/Programs/UnrealBuildTool/Configuration/UEBuildTarget.cs
--- a/Engine/Source/Programs/UnrealBuildTool/Configuration/UEBuildTarget.cs
+++ b/Engine/Source/Programs/UnrealBuildTool/Configuration/UEBuildTarget.cs
@@ -1084,6 +1084,11 @@ namespace UnrealBuildTool
 		/// </summary>
 		public List<FileReference>? EnabledUhtPlugins;
 
+		/// <summary>
+		/// Collection of all required C++ UHT plugins without a C# version
+		/// </summary>
+		public List<string>? RequiredUhtCppPlugins;
+
 		/// <summary>
 		/// All application binaries; may include binaries not built by this target.
 		/// </summary>
@@ -1856,8 +1861,8 @@ namespace UnrealBuildTool
 			// Create the makefile
 			string ExternalMetadata = UEBuildPlatform.GetBuildPlatform(Platform).GetExternalBuildMetadata(ProjectFile);
 			TargetMakefile Makefile = new TargetMakefile(ExternalMetadata, Binaries[0].OutputFilePaths[0], ReceiptFileName, ProjectIntermediateDirectory, TargetType, 
-				Rules.ConfigValueTracker, bDeployAfterCompile, bHasProjectScriptPlugin, bHasRequiredProjectScriptPlugin, UbtPlugins?.ToArray(), EnabledUbtPlugins?.ToArray(), 
-				EnabledUhtPlugins?.ToArray());
+				Rules.ConfigValueTracker, bDeployAfterCompile, bHasProjectScriptPlugin, UbtPlugins?.ToArray(), EnabledUbtPlugins?.ToArray(), 
+				EnabledUhtPlugins?.ToArray(), RequiredUhtCppPlugins?.ToArray());
 			Makefile.IsTestTarget = Rules.IsTestTarget;
 			TargetMakefileBuilder MakefileBuilder = new TargetMakefileBuilder(Makefile, Logger);
 
@@ -3772,7 +3777,11 @@ namespace UnrealBuildTool
 					bHasProjectScriptPlugin = true;
 					if (!bUbtPluginFound)
 					{
-						bHasRequiredProjectScriptPlugin = true;
+						if (RequiredUhtCppPlugins == null)
+						{
+							RequiredUhtCppPlugins = new();
+						}
+						RequiredUhtCppPlugins.Add(Plugin.Name);
 					}
 				}
 
diff --git a/Engine/Source/Programs/UnrealBuildTool/System/ExternalExecution.cs b/Engine/Source/Programs/UnrealBuildTool/System/ExternalExecution.cs
--- a/Engine/Source/Programs/UnrealBuildTool/System/ExternalExecution.cs
+++ b/Engine/Source/Programs/UnrealBuildTool/System/ExternalExecution.cs
@@ -936,7 +936,33 @@ namespace UnrealBuildTool
 		/// </summary>
 		public static void ExecuteHeaderToolIfNecessary(BuildConfiguration BuildConfiguration, FileReference? ProjectFile, TargetMakefile Makefile, string TargetName, ISourceFileWorkingSet WorkingSet, ILogger Logger)
 		{
-			if (!BuildConfiguration.bUseBuiltInUnrealHeaderTool || Makefile.bHasRequiredProjectScriptPlugin || TargetName.Equals("UnrealHeaderTool", StringComparison.InvariantCultureIgnoreCase))
+			if (!BuildConfiguration.bUseBuiltInUnrealHeaderTool)
+			{
+				if (BuildConfiguration.bWarnOnCppUnrealHeaderTool)
+				{
+					Logger.LogWarning("WARNING: C++ UHT being used because 'bUseBuiltInUnrealHeaderTool' is false.  C++ UHT has been deprecated and will be removed in 5.2");
+				}
+				else
+				{
+					Logger.LogInformation("DEPRECATED: C++ UHT being used because 'bUseBuiltInUnrealHeaderTool' is false.  C++ UHT has been deprecated and will be removed in 5.2");
+				}
+			}
+			if (Makefile.RequiredUhtCppPlugins?.Length > 0)
+			{
+				foreach (string PluginName in Makefile.RequiredUhtCppPlugins)
+				{
+					if (BuildConfiguration.bWarnOnCppUnrealHeaderTool)
+					{
+						Logger.LogWarning("WARNING: C++ UHT being used because '{PluginName}' does not have a C# version.  C++ UHT has been deprecated and will be removed in 5.2", PluginName);
+					}
+					else
+					{
+						Logger.LogInformation("DEPRECATED: C++ UHT being used because '{PluginName}' does not have a C# version.  C++ UHT has been deprecated and will be removed in 5.2", PluginName);
+					}
+				}
+			}
+
+			if (!BuildConfiguration.bUseBuiltInUnrealHeaderTool || Makefile.RequiredUhtCppPlugins?.Length > 0 || TargetName.Equals("UnrealHeaderTool", StringComparison.InvariantCultureIgnoreCase))
 			{
 				ExecuteExternalHeaderToolIfNecessary(BuildConfiguration, ProjectFile, Makefile, TargetName, WorkingSet, Logger);
 			}
diff --git a/Engine/Source/Programs/UnrealBuildTool/System/TargetMakefile.cs b/Engine/Source/Programs/UnrealBuildTool/System/TargetMakefile.cs
--- a/Engine/Source/Programs/UnrealBuildTool/System/TargetMakefile.cs
+++ b/Engine/Source/Programs/UnrealBuildTool/System/TargetMakefile.cs
@@ -101,7 +101,7 @@ namespace UnrealBuildTool
 		/// <summary>
 		/// The version number to write
 		/// </summary>
-		public const int CurrentVersion = 32;
+		public const int CurrentVersion = 33;
 
 		/// <summary>
 		/// The time at which the makefile was created
@@ -183,6 +183,11 @@ namespace UnrealBuildTool
 		/// </summary>
 		public FileReference[]? EnabledUhtPlugins;
 
+		/// <summary>
+		/// Collection of UHT C++ plugins that have no C# version
+		/// </summary>
+		public string[]? RequiredUhtCppPlugins;
+
 		/// <summary>
 		/// The array of command-line arguments. The makefile will be invalidated whenever these change.
 		/// </summary>
@@ -302,13 +307,13 @@ namespace UnrealBuildTool
 		/// <param name="ConfigValueTracker">Set of dependencies on config files</param>
 		/// <param name="bDeployAfterCompile">Whether to deploy the target after compiling</param>
 		/// <param name="bHasProjectScriptPlugin">Whether the target has a project script plugin</param>
-		/// <param name="bHasRequiredProjectScriptPlugin">If true, we have a script plugin without a matching UBT plugin</param>
 		/// <param name="UbtPlugins">Collection of UBT plugins</param>
 		/// <param name="EnabledUbtPlugins">Collection of UBT plugins</param>
 		/// <param name="EnabledUhtPlugins">Collection of UBT plugins for UHT</param>
+		/// <param name="RequiredUhtCppPlugins">Collection of required C++ UHT plugins that don't have a C# version</param>
 		public TargetMakefile(string ExternalMetadata, FileReference ExecutableFile, FileReference ReceiptFile, DirectoryReference ProjectIntermediateDirectory, 
 			TargetType TargetType, ConfigValueTracker ConfigValueTracker, bool bDeployAfterCompile, bool bHasProjectScriptPlugin,
-			bool bHasRequiredProjectScriptPlugin, FileReference[]? UbtPlugins, FileReference[]? EnabledUbtPlugins, FileReference[]? EnabledUhtPlugins)
+			FileReference[]? UbtPlugins, FileReference[]? EnabledUbtPlugins, FileReference[]? EnabledUhtPlugins, string[]? RequiredUhtCppPlugins)
 		{
 			this.CreateTimeUtc = UnrealBuildTool.StartTimeUtc;
 			this.ModifiedTimeUtc = CreateTimeUtc;
@@ -321,10 +326,10 @@ namespace UnrealBuildTool
 			this.ConfigValueTracker = ConfigValueTracker;
 			this.bDeployAfterCompile = bDeployAfterCompile;
 			this.bHasProjectScriptPlugin = bHasProjectScriptPlugin;
-			this.bHasRequiredProjectScriptPlugin = bHasRequiredProjectScriptPlugin;
 			this.UbtPlugins = UbtPlugins;
 			this.EnabledUbtPlugins = EnabledUbtPlugins;
 			this.EnabledUhtPlugins = EnabledUhtPlugins;
+			this.RequiredUhtCppPlugins = RequiredUhtCppPlugins;
 			this.Actions = new List<IExternalAction>();
 			this.OutputItems = new List<FileItem>();
 			this.ModuleNameToOutputItems = new Dictionary<string, FileItem[]>(StringComparer.OrdinalIgnoreCase);
@@ -362,10 +367,10 @@ namespace UnrealBuildTool
 			ConfigValueTracker = new ConfigValueTracker(Reader);
 			bDeployAfterCompile = Reader.ReadBool();
 			bHasProjectScriptPlugin = Reader.ReadBool();
-			bHasRequiredProjectScriptPlugin = Reader.ReadBool();
 			UbtPlugins = Reader.ReadArray(() => Reader.ReadFileReference())!;
 			EnabledUbtPlugins = Reader.ReadArray(() => Reader.ReadFileReference())!;
 			EnabledUhtPlugins = Reader.ReadArray(() => Reader.ReadFileReference())!;
+			RequiredUhtCppPlugins = Reader.ReadArray(() => Reader.ReadString())!;
 			AdditionalArguments = Reader.ReadArray(() => Reader.ReadString())!;
 			UHTAdditionalArguments = Reader.ReadArray(() => Reader.ReadString())!;
 			PreBuildScripts = Reader.ReadArray(() => Reader.ReadFileReference())!;
@@ -407,10 +412,10 @@ namespace UnrealBuildTool
 			ConfigValueTracker.Write(Writer);
 			Writer.WriteBool(bDeployAfterCompile);
 			Writer.WriteBool(bHasProjectScriptPlugin);
-			Writer.WriteBool(bHasRequiredProjectScriptPlugin);
 			Writer.WriteArray(UbtPlugins, Item => Writer.WriteFileReference(Item));
 			Writer.WriteArray(EnabledUbtPlugins, Item => Writer.WriteFileReference(Item));
 			Writer.WriteArray(EnabledUhtPlugins, Item => Writer.WriteFileReference(Item));
+			Writer.WriteArray(RequiredUhtCppPlugins, Item => Writer.WriteString(Item));
 			Writer.WriteArray(AdditionalArguments, Item => Writer.WriteString(Item));
 			Writer.WriteArray(UHTAdditionalArguments, Item => Writer.WriteString(Item));
 			Writer.WriteArray(PreBuildScripts, Item => Writer.WriteFileReference(Item));
