UnrealBuildTool: Also check for .git file as that denotes a git worktree, when getting the source file working set

#rnx

diff --git a/Engine/Source/Programs/UnrealBuildTool/System/SourceFileWorkingSet.cs b/Engine/Source/Programs/UnrealBuildTool/System/SourceFileWorkingSet.cs
--- a/Engine/Source/Programs/UnrealBuildTool/System/SourceFileWorkingSet.cs
+++ b/Engine/Source/Programs/UnrealBuildTool/System/SourceFileWorkingSet.cs
@@ -363,7 +363,7 @@ namespace UnrealBuildTool
 			GitSourceFileWorkingSet? WorkingSet  = null;
 
 			// Create the working set for the engine directory
-			if (DirectoryReference.Exists(DirectoryReference.Combine(RootDir, ".git")))
+			if (DirectoryReference.Exists(DirectoryReference.Combine(RootDir, ".git")) || FileReference.Exists(FileReference.Combine(RootDir, ".git")))
 			{
 				WorkingSet = new GitSourceFileWorkingSet(GitPath, RootDir, WorkingSet, Logger);
 			}
@@ -373,11 +373,11 @@ namespace UnrealBuildTool
 			{
 				if(WorkingSet == null || !ProjectDir.IsUnderDirectory(RootDir))
 				{
-					if (DirectoryReference.Exists(DirectoryReference.Combine(ProjectDir, ".git")))
+					if (DirectoryReference.Exists(DirectoryReference.Combine(ProjectDir, ".git")) || FileReference.Exists(FileReference.Combine(ProjectDir, ".git")))
 					{
 						WorkingSet = new GitSourceFileWorkingSet(GitPath, ProjectDir, WorkingSet, Logger);
 					}
-					else if (DirectoryReference.Exists(DirectoryReference.Combine(ProjectDir.ParentDirectory!, ".git")))
+					else if (DirectoryReference.Exists(DirectoryReference.Combine(ProjectDir.ParentDirectory!, ".git")) || FileReference.Exists(FileReference.Combine(ProjectDir.ParentDirectory!, ".git")))
 					{
 						WorkingSet = new GitSourceFileWorkingSet(GitPath, ProjectDir.ParentDirectory!, WorkingSet, Logger);
 					}
