Adding a known compile log event for an xcode library

It comes from sterr so needs to be outside the normal warning/error checking in the CompilerEventMatcher

#rnx

diff --git a/Engine/Source/Programs/Shared/EpicGames.Core/KnownLogEvents.cs b/Engine/Source/Programs/Shared/EpicGames.Core/KnownLogEvents.cs
--- a/Engine/Source/Programs/Shared/EpicGames.Core/KnownLogEvents.cs
+++ b/Engine/Source/Programs/Shared/EpicGames.Core/KnownLogEvents.cs
@@ -215,6 +215,11 @@ namespace EpicGames.Core
 		/// </summary>
 		public static EventId Systemic_MSBuild { get; } = new EventId(750);
 
+		/// <summary>
+		/// A systemic event from XCode
+		/// </summary>
+		public static EventId Systemic_XCode { get; } = new EventId(755);
+
 		/// <summary>
 		/// Robomerge gate is locked
 		/// </summary>
diff --git a/Engine/Source/Programs/UnrealBuildTool/Matchers/CompileEventMatcher.cs b/Engine/Source/Programs/UnrealBuildTool/Matchers/CompileEventMatcher.cs
--- a/Engine/Source/Programs/UnrealBuildTool/Matchers/CompileEventMatcher.cs
+++ b/Engine/Source/Programs/UnrealBuildTool/Matchers/CompileEventMatcher.cs
@@ -53,6 +53,7 @@ namespace UnrealBuildTool.Matchers
 		static readonly Regex s_clangDiagnosticPattern = new Regex($"^\\s*{FilePattern}\\s*{ClangLocationPattern}:\\s*{ClangSeverity}\\s*:");
 		static readonly Regex s_clangNotePattern = new Regex($"^\\s*{FilePattern}\\s*{ClangLocationPattern}:\\s*note:");
 		static readonly Regex s_clangMarkerPattern = new Regex(@"^(\s*)[\^~][\s\^~]*$");
+		static readonly Regex s_xcodeIDEWatchExtensionPattern = new Regex(@"xcodebuild.*Requested but did not find extension point with identifier.*for extension.*\.watchOS of plug-in com\.apple\.dt\.IDEWatchSupportCore");
 
 		static readonly string[] s_invalidExtensions =
 		{
@@ -164,6 +165,11 @@ namespace UnrealBuildTool.Matchers
 					return eventMatch;
 				}
 			}
+			else if (input.TryMatch(s_xcodeIDEWatchExtensionPattern, out Match? match))
+			{
+				LogEventBuilder builder = new LogEventBuilder(input);
+				return builder.ToMatch(LogEventPriority.Normal, LogLevel.Information, KnownLogEvents.Systemic_XCode);
+			}
 			return null;
 		}
 
