UnrealBuildTool: Fix check for community version of visual studio

#rnx

diff --git a/Engine/Source/Programs/UnrealBuildTool/Platform/Windows/MicrosoftPlatformSDK.cs b/Engine/Source/Programs/UnrealBuildTool/Platform/Windows/MicrosoftPlatformSDK.cs
--- a/Engine/Source/Programs/UnrealBuildTool/Platform/Windows/MicrosoftPlatformSDK.cs
+++ b/Engine/Source/Programs/UnrealBuildTool/Platform/Windows/MicrosoftPlatformSDK.cs
@@ -808,15 +808,15 @@ namespace UnrealBuildTool
 						bool bPreview = Catalog != null && Catalog.IsPrerelease();
 
 						string ProductId = Instance.GetProduct().GetId();
-						bool bExpress = ProductId.Equals("Microsoft.VisualStudio.Product.WDExpress", StringComparison.Ordinal);
+						bool bCommunity = ProductId.Equals("Microsoft.VisualStudio.Product.Community", StringComparison.Ordinal);
 
 						DirectoryReference BaseDir = new DirectoryReference(Instance.GetInstallationPath());
-						Installations.Add(new VisualStudioInstallation(Compiler, Version, BaseDir, bExpress, bPreview));
+						Installations.Add(new VisualStudioInstallation(Compiler, Version, BaseDir, bCommunity, bPreview));
 
 						Logger.LogDebug("Found Visual Studio installation: {BaseDir} (Product={ProductId}, Version={Version})", BaseDir, ProductId, Version);
 					}
 
-					Installations = Installations.OrderBy(x => x.bExpress).ThenBy(x => x.bPreview).ThenByDescending(x => x.Version).ToList();
+					Installations = Installations.OrderBy(x => x.bCommunity).ThenBy(x => x.bPreview).ThenByDescending(x => x.Version).ToList();
 				}
 				catch (Exception Ex)
 				{
diff --git a/Engine/Source/Programs/UnrealBuildTool/Platform/Windows/UEBuildWindows.cs b/Engine/Source/Programs/UnrealBuildTool/Platform/Windows/UEBuildWindows.cs
--- a/Engine/Source/Programs/UnrealBuildTool/Platform/Windows/UEBuildWindows.cs
+++ b/Engine/Source/Programs/UnrealBuildTool/Platform/Windows/UEBuildWindows.cs
@@ -754,9 +754,9 @@ namespace UnrealBuildTool
 		public DirectoryReference BaseDir { get; }
 
 		/// <summary>
-		/// Whether it's an express edition of Visual Studio. These versions do not contain a 64-bit compiler.
+		/// Whether it's a community edition of Visual Studio.
 		/// </summary>
-		public bool bExpress { get; }
+		public bool bCommunity { get; }
 
 		/// <summary>
 		/// Whether it's a pre-release version of the IDE.
@@ -766,12 +766,12 @@ namespace UnrealBuildTool
 		/// <summary>
 		/// Constructor
 		/// </summary>
-		public VisualStudioInstallation(WindowsCompiler Compiler, VersionNumber Version, DirectoryReference BaseDir, bool bExpress, bool bPreview)
+		public VisualStudioInstallation(WindowsCompiler Compiler, VersionNumber Version, DirectoryReference BaseDir, bool bCommunity, bool bPreview)
 		{
 			this.Compiler = Compiler;
 			this.Version = Version;
 			this.BaseDir = BaseDir;
-			this.bExpress = bExpress;
+			this.bCommunity = bCommunity;
 			this.bPreview = bPreview;
 		}
 	}
