- Better solution for Xcode 14 requiring plist file for Mac local signing:
* Remove the "Mac short circuit" blocks of code in the project generator as it was skipping the code that sets up PLIST and BUNDLE_IDENT settings (which Xcode14 will need)
* Point Xcode at the temp plist files that project generator makes
* Make sure to copy the final plist on top of the temp plist file so Xcode will pull that one into the .app at the end

diff --git a/Engine/Source/Programs/UnrealBuildTool/Platform/Mac/MacToolChain.cs b/Engine/Source/Programs/UnrealBuildTool/Platform/Mac/MacToolChain.cs
--- a/Engine/Source/Programs/UnrealBuildTool/Platform/Mac/MacToolChain.cs
+++ b/Engine/Source/Programs/UnrealBuildTool/Platform/Mac/MacToolChain.cs
@@ -1,4 +1,4 @@
-// Copyright Epic Games, Inc. All Rights Reserved.
+﻿// Copyright Epic Games, Inc. All Rights Reserved.
 
 using System;
 using System.Collections;
@@ -829,6 +829,14 @@ namespace UnrealBuildTool
 					AppendMacLine(FinalizeAppBundleScript, FormatCopyCommand(TempInfoPlist, String.Format("{0}.app/Contents/Info.plist", ExeName)));
 					AppendMacLine(FinalizeAppBundleScript, "chmod 644 \"{0}.app/Contents/Info.plist\"", ExeName);
 
+					// Also copy it to where Xcode will look for it, now that Xcode 14 requireswe have one set up - it can't point into the .app because that
+					// is where it will copy to, so it will error with reading and writing to the same location
+					string IntermediateDirectory = (ProjectFile == null ? Unreal.EngineDirectory : ProjectFile.Directory) + "/Intermediate/Mac";
+					string XcodeInputPListFile = IntermediateDirectory + "/" + ExeName + "-Info.plist";
+					AppendMacLine(FinalizeAppBundleScript, FormatCopyCommand(TempInfoPlist, XcodeInputPListFile));
+					AppendMacLine(FinalizeAppBundleScript, "chmod 644 \"{0}\"", XcodeInputPListFile);
+
+
 					// Generate PkgInfo file
 					string TempPkgInfo = "$TMPDIR/TempPkgInfo";
 					AppendMacLine(FinalizeAppBundleScript, "echo 'echo -n \"APPL????\"' | bash > \"{0}\"", TempPkgInfo);
diff --git a/Engine/Source/Programs/UnrealBuildTool/ProjectFiles/Xcode/XcodeProjectLegacy.cs b/Engine/Source/Programs/UnrealBuildTool/ProjectFiles/Xcode/XcodeProjectLegacy.cs
--- a/Engine/Source/Programs/UnrealBuildTool/ProjectFiles/Xcode/XcodeProjectLegacy.cs
+++ b/Engine/Source/Programs/UnrealBuildTool/ProjectFiles/Xcode/XcodeProjectLegacy.cs
@@ -1023,160 +1023,157 @@ namespace UnrealBuildTool.XcodeProjectLegacy
 			IOSRunTimeVersion = null;
 			TVOSRunTimeVersion = null;
 
-			// shortcut for mac only
-			if (bSupportMac && !bSupportIOS && !bSupportTVOS)
+			bool bIsUnrealGame = TargetName.Equals("UnrealGame", StringComparison.InvariantCultureIgnoreCase);
+			bool bIsUnrealClient = TargetName.Equals("UnrealClient", StringComparison.InvariantCultureIgnoreCase);
+			DirectoryReference? GameDir = ProjectFile?.Directory;
+			string? GamePath = GameDir != null ? ConvertPath(GameDir.FullName) : null;
+
+			string? IOSRunTimeDevices = null;
+			string? TVOSRunTimeDevices = null;
+			string SupportedPlatforms = bSupportMac ? "macosx" : "";
+
+			bool bAutomaticSigning = false;
+			string? UUID_IOS = "";
+			string? UUID_TVOS = "";
+			string? TEAM_IOS = "";
+			string? TEAM_TVOS = "";
+			string? IOS_CERT = "iPhone Developer";
+			string? TVOS_CERT = "iPhone Developer";
+			string IOS_BUNDLE = "";
+			string TVOS_BUNDLE = "";
+			if (bSupportIOS && InstalledPlatformInfo.IsValidPlatform(UnrealTargetPlatform.IOS, EProjectType.Code))
 			{
-				Content.Append("\t\t\t\tVALID_ARCHS = \"" + string.Join(" ", SupportedMacArchitectures) + "\";" + ProjectFileGenerator.NewLine);
-				Content.Append("\t\t\t\tSUPPORTED_PLATFORMS = \"macosx\";" + ProjectFileGenerator.NewLine);
-				Content.Append("\t\t\t\tPRODUCT_NAME = \"" + MacExecutableFileName + "\";" + ProjectFileGenerator.NewLine);
-				Content.Append("\t\t\t\tCONFIGURATION_BUILD_DIR = \"" + MacExecutableDir + "\";" + ProjectFileGenerator.NewLine);
-				Content.Append("\t\t\t\tGENERATE_INFOPLIST_FILE = YES;" + ProjectFileGenerator.NewLine);
+				IOSPlatform IOSPlatform = ((IOSPlatform)UEBuildPlatform.GetBuildPlatform(UnrealTargetPlatform.IOS));
+				IOSProjectSettings ProjectSettings = IOSPlatform.ReadProjectSettings(ProjectFile);
+				IOSProvisioningData ProvisioningData = IOSPlatform.ReadProvisioningData(ProjectSettings, bForDistribution);
+				IOSRunTimeVersion = ProjectSettings.RuntimeVersion;
+				IOSRunTimeDevices = ProjectSettings.RuntimeDevices;
+				SupportedPlatforms += " iphoneos";
+				bAutomaticSigning = ProjectSettings.bAutomaticSigning;
+				if (!bAutomaticSigning)
+				{
+					UUID_IOS = ProvisioningData.MobileProvisionUUID;
+					IOS_CERT = ProvisioningData.SigningCertificate;
+				}
+				TEAM_IOS = ProvisioningData.TeamUUID;
+				IOS_BUNDLE = ProjectSettings.BundleIdentifier;
 			}
-			else
+
+			if (bSupportTVOS && InstalledPlatformInfo.IsValidPlatform(UnrealTargetPlatform.TVOS, EProjectType.Code))
 			{
-				bool bIsUnrealGame = TargetName.Equals("UnrealGame", StringComparison.InvariantCultureIgnoreCase);
-				bool bIsUnrealClient = TargetName.Equals("UnrealClient", StringComparison.InvariantCultureIgnoreCase);
-				DirectoryReference? GameDir = ProjectFile?.Directory;
-				string? GamePath = GameDir != null ? ConvertPath(GameDir.FullName) : null;
-
-				string? IOSRunTimeDevices = null;
-				string? TVOSRunTimeDevices = null;
-				string SupportedPlatforms = bSupportMac ? "macosx" : "";
-
-				bool bAutomaticSigning = false;
-				string? UUID_IOS = "";
-				string? UUID_TVOS = "";
-				string? TEAM_IOS = "";
-				string? TEAM_TVOS = "";
-				string? IOS_CERT = "iPhone Developer";
-				string? TVOS_CERT = "iPhone Developer";
-				string IOS_BUNDLE = "";
-				string TVOS_BUNDLE = "";
-				if (bSupportIOS && InstalledPlatformInfo.IsValidPlatform(UnrealTargetPlatform.IOS, EProjectType.Code))
+				TVOSPlatform TVOSPlatform = ((TVOSPlatform)UEBuildPlatform.GetBuildPlatform(UnrealTargetPlatform.TVOS));
+				TVOSProjectSettings ProjectSettings = TVOSPlatform.ReadProjectSettings(ProjectFile);
+				TVOSProvisioningData ProvisioningData = TVOSPlatform.ReadProvisioningData(ProjectSettings, bForDistribution);
+				TVOSRunTimeVersion = ProjectSettings.RuntimeVersion;
+				TVOSRunTimeDevices = ProjectSettings.RuntimeDevices;
+				SupportedPlatforms += " appletvos";
+				if (!bAutomaticSigning)
 				{
-					IOSPlatform IOSPlatform = ((IOSPlatform)UEBuildPlatform.GetBuildPlatform(UnrealTargetPlatform.IOS));
-					IOSProjectSettings ProjectSettings = IOSPlatform.ReadProjectSettings(ProjectFile);
-					IOSProvisioningData ProvisioningData = IOSPlatform.ReadProvisioningData(ProjectSettings, bForDistribution);
-					IOSRunTimeVersion = ProjectSettings.RuntimeVersion;
-					IOSRunTimeDevices = ProjectSettings.RuntimeDevices;
-					SupportedPlatforms += " iphoneos";
-					bAutomaticSigning = ProjectSettings.bAutomaticSigning;
-					if (!bAutomaticSigning)
-					{
-						UUID_IOS = ProvisioningData.MobileProvisionUUID;
-						IOS_CERT = ProvisioningData.SigningCertificate;
-					}
-					TEAM_IOS = ProvisioningData.TeamUUID;
-					IOS_BUNDLE = ProjectSettings.BundleIdentifier;
+					UUID_TVOS = ProvisioningData.MobileProvisionUUID;
+					TVOS_CERT = ProvisioningData.SigningCertificate;
 				}
+				TEAM_TVOS = ProvisioningData.TeamUUID;
+				TVOS_BUNDLE = ProjectSettings.BundleIdentifier;
+			}
 
-				if (bSupportTVOS && InstalledPlatformInfo.IsValidPlatform(UnrealTargetPlatform.TVOS, EProjectType.Code))
+			Content.Append("\t\t\t\tSUPPORTED_PLATFORMS = \"" + SupportedPlatforms.Trim() + "\";" + ProjectFileGenerator.NewLine);
+			if (bAutomaticSigning)
+			{
+				Content.Append("\t\t\t\tCODE_SIGN_STYLE = Automatic;" + ProjectFileGenerator.NewLine);
+			}
+			if (IOSRunTimeVersion != null)
+			{
+				Content.Append("\t\t\t\t\"VALID_ARCHS[sdk=iphoneos*]\" = \"" + string.Join(" ", this.SupportedIOSArchitectures) + "\";" + ProjectFileGenerator.NewLine);
+				Content.Append("\t\t\t\tIPHONEOS_DEPLOYMENT_TARGET = " + IOSRunTimeVersion + ";" + ProjectFileGenerator.NewLine);
+				Content.Append("\t\t\t\t\"PRODUCT_NAME[sdk=iphoneos*]\" = \"" + TargetName + "\";" + ProjectFileGenerator.NewLine); // @todo: change to Path.GetFileName(Config.IOSExecutablePath) when we stop using payload
+				Content.Append("\t\t\t\t\"TARGETED_DEVICE_FAMILY[sdk=iphoneos*]\" = \"" + IOSRunTimeDevices + "\";" + ProjectFileGenerator.NewLine);
+				Content.Append("\t\t\t\t\"SDKROOT[sdk=iphoneos]\" = iphoneos;" + ProjectFileGenerator.NewLine);
+				if (!string.IsNullOrEmpty(TEAM_IOS))
 				{
-					TVOSPlatform TVOSPlatform = ((TVOSPlatform)UEBuildPlatform.GetBuildPlatform(UnrealTargetPlatform.TVOS));
-					TVOSProjectSettings ProjectSettings = TVOSPlatform.ReadProjectSettings(ProjectFile);
-					TVOSProvisioningData ProvisioningData = TVOSPlatform.ReadProvisioningData(ProjectSettings, bForDistribution);
-					TVOSRunTimeVersion = ProjectSettings.RuntimeVersion;
-					TVOSRunTimeDevices = ProjectSettings.RuntimeDevices;
-					SupportedPlatforms += " appletvos";
-					if (!bAutomaticSigning)
-					{
-						UUID_TVOS = ProvisioningData.MobileProvisionUUID;
-						TVOS_CERT = ProvisioningData.SigningCertificate;
-					}
-					TEAM_TVOS = ProvisioningData.TeamUUID;
-					TVOS_BUNDLE = ProjectSettings.BundleIdentifier;
+					Content.Append("\t\t\t\t\"DEVELOPMENT_TEAM[sdk=iphoneos*]\" = " + TEAM_IOS + ";" + ProjectFileGenerator.NewLine);
 				}
-
-				Content.Append("\t\t\t\tSUPPORTED_PLATFORMS = \"" + SupportedPlatforms.Trim() + "\";" + ProjectFileGenerator.NewLine);
-				if (bAutomaticSigning)
+				Content.Append("\t\t\t\t\"CODE_SIGN_IDENTITY[sdk=iphoneos*]\" = \"" + IOS_CERT + "\";" + ProjectFileGenerator.NewLine);
+				if (!bAutomaticSigning && !string.IsNullOrEmpty(UUID_IOS))
+				{
+					Content.Append("\t\t\t\t\"PROVISIONING_PROFILE_SPECIFIER[sdk=iphoneos*]\" = \"" + UUID_IOS + "\";" + ProjectFileGenerator.NewLine);
+				}
+				if (ProjectFile != null)
+				{
+					Content.Append("\t\t\t\t\"PRODUCT_BUNDLE_IDENTIFIER[sdk=iphoneos*]\" = " + IOS_BUNDLE + ";" + ProjectFileGenerator.NewLine);
+				}
+			}
+			if (TVOSRunTimeVersion != null)
+			{
+				Content.Append("\t\t\t\t\"VALID_ARCHS[sdk=appletvos*]\" = \"" + string.Join(" ", this.SupportedIOSArchitectures) + "\";" + ProjectFileGenerator.NewLine);
+				Content.Append("\t\t\t\tTVOS_DEPLOYMENT_TARGET = " + TVOSRunTimeVersion + ";" + ProjectFileGenerator.NewLine);
+				Content.Append("\t\t\t\t\"PRODUCT_NAME[sdk=appletvos*]\" = \"" + TargetName + "\";" + ProjectFileGenerator.NewLine); // @todo: change to Path.GetFileName(Config.TVOSExecutablePath) when we stop using payload
+				Content.Append("\t\t\t\t\"TARGETED_DEVICE_FAMILY[sdk=appletvos*]\" = \"" + TVOSRunTimeDevices + "\";" + ProjectFileGenerator.NewLine);
+				Content.Append("\t\t\t\t\"SDKROOT[sdk=appletvos]\" = appletvos;" + ProjectFileGenerator.NewLine);
+				if (!string.IsNullOrEmpty(TEAM_TVOS))
+				{
+					Content.Append("\t\t\t\t\"DEVELOPMENT_TEAM[sdk=appletvos*]\" = " + TEAM_TVOS + ";" + ProjectFileGenerator.NewLine);
+				}
+				Content.Append("\t\t\t\t\"CODE_SIGN_IDENTITY[sdk=appletvos*]\" = \"" + TVOS_CERT + "\";" + ProjectFileGenerator.NewLine);
+				if (!bAutomaticSigning && !string.IsNullOrEmpty(UUID_TVOS))
+				{
+					Content.Append("\t\t\t\t\"PROVISIONING_PROFILE_SPECIFIER[sdk=appletvos*]\" = \"" + UUID_TVOS + "\";" + ProjectFileGenerator.NewLine);
+				}
+				if (ProjectFile != null)
 				{
-					Content.Append("\t\t\t\tCODE_SIGN_STYLE = Automatic;" + ProjectFileGenerator.NewLine);
+					Content.Append("\t\t\t\t\"PRODUCT_BUNDLE_IDENTIFIER[sdk=appletvos*]\" = " + TVOS_BUNDLE + ";" + ProjectFileGenerator.NewLine);
 				}
+			}
+			if (bSupportMac)
+			{
+				Content.Append("\t\t\t\t\"VALID_ARCHS[sdk=macosx*]\" = \"" + string.Join(" ", SupportedMacArchitectures) + "\";" + ProjectFileGenerator.NewLine);
+				Content.Append("\t\t\t\t\"PRODUCT_NAME[sdk=macosx*]\" = \"" + MacExecutableFileName + "\";" + ProjectFileGenerator.NewLine);
+				Content.Append("\t\t\t\t\"CONFIGURATION_BUILD_DIR[sdk=macosx*]\" = \"" + MacExecutableDir + "\";" + ProjectFileGenerator.NewLine);
+				Content.Append("\t\t\t\t\"SDKROOT[sdk=macosx*]\" = macosx;" + ProjectFileGenerator.NewLine);
+
+				// MacToolchain uses the IOS bundle identifier
+				ConfigHierarchy IOSIni = ConfigCache.ReadHierarchy(ConfigHierarchyType.Engine, ProjectFile?.Directory, UnrealTargetPlatform.IOS);
+				string IOSBundleIdentifier;
+				IOSIni.GetString("/Script/IOSRuntimeSettings.IOSRuntimeSettings", "BundleIdentifier", out IOSBundleIdentifier!);
+
+				// we are going to load/sign the UnrealEditor.app when doing modular editor builds, so reuse the shared identifier
+				bool bBuildingEditor = Config.ProjectTarget!.TargetRules!.Type == TargetType.Editor && Config.ProjectTarget!.TargetRules!.LinkType == TargetLinkType.Modular;
+				string MacBundleIdentifier = bBuildingEditor ? ("com.epicgames.UnrealEditor") : (IOSBundleIdentifier.Replace("[PROJECT_NAME]", TargetName).Replace("_", ""));
+				// Xcode 14 wants this to match what is in the final plist
+				Content.Append("\t\t\t\t\"PRODUCT_BUNDLE_IDENTIFIER[sdk=macosx*]\" = " + MacBundleIdentifier + ";" + ProjectFileGenerator.NewLine);
+			}
+
+			if (bIsUnrealGame || bIsUnrealClient)
+			{
 				if (IOSRunTimeVersion != null)
 				{
-					Content.Append("\t\t\t\t\"VALID_ARCHS[sdk=iphoneos*]\" = \"" + string.Join(" ", this.SupportedIOSArchitectures) + "\";" + ProjectFileGenerator.NewLine);
-					Content.Append("\t\t\t\tIPHONEOS_DEPLOYMENT_TARGET = " + IOSRunTimeVersion + ";" + ProjectFileGenerator.NewLine);
-					Content.Append("\t\t\t\t\"PRODUCT_NAME[sdk=iphoneos*]\" = \"" + TargetName + "\";" + ProjectFileGenerator.NewLine); // @todo: change to Path.GetFileName(Config.IOSExecutablePath) when we stop using payload
-					Content.Append("\t\t\t\t\"TARGETED_DEVICE_FAMILY[sdk=iphoneos*]\" = \"" + IOSRunTimeDevices + "\";" + ProjectFileGenerator.NewLine);
-					Content.Append("\t\t\t\t\"SDKROOT[sdk=iphoneos]\" = iphoneos;" + ProjectFileGenerator.NewLine);
-					if (!string.IsNullOrEmpty(TEAM_IOS))
-					{
-						Content.Append("\t\t\t\t\"DEVELOPMENT_TEAM[sdk=iphoneos*]\" = " + TEAM_IOS + ";" + ProjectFileGenerator.NewLine);
-					}
-					Content.Append("\t\t\t\t\"CODE_SIGN_IDENTITY[sdk=iphoneos*]\" = \"" + IOS_CERT + "\";" + ProjectFileGenerator.NewLine);
-					if (!bAutomaticSigning && !string.IsNullOrEmpty(UUID_IOS))
-					{
-						Content.Append("\t\t\t\t\"PROVISIONING_PROFILE_SPECIFIER[sdk=iphoneos*]\" = \"" + UUID_IOS + "\";" + ProjectFileGenerator.NewLine);
-					}
-					if (ProjectFile != null)
-					{
-						Content.Append("\t\t\t\t\"PRODUCT_BUNDLE_IDENTIFIER[sdk=iphoneos*]\" = " + IOS_BUNDLE + ";" + ProjectFileGenerator.NewLine);
-					}
+					Content.Append("\t\t\t\t\"CONFIGURATION_BUILD_DIR[sdk=iphoneos*]\" = \"" + UEDir + "/Engine/Binaries/IOS" + BinariesSubDir + "\";" + ProjectFileGenerator.NewLine);
 				}
 				if (TVOSRunTimeVersion != null)
 				{
-					Content.Append("\t\t\t\t\"VALID_ARCHS[sdk=appletvos*]\" = \"" + string.Join(" ", this.SupportedIOSArchitectures) + "\";" + ProjectFileGenerator.NewLine);
-					Content.Append("\t\t\t\tTVOS_DEPLOYMENT_TARGET = " + TVOSRunTimeVersion + ";" + ProjectFileGenerator.NewLine);
-					Content.Append("\t\t\t\t\"PRODUCT_NAME[sdk=appletvos*]\" = \"" + TargetName + "\";" + ProjectFileGenerator.NewLine); // @todo: change to Path.GetFileName(Config.TVOSExecutablePath) when we stop using payload
-					Content.Append("\t\t\t\t\"TARGETED_DEVICE_FAMILY[sdk=appletvos*]\" = \"" + TVOSRunTimeDevices + "\";" + ProjectFileGenerator.NewLine);
-					Content.Append("\t\t\t\t\"SDKROOT[sdk=appletvos]\" = appletvos;" + ProjectFileGenerator.NewLine);
-					if (!string.IsNullOrEmpty(TEAM_TVOS))
-					{
-						Content.Append("\t\t\t\t\"DEVELOPMENT_TEAM[sdk=appletvos*]\" = " + TEAM_TVOS + ";" + ProjectFileGenerator.NewLine);
-					}
-					Content.Append("\t\t\t\t\"CODE_SIGN_IDENTITY[sdk=appletvos*]\" = \"" + TVOS_CERT + "\";" + ProjectFileGenerator.NewLine);
-					if (!bAutomaticSigning && !string.IsNullOrEmpty(UUID_TVOS))
-					{
-						Content.Append("\t\t\t\t\"PROVISIONING_PROFILE_SPECIFIER[sdk=appletvos*]\" = \"" + UUID_TVOS + "\";" + ProjectFileGenerator.NewLine);
-					}
-					if (ProjectFile != null)
-					{
-						Content.Append("\t\t\t\t\"PRODUCT_BUNDLE_IDENTIFIER[sdk=appletvos*]\" = " + TVOS_BUNDLE + ";" + ProjectFileGenerator.NewLine);
-					}
+					Content.Append("\t\t\t\t\"CONFIGURATION_BUILD_DIR[sdk=appletvos*]\" = \"" + UEDir + "/Engine/Binaries/TVOS" + BinariesSubDir + "\";" + ProjectFileGenerator.NewLine);
 				}
-				if (bSupportMac)
+			}
+			else if (ProjectFile != null)
+			{
+				if (IOSRunTimeVersion != null)
 				{
-					Content.Append("\t\t\t\t\"VALID_ARCHS[sdk=macosx*]\" = \"" + string.Join(" ", SupportedMacArchitectures) + "\";" + ProjectFileGenerator.NewLine);
-					Content.Append("\t\t\t\t\"PRODUCT_NAME[sdk=macosx*]\" = \"" + MacExecutableFileName + "\";" + ProjectFileGenerator.NewLine);
-					Content.Append("\t\t\t\t\"CONFIGURATION_BUILD_DIR[sdk=macosx*]\" = \"" + MacExecutableDir + "\";" + ProjectFileGenerator.NewLine);
-					Content.Append("\t\t\t\t\"SDKROOT[sdk=macosx]\" = macosx;" + ProjectFileGenerator.NewLine);
-					Content.Append("\t\t\t\t\"GENERATE_INFOPLIST_FILE[sdk=macosx*]\" = YES;" + ProjectFileGenerator.NewLine);
+					Content.Append("\t\t\t\t\"CONFIGURATION_BUILD_DIR[sdk=iphoneos*]\" = \"" + GamePath + "/Binaries/IOS" + BinariesSubDir + "\";" + ProjectFileGenerator.NewLine);
 				}
-
-				if (bIsUnrealGame || bIsUnrealClient)
+				if (TVOSRunTimeVersion != null)
 				{
-					if (IOSRunTimeVersion != null)
-					{
-						Content.Append("\t\t\t\t\"CONFIGURATION_BUILD_DIR[sdk=iphoneos*]\" = \"" + UEDir + "/Engine/Binaries/IOS" + BinariesSubDir + "\";" + ProjectFileGenerator.NewLine);
-					}
-					if (TVOSRunTimeVersion != null)
-					{
-						Content.Append("\t\t\t\t\"CONFIGURATION_BUILD_DIR[sdk=appletvos*]\" = \"" + UEDir + "/Engine/Binaries/TVOS" + BinariesSubDir + "\";" + ProjectFileGenerator.NewLine);
-					}
+					Content.Append("\t\t\t\t\"CONFIGURATION_BUILD_DIR[sdk=appletvos*]\" = \"" + GamePath + "/Binaries/TVOS" + BinariesSubDir + "\";" + ProjectFileGenerator.NewLine);
 				}
-				else if (ProjectFile != null)
+			}
+			else
+			{
+				if (IOSRunTimeVersion != null)
 				{
-					if (IOSRunTimeVersion != null)
-					{
-						Content.Append("\t\t\t\t\"CONFIGURATION_BUILD_DIR[sdk=iphoneos*]\" = \"" + GamePath + "/Binaries/IOS" + BinariesSubDir + "\";" + ProjectFileGenerator.NewLine);
-					}
-					if (TVOSRunTimeVersion != null)
-					{
-						Content.Append("\t\t\t\t\"CONFIGURATION_BUILD_DIR[sdk=appletvos*]\" = \"" + GamePath + "/Binaries/TVOS" + BinariesSubDir + "\";" + ProjectFileGenerator.NewLine);
-					}
+					Content.Append("\t\t\t\t\"CONFIGURATION_BUILD_DIR[sdk=iphoneos*]\" = \"" + UEDir + "/Engine/Binaries/IOS" + BinariesSubDir + "\";" + ProjectFileGenerator.NewLine);
 				}
-				else
+				if (TVOSRunTimeVersion != null)
 				{
-					if (IOSRunTimeVersion != null)
-					{
-						Content.Append("\t\t\t\t\"CONFIGURATION_BUILD_DIR[sdk=iphoneos*]\" = \"" + UEDir + "/Engine/Binaries/IOS" + BinariesSubDir + "\";" + ProjectFileGenerator.NewLine);
-					}
-					if (TVOSRunTimeVersion != null)
-					{
-						Content.Append("\t\t\t\t\"CONFIGURATION_BUILD_DIR[sdk=appletvos*]\" = \"" + UEDir + "/Engine/Binaries/TVOS" + BinariesSubDir + "\";" + ProjectFileGenerator.NewLine);
-					}
+					Content.Append("\t\t\t\t\"CONFIGURATION_BUILD_DIR[sdk=appletvos*]\" = \"" + UEDir + "/Engine/Binaries/TVOS" + BinariesSubDir + "\";" + ProjectFileGenerator.NewLine);
 				}
-
 			}
 		}
 
@@ -1202,137 +1199,134 @@ namespace UnrealBuildTool.XcodeProjectLegacy
 			string? IOSRunTimeVersion, TVOSRunTimeVersion;
 			AppendPlatformConfiguration(Content, Config, Config.BuildTarget, ProjectFile, true, !bMacOnly, !bMacOnly, Logger, out IOSRunTimeVersion, out TVOSRunTimeVersion);
 
-			if (!bMacOnly)
-			{
-				bool bIsUnrealGame = Config.BuildTarget.Equals("UnrealGame", StringComparison.InvariantCultureIgnoreCase);
-				bool bIsUnrealClient = Config.BuildTarget.Equals("UnrealClient", StringComparison.InvariantCultureIgnoreCase);
+			bool bIsUnrealGame = Config.BuildTarget.Equals("UnrealGame", StringComparison.InvariantCultureIgnoreCase);
+			bool bIsUnrealClient = Config.BuildTarget.Equals("UnrealClient", StringComparison.InvariantCultureIgnoreCase);
 
-				DirectoryReference? GameDir = ProjectFile?.Directory;
-				string? GamePath = GameDir != null ? ConvertPath(GameDir.FullName) : null;
+			DirectoryReference? GameDir = ProjectFile?.Directory;
+			string? GamePath = GameDir != null ? ConvertPath(GameDir.FullName) : null;
 
-				string IOSInfoPlistPath;
-				string TVOSInfoPlistPath;
-				string MacInfoPlistPath;
-				string? IOSEntitlementPath = null;
-				string? TVOSEntitlementPath = null;
-				if (bIsUnrealGame)
+			string IOSInfoPlistPath;
+			string TVOSInfoPlistPath;
+			string MacInfoPlistPath;
+			string? IOSEntitlementPath = null;
+			string? TVOSEntitlementPath = null;
+			if (bIsUnrealGame)
+			{
+				IOSInfoPlistPath = UEDir + "/Engine/Intermediate/IOS/" + Config.BuildTarget + "-Info.plist";
+				TVOSInfoPlistPath = UEDir + "/Engine/Intermediate/TVOS/" + Config.BuildTarget + "-Info.plist";
+				MacInfoPlistPath = UEDir + "/Engine/Intermediate/Mac/" + MacExecutableFileName + "-Info.plist";
+				IOSEntitlementPath = "";
+				TVOSEntitlementPath = "";
+			}
+			else if (bIsUnrealClient)
+			{
+				IOSInfoPlistPath = UEDir + "/Engine/Intermediate/IOS/UnrealGame-Info.plist";
+				TVOSInfoPlistPath = UEDir + "/Engine/Intermediate/TVOS/UnrealGame-Info.plist";
+				MacInfoPlistPath = UEDir + "/Engine/Intermediate/Mac/" + MacExecutableFileName + "-Info.plist";
+				IOSEntitlementPath = "";
+				TVOSEntitlementPath = "";
+			}
+			else if (ProjectFile != null)
+			{
+				IOSInfoPlistPath = GamePath + "/Intermediate/IOS/" + Config.BuildTarget + "-Info.plist";
+				TVOSInfoPlistPath = GamePath + "/Intermediate/TVOS/" + Config.BuildTarget + "-Info.plist";
+				MacInfoPlistPath = GamePath + "/Intermediate/Mac/" + MacExecutableFileName + "-Info.plist";
+				IOSEntitlementPath = GamePath + "/Intermediate/IOS/" + Config.BuildTarget + ".entitlements";
+				TVOSEntitlementPath = GamePath + "/Intermediate/TVOS/" + Config.BuildTarget + ".entitlements";
+			}
+			else
+			{
+				if (GamePath == null)
 				{
 					IOSInfoPlistPath = UEDir + "/Engine/Intermediate/IOS/" + Config.BuildTarget + "-Info.plist";
 					TVOSInfoPlistPath = UEDir + "/Engine/Intermediate/TVOS/" + Config.BuildTarget + "-Info.plist";
 					MacInfoPlistPath = UEDir + "/Engine/Intermediate/Mac/" + MacExecutableFileName + "-Info.plist";
-					IOSEntitlementPath = "";
-					TVOSEntitlementPath = "";
-				}
-				else if (bIsUnrealClient)
-				{
-					IOSInfoPlistPath = UEDir + "/Engine/Intermediate/IOS/UnrealGame-Info.plist";
-					TVOSInfoPlistPath = UEDir + "/Engine/Intermediate/TVOS/UnrealGame-Info.plist";
-					MacInfoPlistPath = UEDir + "/Engine/Intermediate/Mac/" + MacExecutableFileName + "-Info.plist";
-					IOSEntitlementPath = "";
-					TVOSEntitlementPath = "";
 				}
-				else if (ProjectFile != null)
+				else
 				{
 					IOSInfoPlistPath = GamePath + "/Intermediate/IOS/" + Config.BuildTarget + "-Info.plist";
 					TVOSInfoPlistPath = GamePath + "/Intermediate/TVOS/" + Config.BuildTarget + "-Info.plist";
 					MacInfoPlistPath = GamePath + "/Intermediate/Mac/" + MacExecutableFileName + "-Info.plist";
-					IOSEntitlementPath = GamePath + "/Intermediate/IOS/" + Config.BuildTarget + ".entitlements";
-					TVOSEntitlementPath = GamePath + "/Intermediate/TVOS/" + Config.BuildTarget + ".entitlements";
-				}
-				else
-				{
-					if (GamePath == null)
-					{
-						IOSInfoPlistPath = UEDir + "/Engine/Intermediate/IOS/" + Config.BuildTarget + "-Info.plist";
-						TVOSInfoPlistPath = UEDir + "/Engine/Intermediate/TVOS/" + Config.BuildTarget + "-Info.plist";
-						MacInfoPlistPath = UEDir + "/Engine/Intermediate/Mac/" + MacExecutableFileName + "-Info.plist";
-					}
-					else
-					{
-						IOSInfoPlistPath = GamePath + "/Intermediate/IOS/" + Config.BuildTarget + "-Info.plist";
-						TVOSInfoPlistPath = GamePath + "/Intermediate/TVOS/" + Config.BuildTarget + "-Info.plist";
-						MacInfoPlistPath = GamePath + "/Intermediate/Mac/" + MacExecutableFileName + "-Info.plist";
-					}
 				}
+			}
 
-				if (XcodeProjectFileGenerator.bGeneratingRunIOSProject)
+			if (XcodeProjectFileGenerator.bGeneratingRunIOSProject)
+			{
+				Content.Append("\t\t\t\tINFOPLIST_FILE = \"" + IOSInfoPlistPath + "\";" + ProjectFileGenerator.NewLine);
+				Content.Append("\t\t\t\tCODE_SIGN_ENTITLEMENTS = \"" + IOSEntitlementPath + "\";" + ProjectFileGenerator.NewLine);
+			}
+			else if (XcodeProjectFileGenerator.bGeneratingRunTVOSProject)
+			{
+				Content.Append("\t\t\t\tINFOPLIST_FILE = \"" + TVOSInfoPlistPath + "\";" + ProjectFileGenerator.NewLine);
+				Content.Append("\t\t\t\tCODE_SIGN_ENTITLEMENTS = \"" + TVOSEntitlementPath + "\";" + ProjectFileGenerator.NewLine);
+			}
+			else
+			{
+				Content.Append("\t\t\t\t\"INFOPLIST_FILE[sdk=macosx*]\" = \"" + MacInfoPlistPath + "\";" + ProjectFileGenerator.NewLine);
+				if (IOSRunTimeVersion != null)
 				{
-					Content.Append("\t\t\t\tINFOPLIST_FILE = \"" + IOSInfoPlistPath + "\";" + ProjectFileGenerator.NewLine);
-					Content.Append("\t\t\t\tCODE_SIGN_ENTITLEMENTS = \"" + IOSEntitlementPath + "\";" + ProjectFileGenerator.NewLine);
+					Content.Append("\t\t\t\t\"INFOPLIST_FILE[sdk=iphoneos*]\" = \"" + IOSInfoPlistPath + "\";" + ProjectFileGenerator.NewLine);
+					Content.Append("\t\t\t\t\"CODE_SIGN_ENTITLEMENTS[sdk=iphoneos*]\" = \"" + IOSEntitlementPath + "\";" + ProjectFileGenerator.NewLine);
 				}
-				else if (XcodeProjectFileGenerator.bGeneratingRunTVOSProject)
+				if (TVOSRunTimeVersion != null)
 				{
-					Content.Append("\t\t\t\tINFOPLIST_FILE = \"" + TVOSInfoPlistPath + "\";" + ProjectFileGenerator.NewLine);
-					Content.Append("\t\t\t\tCODE_SIGN_ENTITLEMENTS = \"" + TVOSEntitlementPath + "\";" + ProjectFileGenerator.NewLine);
+					Content.Append("\t\t\t\t\"INFOPLIST_FILE[sdk=appletvos*]\" = \"" + TVOSInfoPlistPath + "\";" + ProjectFileGenerator.NewLine);
+					Content.Append("\t\t\t\t\"CODE_SIGN_ENTITLEMENTS[sdk=appletvos*]\" = \"" + TVOSEntitlementPath + "\";" + ProjectFileGenerator.NewLine);
 				}
-				else
+			}
+
+			// Prepare a temp Info.plist file so Xcode has some basic info about the target immediately after opening the project.
+			// This is needed for the target to pass the settings validation before code signing. UBT will overwrite this plist file later, with proper contents.
+			if (BuildHostPlatform.Current.Platform == UnrealTargetPlatform.Mac)
+			{
+				bool bCreateMacInfoPlist = !File.Exists(MacInfoPlistPath);
+				bool bCreateIOSInfoPlist = !File.Exists(IOSInfoPlistPath) && IOSRunTimeVersion != null;
+				bool bCreateTVOSInfoPlist = !File.Exists(TVOSInfoPlistPath) && TVOSRunTimeVersion != null;
+				if (bCreateMacInfoPlist || bCreateIOSInfoPlist || bCreateTVOSInfoPlist)
 				{
-					Content.Append("\t\t\t\t\"INFOPLIST_FILE[sdk=macosx*]\" = \"" + MacInfoPlistPath + "\";" + ProjectFileGenerator.NewLine);
-					if (IOSRunTimeVersion != null)
+					DirectoryReference? ProjectPath = GameDir;
+					DirectoryReference EngineDir = DirectoryReference.Combine(new DirectoryReference(UEDir), "Engine");
+					string GameName = Config.BuildTarget;
+					bool bIsClient = false;
+					if (ProjectPath == null)
 					{
-						Content.Append("\t\t\t\t\"INFOPLIST_FILE[sdk=iphoneos*]\" = \"" + IOSInfoPlistPath + "\";" + ProjectFileGenerator.NewLine);
-						Content.Append("\t\t\t\t\"CODE_SIGN_ENTITLEMENTS[sdk=iphoneos*]\" = \"" + IOSEntitlementPath + "\";" + ProjectFileGenerator.NewLine);
+						ProjectPath = EngineDir;
 					}
-					if (TVOSRunTimeVersion != null)
+					if (bIsUnrealGame)
 					{
-						Content.Append("\t\t\t\t\"INFOPLIST_FILE[sdk=appletvos*]\" = \"" + TVOSInfoPlistPath + "\";" + ProjectFileGenerator.NewLine);
-						Content.Append("\t\t\t\t\"CODE_SIGN_ENTITLEMENTS[sdk=appletvos*]\" = \"" + TVOSEntitlementPath + "\";" + ProjectFileGenerator.NewLine);
+						ProjectPath = EngineDir;
+						GameName = "UnrealGame";
+						bIsClient = (AppName == "UnrealClient");
 					}
-				}
 
-				// Prepare a temp Info.plist file so Xcode has some basic info about the target immediately after opening the project.
-				// This is needed for the target to pass the settings validation before code signing. UBT will overwrite this plist file later, with proper contents.
-				if (BuildHostPlatform.Current.Platform == UnrealTargetPlatform.Mac)
-				{
-					bool bCreateMacInfoPlist = !File.Exists(MacInfoPlistPath);
-					bool bCreateIOSInfoPlist = !File.Exists(IOSInfoPlistPath) && IOSRunTimeVersion != null;
-					bool bCreateTVOSInfoPlist = !File.Exists(TVOSInfoPlistPath) && TVOSRunTimeVersion != null;
-					if (bCreateMacInfoPlist || bCreateIOSInfoPlist || bCreateTVOSInfoPlist)
+					if (bCreateMacInfoPlist)
 					{
-						DirectoryReference? ProjectPath = GameDir;
-						DirectoryReference EngineDir = DirectoryReference.Combine(new DirectoryReference(UEDir), "Engine");
-						string GameName = Config.BuildTarget;
-						bool bIsClient = false;
-						if (ProjectPath == null)
-						{
-							ProjectPath = EngineDir;
-						}
+						Directory.CreateDirectory(Path.GetDirectoryName(MacInfoPlistPath)!);
+						UEDeployMac.GeneratePList(ProjectPath.FullName, bIsUnrealGame, GameName, Config.BuildTarget, EngineDir.FullName, MacExecutableFileName);
+					}
+					if (bCreateIOSInfoPlist)
+					{
+						// get the receipt
+						FileReference ReceiptFilename;
 						if (bIsUnrealGame)
 						{
-							ProjectPath = EngineDir;
-							GameName = "UnrealGame";
-							bIsClient = (AppName == "UnrealClient");
-						}
-
-						if (bCreateMacInfoPlist)
-						{
-							Directory.CreateDirectory(Path.GetDirectoryName(MacInfoPlistPath)!);
-							UEDeployMac.GeneratePList(ProjectPath.FullName, bIsUnrealGame, GameName, Config.BuildTarget, EngineDir.FullName, MacExecutableFileName);
+							ReceiptFilename = TargetReceipt.GetDefaultPath(Unreal.EngineDirectory, "UnrealGame", UnrealTargetPlatform.IOS, Config.BuildConfig, "");
 						}
-						if (bCreateIOSInfoPlist)
+						else
 						{
-							// get the receipt
-							FileReference ReceiptFilename;
-							if (bIsUnrealGame)
-							{
-								ReceiptFilename = TargetReceipt.GetDefaultPath(Unreal.EngineDirectory, "UnrealGame", UnrealTargetPlatform.IOS, Config.BuildConfig, "");
-							}
-							else
-							{
-								ReceiptFilename = TargetReceipt.GetDefaultPath(ProjectPath, GameName, UnrealTargetPlatform.IOS, Config.BuildConfig, "");
-							}
-							Directory.CreateDirectory(Path.GetDirectoryName(IOSInfoPlistPath)!);
-							bool bSupportPortrait, bSupportLandscape;
-							TargetReceipt? Receipt;
-							TargetReceipt.TryRead(ReceiptFilename, out Receipt);
-							bool bBuildAsFramework = UEDeployIOS.GetCompileAsDll(Receipt);
-							UEDeployIOS.GenerateIOSPList(ProjectFile, Config.BuildConfig, ProjectPath.FullName, bIsUnrealGame, GameName, bIsClient, Config.BuildTarget, EngineDir.FullName, ProjectPath + "/Binaries/IOS/Payload", null, BundleIdentifier, bBuildAsFramework, Logger, out bSupportPortrait, out bSupportLandscape);
-						}
-						if (bCreateTVOSInfoPlist)
-						{
-							Directory.CreateDirectory(Path.GetDirectoryName(TVOSInfoPlistPath)!);
-							UEDeployTVOS.GenerateTVOSPList(ProjectPath.FullName, bIsUnrealGame, GameName, bIsClient, Config.BuildTarget, EngineDir.FullName, ProjectPath + "/Binaries/TVOS/Payload", null, BundleIdentifier, Logger);
+							ReceiptFilename = TargetReceipt.GetDefaultPath(ProjectPath, GameName, UnrealTargetPlatform.IOS, Config.BuildConfig, "");
 						}
+						Directory.CreateDirectory(Path.GetDirectoryName(IOSInfoPlistPath)!);
+						bool bSupportPortrait, bSupportLandscape;
+						TargetReceipt? Receipt;
+						TargetReceipt.TryRead(ReceiptFilename, out Receipt);
+						bool bBuildAsFramework = UEDeployIOS.GetCompileAsDll(Receipt);
+						UEDeployIOS.GenerateIOSPList(ProjectFile, Config.BuildConfig, ProjectPath.FullName, bIsUnrealGame, GameName, bIsClient, Config.BuildTarget, EngineDir.FullName, ProjectPath + "/Binaries/IOS/Payload", null, BundleIdentifier, bBuildAsFramework, Logger, out bSupportPortrait, out bSupportLandscape);
+					}
+					if (bCreateTVOSInfoPlist)
+					{
+						Directory.CreateDirectory(Path.GetDirectoryName(TVOSInfoPlistPath)!);
+						UEDeployTVOS.GenerateTVOSPList(ProjectPath.FullName, bIsUnrealGame, GameName, bIsClient, Config.BuildTarget, EngineDir.FullName, ProjectPath + "/Binaries/TVOS/Payload", null, BundleIdentifier, Logger);
 					}
 				}
 			}
