Mac: Allow to launch UE Editor through XCode without commandline

diff --git a/Engine/Source/Programs/UnrealBuildTool/ProjectFiles/Xcode/XcodeProject.cs b/Engine/Source/Programs/UnrealBuildTool/ProjectFiles/Xcode/XcodeProject.cs
--- a/Engine/Source/Programs/UnrealBuildTool/ProjectFiles/Xcode/XcodeProject.cs
+++ b/Engine/Source/Programs/UnrealBuildTool/ProjectFiles/Xcode/XcodeProject.cs
@@ -3013,7 +3013,7 @@ namespace UnrealBuildTool.XcodeProjectXcconfig
 			Content.WriteLine("      launchStyle = \"0\"");
 			Content.WriteLine("      useCustomWorkingDirectory = \"NO\"");
 			Content.WriteLine("      ignoresPersistentStateOnLaunch = \"NO\"");
-			Content.WriteLine("      debugDocumentVersioning = \"YES\"");
+			Content.WriteLine("      debugDocumentVersioning = \"NO\"");
 			Content.WriteLine("      debugServiceExtension = \"internal\"");
 			Content.WriteLine("      allowLocationSimulation = \"YES\">");
 			Content.WriteLine("      <BuildableProductRunnable");
@@ -3063,7 +3063,7 @@ namespace UnrealBuildTool.XcodeProjectXcconfig
 			Content.WriteLine("      shouldUseLaunchSchemeArgsEnv = \"YES\"");
 			Content.WriteLine("      savedToolIdentifier = \"\"");
 			Content.WriteLine("      useCustomWorkingDirectory = \"NO\"");
-			Content.WriteLine("      debugDocumentVersioning = \"YES\">");
+			Content.WriteLine("      debugDocumentVersioning = \"NO\">");
 			Content.WriteLine("      <BuildableProductRunnable");
 			Content.WriteLine("         runnableDebuggingMode = \"0\">");
 			Content.WriteLine("            <BuildableReference");
diff --git a/Engine/Source/Programs/UnrealBuildTool/ProjectFiles/Xcode/XcodeProjectLegacy.cs b/Engine/Source/Programs/UnrealBuildTool/ProjectFiles/Xcode/XcodeProjectLegacy.cs
--- a/Engine/Source/Programs/UnrealBuildTool/ProjectFiles/Xcode/XcodeProjectLegacy.cs
+++ b/Engine/Source/Programs/UnrealBuildTool/ProjectFiles/Xcode/XcodeProjectLegacy.cs
@@ -1688,7 +1688,7 @@ namespace UnrealBuildTool.XcodeProjectLegacy
 			Content.Append("      launchStyle = \"0\"" + ProjectFileGenerator.NewLine);
 			Content.Append("      useCustomWorkingDirectory = \"NO\"" + ProjectFileGenerator.NewLine);
 			Content.Append("      ignoresPersistentStateOnLaunch = \"NO\"" + ProjectFileGenerator.NewLine);
-			Content.Append("      debugDocumentVersioning = \"YES\"" + ProjectFileGenerator.NewLine);
+			Content.Append("      debugDocumentVersioning = \"NO\"" + ProjectFileGenerator.NewLine);
 			Content.Append("      debugServiceExtension = \"internal\"" + ProjectFileGenerator.NewLine);
 			Content.Append("      allowLocationSimulation = \"YES\">" + ProjectFileGenerator.NewLine);
 			Content.Append("      <BuildableProductRunnable" + ProjectFileGenerator.NewLine);
@@ -1740,7 +1740,7 @@ namespace UnrealBuildTool.XcodeProjectLegacy
 			Content.Append("      shouldUseLaunchSchemeArgsEnv = \"YES\"" + ProjectFileGenerator.NewLine);
 			Content.Append("      savedToolIdentifier = \"\"" + ProjectFileGenerator.NewLine);
 			Content.Append("      useCustomWorkingDirectory = \"NO\"" + ProjectFileGenerator.NewLine);
-			Content.Append("      debugDocumentVersioning = \"YES\">" + ProjectFileGenerator.NewLine);
+			Content.Append("      debugDocumentVersioning = \"NO\">" + ProjectFileGenerator.NewLine);
 			Content.Append("      <BuildableProductRunnable" + ProjectFileGenerator.NewLine);
 			Content.Append("         runnableDebuggingMode = \"0\">" + ProjectFileGenerator.NewLine);
 			Content.Append("            <BuildableReference" + ProjectFileGenerator.NewLine);
