PR #9569: Rider project model: get the list of supported configurations from UnrealBuildTool (Contributed by [REDACTED_PERSON])
#tests Tested in Rider and MSVS

diff --git a/Engine/Source/Programs/UnrealBuildTool/ProjectFiles/Rider/RiderProjectFile.cs b/Engine/Source/Programs/UnrealBuildTool/ProjectFiles/Rider/RiderProjectFile.cs
--- a/Engine/Source/Programs/UnrealBuildTool/ProjectFiles/Rider/RiderProjectFile.cs
+++ b/Engine/Source/Programs/UnrealBuildTool/ProjectFiles/Rider/RiderProjectFile.cs
@@ -30,10 +30,10 @@ namespace UnrealBuildTool
 	internal class RiderProjectFile : ProjectFile
 	{
 		private static readonly XcrunRunner AppleHelper = new XcrunRunner();
-		
-		public DirectoryReference RootPath;
-		public HashSet<TargetType> TargetTypes;
-		public CommandLineArguments Arguments;
+
+		private readonly DirectoryReference RootPath;
+		private readonly HashSet<TargetType> TargetTypes;
+		private readonly CommandLineArguments Arguments;
 
 		private ToolchainInfo RootToolchainInfo = new ToolchainInfo();
 		private UEBuildTarget? CurrentTarget;
@@ -403,7 +403,7 @@ namespace UnrealBuildTool
 		/// <param name="Writer">Writer for the array data</param>
 		/// <param name="ArrayName">Name of the array property</param>
 		/// <param name="Strings">Sequence of strings to write. May be null.</param>
-		static void ExportJsonStringArray(JsonWriter Writer, string ArrayName, IEnumerable<string> Strings)
+		private static void ExportJsonStringArray(JsonWriter Writer, string ArrayName, IEnumerable<string> Strings)
 		{
 			if (Strings == null || !Strings.Any()) return;
 			
diff --git a/Engine/Source/Programs/UnrealBuildTool/ProjectFiles/Rider/RiderProjectFileGenerator.cs b/Engine/Source/Programs/UnrealBuildTool/ProjectFiles/Rider/RiderProjectFileGenerator.cs
--- a/Engine/Source/Programs/UnrealBuildTool/ProjectFiles/Rider/RiderProjectFileGenerator.cs
+++ b/Engine/Source/Programs/UnrealBuildTool/ProjectFiles/Rider/RiderProjectFileGenerator.cs
@@ -2,7 +2,6 @@
 
 using System;
 using System.Collections.Generic;
-using System.IO;
 using System.Linq;
 using EpicGames.Core;
 using UnrealBuildBase;
@@ -24,8 +23,14 @@ namespace UnrealBuildTool
 		/// Don't generate project model for these platforms unless they are specified in "Platforms" console arguments. 
 		/// </summary>
 		/// <returns></returns>
-		private readonly HashSet<UnrealTargetPlatform> DeprecatedPlatforms = new HashSet<UnrealTargetPlatform> { };
+		private readonly HashSet<UnrealTargetPlatform> DeprecatedPlatforms = new HashSet<UnrealTargetPlatform>();
 
+		/// <summary>
+		/// When specified only primary project file (root.json) will be generated.
+		/// Normally project files for specific configurations are generated together with primary project file.
+		/// </summary>
+		[CommandLine("-OnlyPrimaryProjectFile")]
+		bool bOnlyPrimaryProjectFile = false;
 
 		/// <summary>
 		/// Platforms to generate project files for
@@ -87,13 +92,18 @@ namespace UnrealBuildTool
 
 		protected override ProjectFile AllocateProjectFile(FileReference InitFilePath, DirectoryReference BaseDir)
 		{
-			RiderProjectFile projectFile = new RiderProjectFile(InitFilePath, BaseDir,
-				RootPath: InitFilePath.Directory, Arguments: Arguments, TargetTypes: TargetTypes);
+			RiderProjectFile projectFile = new RiderProjectFile(InitFilePath, BaseDir, RootPath: InitFilePath.Directory,
+				Arguments: Arguments, TargetTypes: TargetTypes);
 			return projectFile;
 		}
 
 		protected override bool WriteProjectFiles(PlatformProjectGeneratorCollection PlatformProjectGenerators, ILogger Logger)
 		{
+			if (bOnlyPrimaryProjectFile)
+			{
+				return true;
+			}
+
 			using (ProgressWriter Progress = new ProgressWriter("Writing project files...", true, Logger))
 			{
 				List<ProjectFile> ProjectsToGenerate = new List<ProjectFile>(GeneratedProjectFiles);
@@ -453,7 +463,12 @@ namespace UnrealBuildTool
 
 			GatherProjects(PlatformProjectGenerators, AllGameProjects, Logger);
 
-			WriteProjectFiles(PlatformProjectGenerators, Logger);
+			WritePrimaryProjectFile(UBTProject, PlatformProjectGenerators, Logger);
+			if (!bOnlyPrimaryProjectFile)
+			{
+				WriteProjectFiles(PlatformProjectGenerators, Logger);
+			}
+
 			Logger.LogDebug("Project generation complete ({NumGenerated} generated, {NumImported} imported)", GeneratedProjectFiles.Count,
 				OtherProjectFiles.Count);
 
@@ -493,6 +508,36 @@ namespace UnrealBuildTool
 		protected override bool WritePrimaryProjectFile(ProjectFile? ProjectFile,
 			PlatformProjectGeneratorCollection PlatformProjectGenerators, ILogger Logger)
 		{
+			try
+			{
+				HashSet<UnrealTargetPlatform> ValidPlatforms;
+				List<VCProjectFileGenerator.VCSolutionConfigCombination> ValidSolutionConfigs;
+				VCProjectFileGenerator.CollectSolutionConfigurations(SupportedConfigurations,
+					SupportedPlatforms, AllProjectFiles, Logger, out ValidPlatforms, out ValidSolutionConfigs);
+				FileReference OutputFile = GetRiderProjectLocation("root.json");
+				DirectoryReference.CreateDirectory(OutputFile.Directory);
+				using (JsonWriter Writer = new JsonWriter(OutputFile))
+				{
+					Writer.WriteObjectStart();
+					Writer.WriteArrayStart("ConfigurationsAndPlatforms");
+					foreach (VCProjectFileGenerator.VCSolutionConfigCombination Config in ValidSolutionConfigs)
+					{
+						Writer.WriteObjectStart();
+						Writer.WriteValue("Configuration", Config.Configuration.ToString() );
+						Writer.WriteValue("Platform", Config.Platform.ToString());
+						Writer.WriteValue("Target", Config.TargetConfigurationName.ToString());
+						Writer.WriteObjectEnd();
+					}
+
+					Writer.WriteArrayEnd();
+					Writer.WriteObjectEnd();
+				}
+			}
+			catch (Exception Ex)
+			{
+				Logger.LogWarning("Exception while writing Rider root project file: {Ex}", Ex.ToString());
+			}
+
 			return true;
 		}
 	}
diff --git a/Engine/Source/Programs/UnrealBuildTool/ProjectFiles/VisualStudio/VCProjectFileGenerator.cs b/Engine/Source/Programs/UnrealBuildTool/ProjectFiles/VisualStudio/VCProjectFileGenerator.cs
--- a/Engine/Source/Programs/UnrealBuildTool/ProjectFiles/VisualStudio/VCProjectFileGenerator.cs
+++ b/Engine/Source/Programs/UnrealBuildTool/ProjectFiles/VisualStudio/VCProjectFileGenerator.cs
@@ -346,7 +346,7 @@ namespace UnrealBuildTool
 		/// <summary>
 		/// Used to sort VC solution config names along with the config and platform values
 		/// </summary>
-		class VCSolutionConfigCombination
+		public class VCSolutionConfigCombination
 		{
 			/// <summary>
 			/// Visual Studio solution configuration name for this config+platform
@@ -386,7 +386,7 @@ namespace UnrealBuildTool
 		/// <param name="Configuration">The build configuration</param>
 		/// <param name="TargetType">The type of target being built</param>
 		/// <returns>The generated solution configuration name</returns>
-		string MakeSolutionConfigurationName(UnrealTargetConfiguration Configuration, TargetType TargetType)
+		static string MakeSolutionConfigurationName(UnrealTargetConfiguration Configuration, TargetType TargetType)
 		{
 			string SolutionConfigName = Configuration.ToString();
 
@@ -637,90 +637,19 @@ namespace UnrealBuildTool
 
 			// Solution configuration platforms.  This is just a list of all of the platforms and configurations that
 			// appear in Visual Studio's build configuration selector.
-			List<VCSolutionConfigCombination> SolutionConfigCombinations = new List<VCSolutionConfigCombination>();
+			List<VCSolutionConfigCombination> SolutionConfigCombinations;
 
 			// The "Global" section has source control, solution configurations, project configurations,
 			// preferences, and project hierarchy data
 			{
 				VCSolutionFileContent.AppendLine("Global");
 				{
-					HashSet<UnrealTargetPlatform> PlatformsValidForProjects = new HashSet<UnrealTargetPlatform>();
+					HashSet<UnrealTargetPlatform> PlatformsValidForProjects;
 					{
 						VCSolutionFileContent.AppendLine("	GlobalSection(SolutionConfigurationPlatforms) = preSolution");
 
-						Dictionary<string, Tuple<UnrealTargetConfiguration, TargetType>> SolutionConfigurationsValidForProjects = new Dictionary<string, Tuple<UnrealTargetConfiguration, TargetType>>();
-
-						foreach (UnrealTargetConfiguration CurConfiguration in SupportedConfigurations)
-						{
-							if (InstalledPlatformInfo.IsValidConfiguration(CurConfiguration, EProjectType.Code))
-							{
-								foreach (UnrealTargetPlatform CurPlatform in SupportedPlatforms)
-								{
-									if (InstalledPlatformInfo.IsValidPlatform(CurPlatform, EProjectType.Code))
-									{
-										foreach (ProjectFile CurProject in AllProjectFiles)
-										{
-											if (!CurProject.IsStubProject)
-											{
-												if (CurProject.ProjectTargets.Count == 0)
-												{
-													throw new BuildException("Expecting project '" + CurProject.ProjectFilePath + "' to have at least one ProjectTarget associated with it!");
-												}
-
-												// Figure out the set of valid target configuration names
-												foreach (Project ProjectTarget in CurProject.ProjectTargets)
-												{
-													if (VCProjectFile.IsValidProjectPlatformAndConfiguration(ProjectTarget, CurPlatform, CurConfiguration, Logger))
-													{
-														PlatformsValidForProjects.Add(CurPlatform);
-
-														// Default to a target configuration name of "Game", since that will collapse down to an empty string
-														TargetType TargetType = TargetType.Game;
-														if (ProjectTarget.TargetRules != null)
-														{
-															TargetType = ProjectTarget.TargetRules.Type;
-														}
-
-														string SolutionConfigName = MakeSolutionConfigurationName(CurConfiguration, TargetType);
-														SolutionConfigurationsValidForProjects[SolutionConfigName] = new Tuple<UnrealTargetConfiguration, TargetType>(CurConfiguration, TargetType);
-													}
-												}
-											}
-										}
-									}
-								}
-							}
-						}
-
-						foreach (UnrealTargetPlatform CurPlatform in PlatformsValidForProjects)
-						{
-							foreach (KeyValuePair<string, Tuple<UnrealTargetConfiguration, TargetType>> SolutionConfigKeyValue in SolutionConfigurationsValidForProjects)
-							{
-								// e.g.  "Development|Win64 = Development|Win64"
-								string SolutionConfigName = SolutionConfigKeyValue.Key;
-								UnrealTargetConfiguration Configuration = SolutionConfigKeyValue.Value.Item1;
-								TargetType TargetType = SolutionConfigKeyValue.Value.Item2;
-
-								string SolutionPlatformName = CurPlatform.ToString();
-
-								string SolutionConfigAndPlatformPair = SolutionConfigName + "|" + SolutionPlatformName;
-								SolutionConfigCombinations.Add(
-										new VCSolutionConfigCombination(SolutionConfigAndPlatformPair)
-										{
-											Configuration = Configuration,
-											Platform = CurPlatform,
-											TargetConfigurationName = TargetType
-										}
-									);
-							}
-						}
-
-						// Sort the list of solution platform strings alphabetically (Visual Studio prefers it)
-						SolutionConfigCombinations.Sort(
-								new Comparison<VCSolutionConfigCombination>(
-									(x, y) => { return String.Compare(x.VCSolutionConfigAndPlatformName, y.VCSolutionConfigAndPlatformName, StringComparison.InvariantCultureIgnoreCase); }
-								)
-							);
+						CollectSolutionConfigurations(SupportedConfigurations, SupportedPlatforms, AllProjectFiles,
+							Logger, out PlatformsValidForProjects, out SolutionConfigCombinations);
 
 						HashSet<string> AppendedSolutionConfigAndPlatformNames = new HashSet<string>(StringComparer.InvariantCultureIgnoreCase);
 						foreach (VCSolutionConfigCombination SolutionConfigCombination in SolutionConfigCombinations)
@@ -922,6 +851,97 @@ namespace UnrealBuildTool
 			return bSuccess;
 		}
 
+		public static void CollectSolutionConfigurations(List<UnrealTargetConfiguration> AllConfigurations,
+			List<UnrealTargetPlatform> AllPlatforms, List<ProjectFile> AllProjectFiles, ILogger Logger,
+			out HashSet<UnrealTargetPlatform> OutValidPlatforms, out List<VCSolutionConfigCombination> OutSolutionConfigs)
+		{
+			OutValidPlatforms = new HashSet<UnrealTargetPlatform>();
+			OutSolutionConfigs = new List<VCSolutionConfigCombination>();
+			Dictionary<string, Tuple<UnrealTargetConfiguration, TargetType>> SolutionConfigurationsValidForProjects =
+				new Dictionary<string, Tuple<UnrealTargetConfiguration, TargetType>>();
+
+			foreach (UnrealTargetConfiguration CurConfiguration in AllConfigurations)
+			{
+				if (InstalledPlatformInfo.IsValidConfiguration(CurConfiguration, EProjectType.Code))
+				{
+					foreach (UnrealTargetPlatform CurPlatform in AllPlatforms)
+					{
+						if (InstalledPlatformInfo.IsValidPlatform(CurPlatform, EProjectType.Code))
+						{
+							foreach (ProjectFile CurProject in AllProjectFiles)
+							{
+								if (!CurProject.IsStubProject)
+								{
+									if (CurProject.ProjectTargets.Count == 0)
+									{
+										throw new BuildException("Expecting project '" + CurProject.ProjectFilePath +
+										                         "' to have at least one ProjectTarget associated with it!");
+									}
+
+									// Figure out the set of valid target configuration names
+									foreach (Project ProjectTarget in CurProject.ProjectTargets)
+									{
+										if (VCProjectFile.IsValidProjectPlatformAndConfiguration(ProjectTarget, CurPlatform,
+											    CurConfiguration, Logger))
+										{
+											OutValidPlatforms.Add(CurPlatform);
+
+											// Default to a target configuration name of "Game", since that will collapse down to an empty string
+											TargetType TargetType = TargetType.Game;
+											if (ProjectTarget.TargetRules != null)
+											{
+												TargetType = ProjectTarget.TargetRules.Type;
+											}
+
+											string SolutionConfigName =
+												MakeSolutionConfigurationName(CurConfiguration, TargetType);
+											SolutionConfigurationsValidForProjects[SolutionConfigName] =
+												new Tuple<UnrealTargetConfiguration, TargetType>(CurConfiguration, TargetType);
+										}
+									}
+								}
+							}
+						}
+					}
+				}
+			}
+
+			foreach (UnrealTargetPlatform CurPlatform in OutValidPlatforms)
+			{
+				foreach (KeyValuePair<string, Tuple<UnrealTargetConfiguration, TargetType>> SolutionConfigKeyValue in
+				         SolutionConfigurationsValidForProjects)
+				{
+					// e.g.  "Development|Win64 = Development|Win64"
+					string SolutionConfigName = SolutionConfigKeyValue.Key;
+					UnrealTargetConfiguration Configuration = SolutionConfigKeyValue.Value.Item1;
+					TargetType TargetType = SolutionConfigKeyValue.Value.Item2;
+
+					string SolutionPlatformName = CurPlatform.ToString();
+
+					string SolutionConfigAndPlatformPair = SolutionConfigName + "|" + SolutionPlatformName;
+					OutSolutionConfigs.Add(
+						new VCSolutionConfigCombination(SolutionConfigAndPlatformPair)
+						{
+							Configuration = Configuration,
+							Platform = CurPlatform,
+							TargetConfigurationName = TargetType
+						}
+					);
+				}
+			}
+
+			// Sort the list of solution platform strings alphabetically (Visual Studio prefers it)
+			OutSolutionConfigs.Sort(
+				new Comparison<VCSolutionConfigCombination>(
+					(x, y) =>
+					{
+						return String.Compare(x.VCSolutionConfigAndPlatformName, y.VCSolutionConfigAndPlatformName,
+							StringComparison.InvariantCultureIgnoreCase);
+					}
+				)
+			);
+		}
+
 		protected override void WriteDebugSolutionFiles(PlatformProjectGeneratorCollection PlatformProjectGenerators, DirectoryReference IntermediateProjectFilesPath, ILogger Logger)
 		{
 			//build and collect UnrealVS configuration
