UnrealBuildTool: -CompilerVersion=Latest should filter out preview toolchains. If the latest preview toolchain is desired, -CompilerVersion=Preview is already available

#rnx

diff --git a/Engine/Source/Programs/UnrealBuildTool/Platform/Windows/MicrosoftPlatformSDK.cs b/Engine/Source/Programs/UnrealBuildTool/Platform/Windows/MicrosoftPlatformSDK.cs
--- a/Engine/Source/Programs/UnrealBuildTool/Platform/Windows/MicrosoftPlatformSDK.cs
+++ b/Engine/Source/Programs/UnrealBuildTool/Platform/Windows/MicrosoftPlatformSDK.cs
@@ -516,19 +516,19 @@ namespace UnrealBuildTool
 			}
 			else if (String.Compare(CompilerVersion, "Latest", StringComparison.InvariantCultureIgnoreCase) == 0)
 			{
-				ToolChain = SelectToolChain(ToolChains, x => x.ThenByDescending(x => x.Version), Architecture);
+				ToolChain = SelectToolChain(ToolChains.Where(x => !x.IsPreview), x => x.ThenByDescending(x => x.Version), Architecture);
 				if (ToolChain == null)
 				{
-					DumpToolChains(ToolChains, x => x.ThenByDescending(x => x.Version), Architecture, Logger);
-					throw new BuildException("Unable to find valid C++ toolchain for {0} {1}", Compiler, Architecture.ToString());
+					DumpToolChains(ToolChains, x => x.ThenBy(x => x.IsPreview).ThenByDescending(x => x.Version), Architecture, Logger);
+					throw new BuildException("Unable to find valid latest C++ toolchain for {0} {1}", Compiler, Architecture.ToString());
 				}
 			}
 			else if (String.Compare(CompilerVersion, "Preview", StringComparison.InvariantCultureIgnoreCase) == 0)
 			{
-				ToolChain = SelectToolChain(ToolChains, x => x.ThenByDescending(x => x.IsPreview), Architecture);
-				if (ToolChain == null || !ToolChain.IsPreview)
+				ToolChain = SelectToolChain(ToolChains.Where(x => x.IsPreview), x => x.ThenByDescending(x => x.Version), Architecture);
+				if (ToolChain == null)
 				{
-					DumpToolChains(ToolChains, x => x.ThenByDescending(x => x.IsPreview), Architecture, Logger);
+					DumpToolChains(ToolChains, x => x.ThenByDescending(x => x.IsPreview).ThenByDescending(x => x.Version), Architecture, Logger);
 					throw new BuildException("Unable to find valid preview toolchain for {0} {1}", Compiler, Architecture.ToString());
 				}
 			}
