Added support to inline the gen.cpp files into the module cpp files.

This change does a few things:
* Improves compile times because the gen.cpp is inlined so less header parsing is required.
* Allows us to forward declare more types in the headers where we were not able to because the gen.cpp files did not have all the required headers included. The gen.cpp files required the header that generated it to include all the types because of the constructor and desructor in the gen.cpp need to know the full type.

diff --git a/Engine/Source/Programs/UnrealBuildTool/Configuration/TargetRules.cs b/Engine/Source/Programs/UnrealBuildTool/Configuration/TargetRules.cs
--- a/Engine/Source/Programs/UnrealBuildTool/Configuration/TargetRules.cs
+++ b/Engine/Source/Programs/UnrealBuildTool/Configuration/TargetRules.cs
@@ -1749,6 +1749,13 @@ namespace UnrealBuildTool
 		[XmlConfigFile(Category = "BuildConfiguration")]
 		public bool bPublicSymbolsByDefault = false;
 
+		/// <summary>
+		/// Disable supports for inlining gen.cpps
+		/// </summary>
+		[XmlConfigFile(Name = "bDisableInliningGenCpps")]
+		[CommandLine("-DisableInliningGenCpps")]
+		public bool bDisableInliningGenCpps = false;
+
 		/// <summary>
 		/// Allows overriding the toolchain to be created for this target. This must match the name of a class declared in the UnrealBuildTool assembly.
 		/// </summary>
@@ -3330,6 +3337,11 @@ namespace UnrealBuildTool
 			get { return Inner.bPublicSymbolsByDefault; }
 		}
 
+		public bool bDisableInliningGenCpps
+		{
+			get { return Inner.bDisableInliningGenCpps; }
+		}
+
 		public string? ToolChainName
 		{
 			get { return Inner.ToolChainName; }
diff --git a/Engine/Source/Programs/UnrealBuildTool/Configuration/UEBuildModuleCPP.cs b/Engine/Source/Programs/UnrealBuildTool/Configuration/UEBuildModuleCPP.cs
--- a/Engine/Source/Programs/UnrealBuildTool/Configuration/UEBuildModuleCPP.cs
+++ b/Engine/Source/Programs/UnrealBuildTool/Configuration/UEBuildModuleCPP.cs
@@ -562,6 +562,27 @@ namespace UnrealBuildTool
 					}
 				}
 
+				if (!Target.bDisableInliningGenCpps)
+				{
+					// Remove any generated files from the compile list if they are inlined
+					foreach (FileItem CPPFileItem in CPPFiles)
+					{
+						var ListOfInlinedGenCpps = ModuleCompileEnvironment.MetadataCache.GetListOfInlinedGeneratedCppFiles(CPPFileItem);
+						foreach (string ListOfInlinedGenCppsItem in ListOfInlinedGenCpps)
+						{
+							string? FoundGenCppFile = GeneratedFiles.FirstOrDefault((file) => String.Compare(Path.GetFileName(file), ListOfInlinedGenCppsItem + ".gen.cpp", StringComparison.InvariantCultureIgnoreCase) == 0);
+							if (FoundGenCppFile != null)
+							{
+								GeneratedFiles.Remove(FoundGenCppFile);
+							}
+							else
+							{
+								Logger.LogError("Could not find a UObject marked up header file named '{HeaderFile}'", ListOfInlinedGenCppsItem);
+							}
+						}
+					}
+				}
+
 				if (GeneratedFiles.Count > 0)
 				{
 					bool bMergeUnityFiles = Target.bMergeModuleAndGeneratedUnityFiles && Rules.bMergeUnityFiles;
@@ -1530,6 +1551,15 @@ namespace UnrealBuildTool
 				Result.Definitions.Add("UE_IS_ENGINE_MODULE=0");
 			}
 
+			if (Target.bDisableInliningGenCpps)
+			{
+				Result.Definitions.Add("UE_DISABLE_INLINE_GEN_CPP=1");
+			}
+			else
+			{
+				Result.Definitions.Add("UE_DISABLE_INLINE_GEN_CPP=0");
+			}
+
 			Result.Definitions.AddRange(EngineIncludeOrderHelper.GetDeprecationDefines(Rules.IncludeOrderVersion));
 
 			// For game modules, set the define for the project and target names, which will be used by the IMPLEMENT_PRIMARY_GAME_MODULE macro.
@@ -1587,6 +1617,15 @@ namespace UnrealBuildTool
 				CompileEnvironment.Definitions.Add("UE_IS_ENGINE_MODULE=0");
 			}
 
+			if (Rules.Target.bDisableInliningGenCpps)
+			{
+				CompileEnvironment.Definitions.Add("UE_DISABLE_INLINE_GEN_CPP=1");
+			}
+			else
+			{
+				CompileEnvironment.Definitions.Add("UE_DISABLE_INLINE_GEN_CPP=0");
+			}
+
 			CompileEnvironment.Definitions.AddRange(EngineIncludeOrderHelper.GetDeprecationDefines(Rules.IncludeOrderVersion));
 
 			// Add the module's private definitions.
diff --git a/Engine/Source/Programs/UnrealBuildTool/System/SourceFileMetadataCache.cs b/Engine/Source/Programs/UnrealBuildTool/System/SourceFileMetadataCache.cs
--- a/Engine/Source/Programs/UnrealBuildTool/System/SourceFileMetadataCache.cs
+++ b/Engine/Source/Programs/UnrealBuildTool/System/SourceFileMetadataCache.cs
@@ -53,10 +53,26 @@ namespace UnrealBuildTool
 			public bool bContainsMarkup;
 		}
 
+		/// <summary>
+		/// Information about whether a file contains the inline gen.cpp macro
+		/// </summary>
+		class InlineReflectionInfo
+		{
+			/// <summary>
+			/// Last write time of the file when the data was cached
+			/// </summary>
+			public long LastWriteTimeUtc;
+
+			/// <summary>
+			/// List of files this particular file is inlining
+			/// </summary>
+			public List<string> InlinedFileNames = new List<string>();
+		}
+
 		/// <summary>
 		/// The current file version
 		/// </summary>
-		public const int CurrentVersion = 3;
+		public const int CurrentVersion = 4;
 
 		/// <summary>
 		/// Location of this dependency cache
@@ -83,6 +99,11 @@ namespace UnrealBuildTool
 		/// </summary>
 		ConcurrentDictionary<FileItem, ReflectionInfo> FileToReflectionInfo = new ConcurrentDictionary<FileItem, ReflectionInfo>();
 
+		/// <summary>
+		/// Map from file item to inline info
+		/// </summary>
+		ConcurrentDictionary<FileItem, InlineReflectionInfo> FileToInlineReflectionInfo = new ConcurrentDictionary<FileItem, InlineReflectionInfo>();
+
 		/// <summary>
 		/// Map from file item to source file info
 		/// </summary>
@@ -103,6 +124,11 @@ namespace UnrealBuildTool
 		/// </summary>
 		static readonly Regex ReflectionMarkupRegex = new Regex("^\\s*U(CLASS|STRUCT|ENUM|INTERFACE|DELEGATE)\\b", RegexOptions.Compiled | RegexOptions.Multiline);
 
+		/// <summary>
+		/// 
+		/// </summary>
+		static readonly Regex InlineReflectionMarkupRegex = new Regex(@"UE_INLINE_GENERATED_CPP_BY_NAME\((.+)\)", RegexOptions.Compiled | RegexOptions.Multiline);
+
 		/// <summary>
 		/// Regex that matches #include statements.
 		/// </summary>
@@ -237,6 +263,37 @@ namespace UnrealBuildTool
 			}
 		}
 
+
+		/// <summary>
+		/// Returns a list of inlined generated cpps that this source file contains.
+		/// </summary>
+		/// <param name="SourceFile">The source file to parse</param>
+		/// <returns>List of marked files this source file contains</returns>
+		public IList<string> GetListOfInlinedGeneratedCppFiles(FileItem SourceFile)
+		{
+			if (Parent != null && !SourceFile.Location.IsUnderDirectory(BaseDirectory))
+			{
+				return Parent.GetListOfInlinedGeneratedCppFiles(SourceFile);
+			}
+			else
+			{
+				InlineReflectionInfo? InlineReflectionInfo;
+				if (!FileToInlineReflectionInfo.TryGetValue(SourceFile, out InlineReflectionInfo) || SourceFile.LastWriteTimeUtc.Ticks > InlineReflectionInfo.LastWriteTimeUtc)
+				{
+					InlineReflectionInfo = new InlineReflectionInfo();
+					InlineReflectionInfo.LastWriteTimeUtc = SourceFile.LastWriteTimeUtc.Ticks;
+					MatchCollection FileMatches = InlineReflectionMarkupRegex.Matches(FileReference.ReadAllText(SourceFile.Location));
+					foreach (Match Match in FileMatches)
+					{
+						InlineReflectionInfo.InlinedFileNames.Add(Match.Groups[1].Value);
+					}
+					FileToInlineReflectionInfo[SourceFile] = InlineReflectionInfo;
+					bModified = true;
+				}
+				return InlineReflectionInfo.InlinedFileNames;
+			}
+		}
+
 		/// <summary>
 		/// Parse the first include directive from a source file
 		/// </summary>
@@ -389,6 +446,18 @@ namespace UnrealBuildTool
 
 						FileToReflectionInfo[File] = ReflectionInfo;
 					}
+
+					int FileToInlineMarkupFlagCount = Reader.ReadInt();
+					for (int Idx = 0; Idx < FileToInlineMarkupFlagCount; Idx++)
+					{
+						FileItem File = Reader.ReadCompactFileItem();
+
+						InlineReflectionInfo InlineReflectionInfo = new InlineReflectionInfo();
+						InlineReflectionInfo.LastWriteTimeUtc = Reader.ReadLong();
+						InlineReflectionInfo.InlinedFileNames = Reader.ReadList(() => Reader.ReadString())!;
+
+						FileToInlineReflectionInfo[File] = InlineReflectionInfo;
+					}
 				}
 			}
 			catch(Exception Ex)
@@ -425,6 +494,14 @@ namespace UnrealBuildTool
 						Writer.WriteLong(Pair.Value.LastWriteTimeUtc);
 						Writer.WriteBool(Pair.Value.bContainsMarkup);
 					}
+
+					Writer.WriteInt(FileToInlineReflectionInfo.Count);
+					foreach (KeyValuePair<FileItem, InlineReflectionInfo> Pair in FileToInlineReflectionInfo)
+					{
+						Writer.WriteCompactFileItem(Pair.Key);
+						Writer.WriteLong(Pair.Value.LastWriteTimeUtc);
+						Writer.WriteList(Pair.Value.InlinedFileNames, Item => Writer.WriteString(Item));
+					}
 				}
 			}
 			bModified = false;
diff --git a/Engine/Source/Runtime/CoreUObject/Public/UObject/ObjectMacros.h b/Engine/Source/Runtime/CoreUObject/Public/UObject/ObjectMacros.h
--- a/Engine/Source/Runtime/CoreUObject/Public/UObject/ObjectMacros.h
+++ b/Engine/Source/Runtime/CoreUObject/Public/UObject/ObjectMacros.h
@@ -679,6 +679,14 @@ struct COREUOBJECT_API FReferencerInformationList
 #define UDELEGATE(...)
 #define RIGVM_METHOD(...)
 
+// Used to inline generated cpp files from UObject headers
+#define UE_INLINE_STRINGIFY(name) #name
+#if UE_DISABLE_INLINE_GEN_CPP
+	#define UE_INLINE_GENERATED_CPP_BY_NAME(name) UE_INLINE_STRINGIFY(CoreTypes.h)
+#else
+	#define UE_INLINE_GENERATED_CPP_BY_NAME(name) UE_INLINE_STRINGIFY(name.gen.cpp)
+#endif
+
 // This pair of macros is used to help implement GENERATED_BODY() and GENERATED_USTRUCT_BODY()
 #define BODY_MACRO_COMBINE_INNER(A,B,C,D) A##B##C##D
 #define BODY_MACRO_COMBINE(A,B,C,D) BODY_MACRO_COMBINE_INNER(A,B,C,D)
