Added the ability to override the NumIncludedBytesPerUnityCPP from the module rules. This is needed by some modules to better tune the size of the unity files.

Before:
1 Unity file
Total CPU Time: 162.90625 s
Total time in Parallel executor: 61.62 seconds

After:
6 Unity files
Total CPU Time: 169.15625 s
Total time in Parallel executor: 16.64 seconds

diff --git a/Engine/Plugins/Runtime/Metasound/Source/MetasoundStandardNodes/MetasoundStandardNodes.Build.cs b/Engine/Plugins/Runtime/Metasound/Source/MetasoundStandardNodes/MetasoundStandardNodes.Build.cs
--- a/Engine/Plugins/Runtime/Metasound/Source/MetasoundStandardNodes/MetasoundStandardNodes.Build.cs
+++ b/Engine/Plugins/Runtime/Metasound/Source/MetasoundStandardNodes/MetasoundStandardNodes.Build.cs
@@ -29,6 +29,8 @@ namespace UnrealBuildTool.Rules
 					"MetasoundGraphCore"
 				}
 			);
+
+			NumIncludedBytesPerUnityCPPOverride = 120 * 1024;
 		}
 	}
 }
diff --git a/Engine/Source/Programs/UnrealBuildTool/Configuration/ModuleRules.cs b/Engine/Source/Programs/UnrealBuildTool/Configuration/ModuleRules.cs
--- a/Engine/Source/Programs/UnrealBuildTool/Configuration/ModuleRules.cs
+++ b/Engine/Source/Programs/UnrealBuildTool/Configuration/ModuleRules.cs
@@ -869,6 +869,11 @@ namespace UnrealBuildTool
 		/// </summary>
 		public int MinFilesUsingPrecompiledHeaderOverride = 0;
 
+		/// <summary>
+		/// Overrides Target.NumIncludedBytesPerUnityCPP if non-zero.
+		/// </summary>
+		public int NumIncludedBytesPerUnityCPPOverride = 0;
+
 		/// <summary>
 		/// Module uses a #import so must be built locally when compiling with SN-DBS
 		/// </summary>
diff --git a/Engine/Source/Programs/UnrealBuildTool/Configuration/UEBuildModuleCPP.cs b/Engine/Source/Programs/UnrealBuildTool/Configuration/UEBuildModuleCPP.cs
--- a/Engine/Source/Programs/UnrealBuildTool/Configuration/UEBuildModuleCPP.cs
+++ b/Engine/Source/Programs/UnrealBuildTool/Configuration/UEBuildModuleCPP.cs
@@ -505,6 +505,9 @@ namespace UnrealBuildTool
 				Logger.LogDebug("Module '{ModuleName}' not using unity build mode", this.Name);
 			}
 
+			// Set up the NumIncludedBytesPerUnityCPP for this particular module
+			int NumIncludedBytesPerUnityCPP = (Rules.NumIncludedBytesPerUnityCPPOverride != 0) ? Rules.NumIncludedBytesPerUnityCPPOverride : Target.NumIncludedBytesPerUnityCPP;
+
 			// Set up the environment with which to compile the CPP files
 			CppCompileEnvironment CompileEnvironment = ModuleCompileEnvironment;
 
@@ -630,7 +633,7 @@ namespace UnrealBuildTool
 						else
 						{
 							Unity.GenerateUnityCPPs(Target, GeneratedFileItems, new List<FileItem>(), CompileEnvironment, WorkingSet, (Rules.ShortName ?? Name) + ".gen", IntermediateDirectory, Graph, SourceFileToUnityFile,
-								out List<FileItem> NormalGeneratedFiles, out List<FileItem> AdaptiveGeneratedFiles);
+								out List<FileItem> NormalGeneratedFiles, out List<FileItem> AdaptiveGeneratedFiles, NumIncludedBytesPerUnityCPP);
 							LinkInputFiles.AddRange(CompileFilesWithToolChain(Target, ToolChain, GeneratedCPPCompileEnvironment, ModuleCompileEnvironment, NormalGeneratedFiles, AdaptiveGeneratedFiles, Graph, Logger).ObjectFiles);
 						}
 					}
@@ -645,7 +648,7 @@ namespace UnrealBuildTool
 			if (bModuleUsesUnityBuild)
 			{
 				Unity.GenerateUnityCPPs(Target, CPPFiles, InputFiles.HeaderFiles, CompileEnvironment, WorkingSet, Rules.ShortName ?? Name, IntermediateDirectory, Graph, SourceFileToUnityFile, 
-					out List<FileItem> NormalFiles, out List<FileItem> AdaptiveFiles);
+					out List<FileItem> NormalFiles, out List<FileItem> AdaptiveFiles, NumIncludedBytesPerUnityCPP);
 				LinkInputFiles.AddRange(CompileFilesWithToolChain(Target, ToolChain, CompileEnvironment, ModuleCompileEnvironment, NormalFiles, AdaptiveFiles, Graph, Logger).ObjectFiles);
 			}
 			else if (SpecificFilesToCompile.Count == 0)
diff --git a/Engine/Source/Programs/UnrealBuildTool/System/Unity.cs b/Engine/Source/Programs/UnrealBuildTool/System/Unity.cs
--- a/Engine/Source/Programs/UnrealBuildTool/System/Unity.cs
+++ b/Engine/Source/Programs/UnrealBuildTool/System/Unity.cs
@@ -158,6 +158,7 @@ namespace UnrealBuildTool
 		/// <param name="SourceFileToUnityFile">Receives a mapping of source file to unity file</param>
 		/// <param name="NormalFiles">Receives the files to compile using the normal configuration.</param>
 		/// <param name="AdaptiveFiles">Receives the files to compile using the adaptive unity configuration.</param>
+		/// <param name="NumIncludedBytesPerUnityCPP">An approximate number of bytes of C++ code to target for inclusion in a single unified C++ file.</param>
 		public static void GenerateUnityCPPs(
 			ReadOnlyTargetRules Target,
 			List<FileItem> CPPFiles,
@@ -169,7 +170,8 @@ namespace UnrealBuildTool
 			IActionGraphBuilder Graph,
 			Dictionary<FileItem, FileItem> SourceFileToUnityFile,
 			out List<FileItem> NormalFiles,
-			out List<FileItem> AdaptiveFiles)
+			out List<FileItem> AdaptiveFiles,
+			int NumIncludedBytesPerUnityCPP)
 		{
 			List<FileItem> NewCPPFiles = new List<FileItem>();
 
@@ -180,7 +182,7 @@ namespace UnrealBuildTool
 			// is beneficial when dealing with PCH files. The default PCH creation limit is X unity files so if we generate < X 
 			// this could be fairly slow and we'd rather bump the limit a bit to group them all into the same unity file.
 			// Optimization only makes sense if PCH files are enabled.
-			bool bForceIntoSingleUnityFile = Target.bStressTestUnity || (TotalBytesInCPPFiles < Target.NumIncludedBytesPerUnityCPP * 2 && Target.bUsePCHFiles);
+			bool bForceIntoSingleUnityFile = Target.bStressTestUnity || (TotalBytesInCPPFiles < NumIncludedBytesPerUnityCPP * 2 && Target.bUsePCHFiles);
 
 			// Every single file in the module appears in the working set. Don't bother using adaptive unity for this module.
 			// Otherwise it would make full builds really slow.
@@ -216,7 +218,7 @@ namespace UnrealBuildTool
 				});
 
 				HashSet<FileItem> AdaptiveFileSet = new HashSet<FileItem>(AdaptiveFiles);
-				UnityFileBuilder CPPUnityFileBuilder = new UnityFileBuilder(bForceIntoSingleUnityFile ? -1 : Target.NumIncludedBytesPerUnityCPP);
+				UnityFileBuilder CPPUnityFileBuilder = new UnityFileBuilder(bForceIntoSingleUnityFile ? -1 : NumIncludedBytesPerUnityCPP);
 				foreach (FileItem CPPFile in SortedCPPFiles)
 				{
 					if (!bForceIntoSingleUnityFile && CPPFile.AbsolutePath.IndexOf(".GeneratedWrapper.", StringComparison.InvariantCultureIgnoreCase) != -1)
