UnrealBuildTool: Support automatically downgrading project if platform sdk does not support a version of visual studio
This is required to update the default visual studio version, however that change is currently on hold due to a compiler issue

#rnx

diff --git a/Engine/Source/Programs/UnrealBuildTool/Configuration/UEBuildPlatform.cs b/Engine/Source/Programs/UnrealBuildTool/Configuration/UEBuildPlatform.cs
--- a/Engine/Source/Programs/UnrealBuildTool/Configuration/UEBuildPlatform.cs
+++ b/Engine/Source/Programs/UnrealBuildTool/Configuration/UEBuildPlatform.cs
@@ -189,6 +189,14 @@ namespace UnrealBuildTool
 			return VCProjectFileFormat.Default;
 		}
 
+		/// <summary>
+		/// The version required to support Visual Studio
+		/// </summary>
+		public virtual Version GetVersionRequiredForVisualStudio(VCProjectFileFormat Format)
+		{
+			return new Version();
+		}
+
 		/// <summary>
 		/// Gets all the registered platforms
 		/// </summary>
diff --git a/Engine/Source/Programs/UnrealBuildTool/ProjectFiles/VisualStudio/VCProjectFileGenerator.cs b/Engine/Source/Programs/UnrealBuildTool/ProjectFiles/VisualStudio/VCProjectFileGenerator.cs
--- a/Engine/Source/Programs/UnrealBuildTool/ProjectFiles/VisualStudio/VCProjectFileGenerator.cs
+++ b/Engine/Source/Programs/UnrealBuildTool/ProjectFiles/VisualStudio/VCProjectFileGenerator.cs
@@ -294,19 +294,12 @@ namespace UnrealBuildTool
 				List<VisualStudioInstallation> Installations = MicrosoftPlatformSDK.FindVisualStudioInstallations(Logger).Where(x => WindowsPlatform.HasCompiler(x.Compiler, WindowsArchitecture.x64, Logger)).ToList();
 
 				// Get the corresponding project file format
+				bool Vs2019Available = Installations.Any(x => x.Compiler == WindowsCompiler.VisualStudio2019);
+				bool Vs2022Available = Installations.Any(x => x.Compiler == WindowsCompiler.VisualStudio2022);
 				VCProjectFileFormat Format = VCProjectFileFormat.VisualStudio2019;
-				foreach (VisualStudioInstallation Installation in Installations)
+				if (Vs2022Available && !Vs2019Available)
 				{
-					if (Installation.Compiler == WindowsCompiler.VisualStudio2022)
-					{
-						Format = VCProjectFileFormat.VisualStudio2022;
-						break;
-					}
-					else if (Installation.Compiler == WindowsCompiler.VisualStudio2019)
-					{
-						Format = VCProjectFileFormat.VisualStudio2019;
-						break;
-					}
+					Format = VCProjectFileFormat.VisualStudio2022;
 				}
 				Settings.ProjectFileFormat = Format;
 
@@ -326,7 +319,24 @@ namespace UnrealBuildTool
 								// Reduce the Visual Studio version to the max supported by each platform we plan to include.
 								if (Settings.ProjectFileFormat == VCProjectFileFormat.Default || ProposedFormat < Settings.ProjectFileFormat)
 								{
-									Settings.ProjectFileFormat = ProposedFormat;
+									Logger.LogInformation("Available {SupportedPlatform} SDK does not support Visual Studio 2022.", SupportedPlatform);
+									Version Version = BuildPlatform.GetVersionRequiredForVisualStudio(VCProjectFileFormat.VisualStudio2022);
+									if (Version > new Version())
+									{
+										Logger.LogInformation("Please update {SupportedPlatform} SDK to {Version} if Visual Studio 2022 support is desired.", SupportedPlatform, Version);
+									}
+									if (!Vs2019Available)
+									{
+										Logger.LogInformation("Generated solution cannot be downgraded to Visual Studio 2019 as it is not installed. Please install Visual Studio 2019 if {SupportedPlatform} SDK support is required.", SupportedPlatform);
+									}
+									else
+									{
+										Settings.ProjectFileFormat = ProposedFormat;
+										Logger.LogInformation("Downgrading generated solution to Visual Studio 2019.");
+										Logger.LogInformation("To force Visual Studio 2022 solutions to always be generated add the following to BuildConfiguration.xml:");
+										Logger.LogInformation("  <VCProjectFileGenerator>\r\n    <Version>VisualStudio2022</Version>\r\n  </VCProjectFileGenerator>");
+									}
+									Logger.LogInformation(string.Empty);
 								}
 							}
 						}
