UnrealBuildTool: Fix library path for Intel oneAPI c++ compiler, bump min version

#rnx

diff --git a/Engine/Source/Programs/UnrealBuildTool/Platform/Windows/MicrosoftPlatformSDK.Versions.cs b/Engine/Source/Programs/UnrealBuildTool/Platform/Windows/MicrosoftPlatformSDK.Versions.cs
--- a/Engine/Source/Programs/UnrealBuildTool/Platform/Windows/MicrosoftPlatformSDK.Versions.cs
+++ b/Engine/Source/Programs/UnrealBuildTool/Platform/Windows/MicrosoftPlatformSDK.Versions.cs
@@ -47,13 +47,14 @@ namespace UnrealBuildTool
 
 		/// <summary>
 		/// The default compiler version to be used, if installed. 
+		/// https://www.intel.com/content/www/us/en/developer/articles/tool/oneapi-standalone-components.html#dpcpp-cpp
 		/// </summary>
 		static readonly VersionNumberRange[] PreferredIntelOneApiVersions =
 		{
-			VersionNumberRange.Parse("2022.0.0", "2022.9999"),
+			VersionNumberRange.Parse("2022.2.0", "2022.9999"),
 		};
 
-		static readonly VersionNumber MinimumIntelOneApiVersion = new VersionNumber(2021, 0, 0);
+		static readonly VersionNumber MinimumIntelOneApiVersion = new VersionNumber(2022, 2, 0);
 
 		/// <summary>
 		/// The default set of components that should be suggested to be installed for Visual Studio 2019 or 2022.
diff --git a/Engine/Source/Programs/UnrealBuildTool/Platform/Windows/UEBuildWindows.cs b/Engine/Source/Programs/UnrealBuildTool/Platform/Windows/UEBuildWindows.cs
--- a/Engine/Source/Programs/UnrealBuildTool/Platform/Windows/UEBuildWindows.cs
+++ b/Engine/Source/Programs/UnrealBuildTool/Platform/Windows/UEBuildWindows.cs
@@ -835,7 +835,7 @@ namespace UnrealBuildTool
 
 			// Disable Simplygon support if compiling against the NULL RHI.
 			if (Target.GlobalDefinitions.Contains("USE_NULL_RHI=1"))
-			{				
+			{
 				Target.bCompileCEF3 = false;
 			}
 			
@@ -843,7 +843,10 @@ namespace UnrealBuildTool
 			// as normal.
 			if (Target.StaticAnalyzer == StaticAnalyzer.Clang)
 			{
-				Target.WindowsPlatform.Compiler = WindowsCompiler.Clang;
+				if (!Target.WindowsPlatform.Compiler.IsClang())
+				{
+					Target.WindowsPlatform.Compiler = WindowsCompiler.Clang;
+				}
 				Target.StaticAnalyzer = StaticAnalyzer.Default;
 			}
 			else if (Target.StaticAnalyzer != StaticAnalyzer.None && 
diff --git a/Engine/Source/Programs/UnrealBuildTool/Platform/Windows/VCEnvironment.cs b/Engine/Source/Programs/UnrealBuildTool/Platform/Windows/VCEnvironment.cs
--- a/Engine/Source/Programs/UnrealBuildTool/Platform/Windows/VCEnvironment.cs
+++ b/Engine/Source/Programs/UnrealBuildTool/Platform/Windows/VCEnvironment.cs
@@ -365,7 +365,7 @@ namespace UnrealBuildTool
 			if (Compiler == WindowsCompiler.Intel)
 			{
 				IncludePaths.Add(DirectoryReference.Combine(CompilerDir, "windows", "compiler", "include"));
-				LibraryPaths.Add(DirectoryReference.Combine(CompilerDir, "windows", "compiler", "lib", "intel64"));
+				LibraryPaths.Add(DirectoryReference.Combine(CompilerDir, "windows", "compiler", "lib", "intel64_win"));
 			}
 		}
 
