UnrealBuildTool: Intel compiler doesn't need to disable pch

#rnx

diff --git a/Engine/Source/Programs/UnrealBuildTool/Platform/Windows/UEBuildWindows.cs b/Engine/Source/Programs/UnrealBuildTool/Platform/Windows/UEBuildWindows.cs
--- a/Engine/Source/Programs/UnrealBuildTool/Platform/Windows/UEBuildWindows.cs
+++ b/Engine/Source/Programs/UnrealBuildTool/Platform/Windows/UEBuildWindows.cs
@@ -884,13 +884,6 @@ namespace UnrealBuildTool
 			{
 				Target.bUsePCHFiles = false;
 			}
-			
-			// @todo: Override PCH settings
-			if (Target.WindowsPlatform.Compiler == WindowsCompiler.Intel)
-			{
-				Target.bUseSharedPCHs = false;
-				Target.bUsePCHFiles = false;
-			}
 
 			// E&C support.
 			if (Target.bSupportEditAndContinue || Target.bAdaptiveUnityEnablesEditAndContinue)
