Add a static library version check for iOS/tvOS and Mac. A BuildException will be thrown if a library is not built with the minimum supported Xcode (13.4.1 or earlier for 5.1).

diff --git a/Engine/Source/Programs/Shared/EpicGames.Build/Platform/Mac/ApplePlatformSDK.cs b/Engine/Source/Programs/Shared/EpicGames.Build/Platform/Mac/ApplePlatformSDK.cs
--- a/Engine/Source/Programs/Shared/EpicGames.Build/Platform/Mac/ApplePlatformSDK.cs
+++ b/Engine/Source/Programs/Shared/EpicGames.Build/Platform/Mac/ApplePlatformSDK.cs
@@ -26,7 +26,7 @@ namespace UnrealBuildBase
 			}
 
 			// 8 bits per component, with high getting extra from high 32
-			Match Result = Regex.Match(StringValue, @"^(\d+).(\d+)(.(\d+))?(.(\d+))?$");
+			Match Result = Regex.Match(StringValue, @"^(\d+).(\d+)(.(\d+))?(.(\d+))?(.(\d+))?$");
 			if (Result.Success)
 			{
 				OutValue = UInt64.Parse(Result.Groups[1].Value) << 24 | UInt64.Parse(Result.Groups[2].Value) << 16;
diff --git a/Engine/Source/Programs/UnrealBuildTool/Configuration/UEBuildPlatform.cs b/Engine/Source/Programs/UnrealBuildTool/Configuration/UEBuildPlatform.cs
--- a/Engine/Source/Programs/UnrealBuildTool/Configuration/UEBuildPlatform.cs
+++ b/Engine/Source/Programs/UnrealBuildTool/Configuration/UEBuildPlatform.cs
@@ -582,6 +582,15 @@ namespace UnrealBuildTool
 		{
 		}
 
+		/// <summary>
+		/// Validate a UEBuildModule before it's processed
+		/// <param name="Module">The UEBuildModule that needs to be validated</param>
+		/// <param name="Target">Options for the target being built</param>
+		/// </summary>
+		public virtual void ValidateModule(UEBuildModule Module, ReadOnlyTargetRules Target)
+		{
+		}
+
 		/// <summary>
 		/// Return whether the given platform requires a monolithic build
 		/// </summary>
diff --git a/Engine/Source/Programs/UnrealBuildTool/Configuration/UEBuildTarget.cs b/Engine/Source/Programs/UnrealBuildTool/Configuration/UEBuildTarget.cs
--- a/Engine/Source/Programs/UnrealBuildTool/Configuration/UEBuildTarget.cs
+++ b/Engine/Source/Programs/UnrealBuildTool/Configuration/UEBuildTarget.cs
@@ -2713,6 +2713,13 @@ namespace UnrealBuildTool
 				}
 			}
 
+			// Allow each platform to run any needed validation on the module.
+			// ie: iOS/macOS do a static library version build check for unsupported clang version builds.
+			foreach (UEBuildModule Module in Modules.Values)
+			{
+				UEBuildPlatform.GetBuildPlatform(Platform).ValidateModule(Module, Rules);
+			}
+
 			if (!bCompileMonolithic)
 			{
 				if (Platform == UnrealTargetPlatform.Win64)
diff --git a/Engine/Source/Programs/UnrealBuildTool/Platform/IOS/IOSPlatformSDK.Versions.cs b/Engine/Source/Programs/UnrealBuildTool/Platform/IOS/IOSPlatformSDK.Versions.cs
--- a/Engine/Source/Programs/UnrealBuildTool/Platform/IOS/IOSPlatformSDK.Versions.cs
+++ b/Engine/Source/Programs/UnrealBuildTool/Platform/IOS/IOSPlatformSDK.Versions.cs
@@ -1,5 +1,6 @@
 // Copyright Epic Games, Inc. All Rights Reserved.
 
+using System;
 using Microsoft.Extensions.Logging;
 
 namespace UnrealBuildTool
@@ -17,5 +18,18 @@ namespace UnrealBuildTool
 			MinVersion = "15.0";
 			MaxVersion = null;
 		}
+
+		/// <summary>
+		/// The minimum Clang version that a static library can be built with
+		/// </summary>
+		public override Version MinimumStaticLibClangVersion
+		{
+			get
+			{
+				// return new Version("1200.00.32.29");	// UE 5.0 (Xcode 12.4)
+				// return new Version("1300.00.29.30"); // (old) UE 5.1 sets Xcode13.2.1 as minimum supported
+				return new Version("1316.0.21.2"); // UE 5.1 sets Xcode13.4 as minimum supported
+			}
+		}
 	}
 }
diff --git a/Engine/Source/Programs/UnrealBuildTool/Platform/IOS/UEBuildIOS.cs b/Engine/Source/Programs/UnrealBuildTool/Platform/IOS/UEBuildIOS.cs
--- a/Engine/Source/Programs/UnrealBuildTool/Platform/IOS/UEBuildIOS.cs
+++ b/Engine/Source/Programs/UnrealBuildTool/Platform/IOS/UEBuildIOS.cs
@@ -6,6 +6,7 @@ using System.Text;
 using System.Diagnostics;
 using System.IO;
 using System.Linq;
+using System.Text.RegularExpressions;
 using EpicGames.Core;
 using UnrealBuildBase;
 using Microsoft.Extensions.Logging;
@@ -52,6 +53,13 @@ namespace UnrealBuildTool
 		[CommandLine("-skipcrashlytics")]
 		public bool bSkipCrashlytics = false;
 
+		/// <summary>
+		/// Disables clang build verification checks on static libraries
+		/// </summary>
+		[CommandLine("-skipclangvalidation", Value = "true")]
+		[XmlConfigFile(Category = "BuildConfiguration", Name = "bSkipClangValidation")]
+		public bool bSkipClangValidation = false;
+
 		/// <summary>
 		/// Mark the build for distribution
 		/// </summary>
@@ -154,6 +162,11 @@ namespace UnrealBuildTool
 			get { return Inner.bSkipCrashlytics; }
 		}
 
+		public bool bSkipClangValidation
+		{
+			get { return Inner.bSkipClangValidation; }
+		}
+
 		public bool bForDistribution
 		{
 			get { return Inner.bForDistribution; }
@@ -821,6 +834,52 @@ namespace UnrealBuildTool
 			Target.bCheckSystemHeadersForModification = false;
 		}
 
+		public override void ValidateModule(UEBuildModule Module, ReadOnlyTargetRules Target)
+		{ 
+			if (BuildHostPlatform.Current.Platform == UnrealTargetPlatform.Mac && !Target.IOSPlatform.bSkipClangValidation)
+			{
+				IOSPlatformSDK SDK = (IOSPlatformSDK?)GetSDK() ?? new IOSPlatformSDK(Logger);
+				foreach (FileReference LibLoc in Module.PublicLibraries)
+				{
+					switch (LibLoc.GetExtension())
+					{
+						case ".a":
+							{
+								// When static lib, grep it
+								string Args = "-c \"strings ";
+								Args += LibLoc.FullName;
+								Args += " | grep -m1 -i \\(clang\"";
+								string StdOutResult = Utils.RunLocalProcessAndReturnStdOut("bash", Args);
+								if (string.IsNullOrEmpty(StdOutResult))
+								{
+									continue;
+								}
+
+								// This Regex will extract a 2-4 segment version code from string containing a 2-5 segment code 
+								// ie: if given string: "Apple clang version 14.0.0 (clang-1400.0.17.3.1)"
+								//     it'll extract: "1400.0.17.3"  (note the dropped 5th segment)
+								Match M = Regex.Match(StdOutResult, @"(\(clang-(?<ver>\d+.\d+(.(\d+))?(.(\d+))?)(.(\d+))?\))");
+								if (M.Success)
+								{
+									string LibString = M.Groups["ver"].ToString();
+									Version? LibVersion = new Version(LibString);
+									if (LibVersion != null && LibVersion > SDK.MinimumStaticLibClangVersion)
+									{
+										throw new BuildException("iOS Static Library:'{0}' is built with a version of clang newer than UE supports ({1} > {2}). \nPlease rebuild {3} with the minimum supported version of Xcode/clang.", LibLoc.GetFileName(), LibString, SDK.MinimumStaticLibClangVersion, LibLoc);
+									}
+								}
+							}
+							break;
+
+						default:
+							// For now, we don't validate any other types of libs (dylib, Framework, etc)
+							break;
+					}
+				}
+			}
+		}
+
+
 		/// <summary>
 		/// Allows the platform to override whether the architecture name should be appended to the name of binaries.
 		/// </summary>
diff --git a/Engine/Source/Programs/UnrealBuildTool/Platform/Mac/ApplePlatformSDK.Versions.cs b/Engine/Source/Programs/UnrealBuildTool/Platform/Mac/ApplePlatformSDK.Versions.cs
--- a/Engine/Source/Programs/UnrealBuildTool/Platform/Mac/ApplePlatformSDK.Versions.cs
+++ b/Engine/Source/Programs/UnrealBuildTool/Platform/Mac/ApplePlatformSDK.Versions.cs
@@ -21,7 +21,7 @@ namespace UnrealBuildTool
 		{
 			if (RuntimePlatform.IsMac)
 			{
-				MinVersion = "12.4.0";
+				MinVersion = "13.4.0";
 				MaxVersion = "14.9.9";
 			}
 			else
@@ -36,5 +36,30 @@ namespace UnrealBuildTool
 		{
 			MinVersion = MaxVersion = null;
 		}
+
+		/// <summary>
+		/// The minimum Clang version that a macOS static library can be built with
+		/// </summary>
+		public virtual Version MinimumStaticLibClangVersion
+		{
+			get
+			{
+				// return new Version("1200.00.32.29"); // UE 5.0 (Xcode 12.4)
+				// return new Version("1300.00.29.30"); // UE 5.1 (Xcode 13.2.1)
+				return new Version("1400.00.28.1"); // UE 5.1 (Xcode 14.0)
+			}
+		}
+
+		/// <summary>
+		/// The minimum macOS SDK version that a dynamic library can be built with
+		/// </summary>
+		public virtual Version MinimumDynamicLibSDKVersion
+		{
+			get
+			{
+				return new Version("12.1");		// SDK used in Xcode13.2.1
+			}
+		}
+
 	}
 }
diff --git a/Engine/Source/Programs/UnrealBuildTool/Platform/Mac/ApplePlatformSDK.cs b/Engine/Source/Programs/UnrealBuildTool/Platform/Mac/ApplePlatformSDK.cs
--- a/Engine/Source/Programs/UnrealBuildTool/Platform/Mac/ApplePlatformSDK.cs
+++ b/Engine/Source/Programs/UnrealBuildTool/Platform/Mac/ApplePlatformSDK.cs
@@ -12,7 +12,7 @@ using Microsoft.Extensions.Logging;
 ///////////////////////////////////////////////////////////////////
 // If you are looking for supported version numbers, look in the
 // ApplePlatformSDK.Versions.cs file next to this file, and
-// als Mac/IOSPlatformSDK.Versions.cs
+// als IOS/IOSPlatformSDK.Versions.cs
 ///////////////////////////////////////////////////////////////////
 
 namespace UnrealBuildTool
diff --git a/Engine/Source/Programs/UnrealBuildTool/Platform/Mac/UEBuildMac.cs b/Engine/Source/Programs/UnrealBuildTool/Platform/Mac/UEBuildMac.cs
--- a/Engine/Source/Programs/UnrealBuildTool/Platform/Mac/UEBuildMac.cs
+++ b/Engine/Source/Programs/UnrealBuildTool/Platform/Mac/UEBuildMac.cs
@@ -3,8 +3,9 @@
 
 using System;
 using System.Collections.Generic;
-using EpicGames.Core;
 using System.Linq;
+using System.Text.RegularExpressions;
+using EpicGames.Core;
 using Microsoft.Extensions.Logging;
 
 namespace UnrealBuildTool
@@ -41,6 +42,14 @@ namespace UnrealBuildTool
 		[CommandLine("-NoDSYM", Value = "false")]
 		[XmlConfigFile(Category = "BuildConfiguration", Name = "bUseDSYMFiles")]
 		public bool bUseDSYMFiles = false;
+
+		/// <summary>
+		/// Disables clang build verification checks on static libraries
+		/// </summary>
+		[CommandLine("-skipclangvalidation", Value = "true")]
+		[XmlConfigFile(Category = "BuildConfiguration", Name = "bSkipClangValidation")]
+		public bool bSkipClangValidation = false;
+
 	}
 
 	/// <summary>
@@ -83,6 +92,11 @@ namespace UnrealBuildTool
 			get { return Inner.bEnableUndefinedBehaviorSanitizer; }
 		}
 
+		public bool bSkipClangValidation
+		{
+			get { return Inner.bSkipClangValidation; }
+		}
+
 #pragma warning restore CS1591
 		#endregion
 	}
@@ -177,6 +191,53 @@ namespace UnrealBuildTool
 			Target.bUsePCHFiles = Target.bUsePCHFiles && !bCompilingMultipleArchitectures;
 		}
 
+		public override void ValidateModule(UEBuildModule Module, ReadOnlyTargetRules Target)
+		{ 
+			if (BuildHostPlatform.Current.Platform == UnrealTargetPlatform.Mac && !Target.MacPlatform.bSkipClangValidation)
+			{
+				ApplePlatformSDK SDK = (ApplePlatformSDK?)GetSDK() ?? new ApplePlatformSDK(Logger);
+				// Validate the added public libraries
+				foreach (FileReference LibLoc in Module.PublicLibraries)
+				{
+					switch (LibLoc.GetExtension())
+					{
+						case ".a":
+							{
+								// For a static lib, grep it
+								string Args = "-c \"strings ";
+								Args += LibLoc.FullName;
+								Args += " | grep -m1 -i \\(clang\"";
+								string StdOutResult = Utils.RunLocalProcessAndReturnStdOut("bash", Args);
+								if (string.IsNullOrEmpty(StdOutResult))
+								{
+									continue;
+								}
+
+								// This Regex will extract a 2-4 segment version code from string containing a 2-5 segment code 
+								// ie: if given string: "Apple clang version 14.0.0 (clang-1400.0.17.3.1)"
+								//     it'll extract: "1400.0.17.3"  (note the dropped 5th segment)
+								Match M = Regex.Match(StdOutResult, @"(\(clang-(?<ver>\d+.\d+(.(\d+))?(.(\d+))?)(.(\d+))?\))");
+								if (M.Success)
+								{
+									string LibString = M.Groups["ver"].ToString();
+									Version? LibVersion = new Version(LibString);
+									if (LibVersion != null && LibVersion > SDK.MinimumStaticLibClangVersion)
+									{
+										throw new BuildException("macOS Static Library:'{0}' is built with a version of clang newer than UE supports ({1} > {2}). \nPlease rebuild {3} with the minimum supported version of Xcode/clang.", LibLoc, LibString, SDK.MinimumStaticLibClangVersion, LibLoc);
+									}
+								}
+							}
+							break;
+
+						default:
+							// For now, we don't validate any other types of libs (dylib, Framework, etc)
+							break;
+					}
+				}
+			}
+		}
+
+
 		/// <summary>
 		/// Returns true since we can do this on Mac (with some caveats, that may necessitate this being an option)
 		/// </summary>
