Export all symbols on Android for Debug\Dev builds by default to allow call stacks symbolication

diff --git a/Engine/Source/Programs/UnrealBuildTool/Platform/Android/AndroidToolChain.cs b/Engine/Source/Programs/UnrealBuildTool/Platform/Android/AndroidToolChain.cs
--- a/Engine/Source/Programs/UnrealBuildTool/Platform/Android/AndroidToolChain.cs
+++ b/Engine/Source/Programs/UnrealBuildTool/Platform/Android/AndroidToolChain.cs
@@ -847,16 +847,16 @@ namespace UnrealBuildTool
 			// make sure the DT_SONAME field is set properly (or we can a warning toast at startup on new Android)
 			Result += " -Wl,-soname,libUnreal.so";
 
-			// exclude exported symbols from all static libs
-			//Result += " -Wl,--exclude-libs,ALL";
-
-			string VersionScriptFile = GetVersionScriptFilename(LinkEnvironment);
-			using (StreamWriter Writer = File.CreateText(VersionScriptFile))
+			if (LinkEnvironment.Configuration == CppConfiguration.Shipping)
 			{
-				// Make all symbols (except ones called from Java) hidden
-				Writer.WriteLine("{ global: Java_*; ANativeActivity_onCreate; JNI_OnLoad; local: *; };");
+				string VersionScriptFile = GetVersionScriptFilename(LinkEnvironment);
+				using (StreamWriter Writer = File.CreateText(VersionScriptFile))
+				{
+					// Make all symbols (except ones called from Java) hidden
+					Writer.WriteLine("{ global: Java_*; ANativeActivity_onCreate; JNI_OnLoad; local: *; };");
+				}
+				Result += " -Wl,--version-script=\"" + VersionScriptFile + "\"";
 			}
-			Result += " -Wl,--version-script=\"" + VersionScriptFile + "\"";
 
 			Result += " -Wl,--build-id=sha1";               // add build-id to make debugging easier
 
