UnrealBuildTool: Use correct default visual studio toolchain when clang is the compiler instead of always using 2019. Via UDN.

#rnx

diff --git a/Engine/Platforms/Hololens/Source/Programs/UnrealBuildTool/UEBuildHoloLens.cs b/Engine/Platforms/Hololens/Source/Programs/UnrealBuildTool/UEBuildHoloLens.cs
--- a/Engine/Platforms/Hololens/Source/Programs/UnrealBuildTool/UEBuildHoloLens.cs
+++ b/Engine/Platforms/Hololens/Source/Programs/UnrealBuildTool/UEBuildHoloLens.cs
@@ -215,7 +215,7 @@ namespace UnrealBuildTool
 
 		protected override VCEnvironment CreateVCEnvironment(TargetRules Target)
 		{
-			VCEnvironmentParameters Params = new VCEnvironmentParameters(Target.WindowsPlatform.Compiler, Platform, Target.WindowsPlatform.Architecture, Target.WindowsPlatform.CompilerVersion, Target.WindowsPlatform.WindowsSdkVersion, null, Target.WindowsPlatform.bUseCPPWinRT, Logger);
+			VCEnvironmentParameters Params = new VCEnvironmentParameters(Target.WindowsPlatform.Compiler, Target.WindowsPlatform.ToolChain, Platform, Target.WindowsPlatform.Architecture, Target.WindowsPlatform.CompilerVersion, Target.WindowsPlatform.WindowsSdkVersion, null, Target.WindowsPlatform.bUseCPPWinRT, Logger);
 			return new HoloLensEnvironment(Params, Logger);
 		}
 
diff --git a/Engine/Source/Programs/UnrealBuildTool/Executors/Experimental/FASTBuild.cs b/Engine/Source/Programs/UnrealBuildTool/Executors/Experimental/FASTBuild.cs
--- a/Engine/Source/Programs/UnrealBuildTool/Executors/Experimental/FASTBuild.cs
+++ b/Engine/Source/Programs/UnrealBuildTool/Executors/Experimental/FASTBuild.cs
@@ -767,7 +767,7 @@ namespace UnrealBuildTool
 				// it probably means we are building for another platform.
                 if(BuildType == FBBuildType.Windows)
                 {
-					VCEnv = VCEnvironment.Create(WindowsPlatform.GetDefaultCompiler(null, WindowsArchitecture.x64, Logger), UnrealTargetPlatform.Win64, WindowsArchitecture.x64, null, null, null, false, Logger);
+					VCEnv = VCEnvironment.Create(WindowsPlatform.GetDefaultCompiler(null, WindowsArchitecture.x64, Logger), WindowsCompiler.Default, UnrealTargetPlatform.Win64, WindowsArchitecture.x64, null, null, null, false, Logger);
 				}
             }
 			catch (Exception)
diff --git a/Engine/Source/Programs/UnrealBuildTool/Platform/Windows/UEBuildWindows.cs b/Engine/Source/Programs/UnrealBuildTool/Platform/Windows/UEBuildWindows.cs
--- a/Engine/Source/Programs/UnrealBuildTool/Platform/Windows/UEBuildWindows.cs
+++ b/Engine/Source/Programs/UnrealBuildTool/Platform/Windows/UEBuildWindows.cs
@@ -164,6 +164,11 @@ namespace UnrealBuildTool
 		[CommandLine("-Compiler=")]
 		public WindowsCompiler Compiler = WindowsCompiler.Default;
 
+		/// <summary>
+		/// Version of the toolchain to use on Windows platform when a non-msvc Compiler is in use, to locate include paths etc.
+		/// </summary>
+		public WindowsCompiler ToolChain => Compiler.IsMSVC() ? Compiler : WindowsPlatform.GetDefaultCompiler(Target.ProjectFile, Architecture, Target.Logger);
+
 		/// <summary>
 		/// Architecture of Target.
 		/// </summary>
@@ -819,7 +824,7 @@ namespace UnrealBuildTool
 		[SupportedOSPlatform("windows")]
 		protected virtual VCEnvironment CreateVCEnvironment(TargetRules Target)
 		{
-			return VCEnvironment.Create(Target.WindowsPlatform.Compiler, Platform, Target.WindowsPlatform.Architecture, Target.WindowsPlatform.CompilerVersion, Target.WindowsPlatform.WindowsSdkVersion, null, Target.WindowsPlatform.bUseCPPWinRT, Logger);
+			return VCEnvironment.Create(Target.WindowsPlatform.Compiler, Target.WindowsPlatform.ToolChain, Platform, Target.WindowsPlatform.Architecture, Target.WindowsPlatform.CompilerVersion, Target.WindowsPlatform.WindowsSdkVersion, null, Target.WindowsPlatform.bUseCPPWinRT, Logger);
 		}
 
 		/// <summary>
@@ -858,18 +863,7 @@ namespace UnrealBuildTool
 			// Set the compiler version if necessary
 			if (Target.WindowsPlatform.Compiler == WindowsCompiler.Default)
 			{
-				if (Target.StaticAnalyzer == StaticAnalyzer.PVSStudio && HasCompiler(WindowsCompiler.VisualStudio2019, Target.WindowsPlatform.Architecture, Logger))
-				{
-					Target.WindowsPlatform.Compiler = WindowsCompiler.VisualStudio2019;
-				}
-				if (Target.StaticAnalyzer == StaticAnalyzer.PVSStudio && HasCompiler(WindowsCompiler.VisualStudio2022, Target.WindowsPlatform.Architecture, Logger))
-				{
-					Target.WindowsPlatform.Compiler = WindowsCompiler.VisualStudio2022;
-				}
-				else
-				{
-					Target.WindowsPlatform.Compiler = GetDefaultCompiler(Target.ProjectFile, Target.WindowsPlatform.Architecture, Logger);
-				}
+				Target.WindowsPlatform.Compiler = GetDefaultCompiler(Target.ProjectFile, Target.WindowsPlatform.Architecture, Logger);
 			}
 
 			// Disable linking and ignore build outputs if we're using a static analyzer
diff --git a/Engine/Source/Programs/UnrealBuildTool/Platform/Windows/VCEnvironment.cs b/Engine/Source/Programs/UnrealBuildTool/Platform/Windows/VCEnvironment.cs
--- a/Engine/Source/Programs/UnrealBuildTool/Platform/Windows/VCEnvironment.cs
+++ b/Engine/Source/Programs/UnrealBuildTool/Platform/Windows/VCEnvironment.cs
@@ -374,6 +374,7 @@ namespace UnrealBuildTool
 		/// Creates an environment with the given settings
 		/// </summary>
 		/// <param name="Compiler">The compiler version to use</param>
+		/// <param name="ToolChain">The toolchain version to use, when a non-msvc compiler is used</param>
 		/// <param name="Platform">The platform to target</param>
 		/// <param name="Architecture">The Architecture to target</param>
 		/// <param name="CompilerVersion">The specific toolchain version to use</param>
@@ -383,9 +384,9 @@ namespace UnrealBuildTool
 		/// <param name="Logger">Logger for output</param>
 		/// <returns>New environment object with paths for the given settings</returns>
 		[SupportedOSPlatform("windows")]
-		public static VCEnvironment Create(WindowsCompiler Compiler, UnrealTargetPlatform Platform, WindowsArchitecture Architecture, string? CompilerVersion, string? WindowsSdkVersion, string? SuppliedSdkDirectoryForVersion, bool bUseCPPWinRT, ILogger Logger)
+		public static VCEnvironment Create(WindowsCompiler Compiler, WindowsCompiler ToolChain, UnrealTargetPlatform Platform, WindowsArchitecture Architecture, string? CompilerVersion, string? WindowsSdkVersion, string? SuppliedSdkDirectoryForVersion, bool bUseCPPWinRT, ILogger Logger)
 		{
-			return Create( new VCEnvironmentParameters(Compiler, Platform, Architecture, CompilerVersion, WindowsSdkVersion, SuppliedSdkDirectoryForVersion, bUseCPPWinRT, Logger), Logger );
+			return Create( new VCEnvironmentParameters(Compiler, ToolChain, Platform, Architecture, CompilerVersion, WindowsSdkVersion, SuppliedSdkDirectoryForVersion, bUseCPPWinRT, Logger), Logger );
 		}
 
 		/// <summary>
@@ -445,6 +446,7 @@ namespace UnrealBuildTool
 		/// Creates VC environment construction parameters with the given settings
 		/// </summary>
 		/// <param name="Compiler">The compiler version to use</param>
+		/// <param name="ToolChain">The toolchain version to use, when a non-msvc compiler is used</param>
 		/// <param name="Platform">The platform to target</param>
 		/// <param name="Architecture">The Architecture to target</param>
 		/// <param name="CompilerVersion">The specific toolchain version to use</param>
@@ -454,7 +456,7 @@ namespace UnrealBuildTool
 		/// <param name="Logger">Logger for output</param>
 		/// <returns>Creation parameters for VC environment</returns>
 		[SupportedOSPlatform("windows")]
-		public VCEnvironmentParameters (WindowsCompiler Compiler, UnrealTargetPlatform Platform, WindowsArchitecture Architecture, string? CompilerVersion, string? WindowsSdkVersion, string? SuppliedSdkDirectoryForVersion, bool bUseCPPWinRT, ILogger Logger)
+		public VCEnvironmentParameters (WindowsCompiler Compiler, WindowsCompiler ToolChain, UnrealTargetPlatform Platform, WindowsArchitecture Architecture, string? CompilerVersion, string? WindowsSdkVersion, string? SuppliedSdkDirectoryForVersion, bool bUseCPPWinRT, ILogger Logger)
 		{
 			// Get the compiler version info
 			VersionNumber? SelectedCompilerVersion;
@@ -466,20 +468,16 @@ namespace UnrealBuildTool
 			}
 
 			// Get the toolchain info
-			WindowsCompiler ToolChain;
 			VersionNumber? SelectedToolChainVersion;
 			DirectoryReference? SelectedToolChainDir;
-			if(Compiler.IsClang())
+			if (Compiler.IsClang())
 			{
-				if (WindowsPlatform.TryGetToolChainDir(WindowsCompiler.VisualStudio2019, null, Architecture, Logger, out SelectedToolChainVersion, out SelectedToolChainDir, out SelectedRedistDir))
+				if (ToolChain.IsClang() || ToolChain == WindowsCompiler.Default)
 				{
-					ToolChain = WindowsCompiler.VisualStudio2019;
+					throw new BuildException("{0} is not a valid ToolChain for Compiler {1}", WindowsPlatform.GetCompilerName(ToolChain), WindowsPlatform.GetCompilerName(Compiler));
 				}
-				else if (WindowsPlatform.TryGetToolChainDir(WindowsCompiler.VisualStudio2022, null, Architecture, Logger, out SelectedToolChainVersion, out SelectedToolChainDir, out SelectedRedistDir))
-				{
-					ToolChain = WindowsCompiler.VisualStudio2022;
-				}
-				else
+
+				if (!WindowsPlatform.TryGetToolChainDir(ToolChain, null, Architecture, Logger, out SelectedToolChainVersion, out SelectedToolChainDir, out SelectedRedistDir))
 				{
 					throw new BuildException("{0} or {1} must be installed in order to build this target.", WindowsPlatform.GetCompilerName(WindowsCompiler.VisualStudio2019), WindowsPlatform.GetCompilerName(WindowsCompiler.VisualStudio2022));
 				}
diff --git a/Engine/Source/Programs/UnrealBuildTool/Platform/Windows/VCToolChain.cs b/Engine/Source/Programs/UnrealBuildTool/Platform/Windows/VCToolChain.cs
--- a/Engine/Source/Programs/UnrealBuildTool/Platform/Windows/VCToolChain.cs
+++ b/Engine/Source/Programs/UnrealBuildTool/Platform/Windows/VCToolChain.cs
@@ -2362,7 +2362,7 @@ namespace UnrealBuildTool
 		public static string GetVCIncludePaths(UnrealTargetPlatform Platform, WindowsCompiler Compiler, string? CompilerVersion, ILogger Logger)
 		{
 			// Make sure we've got the environment variables set up for this target
-			VCEnvironment EnvVars = VCEnvironment.Create(Compiler, Platform, WindowsArchitecture.x64, CompilerVersion, null, null, false, Logger);
+			VCEnvironment EnvVars = VCEnvironment.Create(Compiler, WindowsCompiler.Default, Platform, WindowsArchitecture.x64, CompilerVersion, null, null, false, Logger);
 
 			// Also add any include paths from the INCLUDE environment variable.  MSVC is not necessarily running with an environment that
 			// matches what UBT extracted from the vcvars*.bat using SetEnvironmentVariablesFromBatchFile().  We'll use the variables we
diff --git a/Engine/Source/Programs/UnrealBuildTool/ProjectFiles/VisualStudioCode/VSCodeProjectFileGenerator.cs b/Engine/Source/Programs/UnrealBuildTool/ProjectFiles/VisualStudioCode/VSCodeProjectFileGenerator.cs
--- a/Engine/Source/Programs/UnrealBuildTool/ProjectFiles/VisualStudioCode/VSCodeProjectFileGenerator.cs
+++ b/Engine/Source/Programs/UnrealBuildTool/ProjectFiles/VisualStudioCode/VSCodeProjectFileGenerator.cs
@@ -328,7 +328,7 @@ namespace UnrealBuildTool
 			DirectoryReference? SysRootPath = null;
 			if (OperatingSystem.IsWindows())
 			{
-				VCEnvironment Environment = VCEnvironment.Create(WindowsPlatform.GetDefaultCompiler(null, Target.Rules.WindowsPlatform.Architecture, Logger), Target.Platform, Target.Rules.WindowsPlatform.Architecture, null, Target.Rules.WindowsPlatform.WindowsSdkVersion, null, Target.Rules.WindowsPlatform.bUseCPPWinRT, Logger);
+				VCEnvironment Environment = VCEnvironment.Create(WindowsPlatform.GetDefaultCompiler(null, Target.Rules.WindowsPlatform.Architecture, Logger), WindowsCompiler.Default, Target.Platform, Target.Rules.WindowsPlatform.Architecture, null, Target.Rules.WindowsPlatform.WindowsSdkVersion, null, Target.Rules.WindowsPlatform.bUseCPPWinRT, Logger);
 				CompilerPath = FileReference.FromString(Environment.CompilerPath.FullName);
 				UsingClang = false;
 			}
