UnrealBuildTool: Option to allow using clang linker when using clang compiler

diff --git a/Engine/Platforms/Hololens/Source/Programs/UnrealBuildTool/UEBuildHoloLens.cs b/Engine/Platforms/Hololens/Source/Programs/UnrealBuildTool/UEBuildHoloLens.cs
--- a/Engine/Platforms/Hololens/Source/Programs/UnrealBuildTool/UEBuildHoloLens.cs
+++ b/Engine/Platforms/Hololens/Source/Programs/UnrealBuildTool/UEBuildHoloLens.cs
@@ -215,7 +215,7 @@ namespace UnrealBuildTool
 
 		protected override VCEnvironment CreateVCEnvironment(TargetRules Target)
 		{
-			VCEnvironmentParameters Params = new VCEnvironmentParameters(Target.WindowsPlatform.Compiler, Target.WindowsPlatform.ToolChain, Platform, Target.WindowsPlatform.Architecture, Target.WindowsPlatform.CompilerVersion, Target.WindowsPlatform.WindowsSdkVersion, null, Target.WindowsPlatform.bUseCPPWinRT, Logger);
+			VCEnvironmentParameters Params = new VCEnvironmentParameters(Target.WindowsPlatform.Compiler, Target.WindowsPlatform.ToolChain, Platform, Target.WindowsPlatform.Architecture, Target.WindowsPlatform.CompilerVersion, Target.WindowsPlatform.WindowsSdkVersion, null, Target.WindowsPlatform.bUseCPPWinRT, Target.WindowsPlatform.bAllowClangLinker, Logger);
 			return new HoloLensEnvironment(Params, Logger);
 		}
 
diff --git a/Engine/Source/Programs/UnrealBuildTool/Executors/Experimental/FASTBuild.cs b/Engine/Source/Programs/UnrealBuildTool/Executors/Experimental/FASTBuild.cs
--- a/Engine/Source/Programs/UnrealBuildTool/Executors/Experimental/FASTBuild.cs
+++ b/Engine/Source/Programs/UnrealBuildTool/Executors/Experimental/FASTBuild.cs
@@ -767,7 +767,7 @@ namespace UnrealBuildTool
 				// it probably means we are building for another platform.
                 if(BuildType == FBBuildType.Windows)
                 {
-					VCEnv = VCEnvironment.Create(WindowsPlatform.GetDefaultCompiler(null, WindowsArchitecture.x64, Logger), WindowsCompiler.Default, UnrealTargetPlatform.Win64, WindowsArchitecture.x64, null, null, null, false, Logger);
+					VCEnv = VCEnvironment.Create(WindowsPlatform.GetDefaultCompiler(null, WindowsArchitecture.x64, Logger), WindowsCompiler.Default, UnrealTargetPlatform.Win64, WindowsArchitecture.x64, null, null, null, false, false, Logger);
 				}
             }
 			catch (Exception)
diff --git a/Engine/Source/Programs/UnrealBuildTool/Platform/Windows/UEBuildWindows.cs b/Engine/Source/Programs/UnrealBuildTool/Platform/Windows/UEBuildWindows.cs
--- a/Engine/Source/Programs/UnrealBuildTool/Platform/Windows/UEBuildWindows.cs
+++ b/Engine/Source/Programs/UnrealBuildTool/Platform/Windows/UEBuildWindows.cs
@@ -189,6 +189,14 @@ namespace UnrealBuildTool
 		[CommandLine("-CompilerVersion")]
 		public string? CompilerVersion = null;
 
+		/// <summary>
+		/// True if we should use the Clang linker (LLD) when we are compiling with Clang, or Intel linker (xilink\xilib) when we are compiling with Intel oneAPI, otherwise we use the MSVC linker.
+		/// </summary>
+		[ConfigFile(ConfigHierarchyType.Engine, "/Script/WindowsTargetPlatform.WindowsTargetSettings", "bAllowClangLinker")]
+		[XmlConfigFile(Category = "WindowsPlatform")]
+		[CommandLine("-ClangLinker")]
+		public bool bAllowClangLinker = false;
+		
 		/// <summary>
 		/// The specific Windows SDK version to use. This may be a specific version number (for example, "8.1", "10.0" or "10.0.10150.0"), or the string "Latest", to select the newest available version.
 		/// By default, and if it is available, we use the Windows SDK version indicated by WindowsPlatform.DefaultWindowsSdkVersion (otherwise, we use the latest version).
@@ -561,6 +569,11 @@ namespace UnrealBuildTool
 			get { return Inner.bUseCPPWinRT; }
 		}
 
+		public bool bAllowClangLinker
+		{
+			get { return Inner.bAllowClangLinker; }
+		}
+
 		public bool bEnableRayTracing
 		{
 			get { return Inner.bEnableRayTracing; }
@@ -782,17 +795,7 @@ namespace UnrealBuildTool
 	}
 
 	class WindowsPlatform : UEBuildPlatform
-	{
-		/// <summary>
-		/// True if we should use the Clang linker (LLD) when we are compiling with Clang, otherwise we use the MSVC linker
-		/// </summary>
-		public static readonly bool bAllowClangLinker = false;
-
-		/// <summary>
-		/// True if we should use the Intel linker (xilink\xilib) when we are compiling with Intel oneAPI, otherwise we use the MSVC linker
-		/// </summary>
-		public static readonly bool bAllowIntelLinker = true;
-		
+	{		
 		MicrosoftPlatformSDK SDK;
 
 		/// <summary>
@@ -824,7 +827,7 @@ namespace UnrealBuildTool
 		[SupportedOSPlatform("windows")]
 		protected virtual VCEnvironment CreateVCEnvironment(TargetRules Target)
 		{
-			return VCEnvironment.Create(Target.WindowsPlatform.Compiler, Target.WindowsPlatform.ToolChain, Platform, Target.WindowsPlatform.Architecture, Target.WindowsPlatform.CompilerVersion, Target.WindowsPlatform.WindowsSdkVersion, null, Target.WindowsPlatform.bUseCPPWinRT, Logger);
+			return VCEnvironment.Create(Target.WindowsPlatform.Compiler, Target.WindowsPlatform.ToolChain, Platform, Target.WindowsPlatform.Architecture, Target.WindowsPlatform.CompilerVersion, Target.WindowsPlatform.WindowsSdkVersion, null, Target.WindowsPlatform.bUseCPPWinRT, Target.WindowsPlatform.bAllowClangLinker, Logger);
 		}
 
 		/// <summary>
diff --git a/Engine/Source/Programs/UnrealBuildTool/Platform/Windows/VCEnvironment.cs b/Engine/Source/Programs/UnrealBuildTool/Platform/Windows/VCEnvironment.cs
--- a/Engine/Source/Programs/UnrealBuildTool/Platform/Windows/VCEnvironment.cs
+++ b/Engine/Source/Programs/UnrealBuildTool/Platform/Windows/VCEnvironment.cs
@@ -70,6 +70,11 @@ namespace UnrealBuildTool
 		/// </summary>
 		public readonly bool bUseCPPWinRT;
 
+		/// <summary>
+		/// Allow use of Clang linker
+		/// </summary>
+		public readonly bool bAllowClangLinker;
+
 		/// <summary>
 		/// The path to the linker for linking executables
 		/// </summary>
@@ -124,6 +129,7 @@ namespace UnrealBuildTool
 			this.WindowsSdkVersion = Params.WindowsSdkVersion;
 			this.RedistDir = Params.RedistDir;
 			this.bUseCPPWinRT = Params.bUseCPPWinRT;
+			this.bAllowClangLinker = Params.bAllowClangLinker;
 
 			// Get the standard VC paths
 			DirectoryReference VCToolPath = GetVCToolPath(ToolChain, ToolChainDir, Architecture);
@@ -224,14 +230,14 @@ namespace UnrealBuildTool
 		/// <summary>
 		/// Gets the path to the linker.
 		/// </summary>
-		static FileReference GetLinkerToolPath(UnrealTargetPlatform Platform, WindowsCompiler Compiler, DirectoryReference CompilerDir, DirectoryReference DefaultLinkerDir)
+		FileReference GetLinkerToolPath(UnrealTargetPlatform Platform, WindowsCompiler Compiler, DirectoryReference CompilerDir, DirectoryReference DefaultLinkerDir)
 		{
 			// Regardless of the target, if we're linking on a 64 bit machine, we want to use the 64 bit linker (it's faster than the 32 bit linker)
-			if (Compiler == WindowsCompiler.Clang && WindowsPlatform.bAllowClangLinker)
+			if (Compiler == WindowsCompiler.Clang && bAllowClangLinker)
 			{
 				return FileReference.Combine(CompilerDir, "bin", "lld-link.exe");
 			}
-			else if (Compiler == WindowsCompiler.Intel && WindowsPlatform.bAllowIntelLinker)
+			else if (Compiler == WindowsCompiler.Intel && bAllowClangLinker)
 			{
 				return FileReference.Combine(CompilerDir, "windows", "bin", "intel64", "xilink.exe");
 			}
@@ -242,15 +248,15 @@ namespace UnrealBuildTool
 		/// <summary>
 		/// Gets the path to the library linker.
 		/// </summary>
-		static FileReference GetLibraryLinkerToolPath(UnrealTargetPlatform Platform, WindowsCompiler Compiler, DirectoryReference CompilerDir, DirectoryReference DefaultLinkerDir)
+		FileReference GetLibraryLinkerToolPath(UnrealTargetPlatform Platform, WindowsCompiler Compiler, DirectoryReference CompilerDir, DirectoryReference DefaultLinkerDir)
 		{
 			// Regardless of the target, if we're linking on a 64 bit machine, we want to use the 64 bit linker (it's faster than the 32 bit linker)
-			if (Compiler == WindowsCompiler.Clang && WindowsPlatform.bAllowClangLinker)
+			if (Compiler == WindowsCompiler.Clang && bAllowClangLinker)
 			{
 				// @todo: lld-link is not currently working for building .lib
 				//return FileReference.Combine(CompilerDir, "bin", "lld-link.exe");
 			}
-			else if (Compiler == WindowsCompiler.Intel && WindowsPlatform.bAllowIntelLinker)
+			else if (Compiler == WindowsCompiler.Intel && bAllowClangLinker)
 			{
 				return FileReference.Combine(CompilerDir, "windows", "bin", "intel64", "xilib.exe");
 			}
@@ -381,12 +387,13 @@ namespace UnrealBuildTool
 		/// <param name="WindowsSdkVersion">Version of the Windows SDK to use</param>
 		/// <param name="SuppliedSdkDirectoryForVersion">If specified, this is the SDK directory to use, otherwise, attempt to look up via registry. If specified, the WindowsSdkVersion is used directly</param>
 		/// <param name="bUseCPPWinRT">Include the CPP/WinRT language projection</param>
+		/// <param name="bAllowClangLinker">Allow use of Clang linker</param>
 		/// <param name="Logger">Logger for output</param>
 		/// <returns>New environment object with paths for the given settings</returns>
 		[SupportedOSPlatform("windows")]
-		public static VCEnvironment Create(WindowsCompiler Compiler, WindowsCompiler ToolChain, UnrealTargetPlatform Platform, WindowsArchitecture Architecture, string? CompilerVersion, string? WindowsSdkVersion, string? SuppliedSdkDirectoryForVersion, bool bUseCPPWinRT, ILogger Logger)
+		public static VCEnvironment Create(WindowsCompiler Compiler, WindowsCompiler ToolChain, UnrealTargetPlatform Platform, WindowsArchitecture Architecture, string? CompilerVersion, string? WindowsSdkVersion, string? SuppliedSdkDirectoryForVersion, bool bUseCPPWinRT, bool bAllowClangLinker, ILogger Logger)
 		{
-			return Create( new VCEnvironmentParameters(Compiler, ToolChain, Platform, Architecture, CompilerVersion, WindowsSdkVersion, SuppliedSdkDirectoryForVersion, bUseCPPWinRT, Logger), Logger );
+			return Create( new VCEnvironmentParameters(Compiler, ToolChain, Platform, Architecture, CompilerVersion, WindowsSdkVersion, SuppliedSdkDirectoryForVersion, bUseCPPWinRT, bAllowClangLinker, Logger), Logger );
 		}
 
 		/// <summary>
@@ -442,6 +449,9 @@ namespace UnrealBuildTool
 		/// <summary>Include the CPP/WinRT language projection</summary>
 		public bool bUseCPPWinRT;
 
+		/// <summary>Allow use of Clang linker</summary>
+		public bool bAllowClangLinker;
+
 		/// <summary>
 		/// Creates VC environment construction parameters with the given settings
 		/// </summary>
@@ -453,10 +463,11 @@ namespace UnrealBuildTool
 		/// <param name="WindowsSdkVersion">Version of the Windows SDK to use</param>
 		/// <param name="SuppliedSdkDirectoryForVersion">If specified, this is the SDK directory to use, otherwise, attempt to look up via registry. If specified, the WindowsSdkVersion is used directly</param>
 		/// <param name="bUseCPPWinRT">Include the CPP/WinRT language projection</param>
+		/// <param name="bAllowClangLinker">Allow use of Clang linker</param>
 		/// <param name="Logger">Logger for output</param>
 		/// <returns>Creation parameters for VC environment</returns>
 		[SupportedOSPlatform("windows")]
-		public VCEnvironmentParameters (WindowsCompiler Compiler, WindowsCompiler ToolChain, UnrealTargetPlatform Platform, WindowsArchitecture Architecture, string? CompilerVersion, string? WindowsSdkVersion, string? SuppliedSdkDirectoryForVersion, bool bUseCPPWinRT, ILogger Logger)
+		public VCEnvironmentParameters (WindowsCompiler Compiler, WindowsCompiler ToolChain, UnrealTargetPlatform Platform, WindowsArchitecture Architecture, string? CompilerVersion, string? WindowsSdkVersion, string? SuppliedSdkDirectoryForVersion, bool bUseCPPWinRT, bool bAllowClangLinker, ILogger Logger)
 		{
 			// Get the compiler version info
 			VersionNumber? SelectedCompilerVersion;
@@ -523,6 +534,7 @@ namespace UnrealBuildTool
 			this.WindowsSdkVersion = SelectedWindowsSdkVersion;
 			this.RedistDir = SelectedRedistDir;
 			this.bUseCPPWinRT = bUseCPPWinRT;
+			this.bAllowClangLinker = bAllowClangLinker;
 		}
 	}
 }
diff --git a/Engine/Source/Programs/UnrealBuildTool/Platform/Windows/VCToolChain.cs b/Engine/Source/Programs/UnrealBuildTool/Platform/Windows/VCToolChain.cs
--- a/Engine/Source/Programs/UnrealBuildTool/Platform/Windows/VCToolChain.cs
+++ b/Engine/Source/Programs/UnrealBuildTool/Platform/Windows/VCToolChain.cs
@@ -873,7 +873,7 @@ namespace UnrealBuildTool
 			if (Target.WindowsPlatform.Compiler.IsClang())
 			{
 				// Enable codeview ghash for faster lld links
-				if (Target.WindowsPlatform.Compiler == WindowsCompiler.Clang && WindowsPlatform.bAllowClangLinker)
+				if (Target.WindowsPlatform.Compiler == WindowsCompiler.Clang && Target.WindowsPlatform.bAllowClangLinker)
 				{
 					Arguments.Add("-gcodeview-ghash");
 				}
@@ -1024,7 +1024,7 @@ namespace UnrealBuildTool
 
 		protected virtual void AppendLinkArguments(LinkEnvironment LinkEnvironment, List<string> Arguments)
 		{
-			if (Target.WindowsPlatform.Compiler == WindowsCompiler.Clang && WindowsPlatform.bAllowClangLinker)
+			if (Target.WindowsPlatform.Compiler == WindowsCompiler.Clang && Target.WindowsPlatform.bAllowClangLinker)
 			{
 				// @todo clang: The following static libraries aren't linking correctly with Clang:
 				//		tbbmalloc.lib, zlib_64.lib, libpng_64.lib, freetype2412MT.lib, IlmImf.lib
@@ -1065,7 +1065,7 @@ namespace UnrealBuildTool
 			if (LinkEnvironment.bCreateDebugInfo && LinkEnvironment.bUseFastPDBLinking)
 			{
 				// Allow partial PDBs for faster linking
-				if (Target.WindowsPlatform.Compiler == WindowsCompiler.Clang && WindowsPlatform.bAllowClangLinker)
+				if (Target.WindowsPlatform.Compiler == WindowsCompiler.Clang && Target.WindowsPlatform.bAllowClangLinker)
 				{
 					Arguments[Arguments.Count - 1] = "/DEBUG:GHASH";
 				}
@@ -2362,7 +2362,7 @@ namespace UnrealBuildTool
 		public static string GetVCIncludePaths(UnrealTargetPlatform Platform, WindowsCompiler Compiler, string? CompilerVersion, ILogger Logger)
 		{
 			// Make sure we've got the environment variables set up for this target
-			VCEnvironment EnvVars = VCEnvironment.Create(Compiler, WindowsCompiler.Default, Platform, WindowsArchitecture.x64, CompilerVersion, null, null, false, Logger);
+			VCEnvironment EnvVars = VCEnvironment.Create(Compiler, WindowsCompiler.Default, Platform, WindowsArchitecture.x64, CompilerVersion, null, null, false, false, Logger);
 
 			// Also add any include paths from the INCLUDE environment variable.  MSVC is not necessarily running with an environment that
 			// matches what UBT extracted from the vcvars*.bat using SetEnvironmentVariablesFromBatchFile().  We'll use the variables we
diff --git a/Engine/Source/Programs/UnrealBuildTool/ProjectFiles/VisualStudioCode/VSCodeProjectFileGenerator.cs b/Engine/Source/Programs/UnrealBuildTool/ProjectFiles/VisualStudioCode/VSCodeProjectFileGenerator.cs
--- a/Engine/Source/Programs/UnrealBuildTool/ProjectFiles/VisualStudioCode/VSCodeProjectFileGenerator.cs
+++ b/Engine/Source/Programs/UnrealBuildTool/ProjectFiles/VisualStudioCode/VSCodeProjectFileGenerator.cs
@@ -328,7 +328,7 @@ namespace UnrealBuildTool
 			DirectoryReference? SysRootPath = null;
 			if (OperatingSystem.IsWindows())
 			{
-				VCEnvironment Environment = VCEnvironment.Create(WindowsPlatform.GetDefaultCompiler(null, Target.Rules.WindowsPlatform.Architecture, Logger), WindowsCompiler.Default, Target.Platform, Target.Rules.WindowsPlatform.Architecture, null, Target.Rules.WindowsPlatform.WindowsSdkVersion, null, Target.Rules.WindowsPlatform.bUseCPPWinRT, Logger);
+				VCEnvironment Environment = VCEnvironment.Create(WindowsPlatform.GetDefaultCompiler(null, Target.Rules.WindowsPlatform.Architecture, Logger), WindowsCompiler.Default, Target.Platform, Target.Rules.WindowsPlatform.Architecture, null, Target.Rules.WindowsPlatform.WindowsSdkVersion, null, Target.Rules.WindowsPlatform.bUseCPPWinRT, Target.Rules.WindowsPlatform.bAllowClangLinker, Logger);
 				CompilerPath = FileReference.FromString(Environment.CompilerPath.FullName);
 				UsingClang = false;
 			}
