- More modernization of Xcode, focusing on Mac, IOS plist still needs a lot of work
- Premade "template" plist files for Mac, MacEditor, IOS, and added a PlistMode to choose between updating a template plist (from engine) or using a "baked" out plist that Xcode won't really mess with (this is still a work in progress, as it still needs thought on how to handle IOS premade vs Mac template, etc)
- Updating icons to using xcassets instead of .icns file
- Added post build step for UBT to write out build versions (increments each build) to a .xcconfig for Xcode to put into the .plist, but it reads previous value. Not terrible, but it still  needs work

diff --git a/Engine/Build/BatchFiles/Mac/UpdateVersionAfterBuild.sh b/Engine/Build/BatchFiles/Mac/UpdateVersionAfterBuild.sh
new file mode 100755
--- /dev/null
+++ b/Engine/Build/BatchFiles/Mac/UpdateVersionAfterBuild.sh
@@ -0,0 +1,45 @@
+#!/bin/bash
+
+# $1 is path to product directory (engine, program or project)
+# $2 is the platform we are incrementing
+
+PRODUCT_NAME=$(basename $1)
+VERSION_FILE="$1/Build/$2/$PRODUCT_NAME.PackageVersionCounter"
+
+VERSION="0.1"
+# increment version in the counter file
+if [ -f $VERSION_FILE ]; then
+	VERSION=`cat $VERSION_FILE`
+fi
+
+# convert to array with . delimiter
+VERSION_ARRAY=(${VERSION//./ })
+# increment and write out
+echo ${VERSION_ARRAY[0]}.$((VERSION_ARRAY[1]+1)) > "$VERSION_FILE"
+
+
+MAC_VERSION="0.1"
+VERSION_FILE="$1/Build/Mac/$PRODUCT_NAME.PackageVersionCounter"
+if [ -f $VERSION_FILE ]; then
+	MAC_VERSION=`cat $VERSION_FILE`
+fi
+
+IOS_VERSION="0.1"
+VERSION_FILE="$1/Build/IOS/$PRODUCT_NAME.PackageVersionCounter"
+if [ -f $VERSION_FILE ]; then
+	IOS_VERSION=`cat $VERSION_FILE`
+fi
+
+TVOS_VERSION="0.1"
+VERSION_FILE="$1/Build/$2TVOSPRODUCT_NAME.PackageVersionCounter"
+if [ -f $VERSION_FILE ]; then
+	TVOS_VERSION=`cat $VERSION_FILE`
+fi
+
+# now write out an xcconfig containing all platform build versions relative to the project
+XCCONFIG_FILE="$1/Intermediate/Build/Versions.xcconfig"
+
+echo UE_MAC_BUILD_VERSION = $MAC_VERSION > $XCCONFIG_FILE
+echo UE_IOS_BUILD_VERSION = $IOS_VERSION >> $XCCONFIG_FILE
+echo UE_TVOS_BUILD_VERSION = $TVOS_VERSION >> $XCCONFIG_FILE
+
diff --git a/Engine/Build/Mac/Resources/Assets.xcassets/AppIcon.appiconset/Contents.json b/Engine/Build/Mac/Resources/Assets.xcassets/AppIcon.appiconset/Contents.json
new file mode 100644
--- /dev/null
+++ b/Engine/Build/Mac/Resources/Assets.xcassets/AppIcon.appiconset/Contents.json
@@ -0,0 +1,66 @@
+{
+  "images" : [
+    {
+      "filename" : "icon_16x16.png",
+      "idiom" : "mac",
+      "scale" : "1x",
+      "size" : "16x16"
+    },
+    {
+      "filename" : "icon_32x32 1.png",
+      "idiom" : "mac",
+      "scale" : "2x",
+      "size" : "16x16"
+    },
+    {
+      "filename" : "icon_32x32.png",
+      "idiom" : "mac",
+      "scale" : "1x",
+      "size" : "32x32"
+    },
+    {
+      "idiom" : "mac",
+      "scale" : "2x",
+      "size" : "32x32"
+    },
+    {
+      "filename" : "icon_128x128.png",
+      "idiom" : "mac",
+      "scale" : "1x",
+      "size" : "128x128"
+    },
+    {
+      "filename" : "icon_256x256 1.png",
+      "idiom" : "mac",
+      "scale" : "2x",
+      "size" : "128x128"
+    },
+    {
+      "filename" : "icon_256x256.png",
+      "idiom" : "mac",
+      "scale" : "1x",
+      "size" : "256x256"
+    },
+    {
+      "filename" : "icon_512x512 1.png",
+      "idiom" : "mac",
+      "scale" : "2x",
+      "size" : "256x256"
+    },
+    {
+      "filename" : "icon_512x512.png",
+      "idiom" : "mac",
+      "scale" : "1x",
+      "size" : "512x512"
+    },
+    {
+      "idiom" : "mac",
+      "scale" : "2x",
+      "size" : "512x512"
+    }
+  ],
+  "info" : {
+    "author" : "xcode",
+    "version" : 1
+  }
+}
diff --git a/Engine/Build/Mac/Resources/Assets.xcassets/Contents.json b/Engine/Build/Mac/Resources/Assets.xcassets/Contents.json
new file mode 100644
--- /dev/null
+++ b/Engine/Build/Mac/Resources/Assets.xcassets/Contents.json
@@ -0,0 +1,6 @@
+{
+  "info" : {
+    "author" : "xcode",
+    "version" : 1
+  }
+}
diff --git a/Engine/Config/BaseEngine.ini b/Engine/Config/BaseEngine.ini
--- a/Engine/Config/BaseEngine.ini
+++ b/Engine/Config/BaseEngine.ini
@@ -2681,7 +2681,7 @@ bEnableAdvertisingIdentifier=True
 bUseModernXcode=false
 bUseModernCodeSigning=true
 ModernSigningTeam=
-ModernBundleIdentifier=com.YourCompany.[PROJECT_NAME]
+ModernSigningPrefix=com.YourCompany
 
 ; These are the defaults for Android settings, and they need to be in the .ini since UBT reads the .ini settings, without instantiating the class
 [/Script/AndroidRuntimeSettings.AndroidRuntimeSettings]
diff --git a/Engine/Source/Programs/UnrealBuildTool/Platform/IOS/IOSToolChain.cs b/Engine/Source/Programs/UnrealBuildTool/Platform/IOS/IOSToolChain.cs
--- a/Engine/Source/Programs/UnrealBuildTool/Platform/IOS/IOSToolChain.cs
+++ b/Engine/Source/Programs/UnrealBuildTool/Platform/IOS/IOSToolChain.cs
@@ -1211,6 +1211,12 @@ namespace UnrealBuildTool
 				}
 			}
 
+			if (bUseModernXcode)
+			{
+				// update version xcconfig file so xcode can update the .plist with our incrementing build version
+				OutputFiles.Add(UpdateVersionFile(BinaryLinkEnvironment, Executable, Graph));
+			}
+
 			// strip the debug info from the executable if needed. creates a dummy output file for the action graph to track that it's out of date.
 			if (Target.IOSPlatform.bStripSymbols || (Target.Configuration == UnrealTargetConfiguration.Shipping))
 			{
@@ -1232,13 +1238,6 @@ namespace UnrealBuildTool
 				OutputFiles.Add(StripCompleteFile);
 			}
 
-			bool bUseModernXcode = false;
-			if (OperatingSystem.IsMacOS())
-			{
-				ConfigHierarchy Ini = ConfigCache.ReadHierarchy(ConfigHierarchyType.Engine, DirectoryReference.FromFile(Target.ProjectFile), UnrealTargetPlatform.IOS);
-				Ini.TryGetValue("XcodeConfiguration", "bUseModernXcode", out bUseModernXcode);
-			}
-
 			if (!BinaryLinkEnvironment.bIsBuildingDLL && !bUseModernXcode)
 			{
 				// generate the asset catalog
@@ -1562,10 +1561,6 @@ namespace UnrealBuildTool
 				return;
 			}
 
-			ConfigHierarchy Ini = ConfigCache.ReadHierarchy(ConfigHierarchyType.Engine, GetActualProjectDirectory(Target.ProjectFile), UnrealTargetPlatform.IOS);
-			bool bUseModernXcode;
-			Ini.TryGetValue("XcodeConfiguration", "bUseModernXcode", out bUseModernXcode);
-
 			IOSProjectSettings ProjectSettings = ((IOSPlatform)UEBuildPlatform.GetBuildPlatform(Target.Platform)).ReadProjectSettings(Target.ProjectFile);
 			string? BundleID = ProjectSettings.BundleIdentifier;
 
@@ -1621,7 +1616,7 @@ namespace UnrealBuildTool
 			FileReference.Copy(Target.OutputPath, StagedExecutablePath, true);
 			string RemoteShadowDirectoryMac = Target.OutputPath.Directory.FullName;
 
-			if (bUseModernXcode)
+			if (UseModernXcode(Target.ProjectFile))
 			{
 				// generate a run-only project file for codesigning, etc
 				DirectoryReference? GeneratedProjectFile;
diff --git a/Engine/Source/Programs/UnrealBuildTool/Platform/Mac/MacToolChain.cs b/Engine/Source/Programs/UnrealBuildTool/Platform/Mac/MacToolChain.cs
--- a/Engine/Source/Programs/UnrealBuildTool/Platform/Mac/MacToolChain.cs
+++ b/Engine/Source/Programs/UnrealBuildTool/Platform/Mac/MacToolChain.cs
@@ -1,4 +1,4 @@
-﻿// Copyright Epic Games, Inc. All Rights Reserved.
+// Copyright Epic Games, Inc. All Rights Reserved.
 
 using System;
 using System.Collections;
@@ -119,6 +119,61 @@ namespace UnrealBuildTool
 		{
 		}
 
+		public static DirectoryReference FindProductDirectory(FileReference? ProjectFile, DirectoryReference BinaryDir, string? NameIfProgram)
+		{ 
+			// a project file is always used if there is one
+			if (ProjectFile != null)
+			{
+				return ProjectFile.Directory;
+			}
+			// look to see if this is a program, and hunt down where the Programs directory is
+			if (NameIfProgram != null)
+			{
+				return FindProgramDirectoryFromSource(BinaryDir, NameIfProgram, true) ?? Unreal.EngineDirectory;
+			}
+			return Unreal.EngineDirectory;
+		}
+
+		public static DirectoryReference? FindProgramDirectoryFromSource(DirectoryReference StartingDir, string ProgramName, bool bCreateIfNotFound)
+		{
+			DirectoryReference? ProgramFinder = StartingDir;
+			while (ProgramFinder != null &&
+				string.Compare(ProgramFinder.GetDirectoryName(), "Source", true) != 0 &&
+				string.Compare(ProgramFinder.GetDirectoryName(), "Intermediate", true) != 0 &&
+				string.Compare(ProgramFinder.GetDirectoryName(), "Binaries", true) != 0)
+			{
+				ProgramFinder = ProgramFinder.ParentDirectory;
+			}
+
+			// we are now at Source or similar directory, go up one more, then into Programs, and finally the "project" directory
+			if (ProgramFinder != null)
+			{
+				DirectoryReference ProgramsDir = DirectoryReference.Combine(ProgramFinder, "../Programs");
+				// if it doesn't exist, something went wrong, and we can't use this. throw an exception to catch it, for now
+				if (!DirectoryReference.Exists(ProgramsDir))
+				{
+					throw new BuildException($"Unable to find Programs directory for {ProgramName}, when starting in {StartingDir}");
+				}
+
+				ProgramFinder = DirectoryReference.Combine(ProgramsDir, ProgramName);
+				// if it exists, we have a ProductDir we can use for plists, icons, etc
+				if (!DirectoryReference.Exists(ProgramFinder))
+				{
+					if (bCreateIfNotFound)
+					{
+						DirectoryReference.CreateDirectory(ProgramFinder);
+						return ProgramFinder;
+					}
+				}
+				else
+				{
+					return ProgramFinder;
+				}
+			}
+
+			return null;
+		}
+
 		public override void SetUpGlobalEnvironment(ReadOnlyTargetRules Target)
 		{
 			base.SetUpGlobalEnvironment(Target);
@@ -714,7 +769,7 @@ namespace UnrealBuildTool
 				FixDylibDepsScript.Close();
 
 				// For non-console application, prepare a script that will create the app bundle. It'll be run by FinalizeAppBundle action
-				if (bIsBuildingAppBundle)
+				if (bIsBuildingAppBundle && !bUseModernXcode)
 				{
 					FileReference FinalizeAppBundleScriptPath = FileReference.Combine(LinkEnvironment.IntermediateDirectory!, "FinalizeAppBundle.sh");
 					StreamWriter FinalizeAppBundleScript = File.CreateText(FinalizeAppBundleScriptPath.FullName);
@@ -737,10 +792,10 @@ namespace UnrealBuildTool
                     // bundle identifier
                     // plist replacements
                     DirectoryReference? DirRef = (!string.IsNullOrEmpty(UnrealBuildTool.GetRemoteIniPath()) ? new DirectoryReference(UnrealBuildTool.GetRemoteIniPath()!) : (ProjectFile != null ? ProjectFile.Directory : null));
-                    ConfigHierarchy Ini = ConfigCache.ReadHierarchy(ConfigHierarchyType.Engine, DirRef, UnrealTargetPlatform.IOS);
+                    ConfigHierarchy IOSIni = ConfigCache.ReadHierarchy(ConfigHierarchyType.Engine, DirRef, UnrealTargetPlatform.IOS);
 
                     string BundleIdentifier;
-                    Ini.GetString("/Script/IOSRuntimeSettings.IOSRuntimeSettings", "BundleIdentifier", out BundleIdentifier!);
+                    IOSIni.GetString("/Script/IOSRuntimeSettings.IOSRuntimeSettings", "BundleIdentifier", out BundleIdentifier!);
 
                     string ProjectName = GameName;
 					FileReference? UProjectFilePath = ProjectFile;
@@ -829,7 +884,7 @@ namespace UnrealBuildTool
 					AppendMacLine(FinalizeAppBundleScript, FormatCopyCommand(TempInfoPlist, String.Format("{0}.app/Contents/Info.plist", ExeName)));
 					AppendMacLine(FinalizeAppBundleScript, "chmod 644 \"{0}.app/Contents/Info.plist\"", ExeName);
 
-					// Also copy it to where Xcode will look for it, now that Xcode 14 requireswe have one set up - it can't point into the .app because that
+					// Also copy it to where Xcode will look for it, now that Xcode 14 requires we have one set up - it can't point into the .app because that
 					// is where it will copy to, so it will error with reading and writing to the same location
 					string IntermediateDirectory = (ProjectFile == null ? Unreal.EngineDirectory : ProjectFile.Directory) + "/Intermediate/Mac";
 					string XcodeInputPListFile = IntermediateDirectory + "/" + ExeName + "-Info.plist";
@@ -1193,7 +1248,14 @@ namespace UnrealBuildTool
 			bool bIsBuildingAppBundle = !BinaryLinkEnvironment.bIsBuildingDLL && !BinaryLinkEnvironment.bIsBuildingLibrary && !BinaryLinkEnvironment.bIsBuildingConsoleApplication;
 			if (bIsBuildingAppBundle)
 			{
-				OutputFiles.Add(FinalizeAppBundle(BinaryLinkEnvironment, Executable, FixDylibOutputFile, Graph));
+				if (bUseModernXcode)
+				{
+					OutputFiles.Add(UpdateVersionFile(BinaryLinkEnvironment, FixDylibOutputFile, Graph));
+				}
+				else
+				{
+					OutputFiles.Add(FinalizeAppBundle(BinaryLinkEnvironment, Executable, FixDylibOutputFile, Graph));
+				}
 			}
 
 			// Add dsyms that we couldn't add before FixDylibDependencies action was created
diff --git a/Engine/Source/Programs/UnrealBuildTool/ProjectFiles/ProjectFileGenerator.cs b/Engine/Source/Programs/UnrealBuildTool/ProjectFiles/ProjectFileGenerator.cs
--- a/Engine/Source/Programs/UnrealBuildTool/ProjectFiles/ProjectFileGenerator.cs
+++ b/Engine/Source/Programs/UnrealBuildTool/ProjectFiles/ProjectFileGenerator.cs
@@ -1,4 +1,4 @@
-// Copyright Epic Games, Inc. All Rights Reserved.
+﻿// Copyright Epic Games, Inc. All Rights Reserved.
 
 using System;
 using System.Collections.Generic;
@@ -202,7 +202,7 @@ namespace UnrealBuildTool
 		/// <summary>
 		/// Optional list of platforms to generate projects for
 		/// </summary>
-		readonly List<UnrealTargetPlatform> ProjectPlatforms = new List<UnrealTargetPlatform>();
+		protected readonly List<UnrealTargetPlatform> ProjectPlatforms = new List<UnrealTargetPlatform>();
 
 		/// <summary>
 		/// Whether to append the list of platform names after the solution
diff --git a/Engine/Source/Programs/UnrealBuildTool/ProjectFiles/Xcode/XcodeProject.cs b/Engine/Source/Programs/UnrealBuildTool/ProjectFiles/Xcode/XcodeProject.cs
--- a/Engine/Source/Programs/UnrealBuildTool/ProjectFiles/Xcode/XcodeProject.cs
+++ b/Engine/Source/Programs/UnrealBuildTool/ProjectFiles/Xcode/XcodeProject.cs
@@ -89,18 +89,9 @@ namespace UnrealBuildTool.XcodeProjectXcconfig
 		public ProjectTarget? ProjectTarget;
 		public UnrealTargetConfiguration BuildConfig;
 
-		public bool bSupportsMac { get =>
-			XcodeProjectFileGenerator.ProjectFilePlatform.HasFlag(XcodeProjectFileGenerator.XcodeProjectFilePlatform.Mac) &&
-			(ProjectTarget == null || ProjectTarget.SupportedPlatforms.Contains(UnrealTargetPlatform.Mac));
-		}
-		public bool bSupportsIOS { get =>
-			XcodeProjectFileGenerator.ProjectFilePlatform.HasFlag(XcodeProjectFileGenerator.XcodeProjectFilePlatform.iOS) &&
-			(ProjectTarget == null || ProjectTarget.SupportedPlatforms.Contains(UnrealTargetPlatform.IOS));
-		}
-		public bool bSupportsTVOS { get =>
-			XcodeProjectFileGenerator.ProjectFilePlatform.HasFlag(XcodeProjectFileGenerator.XcodeProjectFilePlatform.tvOS) &&
-			(ProjectTarget == null || ProjectTarget.SupportedPlatforms.Contains(UnrealTargetPlatform.TVOS));
-		}
+		public bool bSupportsMac => UnrealData.bSupportsMac && (ProjectTarget == null || ProjectTarget.SupportedPlatforms.Contains(UnrealTargetPlatform.Mac));
+		public bool bSupportsIOS => UnrealData.bSupportsIOS && (ProjectTarget == null || ProjectTarget.SupportedPlatforms.Contains(UnrealTargetPlatform.IOS));
+		public bool bSupportsTVOS => UnrealData.bSupportsTVOS && (ProjectTarget == null || ProjectTarget.SupportedPlatforms.Contains(UnrealTargetPlatform.TVOS));
 	};
 
 	class UnrealBatchedFiles
@@ -149,6 +140,12 @@ namespace UnrealBuildTool.XcodeProjectXcconfig
 		}
 	}
 
+	enum PlistMode
+	{
+		UsePremade,
+		UpdateTemplate,
+	}
+
 	class UnrealData
 	{
 		public bool bIsAppBundle;
@@ -157,6 +154,9 @@ namespace UnrealBuildTool.XcodeProjectXcconfig
 		public bool bIsMergingProjects = false;
 		public bool bWriteCodeSigningSettings = true;
 
+		public PlistMode PlistMode;
+
+
 		public List<UnrealBuildConfig> AllConfigs = new();
 
 		public List<UnrealExtensionInfo> AllExtensions = new();
@@ -165,7 +165,8 @@ namespace UnrealBuildTool.XcodeProjectXcconfig
  
 		public FileReference? UProjectFileLocation = null;
 		public DirectoryReference XcodeProjectFileLocation;
-		
+		public DirectoryReference ProductDir;
+
 		public FileReference? InfoPlistLocation = null;
 		public bool bInfoPlistWasPremade = false;
 		public FileReference? EntitlementsLocation = null;
@@ -178,8 +179,12 @@ namespace UnrealBuildTool.XcodeProjectXcconfig
 		public TVOSProjectSettings? TVOSProjectSettings;
 		public TVOSProvisioningData? TVOSProvisioningData;
 
+		// Name of the product (usually the project name, but UE5.xcodeproj is actually UnrealGame product)
 		public string ProductName;
 
+		// Display name, can be overridden from commandline
+		public string DisplayName;
+
 
 		/// <summary>
 		///  Used to mark the project for distribution (some platforms require this)
@@ -208,11 +213,26 @@ namespace UnrealBuildTool.XcodeProjectXcconfig
 
 		private XcodeProjectFile? ProjectFile;
 
-		public UnrealData(FileReference XcodeProjectFileLocation)
+		public static bool bSupportsMac => XcodeProjectFileGenerator.XcodePlatforms.Contains(UnrealTargetPlatform.Mac);
+		public static bool bSupportsIOS => XcodeProjectFileGenerator.XcodePlatforms.Contains(UnrealTargetPlatform.IOS);
+		public static bool bSupportsTVOS => XcodeProjectFileGenerator.XcodePlatforms.Contains(UnrealTargetPlatform.TVOS);
+
+
+		public UnrealData(FileReference XcodeProjectFileLocation, bool bIsForDistribution, string BundleID, string AppName)
 		{
 			// the .xcodeproj is actually a directory
 			this.XcodeProjectFileLocation = new DirectoryReference(XcodeProjectFileLocation.FullName);
+			// default to engine director, will be fixed in Initialize if needed
+			ProductDir = Unreal.EngineDirectory;
 			ProductName = XcodeProjectFileLocation.GetFileNameWithoutAnyExtensions();
+			if (ProductName == "UE5")
+			{
+				ProductName = "UnrealGame";
+			}
+
+			this.bForDistribution = bIsForDistribution;
+			this.BundleIdentifier = string.IsNullOrEmpty(BundleID) ? "$(UE_SIGNING_PREFIX).$(UE_PRODUCT_NAME_STRIPPED)" : BundleID;
+			this.DisplayName = string.IsNullOrEmpty(AppName) ? "$(UE_PRODUCT_NAME)" : BundleID; ;
 		}
 
 		public FileReference? FindUProjectFileLocation(XcodeProjectFile ProjectFile)
@@ -233,44 +253,71 @@ namespace UnrealBuildTool.XcodeProjectXcconfig
 			return UProjectFileLocation;
 		}
 
-		public bool Initialize(XcodeProjectFile ProjectFile, List<UnrealTargetPlatform> Platforms, List<UnrealTargetConfiguration> Configurations, ILogger Logger)
+		public bool Initialize(XcodeProjectFile ProjectFile, List<UnrealTargetConfiguration> Configurations, ILogger Logger)
 		{
 			this.ProjectFile = ProjectFile;
 			this.Logger = Logger;
 
 			FindUProjectFileLocation(ProjectFile);
 
+			// make sure ProjectDir is something good
+			if (UProjectFileLocation != null)
+			{
+				ProductDir = UProjectFileLocation.Directory;
+			}
+			else if (ProjectFile.ProjectTargets[0].TargetRules!.Type == TargetType.Program)
+			{
+				DirectoryReference? ProgramFinder = DirectoryReference.Combine(ProjectFile.BaseDir);
+				while (ProgramFinder != null && string.Compare(ProgramFinder.GetDirectoryName(), "Source", true) != 0)
+				{
+					ProgramFinder = ProgramFinder.ParentDirectory;
+				}
+				// we are now at Source directory, go up one more, then into Programs, and finally the "project" directory
+				if (ProgramFinder != null)
+				{
+					ProgramFinder = DirectoryReference.Combine(ProgramFinder, "../Programs", ProductName);
+					// if it exists, we have a ProductDir we can use for plists, icons, etc
+					if (DirectoryReference.Exists(ProgramFinder))
+					{
+						ProductDir = ProgramFinder;
+					}
+				}
+			}
+
 			// Figure out all the desired configurations on the unreal side
-			AllConfigs = GetSupportedBuildConfigs(Platforms, Configurations, Logger);
+			AllConfigs = GetSupportedBuildConfigs(XcodeProjectFileGenerator.XcodePlatforms, Configurations, Logger);
 			// if we can't find any configs, we will fail to create a project
 			if (AllConfigs.Count == 0)
 			{
 				return false;
 			}
 
+			// read setings from the configs, now that we have a project
+			ConfigHierarchy SharedPlatformIni = ConfigCache.ReadHierarchy(ConfigHierarchyType.Engine, UProjectFileLocation?.Directory, UnrealTargetPlatform.Mac);
+			SharedPlatformIni.TryGetValue("XcodeConfiguration", "bUseModernCodeSigning", out bUseAutomaticSigning);
+
+			// @todo read from ini
+			PlistMode = PlistMode.UpdateTemplate;
 
 			// this project makes an app bundle (.app directory instead of a raw executable or dylib) if none of the fings make a non-appbundle
 			bIsAppBundle = !AllConfigs.Any(x => x.ProjectTarget!.TargetRules!.bIsBuildingConsoleApplication || x.ProjectTarget.TargetRules.bShouldCompileAsDLL);
 			bHasEditorConfiguration = AllConfigs.Any(x => x.ProjectTarget!.TargetRules!.Type == TargetType.Editor);
 
 			// read config settings
-			if (InstalledPlatformInfo.IsValidPlatform(UnrealTargetPlatform.IOS, EProjectType.Code))
+			if (AllConfigs.Any(x => x.bSupportsIOS))
 			{
 				IOSPlatform IOSPlatform = ((IOSPlatform)UEBuildPlatform.GetBuildPlatform(UnrealTargetPlatform.IOS));
 				IOSProjectSettings = IOSPlatform.ReadProjectSettings(UProjectFileLocation);
-				bUseAutomaticSigning |= IOSProjectSettings.bAutomaticSigning;
 				if (!bUseAutomaticSigning)
 				{
 					IOSProvisioningData = IOSPlatform.ReadProvisioningData(IOSProjectSettings, bForDistribution);
 				}
 			}
 
-
-			if (InstalledPlatformInfo.IsValidPlatform(UnrealTargetPlatform.IOS, EProjectType.Code))
+			if (AllConfigs.Any(x => x.bSupportsTVOS))
 			{
 				TVOSPlatform TVOSPlatform = ((TVOSPlatform)UEBuildPlatform.GetBuildPlatform(UnrealTargetPlatform.TVOS));
 				TVOSProjectSettings = TVOSPlatform.ReadProjectSettings(UProjectFileLocation);
-				bUseAutomaticSigning |= TVOSProjectSettings.bAutomaticSigning;
 				if (!bUseAutomaticSigning)
 				{
 					TVOSProvisioningData = TVOSPlatform.ReadProvisioningData(TVOSProjectSettings, bForDistribution);
@@ -280,6 +327,30 @@ namespace UnrealBuildTool.XcodeProjectXcconfig
 			return true;
 		}
 
+		public string ProjectOrEnginePath(string SubPath, bool bMakeRelative)
+		{
+			string? FinalPath = null;
+			if (UProjectFileLocation != null)
+			{
+				string PathToCheck = Path.Combine(ProductDir.FullName, SubPath);
+				if (File.Exists(PathToCheck) || Directory.Exists(PathToCheck))
+				{
+					FinalPath = PathToCheck;
+				}
+			}
+
+			if (FinalPath == null)
+			{
+				FinalPath = Path.Combine(Unreal.EngineDirectory.FullName, SubPath);
+			}
+			if (bMakeRelative)
+			{
+				FinalPath = new FileReference(FinalPath).MakeRelativeTo(XcodeProjectFileLocation.ParentDirectory!);
+			}
+
+			return FinalPath;
+		}
+
 		public void AddModule(UEBuildModuleCPP Module, CppCompileEnvironment CompileEnvironment)
 		{
 			// we need to keep all _API defines, but we can gather them up to append at the end, instead of using them to compute sameness
@@ -480,28 +551,6 @@ namespace UnrealBuildTool.XcodeProjectXcconfig
 			return BuildConfigs;
 		}
 
-		public static IEnumerable<UnrealTargetPlatform> GetSupportedPlatforms()
-		{
-			List<UnrealTargetPlatform> SupportedPlatforms = new List<UnrealTargetPlatform>();
-
-			if (XcodeProjectFileGenerator.ProjectFilePlatform.HasFlag(XcodeProjectFileGenerator.XcodeProjectFilePlatform.Mac))
-			{
-				SupportedPlatforms.Add(UnrealTargetPlatform.Mac);
-			}
-
-			if (XcodeProjectFileGenerator.ProjectFilePlatform.HasFlag(XcodeProjectFileGenerator.XcodeProjectFilePlatform.iOS))
-			{
-				SupportedPlatforms.Add(UnrealTargetPlatform.IOS);
-			}
-
-			if (XcodeProjectFileGenerator.ProjectFilePlatform.HasFlag(XcodeProjectFileGenerator.XcodeProjectFilePlatform.tvOS))
-			{
-				SupportedPlatforms.Add(UnrealTargetPlatform.TVOS);
-			}
-
-			return SupportedPlatforms;
-		}
-
 		public static IEnumerable<UnrealTargetConfiguration> GetSupportedConfigurations()
 		{
 			return new UnrealTargetConfiguration[] {
@@ -522,7 +571,7 @@ namespace UnrealBuildTool.XcodeProjectXcconfig
 		{
 			foreach (Project ProjectTarget in Proj.ProjectTargets)
 			{
-				foreach (UnrealTargetPlatform Platform in GetSupportedPlatforms())
+				foreach (UnrealTargetPlatform Platform in XcodeProjectFileGenerator.XcodePlatforms)
 				{
 					foreach (UnrealTargetConfiguration Config in GetSupportedConfigurations())
 					{
@@ -992,6 +1041,10 @@ namespace UnrealBuildTool.XcodeProjectXcconfig
 
 		public List<XcodeBuildConfig> BuildConfigs = new();
 
+		public bool bSupportsMac => BuildConfigs.Any(x => x.Info.bSupportsMac);
+		public bool bSupportsIOS => BuildConfigs.Any(x => x.Info.bSupportsIOS);
+		public bool bSupportsTVOS => BuildConfigs.Any(x => x.Info.bSupportsTVOS);
+
 		public XcodeBuildConfigList(string TargetName, List<UnrealBuildConfig> BuildConfigInfos)
 		{
 			if (BuildConfigInfos.Count == 0)
@@ -1121,7 +1174,7 @@ namespace UnrealBuildTool.XcodeProjectXcconfig
 				string BuildToolPath = UEDir + "/Engine/Build/BatchFiles/Mac/XcodeBuild.sh";
 
 				// insert elements to call UBT when building
-				Content.WriteLine(3,	$"buildArgumentsString = \"$(ACTION) $(UE_BUILD_TARGET_NAME) $(PLATFORM_NAME) $(UE_BUILD_TARGET_CONFIG) {UProjectParam}\";");
+				Content.WriteLine(3,	$"buildArgumentsString = \"$(ACTION) $(UE_TARGET_NAME) $(PLATFORM_NAME) $(UE_TARGET_CONFIG) {UProjectParam}\";");
 				Content.WriteLine(3,	$"buildToolPath = \"{BuildToolPath}\";");
 				Content.WriteLine(3,	$"buildWorkingDirectory = \"{UEDir}\";");
 			}
@@ -1170,21 +1223,27 @@ namespace UnrealBuildTool.XcodeProjectXcconfig
 			UnrealData = Project.UnrealData;
 			Project.FileCollection.AddFileReference(ProductGuid, UnrealData.ProductName, Project.UnrealData.bIsAppBundle ? "wrapper.application" : "\"compiled.mach-o.executable\"", "BUILT_PRODUCTS_DIR", "Products");
 
-			// look for Assets
-			DirectoryReference Assets = DirectoryReference.Combine(Project.FileCollection.RootDirectory, "Build/IOS/Resources/Assets.xcassets");
-			if (!DirectoryReference.Exists(Assets))
+			if (Project.UnrealData.bIsAppBundle)
 			{
-				Assets = DirectoryReference.Combine(Unreal.EngineDirectory, "Build/IOS/Resources/Assets.xcassets");
-			}
-			XcodeResourcesBuildPhase ResourcesBuildPhase = new XcodeResourcesBuildPhase(Project.FileCollection);
-			BuildPhases.Add(ResourcesBuildPhase);
-			References.Add(ResourcesBuildPhase);
-
-			ResourcesBuildPhase.AddResource(new FileReference(Assets.FullName));
+				XcodeResourcesBuildPhase ResourcesBuildPhase = new XcodeResourcesBuildPhase(Project.FileCollection);
+				BuildPhases.Add(ResourcesBuildPhase);
+				References.Add(ResourcesBuildPhase);
 
+				// look for Assets
+				if (UnrealData.bSupportsMac)
+				{
+					string AssetsPath = UnrealData.ProjectOrEnginePath("Build/Mac/Resources/Assets.xcassets", false);
+					ResourcesBuildPhase.AddResource(new FileReference(AssetsPath));
+				}
+				if (UnrealData.bSupportsIOS || UnrealData.bSupportsTVOS)
+				{
+					string AssetsPath = UnrealData.ProjectOrEnginePath("Build/IOS/Resources/Assets.xcassets", false);
+					ResourcesBuildPhase.AddResource(new FileReference(AssetsPath));
+				}
+			}
 
 			// if we are run-only, then skip some stuff
-			if (!XcodeProjectFileGenerator.bGeneratingRunIOSProject)
+			if (!XcodeProjectFileGenerator.bGenerateRunOnlyProject)
 			{
 				// create a biuld target only if we have source files to build
 				if (Project.UnrealData.BatchedFiles.Count != 0)
@@ -1208,8 +1267,10 @@ namespace UnrealBuildTool.XcodeProjectXcconfig
 
 		public override void WriteXcconfigFile()
 		{
-			// allow Xcode to generate the final plist file from our input, some INFOPLIST settings and other settings 
-			Xcconfig!.AppendLine("GENERATE_INFOPLIST_FILE = YES");
+			// pull in the stuff that UBT updates
+			DirectoryReference ProjectOrEngineDir = UnrealData.UProjectFileLocation?.Directory ?? Unreal.EngineDirectory;
+			Xcconfig!.AppendLine($"#include? \"{ProjectOrEngineDir}/Intermediate/Build/Versions.xcconfig\"");
+			Xcconfig.AppendLine("");
 
 			// #jira UE-143619: Pre Monterey macOS requires this option for a packaged app to run on iOS15 due to new code signature format. Could be removed once Monterey is miniuS.
 			Xcconfig.AppendLine("OTHER_CODE_SIGN_FLAGS = --generate-entitlement-der");
@@ -1217,10 +1278,13 @@ namespace UnrealBuildTool.XcodeProjectXcconfig
 			Xcconfig.AppendLine($"MACOSX_DEPLOYMENT_TARGET = {MacToolChain.Settings.MacOSVersion};");
 			Xcconfig.AppendLine("INFOPLIST_OUTPUT_FORMAT = xml");
 			Xcconfig.AppendLine("COMBINE_HIDPI_IMAGES = YES");
-			Xcconfig.AppendLine($"PRODUCT_BUNDLE_IDENTIFIER = {UnrealData.BundleIdentifier}");
-			Xcconfig.AppendLine($"ASSETCATALOG_COMPILER_APPICON_NAME = AppIcon");
-			
 
+			// allow Xcode to generate the final plist file from our input, some INFOPLIST settings and other settings 
+			if (UnrealData.PlistMode == PlistMode.UpdateTemplate)
+			{
+				Xcconfig.AppendLine("GENERATE_INFOPLIST_FILE = YES");
+				Xcconfig.AppendLine($"ASSETCATALOG_COMPILER_APPICON_NAME = AppIcon");
+			}
 
 			Xcconfig.Write();
 		}
@@ -1322,7 +1386,7 @@ namespace UnrealBuildTool.XcodeProjectXcconfig
 
 
 			// make an indexing target if we aren't just a run-only project, and it has buildable source files
-			if (!XcodeProjectFileGenerator.bGeneratingRunIOSProject && UnrealData.BatchedFiles.Count != 0)
+			if (!XcodeProjectFileGenerator.bGenerateRunOnlyProject && UnrealData.BatchedFiles.Count != 0)
 			{				
 				// index isn't a dependency of run, it's simply a target that xcode will find to index from
 				XcodeIndexTarget IndexTarget = new XcodeIndexTarget(this);
@@ -1330,7 +1394,6 @@ namespace UnrealBuildTool.XcodeProjectXcconfig
 			}
 		}
 
-
 		public override void Write(StringBuilder Content)
 		{
 			Content.WriteLine("/* Begin PBXProject section */");
@@ -1405,6 +1468,18 @@ namespace UnrealBuildTool.XcodeProjectXcconfig
 			return Orientations;
 		}
 
+		string FindBundleID(string RelativePlistPath)
+		{
+			string FullPath = Path.Combine(UnrealData.XcodeProjectFileLocation.ParentDirectory!.FullName, RelativePlistPath);
+
+			string Identifier = XcodeProjectFile.Plist("Print :CFBundleIdentifier", FullPath);
+			// handle error
+			if (Identifier.StartsWith("Print:"))
+			{
+				return "";
+			}
+			return Identifier;
+		}
 
 		public override void WriteXcconfigFile()
 		{
@@ -1420,17 +1495,25 @@ namespace UnrealBuildTool.XcodeProjectXcconfig
 			// read some settings from new config settings
 			bool bAutomaticSigning = true;
 			string? SigningTeam;
+			string? SigningPrefix;
 			SharedPlatformIni.TryGetValue("XcodeConfiguration", "bUseModernCodeSigning", out bAutomaticSigning);
 			SharedPlatformIni.TryGetValue("XcodeConfiguration", "ModernSigningTeam", out SigningTeam);
+			SharedPlatformIni.TryGetValue("XcodeConfiguration", "ModernSigningPrefix", out SigningPrefix);
 
 			List<string> SupportedDevices = new();
 
 			StringBuilder Content = new();
-			string ProjectName = UnrealData.ProductName;
 
 			// figure out the directory that contains the Binaries and Intermediate directories
 			string ProjectRootDir = (UnrealData.UProjectFileLocation == null) ? $"{UEDir}/Engine" : UnrealData.UProjectFileLocation.Directory.FullName;
 
+			Xcconfig.AppendLine("// Unreal projcet-wide variables");
+			Xcconfig.AppendLine($"UE_PRODUCT_NAME = {UnrealData.ProductName}");
+			Xcconfig.AppendLine($"UE_PRODUCT_NAME_STRIPPED = {UnrealData.ProductName.Replace("_", "").Replace(" ", "")}");
+			Xcconfig.AppendLine($"UE_DISPLAY_NAME = {UnrealData.DisplayName}");
+			Xcconfig.AppendLine($"UE_SIGNING_PREFIX = {SigningPrefix}");
+			Xcconfig.AppendLine("");
+
 
 			// Info.plist generation settings
 
@@ -1438,54 +1521,59 @@ namespace UnrealBuildTool.XcodeProjectXcconfig
 			// IOS
 			////////////////////////
 
-			if (UnrealData.IOSProjectSettings != null)
+			if (ProjectBuildConfigs.bSupportsIOS)
 			{
+				string SettingFilter = XcodeProjectFileGenerator.XcodePlatforms.Count == 1 ? "" : "[sdk=iphoneos*]";
 				ConfigHierarchy Ini = ConfigCache.ReadHierarchy(ConfigHierarchyType.Engine, UnrealData.UProjectFileLocation?.Directory, UnrealTargetPlatform.IOS);
 
 				Xcconfig.AppendLine("// IOS Settings");
 
-				SupportedDevices.Add(UnrealData.IOSProjectSettings.RuntimeDevices);
+				SupportedDevices.Add(UnrealData.IOSProjectSettings!.RuntimeDevices);
 
-				Xcconfig.AppendLine("SDKROOT[sdk=iphoneos*] = iphoneos");
-				Xcconfig.AppendLine($"CONFIGURATION_BUILD_DIR[sdk=iphoneos*] = {ProjectRootDir}/Binaries/IOS/Payload");
+				Xcconfig.AppendLine($"SDKROOT{SettingFilter} = iphoneos");
+				Xcconfig.AppendLine($"CONFIGURATION_BUILD_DIR{SettingFilter} = {ProjectRootDir}/Binaries/IOS/Payload");
 				Xcconfig.AppendLine($"IPHONEOS_DEPLOYMENT_TARGET = {UnrealData.IOSProjectSettings.RuntimeVersion}");
-				Xcconfig.AppendLine($"CURRENT_PROJECT_VERSION[sdk=iphoneos*] = {IOSExports.GetAndUpdateVersionFile(UnrealData.UProjectFileLocation, UnrealTargetPlatform.IOS)}");
 
-				// short version string
-				string BundleShortVersion;
-				Ini.GetString("/Script/IOSRuntimeSettings.IOSRuntimeSettings", "VersionInfo", out BundleShortVersion);
-				Xcconfig.AppendLine($"MARKETING_VERSION[sdk=iphoneos*] = {BundleShortVersion}");
-				
+				if (UnrealData.PlistMode == PlistMode.UpdateTemplate)
+				{
+					Xcconfig.AppendLine($"CURRENT_PROJECT_VERSION{SettingFilter} = $(UE_IOS_BUILD_VERSION)");
+
+					// short version string
+					string BundleShortVersion;
+					Ini.GetString("/Script/IOSRuntimeSettings.IOSRuntimeSettings", "VersionInfo", out BundleShortVersion);
+					Xcconfig.AppendLine($"MARKETING_VERSION{SettingFilter} = {BundleShortVersion}");
+				}
+
 
 				if (UnrealData.IOSProvisioningData != null && UnrealData.bWriteCodeSigningSettings)
 				{
 					if (!string.IsNullOrEmpty(UnrealData.IOSProvisioningData.MobileProvisionUUID))
 					{
-						Xcconfig.AppendLine($"PROVISIONING_PROFILE_SPECIFIER[sdk=iphoneos*] = {UnrealData.IOSProvisioningData.MobileProvisionUUID}");
+						Xcconfig.AppendLine($"PROVISIONING_PROFILE_SPECIFIER{SettingFilter} = {UnrealData.IOSProvisioningData.MobileProvisionUUID}");
 					}
 					string IOSCert = UnrealData.IOSProvisioningData.SigningCertificate ?? (UnrealData.bForDistribution ? "Apple Distribution" : "Apple Developer");
-					Xcconfig.AppendLine($"CODE_SIGN_IDENTITY[sdk=iphoneos*] = {IOSCert}");
+					Xcconfig.AppendLine($"CODE_SIGN_IDENTITY{SettingFilter} = {IOSCert}");
 				}
 
 
 				// Info plist generation settings
-				Xcconfig.AppendLine("INFOPLIST_KEY_UIStatusBarHidden = YES");
-				Xcconfig.AppendLine("INFOPLIST_KEY_UIRequiresFullScreen = YES");
+				//Xcconfig.AppendLine("INFOPLIST_KEY_UIStatusBarHidden = YES");
+				//Xcconfig.AppendLine("INFOPLIST_KEY_UIRequiresFullScreen = YES");
 
 				List<string> SupportedOrientations = GetSupportedOrientations(Ini);
 				Content.WriteLine(3, $"INFOPLIST_KEY_UISupportedInterfaceOrientations = \"{string.Join(" ", SupportedOrientations)}\"");
 
-				if (UnrealData.IOSProjectSettings.BundleName != "" && UnrealData.IOSProjectSettings.BundleName != "[PROJECT_NAME]")
-				{
-					string DisplayName = UnrealData.IOSProjectSettings.BundleDisplayName.Replace("[PROJECT_NAME]", ProjectName).Replace("_", "");
-					Content.WriteLine(3, $"INFOPLIST_KEY_CFBundleName = {DisplayName}");
-				}
+				//if (UnrealData.IOSProjectSettings.BundleName != "" && UnrealData.IOSProjectSettings.BundleName != "[PROJECT_NAME]")
+				//{
+				//	string DisplayName = UnrealData.IOSProjectSettings.BundleDisplayName.Replace("[PROJECT_NAME]", ProjectName).Replace("_", "");
+				//	Content.WriteLine(3, $"INFOPLIST_KEY_CFBundleName = {DisplayName}");
+				//}
 
-				if (UnrealData.IOSProjectSettings.BundleDisplayName != "" && UnrealData.IOSProjectSettings.BundleDisplayName != "[PROJECT_NAME]")
-				{
-					string DisplayName = UnrealData.IOSProjectSettings.BundleDisplayName.Replace("[PROJECT_NAME]", ProjectName).Replace("_", "");
-					Content.WriteLine(3, $"INFOPLIST_KEY_CFBundleDisplayName = {DisplayName}");
-				}
+				//if (UnrealData.IOSProjectSettings.BundleDisplayName != "" && UnrealData.IOSProjectSettings.BundleDisplayName != "[PROJECT_NAME]")
+				//{
+				//	string DisplayName = UnrealData.IOSProjectSettings.BundleDisplayName.Replace("[PROJECT_NAME]", ProjectName).Replace("_", "");
+				//	Content.WriteLine(3, $"INFOPLIST_KEY_CFBundleDisplayName = {DisplayName}");
+				//}
 
 			}
 
@@ -1493,49 +1581,63 @@ namespace UnrealBuildTool.XcodeProjectXcconfig
 			// TVOS
 			////////////////////////
 
-			if (UnrealData.TVOSProjectSettings != null)
+			if (ProjectBuildConfigs.bSupportsTVOS)
 			{
+				string SettingFilter = XcodeProjectFileGenerator.XcodePlatforms.Count == 1 ? "" : "[sdk=appletvos*]";
 				ConfigHierarchy Ini = ConfigCache.ReadHierarchy(ConfigHierarchyType.Engine, UnrealData.UProjectFileLocation?.Directory, UnrealTargetPlatform.TVOS);
 
 				Xcconfig.AppendLine("");
 				Xcconfig.AppendLine("// TVOS Settings");
 
-				SupportedDevices.Add(UnrealData.TVOSProjectSettings.RuntimeDevices);
+				SupportedDevices.Add(UnrealData.TVOSProjectSettings!.RuntimeDevices);
 
-				Xcconfig.AppendLine("SDKROOT[sdk=appletvos*] = appletvos");
-				Xcconfig.AppendLine($"CONFIGURATION_BUILD_DIR[sdk=appletvos*] = {ProjectRootDir}/Binaries/TVOS/Payload");
+				Xcconfig.AppendLine($"SDKROOT{SettingFilter} = appletvos");
+				Xcconfig.AppendLine($"CONFIGURATION_BUILD_DIR{SettingFilter} = {ProjectRootDir}/Binaries/TVOS/Payload");
 				Xcconfig.AppendLine($"TVOS_DEPLOYMENT_TARGET = {UnrealData.TVOSProjectSettings.RuntimeVersion}");
-				Xcconfig.AppendLine($"CURRENT_PROJECT_VERSION[sdk=appletvos*] = {IOSExports.GetAndUpdateVersionFile(UnrealData.UProjectFileLocation, UnrealTargetPlatform.TVOS)}");
 
-				// short version string
-				string BundleShortVersion;
-				Ini.GetString("/Script/IOSRuntimeSettings.IOSRuntimeSettings", "VersionInfo", out BundleShortVersion);
-				Xcconfig.AppendLine($"MARKETING_VERSION[sdk=appletvos*] = {BundleShortVersion}");
+				if (UnrealData.PlistMode == PlistMode.UpdateTemplate)
+				{
+					Xcconfig.AppendLine($"CURRENT_PROJECT_VERSION{SettingFilter} = $(UE_TVOS_BUILD_VERSION)");
+
+					// short version string
+					string BundleShortVersion;
+					Ini.GetString("/Script/IOSRuntimeSettings.IOSRuntimeSettings", "VersionInfo", out BundleShortVersion);
+					Xcconfig.AppendLine($"MARKETING_VERSION{SettingFilter} = {BundleShortVersion}");
+				}
 
 				if (UnrealData.TVOSProvisioningData != null && UnrealData.bWriteCodeSigningSettings)
 				{
 					if (!string.IsNullOrEmpty(UnrealData.TVOSProvisioningData.MobileProvisionUUID))
 					{
-						Xcconfig.AppendLine($"PROVISIONING_PROFILE_SPECIFIER[sdk=iphoneos*] = {UnrealData.TVOSProvisioningData.MobileProvisionUUID}");
+						Xcconfig.AppendLine($"PROVISIONING_PROFILE_SPECIFIER{SettingFilter} = {UnrealData.TVOSProvisioningData.MobileProvisionUUID}");
 					}
 					string TVOSCert = UnrealData.TVOSProvisioningData.SigningCertificate ?? (UnrealData.bForDistribution ? "Apple Distribution" : "Apple Developer");
-					Xcconfig.AppendLine($"CODE_SIGN_IDENTITY[sdk=iphoneos*] = {TVOSCert}");
+					Xcconfig.AppendLine($"CODE_SIGN_IDENTITY{SettingFilter} = {TVOSCert}");
 				}
 
 
 				// Info plist generation settings
-				Xcconfig.AppendLine("INFOPLIST_KEY_UIStatusBarHidden = YES");
-				Xcconfig.AppendLine("INFOPLIST_KEY_UIRequiresFullScreen = YES");
+				//Xcconfig.AppendLine("INFOPLIST_KEY_UIStatusBarHidden = YES");
+				//Xcconfig.AppendLine("INFOPLIST_KEY_UIRequiresFullScreen = YES");
 			}
 
 			////////////////////////
 			// MAC
 			////////////////////////
 
-			Xcconfig.AppendLine("");
-			Xcconfig.AppendLine("// Mac Settings");
-			Xcconfig.AppendLine($"SDKROOT[sdk=macosx*] = macosx");
-			Xcconfig.AppendLine($"CONFIGURATION_BUILD_DIR[sdk=macosx*] = {ProjectRootDir}/Binaries/Mac");
+			if (ProjectBuildConfigs.bSupportsMac)
+			{
+				string SettingFilter = XcodeProjectFileGenerator.XcodePlatforms.Count == 1 ? "" : "[sdk=macosx*]";
+
+				Xcconfig.AppendLine("");
+				Xcconfig.AppendLine("// Mac Settings");
+				Xcconfig.AppendLine($"SDKROOT{SettingFilter} = macosx");
+
+				if (UnrealData.PlistMode == PlistMode.UpdateTemplate)
+				{
+					Xcconfig.AppendLine($"CURRENT_PROJECT_VERSION{SettingFilter} = $(UE_MAC_BUILD_VERSION)");
+				}
+			}
 
 
 			////////////////////////
@@ -1568,11 +1670,6 @@ namespace UnrealBuildTool.XcodeProjectXcconfig
 				string RelativeLoc = UnrealData.EntitlementsLocation.MakeRelativeTo(DirectoryReference.Combine(UnrealData.XcodeProjectFileLocation.ParentDirectory!));
 				Xcconfig.AppendLine($"CODE_SIGN_ENTITLEMENTS = {RelativeLoc}");
 			}
-			if (UnrealData.InfoPlistLocation != null)
-			{
-				string RelativeLoc = UnrealData.InfoPlistLocation.MakeRelativeTo(DirectoryReference.Combine(UnrealData.XcodeProjectFileLocation.ParentDirectory!));
-				Xcconfig.AppendLine($"INFOPLIST_FILE = {RelativeLoc}");
-			}
 
 			Xcconfig.Write();
 
@@ -1581,6 +1678,28 @@ namespace UnrealBuildTool.XcodeProjectXcconfig
 			DirectoryReference? GameDir = UnrealData.UProjectFileLocation?.Directory;
 			string? GamePath = GameDir != null ? XcodeFileCollection.ConvertPath(GameDir.FullName) : null;
 
+			// find engine or premade (from project) plists
+			string MacEditorPlist = UnrealData.ProjectOrEnginePath("Build/Mac/Resources/Info-Editor.plist", true);
+			string MacGamePlist = UnrealData.ProjectOrEnginePath("Build/Mac/Resources/Info.plist", true);
+			string IOSPlist = UnrealData.ProjectOrEnginePath("Build/IOS/Resources/Info.plist", true);
+
+			// now pull the bundle id's out, as xcode will warn if they don't match
+			string MacEditorBundleID = "com.epicgames.UnrealEditor";
+			string MacGameBundleID = UnrealData.BundleIdentifier;
+			string IOSBundleID = UnrealData.BundleIdentifier;
+			if (UnrealData.PlistMode == PlistMode.UsePremade)
+			{
+				if (UnrealData.bSupportsMac)
+				{
+					MacEditorBundleID = FindBundleID(MacEditorPlist);
+					MacGameBundleID = FindBundleID(MacGamePlist);
+				}
+				if (UnrealData.bSupportsIOS || UnrealData.bSupportsTVOS)
+				{
+					IOSBundleID = FindBundleID(IOSPlist);
+				}
+			}
+
 			// Get Mac architectures supported by this project
 			foreach (UnrealBuildConfig Config in UnrealData.AllConfigs)
 			{
@@ -1600,13 +1719,6 @@ namespace UnrealBuildTool.XcodeProjectXcconfig
 				string MacExecutableFileName = MacExecutablePath.GetFileName();
 				string SupportedMacArchitectures = string.Join(" ", UnrealData.GetSupportedMacArchitectures(Config, UnrealData.UProjectFileLocation));
 
-
-				// UnrealClient re-uses UnrealGame
-				string PlistTargetName = bIsUnrealClient ? "UnrealGame" : Config.BuildTarget;
-				string MacInfoPlistPath = $"{ProjectRootDir}/Intermediate/Mac/{MacExecutableFileName}-Info.plist";
-				string IOSInfoPlistPath = $"{ProjectRootDir}/Intermediate/IOS/{PlistTargetName}-Info.plist";
-				string TVOSInfoPlistPath = $"{ProjectRootDir}/Intermediate/TVOS/{PlistTargetName}-Info.plist";
-
 				string SupportedPlatforms = "";
 				if (Config.ProjectTarget!.TargetRules != null)
 				{
@@ -1626,8 +1738,23 @@ namespace UnrealBuildTool.XcodeProjectXcconfig
 				else
 				{
 					// @todo when does this case happen?
-					SupportedPlatforms = "macosx";
+					throw new BuildException("Project without a target??");
+//					SupportedPlatforms = "macosx";
+				}
+
+
+				ConfigXcconfig.AppendLine("// Unreal per-config variables");
+				if (Config.bSupportsMac)
+				{
+					ConfigXcconfig.AppendLine($"UE_MAC_EXECUTABLE_NAME = {MacExecutableFileName}");
 				}
+				if (Config.bSupportsIOS || Config.bSupportsTVOS)
+				{
+					ConfigXcconfig.AppendLine($"UE_IOS_EXECUTABLE_NAME = {Config.BuildTarget}");
+				}
+				ConfigXcconfig.AppendLine($"UE_TARGET_NAME = {Config.BuildTarget}");
+				ConfigXcconfig.AppendLine($"UE_TARGET_CONFIG = {Config.BuildConfig}");
+				ConfigXcconfig.AppendLine("");
 
 				// debug settings
 				if (Config.BuildConfig == UnrealTargetConfiguration.Debug)
@@ -1635,35 +1762,60 @@ namespace UnrealBuildTool.XcodeProjectXcconfig
 					ConfigXcconfig.AppendLine("ENABLE_TESTABILITY = YES");
 				}
 
-				ConfigXcconfig.AppendLine($"UE_BUILD_TARGET_NAME = {Config.BuildTarget}");
-				ConfigXcconfig.AppendLine($"UE_BUILD_TARGET_CONFIG = {Config.BuildConfig}");
-				
+
 				// list the supported platforms, which will narrow down what other settings are active
 				ConfigXcconfig.AppendLine($"SUPPORTED_PLATFORMS = {SupportedPlatforms}");
 
 				// @otodo move VALID_ARCHS up to Project config once we are always universal
 				if (Config.bSupportsMac)
 				{
+					bool bIsEditor = (Config.ProjectTarget!.TargetRules.Type == TargetType.Editor);
+					string SettingFilter = XcodeProjectFileGenerator.XcodePlatforms.Count == 1 ? "" : "[sdk=macosx*]";
+
 					ConfigXcconfig.AppendLine("");
 					ConfigXcconfig.AppendLine("// Mac setup");
-					ConfigXcconfig.AppendLine($"VALID_ARCHS[sdk=macosx*] = {SupportedMacArchitectures}");
-					ConfigXcconfig.AppendLine($"PRODUCT_NAME[sdk=macosx*] = {MacExecutableFileName}");
+					ConfigXcconfig.AppendLine($"VALID_ARCHS{SettingFilter} = {SupportedMacArchitectures}");
+					ConfigXcconfig.AppendLine($"CONFIGURATION_BUILD_DIR{SettingFilter} = {MacExecutableDir}");
+					string MacPlistFile = bIsEditor ? MacEditorPlist : MacGamePlist;
+					ConfigXcconfig.AppendLine($"INFOPLIST_FILE{SettingFilter} = {MacPlistFile}");
+					ConfigXcconfig.AppendLine($"PRODUCT_BUNDLE_IDENTIFIER{SettingFilter} = " + (bIsEditor ? MacEditorBundleID : MacGameBundleID));
+
+					if (UnrealData.PlistMode == PlistMode.UpdateTemplate)
+					{
+						ConfigXcconfig.AppendLine($"PRODUCT_NAME{SettingFilter} = {MacExecutableFileName}");
+						// @todo: get a version for  games, like IOS has
+						ConfigXcconfig.AppendLine($"MARKETING_VERSION{SettingFilter} = {MacToolChain.LoadEngineDisplayVersion()}");
+					}
 				}
 
 				if (Config.bSupportsIOS)
 				{
+					string SettingFilter = XcodeProjectFileGenerator.XcodePlatforms.Count == 1 ? "" : "[sdk=iphoneos*]";
+
 					ConfigXcconfig.AppendLine("");
 					ConfigXcconfig.AppendLine("// IOS setup");
-					ConfigXcconfig.AppendLine($"PRODUCT_NAME[sdk=iphoneos*] = {Config.BuildTarget}"); // @todo: change to Path.GetFileName(Config.IOSExecutablePath) when we stop using payload
+					ConfigXcconfig.AppendLine($"INFOPLIST_FILE{SettingFilter} = {IOSPlist}");
+					ConfigXcconfig.AppendLine($"PRODUCT_BUNDLE_IDENTIFIER{SettingFilter} = {IOSBundleID}");
+					if (UnrealData.PlistMode == PlistMode.UpdateTemplate)
+					{
+						ConfigXcconfig.AppendLine($"PRODUCT_NAME{SettingFilter} = {Config.BuildTarget}"); // @todo: change to Path.GetFileName(Config.IOSExecutablePath) when we stop using payload
+					}
 				}
 				if (Config.bSupportsTVOS)
 				{
+					string SettingFilter = XcodeProjectFileGenerator.XcodePlatforms.Count == 1 ? "" : "[sdk=appletvos*]";
+
 					ConfigXcconfig.AppendLine("");
 					ConfigXcconfig.AppendLine("// TVOS setup");
-					ConfigXcconfig.AppendLine($"PRODUCT_NAME[sdk=appletvos*] = {Config.BuildTarget}"); // @todo: change to Path.GetFileName(Config.IOSExecutablePath) when we stop using payload
+					ConfigXcconfig.AppendLine($"INFOPLIST_FILE{SettingFilter} = {IOSPlist}");
+					ConfigXcconfig.AppendLine($"PRODUCT_BUNDLE_IDENTIFIER{SettingFilter} = {IOSBundleID}");
+					if (UnrealData.PlistMode == PlistMode.UpdateTemplate)
+					{
+						ConfigXcconfig.AppendLine($"PRODUCT_NAME{SettingFilter} = {Config.BuildTarget}"); // @todo: change to Path.GetFileName(Config.IOSExecutablePath) when we stop using payload
+					}
 				}
 
-// @todo this is a reminder of plist settings we are not handling yet
+				// @todo this is a reminder of plist settings we are not handling yet
 #if false
 				//// Prepare a temp Info.plist file so Xcode has some basic info about the target immediately after opening the project.
 				//// This is needed for the target to pass the settings validation before code signing. UBT will overwrite this plist file later, with proper contents.
@@ -2126,11 +2278,7 @@ namespace UnrealBuildTool.XcodeProjectXcconfig
 		public XcodeProjectFile(FileReference InitFilePath, DirectoryReference BaseDir, bool bIsForDistribution, string BundleID, string AppName)
 			: base(InitFilePath, BaseDir)
 		{
-			UnrealData = new UnrealData(InitFilePath);
-
-			UnrealData.bForDistribution = bIsForDistribution;
-			UnrealData.BundleIdentifier = BundleID;
-			UnrealData.AppName = AppName;
+			UnrealData = new UnrealData(InitFilePath, bIsForDistribution, BundleID, AppName);
 
 			// create the container for all the files that will 
 			FileCollection = new XcodeFileCollection(this);
@@ -2257,24 +2405,12 @@ namespace UnrealBuildTool.XcodeProjectXcconfig
 				return LegacyProjectFile.WriteProjectFile(InPlatforms, InConfigurations, PlatformProjectGenerators, Logger);
 			}
 
-			if (UnrealData.Initialize(this, InPlatforms, InConfigurations, Logger) == false)
+			if (UnrealData.Initialize(this, InConfigurations, Logger) == false)
 			{
 				// if we failed to initialize, we silently return to move on (it's not an error, it's a project with nothing to do)
 				return true;
 			}
 
-			if (string.IsNullOrEmpty(UnrealData.BundleIdentifier))
-			{
-				ConfigHierarchy Ini = ConfigCache.ReadHierarchy(ConfigHierarchyType.Engine, UnrealData.UProjectFileLocation?.Directory, UnrealTargetPlatform.Mac);
-				string? BundleId;
-				if (Ini.TryGetValue("XcodeConfiguration", "ModernBundleIdentifier", out BundleId))
-				{
-					// adapted from IOSProjectSettings
-					string ProjectName = (UnrealData.UProjectFileLocation != null) ? UnrealData.ProductName : "UnrealGame";
-					UnrealData.BundleIdentifier = BundleId.Replace("[PROJECT_NAME]", ProjectName).Replace("_", "");
-				}	
-			}
-
 			// look for an existing project to use as a template (if none found, create one from scratch)
 			DirectoryReference BuildDirLocation = UnrealData.UProjectFileLocation == null ? Unreal.EngineDirectory : UnrealData.UProjectFileLocation.Directory;
 			string ExistingProjectName = UnrealData.ProductName;
@@ -2404,7 +2540,7 @@ namespace UnrealBuildTool.XcodeProjectXcconfig
 				}
 			}
 
-
+#if false
 			/////////////////////
 			// Info.plist
 			/////////////////////
@@ -2460,22 +2596,35 @@ namespace UnrealBuildTool.XcodeProjectXcconfig
 				}
 			}
 
-// @todo this is a reminder of plist settings we are not handling yet
+#endif
+
+
+
 #if false
-			Text.AppendLine("\t\t\t<key>CFBundleURLName</key>");
-			Text.AppendLine("\t\t\t<string>com.Epic.Unreal</string>");
-			Text.AppendLine("\t\t\t<key>CFBundleURLSchemes</key>");
-			Text.AppendLine("\t\t\t<array>");
-			Text.AppendLine(string.Format("\t\t\t\t<string>{0}</string>", bIsUnrealGame ? "UnrealGame" : GameName));
+
+			// in CFBundleURLSchemes array
+			if (bEnableFacebookSupport)
+			{
+				// This is needed for facebook login to redirect back to the app after completion.
+				Text.AppendLine(string.Format("\t\t\t\t<string>fb{0}</string>", FacebookAppID));
+			}
 			if (bEnableGoogleSupport)
 			{
 				Text.AppendLine(string.Format("\t\t\t\t<string>{0}</string>", GoogleReversedClientId));
 			}
-			Text.AppendLine("\t\t\t</array>");
-			Text.AppendLine("\t\t</dict>");
-			Text.AppendLine("\t</array>");
+
+			Text.AppendLine("\t<key>CFBundleVersion</key>");
+			Text.AppendLine(string.Format("\t<string>{0}</string>", VersionUtilities.UpdateBundleVersion(OldPListData, InEngineDir)));
+			Text.AppendLine("\t<key>CFBundleShortVersionString</key>");
+			Text.AppendLine(string.Format("\t<string>{0}</string>", BundleShortVersion));
 
 
+#endif
+
+
+			// deleted
+#if false
+
 			Text.AppendLine("\t<key>CFBundleDevelopmentRegion</key>");
 			Text.AppendLine("\t<string>English</string>");
 
@@ -2485,10 +2634,20 @@ namespace UnrealBuildTool.XcodeProjectXcconfig
 				(bIsClient ? GameName + "Client" : GameName);
 			Text.AppendLine(string.Format("\t<string>{0}</string>", BundleExecutable));
 
-			Text.AppendLine("\t<key>CFBundleVersion</key>");
-			Text.AppendLine(string.Format("\t<string>{0}</string>", VersionUtilities.UpdateBundleVersion(OldPListData, InEngineDir)));
-			Text.AppendLine("\t<key>CFBundleShortVersionString</key>");
-			Text.AppendLine(string.Format("\t<string>{0}</string>", BundleShortVersion));
+ 
+#endif
+
+			// moved into .plist
+#if false
+			Text.AppendLine("\t\t\t<key>CFBundleURLName</key>");
+			Text.AppendLine("\t\t\t<string>com.Epic.Unreal</string>");
+			Text.AppendLine("\t\t\t<key>CFBundleURLSchemes</key>");
+			Text.AppendLine("\t\t\t<array>");
+			Text.AppendLine(string.Format("\t\t\t\t<string>{0}</string>", bIsUnrealGame ? "UnrealGame" : GameName));
+			Text.AppendLine("\t\t\t</array>");
+			Text.AppendLine("\t\t</dict>");
+			Text.AppendLine("\t</array>");
+
 
 			Text.AppendLine("\t<key>UIRequiredDeviceCapabilities</key>");
 			Text.AppendLine("\t<array>");
@@ -2673,12 +2832,12 @@ namespace UnrealBuildTool.XcodeProjectXcconfig
 #endif
 		}
 
-		private string Plist(string Command, string PlistFile)
+		public static string Plist(string Command, string PlistFile)
 		{
 			Command = Command.Replace("\"", "\\\"");
 			string Output = Utils.RunLocalProcessAndReturnStdOut("/usr/libexec/PlistBuddy", $"-c \"{Command}\" \"{PlistFile}\"");
 
-			Console.WriteLine($"{Command} ==> {Output}");
+//			Console.WriteLine($"{Command} ==> {Output}");
 
 			return Output;
 		}
@@ -2963,7 +3122,7 @@ namespace UnrealBuildTool.XcodeProjectXcconfig
 				}
 			}
 
-			string DefaultConfiguration = bHasEditorConfiguration && !XcodeProjectFileGenerator.bGeneratingRunIOSProject && !XcodeProjectFileGenerator.bGeneratingRunTVOSProject ? "Development Editor" : "Development";
+			string DefaultConfiguration = (bHasEditorConfiguration && !XcodeProjectFileGenerator.bGenerateRunOnlyProject) ? "Development Editor" : "Development";
 
 			StringBuilder Content = new StringBuilder();
 
diff --git a/Engine/Source/Programs/UnrealBuildTool/ProjectFiles/Xcode/XcodeProjectFileGenerator.cs b/Engine/Source/Programs/UnrealBuildTool/ProjectFiles/Xcode/XcodeProjectFileGenerator.cs
--- a/Engine/Source/Programs/UnrealBuildTool/ProjectFiles/Xcode/XcodeProjectFileGenerator.cs
+++ b/Engine/Source/Programs/UnrealBuildTool/ProjectFiles/Xcode/XcodeProjectFileGenerator.cs
@@ -243,10 +243,10 @@ namespace UnrealBuildTool
 			}
 
 			string ProjectName = PrimaryProjectName;
-			if (ProjectFilePlatform != XcodeProjectFilePlatform.All)
-			{
-				ProjectName += ProjectFilePlatform == XcodeProjectFilePlatform.Mac ? "_Mac" : (ProjectFilePlatform == XcodeProjectFilePlatform.iOS ? "_IOS" : "_TVOS");
-			}
+			//if (ProjectFilePlatform != XcodeProjectFilePlatform.All)
+			//{
+			//	ProjectName += ProjectFilePlatform == XcodeProjectFilePlatform.Mac ? "_Mac" : (ProjectFilePlatform == XcodeProjectFilePlatform.iOS ? "_IOS" : "_TVOS");
+			//}
 			string WorkspaceDataFilePath = PrimaryProjectPath + "/" + ProjectName + ".xcworkspace/contents.xcworkspacedata";
 			bSuccess = WriteFileIfChanged(WorkspaceDataFilePath, WorkspaceDataContent.ToString(), Logger, new UTF8Encoding());
 			if (bSuccess)
@@ -269,47 +269,55 @@ namespace UnrealBuildTool
 			return WriteXcodeWorkspace(Logger);
 		}
 
-		[Flags]
-		public enum XcodeProjectFilePlatform
-		{
-			Mac = 1 << 0,
-			iOS = 1 << 1,
-			tvOS = 1 << 2,
-			All = Mac | iOS | tvOS
-		}
-
-		/// Which platforms we should generate targets for
-		static public XcodeProjectFilePlatform ProjectFilePlatform = XcodeProjectFilePlatform.All;
-
-		/// Should we generate a special project to use for iOS signing instead of a normal one
-		static public bool bGeneratingRunIOSProject = false;
+		/// <summary>
+		/// A static copy of ProjectPlatforms from the base class
+		/// </summary>
+		static public List<UnrealTargetPlatform> XcodePlatforms = new();
 
-		/// Should we generate a special project to use for tvOS signing instead of a normal one
-		static public bool bGeneratingRunTVOSProject = false;
+		/// <summary>
+		/// Should we generate only a run project (no build/index targets)
+		/// </summary>
+		static public bool bGenerateRunOnlyProject = false;
 
 		/// <inheritdoc/>
 		protected override void ConfigureProjectFileGeneration(string[] Arguments, ref bool IncludeAllPlatforms, ILogger Logger)
 		{
 			// Call parent implementation first
 			base.ConfigureProjectFileGeneration(Arguments, ref IncludeAllPlatforms, Logger);
-			ProjectFilePlatform = IncludeAllPlatforms ? XcodeProjectFilePlatform.All : XcodeProjectFilePlatform.Mac;
 
-			foreach (string CurArgument in Arguments)
+			if (ProjectPlatforms.Count > 0)
 			{
-				if (CurArgument.StartsWith("-iOSDeployOnly", StringComparison.InvariantCultureIgnoreCase))
+				XcodePlatforms.AddRange(ProjectPlatforms);
+			}
+			else
+			{
+				// add platforms that have synced platform support
+				if (InstalledPlatformInfo.IsValidPlatform(UnrealTargetPlatform.Mac, EProjectType.Code))
 				{
-					bGeneratingRunIOSProject = true;
-					break;
+					XcodePlatforms.Add(UnrealTargetPlatform.Mac);
+				}
+				if (InstalledPlatformInfo.IsValidPlatform(UnrealTargetPlatform.IOS, EProjectType.Code))
+				{
+					XcodePlatforms.Add(UnrealTargetPlatform.IOS);
 				}
+				if (InstalledPlatformInfo.IsValidPlatform(UnrealTargetPlatform.TVOS, EProjectType.Code))
+				{
+					XcodePlatforms.Add(UnrealTargetPlatform.TVOS);
+				}
+			}
 
-				if (CurArgument.StartsWith("-tvOSDeployOnly", StringComparison.InvariantCultureIgnoreCase))
+			foreach (string CurArgument in Arguments)
+			{
+				if (CurArgument.Contains("-iOSDeployOnly", StringComparison.InvariantCultureIgnoreCase) ||
+					CurArgument.Contains("-tvOSDeployOnly", StringComparison.InvariantCultureIgnoreCase) ||
+					CurArgument.Contains("-DeployOnly", StringComparison.InvariantCultureIgnoreCase))
 				{
-					bGeneratingRunTVOSProject = true;
+					bGenerateRunOnlyProject = true;
 					break;
 				}
 			}
 
-			if (bGeneratingRunIOSProject || bGeneratingRunTVOSProject)
+			if (bGenerateRunOnlyProject)
 			{
 				bIncludeEnginePrograms = false;
 				//bIncludeEngineSource = false;
diff --git a/Engine/Source/Programs/UnrealBuildTool/ProjectFiles/Xcode/XcodeProjectLegacy.cs b/Engine/Source/Programs/UnrealBuildTool/ProjectFiles/Xcode/XcodeProjectLegacy.cs
--- a/Engine/Source/Programs/UnrealBuildTool/ProjectFiles/Xcode/XcodeProjectLegacy.cs
+++ b/Engine/Source/Programs/UnrealBuildTool/ProjectFiles/Xcode/XcodeProjectLegacy.cs
@@ -622,7 +622,7 @@ namespace UnrealBuildTool.XcodeProjectLegacy
 		{
 			List<string> DependencyGuids = new List<string>();
 			// depends on the Run target if we want one
-			if (!XcodeProjectFileGenerator.bGeneratingRunIOSProject && !XcodeProjectFileGenerator.bGeneratingRunTVOSProject)
+			if (!XcodeProjectFileGenerator.bGenerateRunOnlyProject)
 			{
 				DependencyGuids.Add(TargetDependencyGuid);
 			}
@@ -1023,164 +1023,164 @@ namespace UnrealBuildTool.XcodeProjectLegacy
 			IOSRunTimeVersion = null;
 			TVOSRunTimeVersion = null;
 
-			bool bIsUnrealGame = TargetName.Equals("UnrealGame", StringComparison.InvariantCultureIgnoreCase);
-			bool bIsUnrealClient = TargetName.Equals("UnrealClient", StringComparison.InvariantCultureIgnoreCase);
-			DirectoryReference? GameDir = ProjectFile?.Directory;
-			string? GamePath = GameDir != null ? ConvertPath(GameDir.FullName) : null;
-
-			string? IOSRunTimeDevices = null;
-			string? TVOSRunTimeDevices = null;
-			string SupportedPlatforms = bSupportMac ? "macosx" : "";
-
-			bool bAutomaticSigning = false;
-			string? UUID_IOS = "";
-			string? UUID_TVOS = "";
-			string? TEAM_IOS = "";
-			string? TEAM_TVOS = "";
-			string? IOS_CERT = "iPhone Developer";
-			string? TVOS_CERT = "iPhone Developer";
-			string IOS_BUNDLE = "";
-			string TVOS_BUNDLE = "";
-			if (bSupportIOS && InstalledPlatformInfo.IsValidPlatform(UnrealTargetPlatform.IOS, EProjectType.Code))
-			{
-				IOSPlatform IOSPlatform = ((IOSPlatform)UEBuildPlatform.GetBuildPlatform(UnrealTargetPlatform.IOS));
-				IOSProjectSettings ProjectSettings = IOSPlatform.ReadProjectSettings(ProjectFile);
-				IOSProvisioningData ProvisioningData = IOSPlatform.ReadProvisioningData(ProjectSettings, bForDistribution);
-				IOSRunTimeVersion = ProjectSettings.RuntimeVersion;
-				IOSRunTimeDevices = ProjectSettings.RuntimeDevices;
-				SupportedPlatforms += " iphoneos";
-				bAutomaticSigning = ProjectSettings.bAutomaticSigning;
-				if (!bAutomaticSigning)
+				bool bIsUnrealGame = TargetName.Equals("UnrealGame", StringComparison.InvariantCultureIgnoreCase);
+				bool bIsUnrealClient = TargetName.Equals("UnrealClient", StringComparison.InvariantCultureIgnoreCase);
+				DirectoryReference? GameDir = ProjectFile?.Directory;
+				string? GamePath = GameDir != null ? ConvertPath(GameDir.FullName) : null;
+
+				string? IOSRunTimeDevices = null;
+				string? TVOSRunTimeDevices = null;
+				string SupportedPlatforms = bSupportMac ? "macosx" : "";
+
+				bool bAutomaticSigning = false;
+				string? UUID_IOS = "";
+				string? UUID_TVOS = "";
+				string? TEAM_IOS = "";
+				string? TEAM_TVOS = "";
+				string? IOS_CERT = "iPhone Developer";
+				string? TVOS_CERT = "iPhone Developer";
+				string IOS_BUNDLE = "";
+				string TVOS_BUNDLE = "";
+				if (bSupportIOS && InstalledPlatformInfo.IsValidPlatform(UnrealTargetPlatform.IOS, EProjectType.Code))
 				{
-					UUID_IOS = ProvisioningData.MobileProvisionUUID;
-					IOS_CERT = ProvisioningData.SigningCertificate;
+					IOSPlatform IOSPlatform = ((IOSPlatform)UEBuildPlatform.GetBuildPlatform(UnrealTargetPlatform.IOS));
+					IOSProjectSettings ProjectSettings = IOSPlatform.ReadProjectSettings(ProjectFile);
+					IOSProvisioningData ProvisioningData = IOSPlatform.ReadProvisioningData(ProjectSettings, bForDistribution);
+					IOSRunTimeVersion = ProjectSettings.RuntimeVersion;
+					IOSRunTimeDevices = ProjectSettings.RuntimeDevices;
+					SupportedPlatforms += " iphoneos";
+					bAutomaticSigning = ProjectSettings.bAutomaticSigning;
+					if (!bAutomaticSigning)
+					{
+						UUID_IOS = ProvisioningData.MobileProvisionUUID;
+						IOS_CERT = ProvisioningData.SigningCertificate;
+					}
+					TEAM_IOS = ProvisioningData.TeamUUID;
+					IOS_BUNDLE = ProjectSettings.BundleIdentifier;
 				}
-				TEAM_IOS = ProvisioningData.TeamUUID;
-				IOS_BUNDLE = ProjectSettings.BundleIdentifier;
-			}
 
-			if (bSupportTVOS && InstalledPlatformInfo.IsValidPlatform(UnrealTargetPlatform.TVOS, EProjectType.Code))
-			{
-				TVOSPlatform TVOSPlatform = ((TVOSPlatform)UEBuildPlatform.GetBuildPlatform(UnrealTargetPlatform.TVOS));
-				TVOSProjectSettings ProjectSettings = TVOSPlatform.ReadProjectSettings(ProjectFile);
-				TVOSProvisioningData ProvisioningData = TVOSPlatform.ReadProvisioningData(ProjectSettings, bForDistribution);
-				TVOSRunTimeVersion = ProjectSettings.RuntimeVersion;
-				TVOSRunTimeDevices = ProjectSettings.RuntimeDevices;
-				SupportedPlatforms += " appletvos";
-				if (!bAutomaticSigning)
+				if (bSupportTVOS && InstalledPlatformInfo.IsValidPlatform(UnrealTargetPlatform.TVOS, EProjectType.Code))
 				{
-					UUID_TVOS = ProvisioningData.MobileProvisionUUID;
-					TVOS_CERT = ProvisioningData.SigningCertificate;
+					TVOSPlatform TVOSPlatform = ((TVOSPlatform)UEBuildPlatform.GetBuildPlatform(UnrealTargetPlatform.TVOS));
+					TVOSProjectSettings ProjectSettings = TVOSPlatform.ReadProjectSettings(ProjectFile);
+					TVOSProvisioningData ProvisioningData = TVOSPlatform.ReadProvisioningData(ProjectSettings, bForDistribution);
+					TVOSRunTimeVersion = ProjectSettings.RuntimeVersion;
+					TVOSRunTimeDevices = ProjectSettings.RuntimeDevices;
+					SupportedPlatforms += " appletvos";
+					if (!bAutomaticSigning)
+					{
+						UUID_TVOS = ProvisioningData.MobileProvisionUUID;
+						TVOS_CERT = ProvisioningData.SigningCertificate;
+					}
+					TEAM_TVOS = ProvisioningData.TeamUUID;
+					TVOS_BUNDLE = ProjectSettings.BundleIdentifier;
 				}
-				TEAM_TVOS = ProvisioningData.TeamUUID;
-				TVOS_BUNDLE = ProjectSettings.BundleIdentifier;
-			}
 
-			Content.Append("\t\t\t\tSUPPORTED_PLATFORMS = \"" + SupportedPlatforms.Trim() + "\";" + ProjectFileGenerator.NewLine);
-			if (bAutomaticSigning)
-			{
-				Content.Append("\t\t\t\tCODE_SIGN_STYLE = Automatic;" + ProjectFileGenerator.NewLine);
-			}
-			if (IOSRunTimeVersion != null)
-			{
-				Content.Append("\t\t\t\t\"VALID_ARCHS[sdk=iphoneos*]\" = \"" + string.Join(" ", this.SupportedIOSArchitectures) + "\";" + ProjectFileGenerator.NewLine);
-				Content.Append("\t\t\t\tIPHONEOS_DEPLOYMENT_TARGET = " + IOSRunTimeVersion + ";" + ProjectFileGenerator.NewLine);
-				Content.Append("\t\t\t\t\"PRODUCT_NAME[sdk=iphoneos*]\" = \"" + TargetName + "\";" + ProjectFileGenerator.NewLine); // @todo: change to Path.GetFileName(Config.IOSExecutablePath) when we stop using payload
-				Content.Append("\t\t\t\t\"TARGETED_DEVICE_FAMILY[sdk=iphoneos*]\" = \"" + IOSRunTimeDevices + "\";" + ProjectFileGenerator.NewLine);
-				Content.Append("\t\t\t\t\"SDKROOT[sdk=iphoneos]\" = iphoneos;" + ProjectFileGenerator.NewLine);
-				if (!string.IsNullOrEmpty(TEAM_IOS))
-				{
-					Content.Append("\t\t\t\t\"DEVELOPMENT_TEAM[sdk=iphoneos*]\" = " + TEAM_IOS + ";" + ProjectFileGenerator.NewLine);
-				}
-				Content.Append("\t\t\t\t\"CODE_SIGN_IDENTITY[sdk=iphoneos*]\" = \"" + IOS_CERT + "\";" + ProjectFileGenerator.NewLine);
-				if (!bAutomaticSigning && !string.IsNullOrEmpty(UUID_IOS))
-				{
-					Content.Append("\t\t\t\t\"PROVISIONING_PROFILE_SPECIFIER[sdk=iphoneos*]\" = \"" + UUID_IOS + "\";" + ProjectFileGenerator.NewLine);
-				}
-				if (ProjectFile != null)
-				{
-					Content.Append("\t\t\t\t\"PRODUCT_BUNDLE_IDENTIFIER[sdk=iphoneos*]\" = " + IOS_BUNDLE + ";" + ProjectFileGenerator.NewLine);
-				}
-			}
-			if (TVOSRunTimeVersion != null)
-			{
-				Content.Append("\t\t\t\t\"VALID_ARCHS[sdk=appletvos*]\" = \"" + string.Join(" ", this.SupportedIOSArchitectures) + "\";" + ProjectFileGenerator.NewLine);
-				Content.Append("\t\t\t\tTVOS_DEPLOYMENT_TARGET = " + TVOSRunTimeVersion + ";" + ProjectFileGenerator.NewLine);
-				Content.Append("\t\t\t\t\"PRODUCT_NAME[sdk=appletvos*]\" = \"" + TargetName + "\";" + ProjectFileGenerator.NewLine); // @todo: change to Path.GetFileName(Config.TVOSExecutablePath) when we stop using payload
-				Content.Append("\t\t\t\t\"TARGETED_DEVICE_FAMILY[sdk=appletvos*]\" = \"" + TVOSRunTimeDevices + "\";" + ProjectFileGenerator.NewLine);
-				Content.Append("\t\t\t\t\"SDKROOT[sdk=appletvos]\" = appletvos;" + ProjectFileGenerator.NewLine);
-				if (!string.IsNullOrEmpty(TEAM_TVOS))
+				Content.Append("\t\t\t\tSUPPORTED_PLATFORMS = \"" + SupportedPlatforms.Trim() + "\";" + ProjectFileGenerator.NewLine);
+				if (bAutomaticSigning)
 				{
-					Content.Append("\t\t\t\t\"DEVELOPMENT_TEAM[sdk=appletvos*]\" = " + TEAM_TVOS + ";" + ProjectFileGenerator.NewLine);
+					Content.Append("\t\t\t\tCODE_SIGN_STYLE = Automatic;" + ProjectFileGenerator.NewLine);
 				}
-				Content.Append("\t\t\t\t\"CODE_SIGN_IDENTITY[sdk=appletvos*]\" = \"" + TVOS_CERT + "\";" + ProjectFileGenerator.NewLine);
-				if (!bAutomaticSigning && !string.IsNullOrEmpty(UUID_TVOS))
-				{
-					Content.Append("\t\t\t\t\"PROVISIONING_PROFILE_SPECIFIER[sdk=appletvos*]\" = \"" + UUID_TVOS + "\";" + ProjectFileGenerator.NewLine);
-				}
-				if (ProjectFile != null)
-				{
-					Content.Append("\t\t\t\t\"PRODUCT_BUNDLE_IDENTIFIER[sdk=appletvos*]\" = " + TVOS_BUNDLE + ";" + ProjectFileGenerator.NewLine);
-				}
-			}
-			if (bSupportMac)
-			{
-				Content.Append("\t\t\t\t\"VALID_ARCHS[sdk=macosx*]\" = \"" + string.Join(" ", SupportedMacArchitectures) + "\";" + ProjectFileGenerator.NewLine);
-				Content.Append("\t\t\t\t\"PRODUCT_NAME[sdk=macosx*]\" = \"" + MacExecutableFileName + "\";" + ProjectFileGenerator.NewLine);
-				Content.Append("\t\t\t\t\"CONFIGURATION_BUILD_DIR[sdk=macosx*]\" = \"" + MacExecutableDir + "\";" + ProjectFileGenerator.NewLine);
-				Content.Append("\t\t\t\t\"SDKROOT[sdk=macosx*]\" = macosx;" + ProjectFileGenerator.NewLine);
-
-				// MacToolchain uses the IOS bundle identifier
-				ConfigHierarchy IOSIni = ConfigCache.ReadHierarchy(ConfigHierarchyType.Engine, ProjectFile?.Directory, UnrealTargetPlatform.IOS);
-				string IOSBundleIdentifier;
-				IOSIni.GetString("/Script/IOSRuntimeSettings.IOSRuntimeSettings", "BundleIdentifier", out IOSBundleIdentifier!);
-
-				// we are going to load/sign the UnrealEditor.app when doing modular editor builds, so reuse the shared identifier
-				bool bBuildingEditor = Config.ProjectTarget!.TargetRules!.Type == TargetType.Editor && Config.ProjectTarget!.TargetRules!.LinkType == TargetLinkType.Modular;
-				string MacBundleIdentifier = bBuildingEditor ? ("com.epicgames.UnrealEditor") : (IOSBundleIdentifier.Replace("[PROJECT_NAME]", TargetName).Replace("_", ""));
-				// Xcode 14 wants this to match what is in the final plist
-				Content.Append("\t\t\t\t\"PRODUCT_BUNDLE_IDENTIFIER[sdk=macosx*]\" = " + MacBundleIdentifier + ";" + ProjectFileGenerator.NewLine);
-			}
-
-			if (bIsUnrealGame || bIsUnrealClient)
-			{
 				if (IOSRunTimeVersion != null)
 				{
-					Content.Append("\t\t\t\t\"CONFIGURATION_BUILD_DIR[sdk=iphoneos*]\" = \"" + UEDir + "/Engine/Binaries/IOS" + BinariesSubDir + "\";" + ProjectFileGenerator.NewLine);
+					Content.Append("\t\t\t\t\"VALID_ARCHS[sdk=iphoneos*]\" = \"" + string.Join(" ", this.SupportedIOSArchitectures) + "\";" + ProjectFileGenerator.NewLine);
+					Content.Append("\t\t\t\tIPHONEOS_DEPLOYMENT_TARGET = " + IOSRunTimeVersion + ";" + ProjectFileGenerator.NewLine);
+					Content.Append("\t\t\t\t\"PRODUCT_NAME[sdk=iphoneos*]\" = \"" + TargetName + "\";" + ProjectFileGenerator.NewLine); // @todo: change to Path.GetFileName(Config.IOSExecutablePath) when we stop using payload
+					Content.Append("\t\t\t\t\"TARGETED_DEVICE_FAMILY[sdk=iphoneos*]\" = \"" + IOSRunTimeDevices + "\";" + ProjectFileGenerator.NewLine);
+					Content.Append("\t\t\t\t\"SDKROOT[sdk=iphoneos]\" = iphoneos;" + ProjectFileGenerator.NewLine);
+					if (!string.IsNullOrEmpty(TEAM_IOS))
+					{
+						Content.Append("\t\t\t\t\"DEVELOPMENT_TEAM[sdk=iphoneos*]\" = " + TEAM_IOS + ";" + ProjectFileGenerator.NewLine);
+					}
+					Content.Append("\t\t\t\t\"CODE_SIGN_IDENTITY[sdk=iphoneos*]\" = \"" + IOS_CERT + "\";" + ProjectFileGenerator.NewLine);
+					if (!bAutomaticSigning && !string.IsNullOrEmpty(UUID_IOS))
+					{
+						Content.Append("\t\t\t\t\"PROVISIONING_PROFILE_SPECIFIER[sdk=iphoneos*]\" = \"" + UUID_IOS + "\";" + ProjectFileGenerator.NewLine);
+					}
+					if (ProjectFile != null)
+					{
+						Content.Append("\t\t\t\t\"PRODUCT_BUNDLE_IDENTIFIER[sdk=iphoneos*]\" = " + IOS_BUNDLE + ";" + ProjectFileGenerator.NewLine);
+					}
 				}
 				if (TVOSRunTimeVersion != null)
 				{
-					Content.Append("\t\t\t\t\"CONFIGURATION_BUILD_DIR[sdk=appletvos*]\" = \"" + UEDir + "/Engine/Binaries/TVOS" + BinariesSubDir + "\";" + ProjectFileGenerator.NewLine);
+					Content.Append("\t\t\t\t\"VALID_ARCHS[sdk=appletvos*]\" = \"" + string.Join(" ", this.SupportedIOSArchitectures) + "\";" + ProjectFileGenerator.NewLine);
+					Content.Append("\t\t\t\tTVOS_DEPLOYMENT_TARGET = " + TVOSRunTimeVersion + ";" + ProjectFileGenerator.NewLine);
+					Content.Append("\t\t\t\t\"PRODUCT_NAME[sdk=appletvos*]\" = \"" + TargetName + "\";" + ProjectFileGenerator.NewLine); // @todo: change to Path.GetFileName(Config.TVOSExecutablePath) when we stop using payload
+					Content.Append("\t\t\t\t\"TARGETED_DEVICE_FAMILY[sdk=appletvos*]\" = \"" + TVOSRunTimeDevices + "\";" + ProjectFileGenerator.NewLine);
+					Content.Append("\t\t\t\t\"SDKROOT[sdk=appletvos]\" = appletvos;" + ProjectFileGenerator.NewLine);
+					if (!string.IsNullOrEmpty(TEAM_TVOS))
+					{
+						Content.Append("\t\t\t\t\"DEVELOPMENT_TEAM[sdk=appletvos*]\" = " + TEAM_TVOS + ";" + ProjectFileGenerator.NewLine);
+					}
+					Content.Append("\t\t\t\t\"CODE_SIGN_IDENTITY[sdk=appletvos*]\" = \"" + TVOS_CERT + "\";" + ProjectFileGenerator.NewLine);
+					if (!bAutomaticSigning && !string.IsNullOrEmpty(UUID_TVOS))
+					{
+						Content.Append("\t\t\t\t\"PROVISIONING_PROFILE_SPECIFIER[sdk=appletvos*]\" = \"" + UUID_TVOS + "\";" + ProjectFileGenerator.NewLine);
+					}
+					if (ProjectFile != null)
+					{
+						Content.Append("\t\t\t\t\"PRODUCT_BUNDLE_IDENTIFIER[sdk=appletvos*]\" = " + TVOS_BUNDLE + ";" + ProjectFileGenerator.NewLine);
+					}
 				}
-			}
-			else if (ProjectFile != null)
-			{
-				if (IOSRunTimeVersion != null)
+				if (bSupportMac)
 				{
-					Content.Append("\t\t\t\t\"CONFIGURATION_BUILD_DIR[sdk=iphoneos*]\" = \"" + GamePath + "/Binaries/IOS" + BinariesSubDir + "\";" + ProjectFileGenerator.NewLine);
+					Content.Append("\t\t\t\t\"VALID_ARCHS[sdk=macosx*]\" = \"" + string.Join(" ", SupportedMacArchitectures) + "\";" + ProjectFileGenerator.NewLine);
+					Content.Append("\t\t\t\t\"PRODUCT_NAME[sdk=macosx*]\" = \"" + MacExecutableFileName + "\";" + ProjectFileGenerator.NewLine);
+					Content.Append("\t\t\t\t\"CONFIGURATION_BUILD_DIR[sdk=macosx*]\" = \"" + MacExecutableDir + "\";" + ProjectFileGenerator.NewLine);
+					Content.Append("\t\t\t\t\"SDKROOT[sdk=macosx*]\" = macosx;" + ProjectFileGenerator.NewLine);
+
+					// MacToolchain uses the IOS bundle identifier
+					ConfigHierarchy IOSIni = ConfigCache.ReadHierarchy(ConfigHierarchyType.Engine, ProjectFile?.Directory, UnrealTargetPlatform.IOS);
+					string IOSBundleIdentifier;
+					IOSIni.GetString("/Script/IOSRuntimeSettings.IOSRuntimeSettings", "BundleIdentifier", out IOSBundleIdentifier!);
+
+					// we are going to load/sign the UnrealEditor.app when doing modular editor builds, so reuse the shared identifier
+					bool bBuildingEditor = Config.ProjectTarget!.TargetRules!.Type == TargetType.Editor && Config.ProjectTarget!.TargetRules!.LinkType == TargetLinkType.Modular;
+					string MacBundleIdentifier = bBuildingEditor ? ("com.epicgames.UnrealEditor") : (IOSBundleIdentifier.Replace("[PROJECT_NAME]", TargetName).Replace("_", ""));
+					// Xcode 14 wants this to match what is in the final plist
+					Content.Append("\t\t\t\t\"PRODUCT_BUNDLE_IDENTIFIER[sdk=macosx*]\" = " + MacBundleIdentifier + ";" + ProjectFileGenerator.NewLine);
 				}
-				if (TVOSRunTimeVersion != null)
+
+				if (bIsUnrealGame || bIsUnrealClient)
 				{
-					Content.Append("\t\t\t\t\"CONFIGURATION_BUILD_DIR[sdk=appletvos*]\" = \"" + GamePath + "/Binaries/TVOS" + BinariesSubDir + "\";" + ProjectFileGenerator.NewLine);
+					if (IOSRunTimeVersion != null)
+					{
+						Content.Append("\t\t\t\t\"CONFIGURATION_BUILD_DIR[sdk=iphoneos*]\" = \"" + UEDir + "/Engine/Binaries/IOS" + BinariesSubDir + "\";" + ProjectFileGenerator.NewLine);
+					}
+					if (TVOSRunTimeVersion != null)
+					{
+						Content.Append("\t\t\t\t\"CONFIGURATION_BUILD_DIR[sdk=appletvos*]\" = \"" + UEDir + "/Engine/Binaries/TVOS" + BinariesSubDir + "\";" + ProjectFileGenerator.NewLine);
+					}
 				}
-			}
-			else
-			{
-				if (IOSRunTimeVersion != null)
+				else if (ProjectFile != null)
 				{
-					Content.Append("\t\t\t\t\"CONFIGURATION_BUILD_DIR[sdk=iphoneos*]\" = \"" + UEDir + "/Engine/Binaries/IOS" + BinariesSubDir + "\";" + ProjectFileGenerator.NewLine);
+					if (IOSRunTimeVersion != null)
+					{
+						Content.Append("\t\t\t\t\"CONFIGURATION_BUILD_DIR[sdk=iphoneos*]\" = \"" + GamePath + "/Binaries/IOS" + BinariesSubDir + "\";" + ProjectFileGenerator.NewLine);
+					}
+					if (TVOSRunTimeVersion != null)
+					{
+						Content.Append("\t\t\t\t\"CONFIGURATION_BUILD_DIR[sdk=appletvos*]\" = \"" + GamePath + "/Binaries/TVOS" + BinariesSubDir + "\";" + ProjectFileGenerator.NewLine);
+					}
 				}
-				if (TVOSRunTimeVersion != null)
+				else
 				{
-					Content.Append("\t\t\t\t\"CONFIGURATION_BUILD_DIR[sdk=appletvos*]\" = \"" + UEDir + "/Engine/Binaries/TVOS" + BinariesSubDir + "\";" + ProjectFileGenerator.NewLine);
+					if (IOSRunTimeVersion != null)
+					{
+						Content.Append("\t\t\t\t\"CONFIGURATION_BUILD_DIR[sdk=iphoneos*]\" = \"" + UEDir + "/Engine/Binaries/IOS" + BinariesSubDir + "\";" + ProjectFileGenerator.NewLine);
+					}
+					if (TVOSRunTimeVersion != null)
+					{
+						Content.Append("\t\t\t\t\"CONFIGURATION_BUILD_DIR[sdk=appletvos*]\" = \"" + UEDir + "/Engine/Binaries/TVOS" + BinariesSubDir + "\";" + ProjectFileGenerator.NewLine);
+					}
 				}
-			}
 		}
 
 		private void AppendNativeTargetBuildConfiguration(StringBuilder Content, XcodeBuildConfig Config, string ConfigGuid, FileReference? ProjectFile, ILogger Logger)
 		{
 			bool bMacOnly = true;
-			if (Config.ProjectTarget!.TargetRules != null && XcodeProjectFileGenerator.ProjectFilePlatform.HasFlag(XcodeProjectFileGenerator.XcodeProjectFilePlatform.iOS))
+			if (Config.ProjectTarget!.TargetRules != null && XcodeProjectFileGenerator.XcodePlatforms.Contains(UnrealTargetPlatform.IOS))
 			{
 				if (Config.ProjectTarget.SupportedPlatforms.Contains(UnrealTargetPlatform.IOS))
 				{
@@ -1199,137 +1199,140 @@ namespace UnrealBuildTool.XcodeProjectLegacy
 			string? IOSRunTimeVersion, TVOSRunTimeVersion;
 			AppendPlatformConfiguration(Content, Config, Config.BuildTarget, ProjectFile, true, !bMacOnly, !bMacOnly, Logger, out IOSRunTimeVersion, out TVOSRunTimeVersion);
 
-			bool bIsUnrealGame = Config.BuildTarget.Equals("UnrealGame", StringComparison.InvariantCultureIgnoreCase);
-			bool bIsUnrealClient = Config.BuildTarget.Equals("UnrealClient", StringComparison.InvariantCultureIgnoreCase);
+				bool bIsUnrealGame = Config.BuildTarget.Equals("UnrealGame", StringComparison.InvariantCultureIgnoreCase);
+				bool bIsUnrealClient = Config.BuildTarget.Equals("UnrealClient", StringComparison.InvariantCultureIgnoreCase);
 
-			DirectoryReference? GameDir = ProjectFile?.Directory;
-			string? GamePath = GameDir != null ? ConvertPath(GameDir.FullName) : null;
+				DirectoryReference? GameDir = ProjectFile?.Directory;
+				string? GamePath = GameDir != null ? ConvertPath(GameDir.FullName) : null;
 
-			string IOSInfoPlistPath;
-			string TVOSInfoPlistPath;
-			string MacInfoPlistPath;
-			string? IOSEntitlementPath = null;
-			string? TVOSEntitlementPath = null;
-			if (bIsUnrealGame)
-			{
-				IOSInfoPlistPath = UEDir + "/Engine/Intermediate/IOS/" + Config.BuildTarget + "-Info.plist";
-				TVOSInfoPlistPath = UEDir + "/Engine/Intermediate/TVOS/" + Config.BuildTarget + "-Info.plist";
-				MacInfoPlistPath = UEDir + "/Engine/Intermediate/Mac/" + MacExecutableFileName + "-Info.plist";
-				IOSEntitlementPath = "";
-				TVOSEntitlementPath = "";
-			}
-			else if (bIsUnrealClient)
-			{
-				IOSInfoPlistPath = UEDir + "/Engine/Intermediate/IOS/UnrealGame-Info.plist";
-				TVOSInfoPlistPath = UEDir + "/Engine/Intermediate/TVOS/UnrealGame-Info.plist";
-				MacInfoPlistPath = UEDir + "/Engine/Intermediate/Mac/" + MacExecutableFileName + "-Info.plist";
-				IOSEntitlementPath = "";
-				TVOSEntitlementPath = "";
-			}
-			else if (ProjectFile != null)
-			{
-				IOSInfoPlistPath = GamePath + "/Intermediate/IOS/" + Config.BuildTarget + "-Info.plist";
-				TVOSInfoPlistPath = GamePath + "/Intermediate/TVOS/" + Config.BuildTarget + "-Info.plist";
-				MacInfoPlistPath = GamePath + "/Intermediate/Mac/" + MacExecutableFileName + "-Info.plist";
-				IOSEntitlementPath = GamePath + "/Intermediate/IOS/" + Config.BuildTarget + ".entitlements";
-				TVOSEntitlementPath = GamePath + "/Intermediate/TVOS/" + Config.BuildTarget + ".entitlements";
-			}
-			else
-			{
-				if (GamePath == null)
+				string IOSInfoPlistPath;
+				string TVOSInfoPlistPath;
+				string MacInfoPlistPath;
+				string? IOSEntitlementPath = null;
+				string? TVOSEntitlementPath = null;
+				if (bIsUnrealGame)
 				{
 					IOSInfoPlistPath = UEDir + "/Engine/Intermediate/IOS/" + Config.BuildTarget + "-Info.plist";
 					TVOSInfoPlistPath = UEDir + "/Engine/Intermediate/TVOS/" + Config.BuildTarget + "-Info.plist";
 					MacInfoPlistPath = UEDir + "/Engine/Intermediate/Mac/" + MacExecutableFileName + "-Info.plist";
+					IOSEntitlementPath = "";
+					TVOSEntitlementPath = "";
 				}
-				else
+				else if (bIsUnrealClient)
+				{
+					IOSInfoPlistPath = UEDir + "/Engine/Intermediate/IOS/UnrealGame-Info.plist";
+					TVOSInfoPlistPath = UEDir + "/Engine/Intermediate/TVOS/UnrealGame-Info.plist";
+					MacInfoPlistPath = UEDir + "/Engine/Intermediate/Mac/" + MacExecutableFileName + "-Info.plist";
+					IOSEntitlementPath = "";
+					TVOSEntitlementPath = "";
+				}
+				else if (ProjectFile != null)
 				{
 					IOSInfoPlistPath = GamePath + "/Intermediate/IOS/" + Config.BuildTarget + "-Info.plist";
 					TVOSInfoPlistPath = GamePath + "/Intermediate/TVOS/" + Config.BuildTarget + "-Info.plist";
 					MacInfoPlistPath = GamePath + "/Intermediate/Mac/" + MacExecutableFileName + "-Info.plist";
+					IOSEntitlementPath = GamePath + "/Intermediate/IOS/" + Config.BuildTarget + ".entitlements";
+					TVOSEntitlementPath = GamePath + "/Intermediate/TVOS/" + Config.BuildTarget + ".entitlements";
+				}
+				else
+				{
+					if (GamePath == null)
+					{
+						IOSInfoPlistPath = UEDir + "/Engine/Intermediate/IOS/" + Config.BuildTarget + "-Info.plist";
+						TVOSInfoPlistPath = UEDir + "/Engine/Intermediate/TVOS/" + Config.BuildTarget + "-Info.plist";
+						MacInfoPlistPath = UEDir + "/Engine/Intermediate/Mac/" + MacExecutableFileName + "-Info.plist";
+					}
+					else
+					{
+						IOSInfoPlistPath = GamePath + "/Intermediate/IOS/" + Config.BuildTarget + "-Info.plist";
+						TVOSInfoPlistPath = GamePath + "/Intermediate/TVOS/" + Config.BuildTarget + "-Info.plist";
+						MacInfoPlistPath = GamePath + "/Intermediate/Mac/" + MacExecutableFileName + "-Info.plist";
+					}
 				}
-			}
 
-			if (XcodeProjectFileGenerator.bGeneratingRunIOSProject)
-			{
-				Content.Append("\t\t\t\tINFOPLIST_FILE = \"" + IOSInfoPlistPath + "\";" + ProjectFileGenerator.NewLine);
-				Content.Append("\t\t\t\tCODE_SIGN_ENTITLEMENTS = \"" + IOSEntitlementPath + "\";" + ProjectFileGenerator.NewLine);
-			}
-			else if (XcodeProjectFileGenerator.bGeneratingRunTVOSProject)
-			{
-				Content.Append("\t\t\t\tINFOPLIST_FILE = \"" + TVOSInfoPlistPath + "\";" + ProjectFileGenerator.NewLine);
-				Content.Append("\t\t\t\tCODE_SIGN_ENTITLEMENTS = \"" + TVOSEntitlementPath + "\";" + ProjectFileGenerator.NewLine);
-			}
-			else
-			{
-				Content.Append("\t\t\t\t\"INFOPLIST_FILE[sdk=macosx*]\" = \"" + MacInfoPlistPath + "\";" + ProjectFileGenerator.NewLine);
-				if (IOSRunTimeVersion != null)
+				if (XcodeProjectFileGenerator.bGenerateRunOnlyProject)
 				{
-					Content.Append("\t\t\t\t\"INFOPLIST_FILE[sdk=iphoneos*]\" = \"" + IOSInfoPlistPath + "\";" + ProjectFileGenerator.NewLine);
-					Content.Append("\t\t\t\t\"CODE_SIGN_ENTITLEMENTS[sdk=iphoneos*]\" = \"" + IOSEntitlementPath + "\";" + ProjectFileGenerator.NewLine);
+					if (XcodeProjectFileGenerator.XcodePlatforms.Contains(UnrealTargetPlatform.IOS))
+				{
+					Content.Append("\t\t\t\tINFOPLIST_FILE = \"" + IOSInfoPlistPath + "\";" + ProjectFileGenerator.NewLine);
+					Content.Append("\t\t\t\tCODE_SIGN_ENTITLEMENTS = \"" + IOSEntitlementPath + "\";" + ProjectFileGenerator.NewLine);
 				}
-				if (TVOSRunTimeVersion != null)
+					else
 				{
-					Content.Append("\t\t\t\t\"INFOPLIST_FILE[sdk=appletvos*]\" = \"" + TVOSInfoPlistPath + "\";" + ProjectFileGenerator.NewLine);
-					Content.Append("\t\t\t\t\"CODE_SIGN_ENTITLEMENTS[sdk=appletvos*]\" = \"" + TVOSEntitlementPath + "\";" + ProjectFileGenerator.NewLine);
+					Content.Append("\t\t\t\tINFOPLIST_FILE = \"" + TVOSInfoPlistPath + "\";" + ProjectFileGenerator.NewLine);
+					Content.Append("\t\t\t\tCODE_SIGN_ENTITLEMENTS = \"" + TVOSEntitlementPath + "\";" + ProjectFileGenerator.NewLine);
 				}
-			}
-
-			// Prepare a temp Info.plist file so Xcode has some basic info about the target immediately after opening the project.
-			// This is needed for the target to pass the settings validation before code signing. UBT will overwrite this plist file later, with proper contents.
-			if (BuildHostPlatform.Current.Platform == UnrealTargetPlatform.Mac)
-			{
-				bool bCreateMacInfoPlist = !File.Exists(MacInfoPlistPath);
-				bool bCreateIOSInfoPlist = !File.Exists(IOSInfoPlistPath) && IOSRunTimeVersion != null;
-				bool bCreateTVOSInfoPlist = !File.Exists(TVOSInfoPlistPath) && TVOSRunTimeVersion != null;
-				if (bCreateMacInfoPlist || bCreateIOSInfoPlist || bCreateTVOSInfoPlist)
+				}
+				else
 				{
-					DirectoryReference? ProjectPath = GameDir;
-					DirectoryReference EngineDir = DirectoryReference.Combine(new DirectoryReference(UEDir), "Engine");
-					string GameName = Config.BuildTarget;
-					bool bIsClient = false;
-					if (ProjectPath == null)
+					Content.Append("\t\t\t\t\"INFOPLIST_FILE[sdk=macosx*]\" = \"" + MacInfoPlistPath + "\";" + ProjectFileGenerator.NewLine);
+					if (IOSRunTimeVersion != null)
 					{
-						ProjectPath = EngineDir;
+						Content.Append("\t\t\t\t\"INFOPLIST_FILE[sdk=iphoneos*]\" = \"" + IOSInfoPlistPath + "\";" + ProjectFileGenerator.NewLine);
+						Content.Append("\t\t\t\t\"CODE_SIGN_ENTITLEMENTS[sdk=iphoneos*]\" = \"" + IOSEntitlementPath + "\";" + ProjectFileGenerator.NewLine);
 					}
-					if (bIsUnrealGame)
+					if (TVOSRunTimeVersion != null)
 					{
-						ProjectPath = EngineDir;
-						GameName = "UnrealGame";
-						bIsClient = (AppName == "UnrealClient");
+						Content.Append("\t\t\t\t\"INFOPLIST_FILE[sdk=appletvos*]\" = \"" + TVOSInfoPlistPath + "\";" + ProjectFileGenerator.NewLine);
+						Content.Append("\t\t\t\t\"CODE_SIGN_ENTITLEMENTS[sdk=appletvos*]\" = \"" + TVOSEntitlementPath + "\";" + ProjectFileGenerator.NewLine);
 					}
+				}
 
-					if (bCreateMacInfoPlist)
-					{
-						Directory.CreateDirectory(Path.GetDirectoryName(MacInfoPlistPath)!);
-						UEDeployMac.GeneratePList(ProjectPath.FullName, bIsUnrealGame, GameName, Config.BuildTarget, EngineDir.FullName, MacExecutableFileName);
-					}
-					if (bCreateIOSInfoPlist)
+				// Prepare a temp Info.plist file so Xcode has some basic info about the target immediately after opening the project.
+				// This is needed for the target to pass the settings validation before code signing. UBT will overwrite this plist file later, with proper contents.
+				if (BuildHostPlatform.Current.Platform == UnrealTargetPlatform.Mac)
+				{
+					bool bCreateMacInfoPlist = !File.Exists(MacInfoPlistPath);
+					bool bCreateIOSInfoPlist = !File.Exists(IOSInfoPlistPath) && IOSRunTimeVersion != null;
+					bool bCreateTVOSInfoPlist = !File.Exists(TVOSInfoPlistPath) && TVOSRunTimeVersion != null;
+					if (bCreateMacInfoPlist || bCreateIOSInfoPlist || bCreateTVOSInfoPlist)
 					{
-						// get the receipt
-						FileReference ReceiptFilename;
+						DirectoryReference? ProjectPath = GameDir;
+						DirectoryReference EngineDir = DirectoryReference.Combine(new DirectoryReference(UEDir), "Engine");
+						string GameName = Config.BuildTarget;
+						bool bIsClient = false;
+						if (ProjectPath == null)
+						{
+							ProjectPath = EngineDir;
+						}
 						if (bIsUnrealGame)
 						{
-							ReceiptFilename = TargetReceipt.GetDefaultPath(Unreal.EngineDirectory, "UnrealGame", UnrealTargetPlatform.IOS, Config.BuildConfig, "");
+							ProjectPath = EngineDir;
+							GameName = "UnrealGame";
+							bIsClient = (AppName == "UnrealClient");
+						}
+
+						if (bCreateMacInfoPlist)
+						{
+							Directory.CreateDirectory(Path.GetDirectoryName(MacInfoPlistPath)!);
+							UEDeployMac.GeneratePList(ProjectPath.FullName, bIsUnrealGame, GameName, Config.BuildTarget, EngineDir.FullName, MacExecutableFileName);
 						}
-						else
+						if (bCreateIOSInfoPlist)
 						{
-							ReceiptFilename = TargetReceipt.GetDefaultPath(ProjectPath, GameName, UnrealTargetPlatform.IOS, Config.BuildConfig, "");
+							// get the receipt
+							FileReference ReceiptFilename;
+							if (bIsUnrealGame)
+							{
+								ReceiptFilename = TargetReceipt.GetDefaultPath(Unreal.EngineDirectory, "UnrealGame", UnrealTargetPlatform.IOS, Config.BuildConfig, "");
+							}
+							else
+							{
+								ReceiptFilename = TargetReceipt.GetDefaultPath(ProjectPath, GameName, UnrealTargetPlatform.IOS, Config.BuildConfig, "");
+							}
+							Directory.CreateDirectory(Path.GetDirectoryName(IOSInfoPlistPath)!);
+							bool bSupportPortrait, bSupportLandscape;
+							TargetReceipt? Receipt;
+							TargetReceipt.TryRead(ReceiptFilename, out Receipt);
+							bool bBuildAsFramework = UEDeployIOS.GetCompileAsDll(Receipt);
+							UEDeployIOS.GenerateIOSPList(ProjectFile, Config.BuildConfig, ProjectPath.FullName, bIsUnrealGame, GameName, bIsClient, Config.BuildTarget, EngineDir.FullName, ProjectPath + "/Binaries/IOS/Payload", null, BundleIdentifier, bBuildAsFramework, Logger, out bSupportPortrait, out bSupportLandscape);
+						}
+						if (bCreateTVOSInfoPlist)
+						{
+							Directory.CreateDirectory(Path.GetDirectoryName(TVOSInfoPlistPath)!);
+							UEDeployTVOS.GenerateTVOSPList(ProjectPath.FullName, bIsUnrealGame, GameName, bIsClient, Config.BuildTarget, EngineDir.FullName, ProjectPath + "/Binaries/TVOS/Payload", null, BundleIdentifier, Logger);
 						}
-						Directory.CreateDirectory(Path.GetDirectoryName(IOSInfoPlistPath)!);
-						bool bSupportPortrait, bSupportLandscape;
-						TargetReceipt? Receipt;
-						TargetReceipt.TryRead(ReceiptFilename, out Receipt);
-						bool bBuildAsFramework = UEDeployIOS.GetCompileAsDll(Receipt);
-						UEDeployIOS.GenerateIOSPList(ProjectFile, Config.BuildConfig, ProjectPath.FullName, bIsUnrealGame, GameName, bIsClient, Config.BuildTarget, EngineDir.FullName, ProjectPath + "/Binaries/IOS/Payload", null, BundleIdentifier, bBuildAsFramework, Logger, out bSupportPortrait, out bSupportLandscape);
-					}
-					if (bCreateTVOSInfoPlist)
-					{
-						Directory.CreateDirectory(Path.GetDirectoryName(TVOSInfoPlistPath)!);
-						UEDeployTVOS.GenerateTVOSPList(ProjectPath.FullName, bIsUnrealGame, GameName, bIsClient, Config.BuildTarget, EngineDir.FullName, ProjectPath + "/Binaries/TVOS/Payload", null, BundleIdentifier, Logger);
 					}
 				}
-			}
 
 			// #jira UE-143619: Pre Monterey macOS requires this option for a packaged app to run on iOS15 due to new code signature format. Could be removed once Monterey is minimum OS.
 			Content.Append("\t\t\t\tOTHER_CODE_SIGN_FLAGS = \"--generate-entitlement-der\";" + ProjectFileGenerator.NewLine);
@@ -1340,14 +1343,17 @@ namespace UnrealBuildTool.XcodeProjectLegacy
 
 			//#jira UE-50382 Xcode Address Sanitizer feature does not work on iOS
 			// address sanitizer dylib loader depends on the SDKROOT parameter. For macosx or default (missing, translated as macosx), the path is incorrect for iphone/appletv
-			if (XcodeProjectFileGenerator.bGeneratingRunIOSProject)
+			if (XcodeProjectFileGenerator.bGenerateRunOnlyProject)
+			{
+				if (XcodeProjectFileGenerator.XcodePlatforms.Contains(UnrealTargetPlatform.IOS))
 			{
 				Content.Append("\t\t\t\tSDKROOT = iphoneos;" + ProjectFileGenerator.NewLine);
 			}
-			else if (XcodeProjectFileGenerator.bGeneratingRunTVOSProject)
+				else
 			{
 				Content.Append("\t\t\t\tSDKROOT = appletvos;" + ProjectFileGenerator.NewLine);
 			}
+			}
 			else
 			{
 				Content.Append("\t\t\t\tSDKROOT = macosx;" + ProjectFileGenerator.NewLine);
@@ -1362,7 +1368,7 @@ namespace UnrealBuildTool.XcodeProjectLegacy
 		private void AppendLegacyTargetBuildConfiguration(StringBuilder Content, XcodeBuildConfig Config, string ConfigGuid, FileReference? ProjectFile)
 		{
 			bool bMacOnly = true;
-			if (Config.ProjectTarget!.TargetRules != null && XcodeProjectFileGenerator.ProjectFilePlatform.HasFlag(XcodeProjectFileGenerator.XcodeProjectFilePlatform.iOS))
+			if (Config.ProjectTarget!.TargetRules != null && XcodeProjectFileGenerator.XcodePlatforms.Contains(UnrealTargetPlatform.IOS))
 			{
 				if (Config.ProjectTarget.SupportedPlatforms.Contains(UnrealTargetPlatform.IOS))
 				{
@@ -1634,7 +1640,7 @@ namespace UnrealBuildTool.XcodeProjectLegacy
 				}
 			}
 
-			string DefaultConfiguration = bHasEditorConfiguration && !XcodeProjectFileGenerator.bGeneratingRunIOSProject && !XcodeProjectFileGenerator.bGeneratingRunTVOSProject ? "Development Editor" : "Development";
+			string DefaultConfiguration = bHasEditorConfiguration && !XcodeProjectFileGenerator.bGenerateRunOnlyProject ? "Development Editor" : "Development";
 
 			StringBuilder Content = new StringBuilder();
 
@@ -1807,28 +1813,6 @@ namespace UnrealBuildTool.XcodeProjectLegacy
 			File.WriteAllText(ManagementFile.FullName, Content.ToString(), new UTF8Encoding());
 		}
 
-		public static IEnumerable<UnrealTargetPlatform> GetSupportedPlatforms()
-		{
-			List<UnrealTargetPlatform> SupportedPlatforms = new List<UnrealTargetPlatform>();
-
-			if (XcodeProjectFileGenerator.ProjectFilePlatform.HasFlag(XcodeProjectFileGenerator.XcodeProjectFilePlatform.Mac))
-			{
-				SupportedPlatforms.Add(UnrealTargetPlatform.Mac);
-			}
-
-			if (XcodeProjectFileGenerator.ProjectFilePlatform.HasFlag(XcodeProjectFileGenerator.XcodeProjectFilePlatform.iOS))
-			{
-				SupportedPlatforms.Add(UnrealTargetPlatform.IOS);
-			}
-
-			if (XcodeProjectFileGenerator.ProjectFilePlatform.HasFlag(XcodeProjectFileGenerator.XcodeProjectFilePlatform.tvOS))
-			{
-				SupportedPlatforms.Add(UnrealTargetPlatform.TVOS);
-			}
-
-			return SupportedPlatforms;
-		}
-
 		public static IEnumerable<UnrealTargetConfiguration> GetSupportedConfigurations()
 		{
 			return new UnrealTargetConfiguration[] {
@@ -1849,7 +1833,7 @@ namespace UnrealBuildTool.XcodeProjectLegacy
 		{
 			foreach (Project ProjectTarget in ProjectTargets)
 			{
-				foreach (UnrealTargetPlatform Platform in GetSupportedPlatforms())
+				foreach (UnrealTargetPlatform Platform in XcodeProjectFileGenerator.XcodePlatforms)
 				{
 					foreach (UnrealTargetConfiguration Config in GetSupportedConfigurations())
 					{
@@ -1956,7 +1940,7 @@ namespace UnrealBuildTool.XcodeProjectLegacy
 			AppendCopyExtensionsBuildPhaseSection(ProjectFileContent, PBXCopyExtensionsBuildPhaseSection, CopyExtensionsBuildPhaseGuid);
 			ProjectFileContent.Append(PBXResourcesBuildPhaseSection);
 			AppendContainerItemProxySection(ProjectFileContent, BuildTargetName, BuildTargetGuid, TargetProxyGuid, ProjectGuid);
-			if (!XcodeProjectFileGenerator.bGeneratingRunIOSProject)
+			if (!XcodeProjectFileGenerator.bGenerateRunOnlyProject)
 			{
 				AppendTargetDependencySection(ProjectFileContent, BuildTargetName, BuildTargetGuid, TargetDependencyGuid, TargetProxyGuid);
 			}
diff --git a/Engine/Source/Programs/UnrealBuildTool/ToolChain/AppleToolChain.cs b/Engine/Source/Programs/UnrealBuildTool/ToolChain/AppleToolChain.cs
--- a/Engine/Source/Programs/UnrealBuildTool/ToolChain/AppleToolChain.cs
+++ b/Engine/Source/Programs/UnrealBuildTool/ToolChain/AppleToolChain.cs
@@ -161,6 +161,17 @@ namespace UnrealBuildTool
 		}
 
 		protected FileReference? ProjectFile;
+		
+		protected static bool UseModernXcode(FileReference? ProjectFile)
+		{
+			// Modern Xcode mode does this now
+			bool _bUseModernXcode = false;
+			ConfigHierarchy Ini = ConfigCache.ReadHierarchy(ConfigHierarchyType.Engine, ProjectFile?.Directory, UnrealTargetPlatform.Mac);
+			Ini.TryGetValue("XcodeConfiguration", "bUseModernXcode", out _bUseModernXcode);
+			return _bUseModernXcode;
+		}
+
+		protected bool bUseModernXcode => UseModernXcode(ProjectFile);
 
 		public AppleToolChain(FileReference? InProjectFile, ClangToolChainOptions InOptions, ILogger InLogger) : base(InOptions, InLogger)
 		{
@@ -188,6 +199,41 @@ namespace UnrealBuildTool
 			Utils.RunLocalProcessAndLogOutput(StartInfo, Logger);
 		}
 
+		/// <summary>
+		/// Writes a versions.xcconfig file for xcode to pull in when making an app plist
+		/// </summary>
+		/// <param name="LinkEnvironment"></param>
+		/// <param name="Prerequisite">FileItem describing the Prerequisite that this will this depends on (executable or similar) </param>
+		/// <param name="Graph">List of actions to be executed. Additional actions will be added to this list.</param>
+		protected FileItem UpdateVersionFile(LinkEnvironment LinkEnvironment, FileItem Prerequisite, IActionGraphBuilder Graph)
+		{
+			FileItem DestFile;
+
+			// Make the compile action
+			Action UpdateVersionAction = Graph.CreateAction(ActionType.CreateAppBundle);
+			UpdateVersionAction.WorkingDirectory = GetMacDevSrcRoot();
+			UpdateVersionAction.CommandPath = BuildHostPlatform.Current.Shell;
+			UpdateVersionAction.CommandDescription = "";
+
+			// @todo programs right nhow are sharing the Engine build version - one reason for this is that we can't get to the Engine/Programs directory from here
+			// (we can't even get to the Engine/Source/Programs directory without searching on disk), and if we did, we would create a _lot_ of Engine/Programs directories
+			// on disk that don't exist in p4. So, we just re-use Engine version, not Project version
+			//				DirectoryReference ProductDirectory = FindProductDirectory(ProjectFile, LinkEnvironment.OutputDirectory!, Graph.Makefile.TargetType);
+			DirectoryReference ProductDirectory = (ProjectFile?.Directory) ?? Unreal.EngineDirectory;
+			FileReference OutputVersionFile = FileReference.Combine(ProductDirectory, "Intermediate/Build/Versions.xcconfig");
+			DestFile = FileItem.GetItemByFileReference(OutputVersionFile);
+
+			// make path to the script
+			FileItem BundleScript = FileItem.GetItemByFileReference(FileReference.Combine(Unreal.EngineDirectory, "Build/BatchFiles/Mac/UpdateVersionAfterBuild.sh"));
+			UpdateVersionAction.CommandArguments = $"\"{BundleScript.AbsolutePath}\" {ProductDirectory} {LinkEnvironment.Platform}";
+			UpdateVersionAction.PrerequisiteItems.Add(Prerequisite);
+			UpdateVersionAction.ProducedItems.Add(DestFile);
+			UpdateVersionAction.StatusDescription = $"Updating version file: {OutputVersionFile}";
+
+			return DestFile;
+		}
+
+
 		/// <inheritdoc/>
 		protected override string EscapePreprocessorDefinition(string Definition)
 		{
