UnrealBuildTool: Revert initial visual studio selection logic back to how it was before the downgrade logic was added. It will still downgrade to 2019 if a platform sdk requires it.

#rnx

diff --git a/Engine/Source/Programs/UnrealBuildTool/ProjectFiles/VisualStudio/VCProjectFileGenerator.cs b/Engine/Source/Programs/UnrealBuildTool/ProjectFiles/VisualStudio/VCProjectFileGenerator.cs
--- a/Engine/Source/Programs/UnrealBuildTool/ProjectFiles/VisualStudio/VCProjectFileGenerator.cs
+++ b/Engine/Source/Programs/UnrealBuildTool/ProjectFiles/VisualStudio/VCProjectFileGenerator.cs
@@ -294,15 +294,24 @@ namespace UnrealBuildTool
 				List<VisualStudioInstallation> Installations = MicrosoftPlatformSDK.FindVisualStudioInstallations(Logger).Where(x => WindowsPlatform.HasCompiler(x.Compiler, WindowsArchitecture.x64, Logger)).ToList();
 
 				// Get the corresponding project file format
-				bool Vs2019Available = Installations.Any(x => x.Compiler == WindowsCompiler.VisualStudio2019);
-				bool Vs2022Available = Installations.Any(x => x.Compiler == WindowsCompiler.VisualStudio2022);
 				VCProjectFileFormat Format = VCProjectFileFormat.VisualStudio2019;
-				if (Vs2022Available && !Vs2019Available)
+				foreach (VisualStudioInstallation Installation in Installations)
 				{
-					Format = VCProjectFileFormat.VisualStudio2022;
+					if (Installation.Compiler == WindowsCompiler.VisualStudio2022)
+					{
+						Format = VCProjectFileFormat.VisualStudio2022;
+						break;
+					}
+					else if (Installation.Compiler == WindowsCompiler.VisualStudio2019)
+					{
+						Format = VCProjectFileFormat.VisualStudio2019;
+						break;
+					}
 				}
 				Settings.ProjectFileFormat = Format;
 
+				bool Vs2019Available = Installations.Any(x => x.Compiler == WindowsCompiler.VisualStudio2019);
+
 				// Allow the SDKs to override
 				foreach (UnrealTargetPlatform SupportedPlatform in SupportedPlatforms)
 				{
