Added LTO support to Android toolchain

diff --git a/Engine/Source/Programs/UnrealBuildTool/Platform/Android/AndroidToolChain.cs b/Engine/Source/Programs/UnrealBuildTool/Platform/Android/AndroidToolChain.cs
--- a/Engine/Source/Programs/UnrealBuildTool/Platform/Android/AndroidToolChain.cs
+++ b/Engine/Source/Programs/UnrealBuildTool/Platform/Android/AndroidToolChain.cs
@@ -711,9 +711,6 @@ namespace UnrealBuildTool
 				Result += " -Wno-backend-plugin";
 
 				Result += string.Format(" -fprofile-use=\"{0}.profdata\"", Path.Combine(CompileEnvironment.PGODirectory!, CompileEnvironment.PGOFilenamePrefix!));
-
-				//TODO: measure LTO.
-				//Result += " -flto=thin";
 			}
 			else if (CompileEnvironment.bPGOProfile)
 			{
@@ -722,9 +719,19 @@ namespace UnrealBuildTool
 				Result += " -fprofile-generate";
 				//  for sampling-based profile collection to generate minimal debug information:
 				Result += " -gline-tables-only";
+			}
 
-				//TODO: check LTO improves perf.
-				//Result += " -flto=thin";
+			bool bEnableLTO = CompileEnvironment.bPGOOptimize || CompileEnvironment.bPGOProfile;
+			if (CompileEnvironment.bAllowLTCG || bEnableLTO)
+			{
+				if ((Options & ClangToolChainOptions.EnableThinLTO) != 0)
+				{
+					Result += " -flto=thin";
+				}
+				else
+				{
+					Result += " -flto";
+				}
 			}
 
 			if (Architecture == "-arm64")
@@ -872,9 +879,6 @@ namespace UnrealBuildTool
 				Result += " -Wno-profile-instr-unprofiled";
 
 				Result += string.Format(" -fprofile-use=\"{0}.profdata\"", Path.Combine(LinkEnvironment.PGODirectory!, LinkEnvironment.PGOFilenamePrefix!));
-
-				//TODO: check LTO improves perf.
-				//Result += " -flto=thin";
 			}
 			else if (LinkEnvironment.bPGOProfile)
 			{
@@ -884,6 +888,19 @@ namespace UnrealBuildTool
 				Result += " -gline-tables-only";
 			}
 
+			bool bEnableLTO = LinkEnvironment.bPGOOptimize || LinkEnvironment.bPGOProfile;
+			if (LinkEnvironment.bAllowLTCG || bEnableLTO)
+			{
+				if ((Options & ClangToolChainOptions.EnableThinLTO) != 0)
+				{
+					Result += " -flto=thin";
+				}
+				else
+				{
+					Result += " -flto";
+				}
+			}
+
 			// verbose output from the linker
 			// Result += " -v";
 
diff --git a/Engine/Source/Programs/UnrealBuildTool/Platform/Android/UEBuildAndroid.cs b/Engine/Source/Programs/UnrealBuildTool/Platform/Android/UEBuildAndroid.cs
--- a/Engine/Source/Programs/UnrealBuildTool/Platform/Android/UEBuildAndroid.cs
+++ b/Engine/Source/Programs/UnrealBuildTool/Platform/Android/UEBuildAndroid.cs
@@ -456,7 +456,7 @@ namespace UnrealBuildTool
 			if (Target.bPGOOptimize || Target.bPGOProfile)
 			{
 				Logger.LogInformation("PGO {PgoType} build", Target.bPGOOptimize ? "optimize" : "profile");
-				if(Target.bPGOOptimize)
+				if (Target.bPGOOptimize)
 				{
 					CompileEnvironment.PGODirectory = Path.Combine(DirectoryReference.FromFile(Target.ProjectFile)!.FullName, "Platforms", "Android", "Build", "PGO");
 					CompileEnvironment.PGOFilenamePrefix = string.Format("{0}-Android", Target.Name);
@@ -469,6 +469,14 @@ namespace UnrealBuildTool
 				}
 			}
 
+			// For consistency with other platforms, also enable LTO whenever doing profile-guided optimizations.
+			// Obviously both PGI (instrumented) and PGO (optimized) binaries need to have that
+			if (CompileEnvironment.bPGOProfile || CompileEnvironment.bPGOOptimize)
+			{
+				CompileEnvironment.bAllowLTCG = true;
+				LinkEnvironment.bAllowLTCG = true;
+			}
+
 			CompileEnvironment.Definitions.Add("INT64_T_TYPES_NOT_LONG_LONG=1");
 
 			SetUpSpecificEnvironment(Target, CompileEnvironment, LinkEnvironment, Logger);
@@ -531,6 +539,10 @@ namespace UnrealBuildTool
 		public override UEToolChain CreateToolChain(ReadOnlyTargetRules Target)
 		{
 			ClangToolChainOptions Options = CreateToolChainOptions(Target.AndroidPlatform.TargetRules);
+			if (Target.bAllowLTCG && Target.bPreferThinLTO)
+			{
+				Options |= ClangToolChainOptions.EnableThinLTO;
+			}
 			return new AndroidToolChain(Target.ProjectFile, Target.AndroidPlatform.Architectures, Target.AndroidPlatform.GPUArchitectures, Options, Logger);
 		}
 		public virtual UEToolChain CreateTempToolChainForProject(FileReference? ProjectFile)
