Set the minimum compile sdk for Android to 30

#android

diff --git a/Engine/Source/Programs/UnrealBuildTool/Platform/Android/UEDeployAndroid.cs b/Engine/Source/Programs/UnrealBuildTool/Platform/Android/UEDeployAndroid.cs
--- a/Engine/Source/Programs/UnrealBuildTool/Platform/Android/UEDeployAndroid.cs
+++ b/Engine/Source/Programs/UnrealBuildTool/Platform/Android/UEDeployAndroid.cs
@@ -31,7 +31,7 @@ namespace UnrealBuildTool
 		private const string ANDROID_VULKAN_VALIDATION_LAYER = "libVkLayer_khronos_validation.so";
 
 		// Minimum Android SDK that must be used for Java compiling
-		readonly int MinimumSDKLevel = 28;
+		readonly int MinimumSDKLevel = 30;
 
 		// Minimum SDK version needed for App Bundles
 		readonly int MinimumSDKLevelForBundle = 21;
