Sort mode names before displaying them.

diff --git a/Engine/Source/Programs/UnrealBuildTool/UnrealBuildTool.cs b/Engine/Source/Programs/UnrealBuildTool/UnrealBuildTool.cs
--- a/Engine/Source/Programs/UnrealBuildTool/UnrealBuildTool.cs
+++ b/Engine/Source/Programs/UnrealBuildTool/UnrealBuildTool.cs
@@ -526,7 +526,9 @@ namespace UnrealBuildTool
 					// Try to get the correct mode
 					if(!ModeNameToType.TryGetValue(Options.Mode, out ModeType))
 					{
-						Logger.LogError("No mode named '{Name}'. Available modes are:\n  {ModeList}", Options.Mode, String.Join("\n  ", ModeNameToType.Keys));
+						List<string> ModuleNameList = ModeNameToType.Keys.ToList();
+						ModuleNameList.Sort(StringComparer.OrdinalIgnoreCase);
+						Logger.LogError("No mode named '{Name}'. Available modes are:\n  {ModeList}", Options.Mode, String.Join("\n  ", ModuleNameList));
 						return 1;
 					}
 				}
