- Added a mkdir to the FInalizeAppBundle script before copying to it

diff --git a/Engine/Source/Programs/UnrealBuildTool/Platform/Mac/MacToolChain.cs b/Engine/Source/Programs/UnrealBuildTool/Platform/Mac/MacToolChain.cs
--- a/Engine/Source/Programs/UnrealBuildTool/Platform/Mac/MacToolChain.cs
+++ b/Engine/Source/Programs/UnrealBuildTool/Platform/Mac/MacToolChain.cs
@@ -888,6 +888,7 @@ namespace UnrealBuildTool
 					// is where it will copy to, so it will error with reading and writing to the same location
 					string IntermediateDirectory = (ProjectFile == null ? Unreal.EngineDirectory : ProjectFile.Directory) + "/Intermediate/Mac";
 					string XcodeInputPListFile = IntermediateDirectory + "/" + ExeName + "-Info.plist";
+					AppendMacLine(FinalizeAppBundleScript, "mkdir \"{0}\"", IntermediateDirectory);
 					AppendMacLine(FinalizeAppBundleScript, FormatCopyCommand(TempInfoPlist, XcodeInputPListFile));
 					AppendMacLine(FinalizeAppBundleScript, "chmod 644 \"{0}\"", XcodeInputPListFile);
 
