UnrealBuildTool: Fix another instance where all writable source files would incorrectly add many files to the adaptive working set

#rnx

diff --git a/Engine/Source/Programs/UnrealBuildTool/Configuration/UEBuildModuleCPP.cs b/Engine/Source/Programs/UnrealBuildTool/Configuration/UEBuildModuleCPP.cs
--- a/Engine/Source/Programs/UnrealBuildTool/Configuration/UEBuildModuleCPP.cs
+++ b/Engine/Source/Programs/UnrealBuildTool/Configuration/UEBuildModuleCPP.cs
@@ -655,6 +655,11 @@ namespace UnrealBuildTool
 			{
 				Unity.GetAdaptiveFiles(Target, CPPFiles, InputFiles.HeaderFiles, CompileEnvironment, WorkingSet, Rules.ShortName ?? Name, IntermediateDirectory, Graph, 
 					out List<FileItem> NormalFiles, out List<FileItem> AdaptiveFiles);
+				if (NormalFiles.Where(file => !file.HasExtension(".gen.cpp")).Count() == 0)
+				{
+					NormalFiles = CPPFiles;
+					AdaptiveFiles.RemoveAll(new HashSet<FileItem>(NormalFiles).Contains);
+				}
 				LinkInputFiles.AddRange(CompileFilesWithToolChain(Target, ToolChain, CompileEnvironment, ModuleCompileEnvironment, NormalFiles, AdaptiveFiles, Graph, Logger).ObjectFiles);
 			}
 			else
