Fixed an issue the BF was seeing with compiling non-unity builds that use the inline gen.cpp macro.

diff --git a/Engine/Source/Programs/UnrealBuildTool/Configuration/UEBuildModuleCPP.cs b/Engine/Source/Programs/UnrealBuildTool/Configuration/UEBuildModuleCPP.cs
--- a/Engine/Source/Programs/UnrealBuildTool/Configuration/UEBuildModuleCPP.cs
+++ b/Engine/Source/Programs/UnrealBuildTool/Configuration/UEBuildModuleCPP.cs
@@ -565,29 +565,30 @@ namespace UnrealBuildTool
 					}
 				}
 
-				if (!Target.bDisableInliningGenCpps)
+				if (GeneratedFiles.Count > 0)
 				{
-					// Remove any generated files from the compile list if they are inlined
-					foreach (FileItem CPPFileItem in CPPFiles)
+					if (!Target.bDisableInliningGenCpps)
 					{
-						var ListOfInlinedGenCpps = ModuleCompileEnvironment.MetadataCache.GetListOfInlinedGeneratedCppFiles(CPPFileItem);
-						foreach (string ListOfInlinedGenCppsItem in ListOfInlinedGenCpps)
+						// Remove any generated files from the compile list if they are inlined
+						foreach (FileItem CPPFileItem in CPPFiles)
 						{
-							string? FoundGenCppFile = GeneratedFiles.FirstOrDefault((file) => String.Compare(Path.GetFileName(file), ListOfInlinedGenCppsItem + ".gen.cpp", StringComparison.InvariantCultureIgnoreCase) == 0);
-							if (FoundGenCppFile != null)
-							{
-								GeneratedFiles.Remove(FoundGenCppFile);
-							}
-							else
+							var ListOfInlinedGenCpps = ModuleCompileEnvironment.MetadataCache.GetListOfInlinedGeneratedCppFiles(CPPFileItem);
+							foreach (string ListOfInlinedGenCppsItem in ListOfInlinedGenCpps)
 							{
-								Logger.LogError("Could not find a UObject marked up header file named '{HeaderFile}'", ListOfInlinedGenCppsItem);
+								string GeneratedCppName = ListOfInlinedGenCppsItem + ".gen.cpp";
+								string? FoundGenCppFile = GeneratedFiles.FirstOrDefault((file) => String.Compare(Path.GetFileName(file), GeneratedCppName, StringComparison.InvariantCultureIgnoreCase) == 0);
+								if (FoundGenCppFile != null)
+								{
+									GeneratedFiles.Remove(FoundGenCppFile);
+								}
+								else
+								{
+									Logger.LogError("'{CPPFileItem}' is looking for a generated cpp with named '{HeaderFile}'", CPPFileItem.AbsolutePath, GeneratedCppName);
+								}
 							}
 						}
 					}
-				}
 
-				if (GeneratedFiles.Count > 0)
-				{
 					bool bMergeUnityFiles = Target.bMergeModuleAndGeneratedUnityFiles && Rules.bMergeUnityFiles;
 
 					// Create a compile environment for the generated files. We can disable creating debug info here to improve link times.
