Do not assume LTO when PGO is enabled on Android as it doesn't work well with NDK21

diff --git a/Engine/Source/Programs/UnrealBuildTool/Platform/Android/AndroidToolChain.cs b/Engine/Source/Programs/UnrealBuildTool/Platform/Android/AndroidToolChain.cs
--- a/Engine/Source/Programs/UnrealBuildTool/Platform/Android/AndroidToolChain.cs
+++ b/Engine/Source/Programs/UnrealBuildTool/Platform/Android/AndroidToolChain.cs
@@ -721,8 +721,7 @@ namespace UnrealBuildTool
 				Result += " -gline-tables-only";
 			}
 
-			bool bEnableLTO = CompileEnvironment.bPGOOptimize || CompileEnvironment.bPGOProfile;
-			if (CompileEnvironment.bAllowLTCG || bEnableLTO)
+			if (CompileEnvironment.bAllowLTCG)
 			{
 				if ((Options & ClangToolChainOptions.EnableThinLTO) != 0)
 				{
@@ -888,8 +887,7 @@ namespace UnrealBuildTool
 				Result += " -gline-tables-only";
 			}
 
-			bool bEnableLTO = LinkEnvironment.bPGOOptimize || LinkEnvironment.bPGOProfile;
-			if (LinkEnvironment.bAllowLTCG || bEnableLTO)
+			if (LinkEnvironment.bAllowLTCG)
 			{
 				if ((Options & ClangToolChainOptions.EnableThinLTO) != 0)
 				{
diff --git a/Engine/Source/Programs/UnrealBuildTool/Platform/Android/UEBuildAndroid.cs b/Engine/Source/Programs/UnrealBuildTool/Platform/Android/UEBuildAndroid.cs
--- a/Engine/Source/Programs/UnrealBuildTool/Platform/Android/UEBuildAndroid.cs
+++ b/Engine/Source/Programs/UnrealBuildTool/Platform/Android/UEBuildAndroid.cs
@@ -469,14 +469,6 @@ namespace UnrealBuildTool
 				}
 			}
 
-			// For consistency with other platforms, also enable LTO whenever doing profile-guided optimizations.
-			// Obviously both PGI (instrumented) and PGO (optimized) binaries need to have that
-			if (CompileEnvironment.bPGOProfile || CompileEnvironment.bPGOOptimize)
-			{
-				CompileEnvironment.bAllowLTCG = true;
-				LinkEnvironment.bAllowLTCG = true;
-			}
-
 			CompileEnvironment.Definitions.Add("INT64_T_TYPES_NOT_LONG_LONG=1");
 
 			SetUpSpecificEnvironment(Target, CompileEnvironment, LinkEnvironment, Logger);
