Implemented ThinLTO cache and pruning arguments in UBT

diff --git a/Engine/Source/Programs/UnrealBuildTool/Configuration/TargetRules.cs b/Engine/Source/Programs/UnrealBuildTool/Configuration/TargetRules.cs
--- a/Engine/Source/Programs/UnrealBuildTool/Configuration/TargetRules.cs
+++ b/Engine/Source/Programs/UnrealBuildTool/Configuration/TargetRules.cs
@@ -1561,6 +1561,21 @@ namespace UnrealBuildTool
 		[XmlConfigFile(Category = "BuildConfiguration")]
 		public bool bPreferThinLTO = false;
 
+		/// <summary>
+		/// Directory where to put the ThinLTO cache on supported platforms.
+		/// </summary>
+		[CommandLine("-ThinLTOCacheDirectory")]
+		[XmlConfigFile(Category = "BuildConfiguration")]
+		public string? ThinLTOCacheDirectory = null;
+
+		/// <summary>
+		/// Arguments that will be applied to prune the ThinLTO cache on supported platforms.
+		/// Arguments will only be applied if ThinLTOCacheDirectory is set.
+		/// </summary>
+		[CommandLine("-ThinLTOCachePruningArguments")]
+		[XmlConfigFile(Category = "BuildConfiguration")]
+		public string? ThinLTOCachePruningArguments = null;
+
 		/// <summary>
 		/// Whether to enable Profile Guided Optimization (PGO) instrumentation in this build.
 		/// </summary>
@@ -3289,6 +3304,16 @@ namespace UnrealBuildTool
 			get { return Inner.CrashDiagnosticDirectory; }
 		}
 
+		public string? ThinLTOCacheDirectory
+		{
+			get { return Inner.ThinLTOCacheDirectory; }
+		}
+
+		public string? ThinLTOCachePruningArguments
+		{
+			get { return Inner.ThinLTOCachePruningArguments; }
+		}
+
 		public string? BundleVersion
 		{
 			get { return Inner.BundleVersion; }
diff --git a/Engine/Source/Programs/UnrealBuildTool/Configuration/UEBuildTarget.cs b/Engine/Source/Programs/UnrealBuildTool/Configuration/UEBuildTarget.cs
--- a/Engine/Source/Programs/UnrealBuildTool/Configuration/UEBuildTarget.cs
+++ b/Engine/Source/Programs/UnrealBuildTool/Configuration/UEBuildTarget.cs
@@ -3918,6 +3918,8 @@ namespace UnrealBuildTool
 			GlobalLinkEnvironment.bUsePDBFiles = Rules.bUsePDBFiles;
 			GlobalLinkEnvironment.PackagePath = Rules.PackagePath;
 			GlobalLinkEnvironment.CrashDiagnosticDirectory = Rules.CrashDiagnosticDirectory;
+			GlobalLinkEnvironment.ThinLTOCacheDirectory = Rules.ThinLTOCacheDirectory;
+			GlobalLinkEnvironment.ThinLTOCachePruningArguments = Rules.ThinLTOCachePruningArguments;
 			GlobalLinkEnvironment.BundleDirectory = BuildPlatform.GetBundleDirectory(Rules, Binaries[0].OutputFilePaths);
 			GlobalLinkEnvironment.BundleVersion = Rules.BundleVersion;
 			GlobalLinkEnvironment.bAllowLTCG = Rules.bAllowLTCG;
diff --git a/Engine/Source/Programs/UnrealBuildTool/System/LinkEnvironment.cs b/Engine/Source/Programs/UnrealBuildTool/System/LinkEnvironment.cs
--- a/Engine/Source/Programs/UnrealBuildTool/System/LinkEnvironment.cs
+++ b/Engine/Source/Programs/UnrealBuildTool/System/LinkEnvironment.cs
@@ -277,6 +277,17 @@ namespace UnrealBuildTool
 		/// </summary>
 		public string? CrashDiagnosticDirectory;
 
+		/// <summary>
+		/// Directory where to put the ThinLTO cache for platforms that support it
+		/// </summary>
+		public string? ThinLTOCacheDirectory;
+
+		/// <summary>
+		/// Arguments that will be applied to prune the ThinLTO cache for platforms that support it.
+		/// The arguments will only be applied if ThinLTOCacheDirectory is set.
+		/// </summary>
+		public string? ThinLTOCachePruningArguments;
+
 		/// <summary>
 		/// Bundle version for Mac apps
 		/// </summary>
@@ -382,6 +393,8 @@ namespace UnrealBuildTool
 			bPrintTimingInfo = Other.bPrintTimingInfo;
 			PackagePath = Other.PackagePath;
 			CrashDiagnosticDirectory = Other.CrashDiagnosticDirectory;
+			ThinLTOCacheDirectory = Other.ThinLTOCacheDirectory;
+			ThinLTOCachePruningArguments = Other.ThinLTOCachePruningArguments;
 			BundleVersion = Other.BundleVersion;
 			InstallName = Other.InstallName;
 			InputFiles.AddRange(Other.InputFiles);
