UBT: Add support for throwing exceptions with structured logging events.

diff --git a/Engine/Source/Programs/UnrealBuildTool/System/BuildException.cs b/Engine/Source/Programs/UnrealBuildTool/System/BuildException.cs
--- a/Engine/Source/Programs/UnrealBuildTool/System/BuildException.cs
+++ b/Engine/Source/Programs/UnrealBuildTool/System/BuildException.cs
@@ -1,5 +1,7 @@
 // Copyright Epic Games, Inc. All Rights Reserved.
 
+using EpicGames.Core;
+using Microsoft.Extensions.Logging;
 using System;
 using System.Collections.Generic;
 using System.Text;
@@ -25,7 +27,7 @@ namespace UnrealBuildTool
 		/// </summary>
 		/// <param name="InnerException">An inner exception to wrap</param>
 		/// <param name="Message">The error message to display.</param>
-		public BuildException(Exception InnerException, string Message)
+		public BuildException(Exception? InnerException, string Message)
 			: base(Message, InnerException)
 		{
 		}
@@ -35,7 +37,7 @@ namespace UnrealBuildTool
 		/// </summary>
 		/// <param name="Format">Formatting string for the error message</param>
 		/// <param name="Arguments">Arguments for the formatting string</param>
-		public BuildException(string Format, params object?[] Arguments) 
+		public BuildException(string Format, params object?[] Arguments)
 			: base(String.Format(Format, Arguments))
 		{
 		}
@@ -46,7 +48,7 @@ namespace UnrealBuildTool
 		/// <param name="InnerException">The inner exception being wrapped</param>
 		/// <param name="Format">Format for the message string</param>
 		/// <param name="Arguments">Format arguments</param>
-		public BuildException(Exception InnerException, string Format, params object?[] Arguments) 
+		public BuildException(Exception InnerException, string Format, params object?[] Arguments)
 			: base(String.Format(Format, Arguments), InnerException)
 		{
 		}
@@ -60,4 +62,69 @@ namespace UnrealBuildTool
 			return Message;
 		}
 	}
+
+	/// <summary>
+	/// Implementation of <see cref="BuildException"/> that captures a full structured logging event.
+	/// </summary>
+	class BuildLogEventException : BuildException
+	{
+		/// <summary>
+		/// The event object
+		/// </summary>
+		public LogEvent Event { get; }
+
+		/// <summary>
+		/// Constructor
+		/// </summary>
+		/// <param name="Event">Event to construct from</param>
+		public BuildLogEventException(LogEvent Event)
+			: this(null, Event)
+		{
+		}
+
+		/// <summary>
+		/// Constructor
+		/// </summary>
+		/// <param name="InnerException">The inner exception</param>
+		/// <param name="Event">Event to construct from</param>
+		public BuildLogEventException(Exception? InnerException, LogEvent Event)
+			: base(InnerException, Event.ToString())
+		{
+			this.Event = Event;
+		}
+
+		/// <summary>
+		/// Constructor
+		/// </summary>
+		/// <param name="Format">Structured logging format string</param>
+		/// <param name="Args">Argument objects</param>
+		public BuildLogEventException(string Format, params object[] Args)
+			: this(LogEvent.Create(LogLevel.Error, Format, Args))
+		{
+		}
+
+		/// <summary>
+		/// Constructor
+		/// </summary>
+		/// <param name="InnerException">Inner exception to wrap</param>
+		/// <param name="Format">Structured logging format string</param>
+		/// <param name="Args">Argument objects</param>
+		public BuildLogEventException(Exception? InnerException, string Format, params object[] Args)
+			: this(InnerException, LogEvent.Create(LogLevel.Error, default, InnerException, Format, Args))
+		{
+		}
+
+		/// <summary>
+		/// Constructor
+		/// </summary>
+		/// <param name="EventId">Event id for the error</param>
+		/// <param name="InnerException">Inner exception to wrap</param>
+		/// <param name="Format">Structured logging format string</param>
+		/// <param name="Args">Argument objects</param>
+		public BuildLogEventException(Exception? InnerException, EventId EventId, string Format, params object[] Args)
+			: this(InnerException, LogEvent.Create(LogLevel.Error, EventId, InnerException, Format, Args))
+		{
+		}
+	}
 }
+
diff --git a/Engine/Source/Programs/UnrealBuildTool/UnrealBuildTool.cs b/Engine/Source/Programs/UnrealBuildTool/UnrealBuildTool.cs
--- a/Engine/Source/Programs/UnrealBuildTool/UnrealBuildTool.cs
+++ b/Engine/Source/Programs/UnrealBuildTool/UnrealBuildTool.cs
@@ -644,6 +644,13 @@ namespace UnrealBuildTool
 				Logger.LogDebug(Ex, "{Ex}", ExceptionUtils.FormatExceptionDetails(Ex));
 				return (int)Ex.Result;
 			}
+			catch (BuildLogEventException Ex)
+			{
+				// BuildExceptions should have nicely formatted messages. We can log these directly.
+				Logger.Log(Ex.Event.Level, Ex.Event.Id, Ex.Event, Ex, (s, e) => s.ToString());
+				Logger.LogDebug(Ex, "{Ex}", ExceptionUtils.FormatExceptionDetails(Ex));
+				return (int)CompilationResult.OtherCompilationError;
+			}
 			catch (BuildException Ex)
 			{
 				// BuildExceptions should have nicely formatted messages. We can log these directly.
