UBT: Use structured log event to include path to the plugin file in error message.

diff --git a/Engine/Source/Programs/UnrealBuildTool/Configuration/UEBuildTarget.cs b/Engine/Source/Programs/UnrealBuildTool/Configuration/UEBuildTarget.cs
--- a/Engine/Source/Programs/UnrealBuildTool/Configuration/UEBuildTarget.cs
+++ b/Engine/Source/Programs/UnrealBuildTool/Configuration/UEBuildTarget.cs
@@ -1,4 +1,4 @@
-﻿// Copyright Epic Games, Inc. All Rights Reserved.
+// Copyright Epic Games, Inc. All Rights Reserved.
 
 using System;
 using System.Collections.Generic;
@@ -3621,7 +3621,8 @@ namespace UnrealBuildTool
 				// Disable any plugin which does not support the target platform. The editor should update such references in the .uproject file on load.
 				if (!Rules.bIncludePluginsForTargetPlatforms && !Info.Descriptor.SupportsTargetPlatform(Platform))
 				{
-					throw new BuildException("{0} is referenced via {1} with a mismatched 'SupportedTargetPlatforms' field. This will cause problems in packaged builds, because the .uplugin file will not be staged. Launch the editor to update references from your project file, or update references from other plugins manually.", Info.File.GetFileName(), ReferenceChain);
+					LogValue PluginLogValue = LogValue.SourceFile(Info.File, Info.File.GetFileName());
+					throw new BuildLogEventException("{Plugin} is referenced via {ReferenceChain} with a mismatched 'SupportedTargetPlatforms' field. This will cause problems in packaged builds, because the .uplugin file will not be staged. Launch the editor to update references from your project file, or update references from other plugins manually.", PluginLogValue, ReferenceChain);
 				}
 
 				// Disable any plugin that requires the build platform
