UnrealBuildTool: Don't throw entire action graph as an exception when there is a cycle

#rnx

diff --git a/Engine/Source/Programs/UnrealBuildTool/Modes/BuildMode.cs b/Engine/Source/Programs/UnrealBuildTool/Modes/BuildMode.cs
--- a/Engine/Source/Programs/UnrealBuildTool/Modes/BuildMode.cs
+++ b/Engine/Source/Programs/UnrealBuildTool/Modes/BuildMode.cs
@@ -346,7 +346,7 @@ namespace UnrealBuildTool
 				}
 
 				// Link all the actions together
-				ActionGraph.Link(MergedActions);
+				ActionGraph.Link(MergedActions, Logger);
 
 				// Get all the actions that are prerequisites for these targets. This forms the list of actions that we want executed.
 				List<LinkedAction> PrerequisiteActions = ActionGraph.GatherPrerequisiteActions(MergedActions, MergedOutputItems);
@@ -502,7 +502,7 @@ namespace UnrealBuildTool
 
 				// Link the action graph again to sort it
 				List<LinkedAction> MergedActionsToExecute = ActionToOutdatedFlag.Where(x => x.Value).Select(x => x.Key).ToList();
-				ActionGraph.Link(MergedActionsToExecute);
+				ActionGraph.Link(MergedActionsToExecute, Logger);
 
 				// Allow hot reload to override the actions
 				int HotReloadTargetIdx = -1;
@@ -533,7 +533,7 @@ namespace UnrealBuildTool
 					Dictionary<LinkedAction, bool> LiveActionToOutdatedFlag = new Dictionary<LinkedAction, bool>(MergedActionsToExecute.Count);
 					ActionGraph.GatherAllOutdatedActions(MergedActionsToExecute, History, LiveActionToOutdatedFlag, CppDependencies, BuildConfiguration.bIgnoreOutdatedImportLibraries, Logger);
 					List<LinkedAction> LiveCodingActionsToExecute = LiveActionToOutdatedFlag.Where(x => x.Value).Select(x => x.Key).ToList();
-					ActionGraph.Link(LiveCodingActionsToExecute);
+					ActionGraph.Link(LiveCodingActionsToExecute, Logger);
 					MergedActionsToExecute = LiveCodingActionsToExecute;
 				}
 
diff --git a/Engine/Source/Programs/UnrealBuildTool/Modes/ExecuteMode.cs b/Engine/Source/Programs/UnrealBuildTool/Modes/ExecuteMode.cs
--- a/Engine/Source/Programs/UnrealBuildTool/Modes/ExecuteMode.cs
+++ b/Engine/Source/Programs/UnrealBuildTool/Modes/ExecuteMode.cs
@@ -54,7 +54,7 @@ namespace UnrealBuildTool
 			// Link the action graph
 			using (GlobalTracer.Instance.BuildSpan("ActionGraph.Link()").StartActive())
 			{
-				ActionGraph.Link(Actions);
+				ActionGraph.Link(Actions, Logger);
 			}
 
 			// Execute the actions
diff --git a/Engine/Source/Programs/UnrealBuildTool/Modes/GenerateClangDatabase.cs b/Engine/Source/Programs/UnrealBuildTool/Modes/GenerateClangDatabase.cs
--- a/Engine/Source/Programs/UnrealBuildTool/Modes/GenerateClangDatabase.cs
+++ b/Engine/Source/Programs/UnrealBuildTool/Modes/GenerateClangDatabase.cs
@@ -151,7 +151,7 @@ namespace UnrealBuildTool
 						// Create the makefile
 						TargetMakefile Makefile = Target.Build(BuildConfiguration, WorkingSet, TargetDescriptor, Logger);
 						List<LinkedAction> Actions = Makefile.Actions.ConvertAll(x => new LinkedAction(x, TargetDescriptor));
-						ActionGraph.Link(Actions);
+						ActionGraph.Link(Actions, Logger);
 
 						// Filter all the actions to execute
 						HashSet<FileItem> PrerequisiteItems = new HashSet<FileItem>(Makefile.Actions.SelectMany(x => x.ProducedItems).Where(x => x.HasExtension(".h") || x.HasExtension(".cpp")));
diff --git a/Engine/Source/Programs/UnrealBuildTool/Modes/JsonExportMode.cs b/Engine/Source/Programs/UnrealBuildTool/Modes/JsonExportMode.cs
--- a/Engine/Source/Programs/UnrealBuildTool/Modes/JsonExportMode.cs
+++ b/Engine/Source/Programs/UnrealBuildTool/Modes/JsonExportMode.cs
@@ -59,7 +59,7 @@ namespace UnrealBuildTool
 						// Create the makefile
 						TargetMakefile Makefile = Target.Build(BuildConfiguration, WorkingSet, TargetDescriptor, Logger);
 						List<LinkedAction> Actions = Makefile.Actions.ConvertAll(x => new LinkedAction(x, TargetDescriptor));
-						ActionGraph.Link(Actions);
+						ActionGraph.Link(Actions, Logger);
 
 						// Filter all the actions to execute
 						HashSet<FileItem> PrerequisiteItems = new HashSet<FileItem>(Makefile.Actions.SelectMany(x => x.ProducedItems).Where(x => x.HasExtension(".h") || x.HasExtension(".cpp")));
diff --git a/Engine/Source/Programs/UnrealBuildTool/System/ActionGraph.cs b/Engine/Source/Programs/UnrealBuildTool/System/ActionGraph.cs
--- a/Engine/Source/Programs/UnrealBuildTool/System/ActionGraph.cs
+++ b/Engine/Source/Programs/UnrealBuildTool/System/ActionGraph.cs
@@ -40,7 +40,8 @@ namespace UnrealBuildTool
 		/// Links the actions together and sets up their dependencies
 		/// </summary>
 		/// <param name="Actions">List of actions in the graph</param>
-		public static void Link(List<LinkedAction> Actions)
+		/// <param name="Logger">Logger for output</param>
+		public static void Link(List<LinkedAction> Actions, ILogger Logger)
 		{
 			// Build a map from item to its producing action
 			Dictionary<FileItem, LinkedAction> ItemToProducingAction = new Dictionary<FileItem, LinkedAction>();
@@ -53,7 +54,7 @@ namespace UnrealBuildTool
 			}
 
 			// Check for cycles
-			DetectActionGraphCycles(Actions, ItemToProducingAction);
+			DetectActionGraphCycles(Actions, ItemToProducingAction, Logger);
 
 			// Use this map to add all the prerequisite actions
 			foreach (LinkedAction Action in Actions)
@@ -437,7 +438,7 @@ namespace UnrealBuildTool
 		/// <summary>
 		/// Checks for cycles in the action graph.
 		/// </summary>
-		static void DetectActionGraphCycles(List<LinkedAction> Actions, Dictionary<FileItem, LinkedAction> ItemToProducingAction)
+		static void DetectActionGraphCycles(List<LinkedAction> Actions, Dictionary<FileItem, LinkedAction> ItemToProducingAction, ILogger Logger)
 		{
 			// Starting with actions that only depend on non-produced items, iteratively expand a set of actions that are only dependent on
 			// non-cyclical actions.
@@ -506,11 +507,11 @@ namespace UnrealBuildTool
 				}
 
 				// Describe the cyclical actions.
-				string CycleDescription = "";
 				foreach (LinkedAction Action in Actions)
 				{
 					if (!ActionIsNonCyclical.ContainsKey(Action))
 					{
+						string CycleDescription = "";
 						CycleDescription += $"Action #{ActionToIndex[Action].ToString()}: {Action.CommandPath}\n";
 						CycleDescription += $"\twith arguments: {Action.CommandArguments}\n";
 						foreach (FileItem PrerequisiteItem in Action.PrerequisiteItems)
@@ -554,11 +555,11 @@ namespace UnrealBuildTool
 							CycleDescription += "\t\tNone?? Coding error!\n";
 						}
 
-						CycleDescription += "\n\n";
+						Logger.LogInformation(CycleDescription);
 					}
 				}
 
-				throw new BuildException($"Action graph contains cycle!\n\n{CycleDescription}");
+				throw new BuildException($"Action graph contains cycle!");
 			}
 		}
 
