Fixed an issue where the makefile would not be invalidated if you switched back and forth between using inline gen cpp files when unity files were disabled.  This would sometimes result in linker issues.

diff --git a/Engine/Source/Programs/UnrealBuildTool/Configuration/UEBuildTarget.cs b/Engine/Source/Programs/UnrealBuildTool/Configuration/UEBuildTarget.cs
--- a/Engine/Source/Programs/UnrealBuildTool/Configuration/UEBuildTarget.cs
+++ b/Engine/Source/Programs/UnrealBuildTool/Configuration/UEBuildTarget.cs
@@ -1999,6 +1999,19 @@ namespace UnrealBuildTool
 				}
 			}
 
+			// Cache inline gen cpp data
+			using (GlobalTracer.Instance.BuildSpan("CacheInlineGenCppData").StartActive())
+			{
+				Makefile.DirectoryToSourceFiles.AsParallel()
+					.ForAll((kvp) =>
+					{
+						foreach (TargetMakefileSourceFileInfo SourceFileInfo in kvp.Value)
+						{
+							SourceFileInfo.InlineGenCppHash = TargetMakefileSourceFileInfo.CalculateInlineGenCppHash(MetadataCache.GetListOfInlinedGeneratedCppFiles(SourceFileInfo.SourceFileItem!));
+						}
+					});
+			}
+
 			// Prepare all the runtime dependencies, copying them from their source folders if necessary
 			List<RuntimeDependency> RuntimeDependencies = new List<RuntimeDependency>();
 			Dictionary<FileReference, FileReference> RuntimeDependencyTargetFileToSourceFile = new Dictionary<FileReference, FileReference>();
diff --git a/Engine/Source/Programs/UnrealBuildTool/System/TargetMakefile.cs b/Engine/Source/Programs/UnrealBuildTool/System/TargetMakefile.cs
--- a/Engine/Source/Programs/UnrealBuildTool/System/TargetMakefile.cs
+++ b/Engine/Source/Programs/UnrealBuildTool/System/TargetMakefile.cs
@@ -13,6 +13,56 @@ using UnrealBuildBase;
 
 namespace UnrealBuildTool
 {
+	class TargetMakefileSourceFileInfo
+	{
+		public FileItem? SourceFileItem;
+		
+		/// <summary>
+		/// Hash of all the gen.cpp files being inlined in this source file
+		/// </summary>
+		public int InlineGenCppHash = 0;
+
+		public TargetMakefileSourceFileInfo(FileItem FileItem)
+		{
+			SourceFileItem = FileItem;
+		}
+
+		public TargetMakefileSourceFileInfo(BinaryArchiveReader Reader)
+		{
+			SourceFileItem = Reader.ReadFileItem();
+			InlineGenCppHash = Reader.ReadInt();
+		}
+
+		public void Write(BinaryArchiveWriter Writer)
+		{
+			Writer.WriteFileItem(SourceFileItem);
+			Writer.WriteInt(InlineGenCppHash);
+		}
+
+		/// <summary>
+		/// Helper function to calculate a hash for the passed in inline gen.cpp file names
+		/// </summary>
+		/// <param name="InlinedGenCppNames">Gen.cpp file names</param>
+		/// <returns>Hash of file names</returns>
+		static public int CalculateInlineGenCppHash(IEnumerable<string> InlinedGenCppNames)
+		{
+			int InlineHash = 0;
+			foreach (string FileName in InlinedGenCppNames)
+			{
+				int Result = 0;
+				if (FileName != null)
+				{
+					for (int Idx = 0; Idx < FileName.Length; Idx++)
+					{
+						Result = (Result * 13) + FileName[Idx];
+					}
+				}
+				InlineHash += Result;
+			}
+			return InlineHash;
+		}
+	}
+
 	class TargetMakefileBuilder : IActionGraphBuilder
 	{
 		private readonly ILogger Logger;
@@ -62,7 +112,7 @@ namespace UnrealBuildTool
 		/// <inheritdoc/>
 		public void AddSourceFiles(DirectoryItem SourceDir, FileItem[] SourceFiles)
 		{
-			Makefile.DirectoryToSourceFiles[SourceDir] = SourceFiles;
+			Makefile.DirectoryToSourceFiles[SourceDir] = SourceFiles.Select(fi => new TargetMakefileSourceFileInfo(fi)).ToArray();
 		}
 
 		/// <inheritdoc/>
@@ -242,7 +292,7 @@ namespace UnrealBuildTool
 		/// <summary>
 		/// Set of all source directories. Any files being added or removed from these directories will invalidate the makefile.
 		/// </summary>
-		public Dictionary<DirectoryItem, FileItem[]> DirectoryToSourceFiles;
+		public Dictionary<DirectoryItem, TargetMakefileSourceFileInfo[]> DirectoryToSourceFiles;
 
 		/// <summary>
 		/// The set of source files that UnrealBuildTool determined to be part of the programmer's "working set". Used for adaptive non-unity builds.
@@ -335,7 +385,7 @@ namespace UnrealBuildTool
 			this.ModuleNameToOutputItems = new Dictionary<string, FileItem[]>(StringComparer.OrdinalIgnoreCase);
 			this.HotReloadModuleNames = new HashSet<string>(StringComparer.OrdinalIgnoreCase);
 			this.SourceDirectories = new List<DirectoryItem>();
-			this.DirectoryToSourceFiles = new Dictionary<DirectoryItem, FileItem[]>();
+			this.DirectoryToSourceFiles = new();
 			this.WorkingSet = new HashSet<FileItem>();
 			this.CandidatesForWorkingSet = new HashSet<FileItem>();
 			this.UObjectModules = new List<UHTModuleInfo>();
@@ -381,7 +431,7 @@ namespace UnrealBuildTool
 			ModuleNameToOutputItems = Reader.ReadDictionary(() => Reader.ReadString()!, () => Reader.ReadArray(() => Reader.ReadFileItem()), StringComparer.OrdinalIgnoreCase)!;
 			HotReloadModuleNames = Reader.ReadHashSet(() => Reader.ReadString(), StringComparer.OrdinalIgnoreCase)!;
 			SourceDirectories = Reader.ReadList(() => Reader.ReadDirectoryItem())!;
-			DirectoryToSourceFiles = Reader.ReadDictionary(() => Reader.ReadDirectoryItem()!, () => Reader.ReadArray(() => Reader.ReadFileItem()))!;
+			DirectoryToSourceFiles = Reader.ReadDictionary(() => Reader.ReadDirectoryItem()!, () => Reader.ReadArray(() =>  new TargetMakefileSourceFileInfo(Reader)))!;
 			WorkingSet = Reader.ReadHashSet(() => Reader.ReadFileItem())!;
 			CandidatesForWorkingSet = Reader.ReadHashSet(() => Reader.ReadFileItem())!;
 			UObjectModules = Reader.ReadList(() => new UHTModuleInfo(Reader))!;
@@ -426,7 +476,7 @@ namespace UnrealBuildTool
 			Writer.WriteDictionary(ModuleNameToOutputItems, k => Writer.WriteString(k), v => Writer.WriteArray(v, e => Writer.WriteFileItem(e)));
 			Writer.WriteHashSet(HotReloadModuleNames, x => Writer.WriteString(x));
 			Writer.WriteList(SourceDirectories, x => Writer.WriteDirectoryItem(x));
-			Writer.WriteDictionary(DirectoryToSourceFiles, k => Writer.WriteDirectoryItem(k), v => Writer.WriteArray(v, e => Writer.WriteFileItem(e)));
+			Writer.WriteDictionary(DirectoryToSourceFiles, k => Writer.WriteDirectoryItem(k), v => Writer.WriteArray(v, e => e.Write(Writer)));
 			Writer.WriteHashSet(WorkingSet, x => Writer.WriteFileItem(x));
 			Writer.WriteHashSet(CandidatesForWorkingSet, x => Writer.WriteFileItem(x));
 			Writer.WriteList(UObjectModules, e => e.Write(Writer));
@@ -654,23 +704,26 @@ namespace UnrealBuildTool
 				// Get the list of excluded folder names for this platform
 				ReadOnlyHashSet<string> ExcludedFolderNames = UEBuildPlatform.GetBuildPlatform(Platform).GetExcludedFolderNames();
 
+				// Load the metadata cache
+				SourceFileMetadataCache MetadataCache = SourceFileMetadataCache.CreateHierarchy(ProjectFile, Logger);
+
 				// Check if any source files have been added or removed
-				foreach((DirectoryItem InputDirectory, FileItem[] Files) in Makefile.DirectoryToSourceFiles)
+				foreach ((DirectoryItem InputDirectory, TargetMakefileSourceFileInfo[] SourceFileInfos) in Makefile.DirectoryToSourceFiles)
 				{
 					if(!InputDirectory.Exists || InputDirectory.LastWriteTimeUtc > Makefile.CreateTimeUtc)
 					{
 						FileItem[] SourceFiles = UEBuildModuleCPP.GetSourceFiles(InputDirectory);
-						if(SourceFiles.Length < Files.Length)
+						if(SourceFiles.Length < SourceFileInfos.Length)
 						{
 							ReasonNotLoaded = "source file removed";
 							return false;
 						}
-						else if(SourceFiles.Length > Files.Length)
+						else if(SourceFiles.Length > SourceFileInfos.Length)
 						{
 							ReasonNotLoaded = "source file added";
 							return false;
 						}
-						else if(SourceFiles.Intersect(Files).Count() != SourceFiles.Length)
+						else if(SourceFiles.Intersect(SourceFileInfos.Select((sfi) => sfi.SourceFileItem)).Count() != SourceFiles.Length)
 						{
 							ReasonNotLoaded = "source file modified";
 							return false;
@@ -684,6 +737,17 @@ namespace UnrealBuildTool
 								return false;
 							}
 						}
+
+						// Make sure the inlined gen.cpp files didn't change
+						foreach (TargetMakefileSourceFileInfo SourceFileInfo in SourceFileInfos)
+						{
+							int InlineHash = TargetMakefileSourceFileInfo.CalculateInlineGenCppHash(MetadataCache.GetListOfInlinedGeneratedCppFiles(SourceFileInfo.SourceFileItem!));
+							if (SourceFileInfo.InlineGenCppHash != InlineHash)
+							{
+								ReasonNotLoaded = "inlined gen.cpp files changed";
+								return false;
+							}
+						}
 					}
 				}
 
@@ -733,9 +797,6 @@ namespace UnrealBuildTool
 					}
 				}
 
-				// Load the metadata cache
-				SourceFileMetadataCache MetadataCache = SourceFileMetadataCache.CreateHierarchy(ProjectFile, Logger);
-
 				// Find the set of files that contain reflection markup
 				ConcurrentBag<FileItem> NewFilesWithMarkupBag = new ConcurrentBag<FileItem>();
 				using (ThreadPoolWorkQueue Queue = new ThreadPoolWorkQueue())
