Optimized how files were iterated on when scanning the directories.

diff --git a/Engine/Source/Programs/UnrealBuildTool/System/ExternalExecution.cs b/Engine/Source/Programs/UnrealBuildTool/System/ExternalExecution.cs
--- a/Engine/Source/Programs/UnrealBuildTool/System/ExternalExecution.cs
+++ b/Engine/Source/Programs/UnrealBuildTool/System/ExternalExecution.cs
@@ -433,19 +433,14 @@ namespace UnrealBuildTool
 		/// <param name="Headers">Receives the list of headers that was found</param>
 		static void FindHeaders(DirectoryItem BaseDirectory, ReadOnlyHashSet<string> ExcludeFolders, List<FileItem> Headers)
 		{
-			foreach (FileItem File in BaseDirectory.EnumerateFiles())
+			if (BaseDirectory.TryGetFile(".ubtignore", out FileItem? OutIgnoreFile))
 			{
-				if (File.Name == ".ubtignore")
-				{
-					return;
-				}
-
-				if (File.HasExtension(".h"))
-				{
-					Headers.Add(File);
-				}
+				return;
 			}
 
+			// Check for all the headers in this folder
+			Headers.AddRange(BaseDirectory.EnumerateFiles().Where((fi) => fi.HasExtension(".h")));
+
 			foreach (DirectoryItem SubDirectory in BaseDirectory.EnumerateDirectories())
 			{
 				if(!ExcludeFolders.Contains(SubDirectory.Name))
diff --git a/Engine/Source/Programs/UnrealBuildTool/System/FileMetadataPrefetch.cs b/Engine/Source/Programs/UnrealBuildTool/System/FileMetadataPrefetch.cs
--- a/Engine/Source/Programs/UnrealBuildTool/System/FileMetadataPrefetch.cs
+++ b/Engine/Source/Programs/UnrealBuildTool/System/FileMetadataPrefetch.cs
@@ -160,27 +160,11 @@ namespace UnrealBuildTool
 		{
 			foreach(DirectoryItem SubDirectory in Directory.EnumerateDirectories())
 			{
-				bool bIsUPlugin = false;
-				bool bIsIgnore = false;
-				foreach (FileItem FileItem in SubDirectory.EnumerateFiles())
-				{
-					if (FileItem.HasExtension(".uplugin"))
-					{
-						bIsUPlugin = true;
-						break;
-					}
-					if (FileItem.Name == ".ubtignore")
-					{
-						bIsIgnore = true;
-						break;
-					}
-				}
-
-				if (bIsUPlugin)
+				if (SubDirectory.EnumerateFiles().Any((fi) => fi.HasExtension(".uplugin")))
 				{
 					Enqueue(() => ScanDirectoryTree(DirectoryItem.Combine(SubDirectory, "Source")));
 				}
-				else if (!bIsIgnore)
+				else if (!SubDirectory.TryGetFile(".ubtignore", out FileItem? OutFile))
 				{
 					Enqueue(() => ScanPluginFolder(SubDirectory));
 				}
@@ -193,12 +177,9 @@ namespace UnrealBuildTool
 		/// <param name="Directory">Root of the directory tree</param>
 		static void ScanDirectoryTree(DirectoryItem Directory)
 		{
-			foreach (FileItem FileItem in Directory.EnumerateFiles())
+			if (Directory.TryGetFile(".ubtignore", out FileItem? OutFile))
 			{
-				if (FileItem.Name == ".ubtignore")
-				{
-					return;
-				}
+				return;
 			}
 
 			foreach (DirectoryItem SubDirectory in Directory.EnumerateDirectories())
