Fixed a bug where the dependency file info was not cached properly if it changed. This resulted in the dependency files getting parsed every time we started UBT.

diff --git a/Engine/Source/Programs/UnrealBuildTool/System/CppDependencyCache.cs b/Engine/Source/Programs/UnrealBuildTool/System/CppDependencyCache.cs
--- a/Engine/Source/Programs/UnrealBuildTool/System/CppDependencyCache.cs
+++ b/Engine/Source/Programs/UnrealBuildTool/System/CppDependencyCache.cs
@@ -302,7 +302,7 @@ namespace UnrealBuildTool
 					if (!Partition.DependencyFileToInfo.TryGetValue(InputFile, out Info) || InputFile.LastWriteTimeUtc.Ticks > Info.LastWriteTimeUtc)
 					{
 						Info = ReadDependencyInfo(InputFile);
-						Partition.DependencyFileToInfo.TryAdd(InputFile, Info);
+						Partition.DependencyFileToInfo.AddOrUpdate(InputFile, Info, (k, v) => Info);
 						Partition.bModified = true;
 					}
 
