Another optimization for when we iterate through header files when scanning the directories.

diff --git a/Engine/Source/Programs/UnrealBuildTool/System/TargetMakefile.cs b/Engine/Source/Programs/UnrealBuildTool/System/TargetMakefile.cs
--- a/Engine/Source/Programs/UnrealBuildTool/System/TargetMakefile.cs
+++ b/Engine/Source/Programs/UnrealBuildTool/System/TargetMakefile.cs
@@ -889,18 +889,15 @@ namespace UnrealBuildTool
 		/// <param name="Queue">Queue to add sub-tasks to</param>
 		static void FindFilesWithMarkup(DirectoryItem Directory, SourceFileMetadataCache MetadataCache, ReadOnlyHashSet<string> ExcludedFolderNames, ConcurrentBag<FileItem> FilesWithMarkup, ThreadPoolWorkQueue Queue)
 		{
-			// Check for all the headers in this folder
-			foreach (FileItem File in Directory.EnumerateFiles())
+			if (Directory.TryGetFile(".ubtignore", out FileItem? OutFile))
 			{
-				if (File.Name == ".ubtignore")
-				{
-					return;
-				}
+				return;
+			}
 
-				if (File.HasExtension(".h") && MetadataCache.ContainsReflectionMarkup(File))
-				{
-					FilesWithMarkup.Add(File);
-				}
+			// Check for all the headers in this folder
+			foreach (FileItem File in Directory.EnumerateFiles().Where((fi) => fi.HasExtension(".h") && MetadataCache.ContainsReflectionMarkup(fi)))
+			{
+				FilesWithMarkup.Add(File);
 			}
 
 			// Search through all the subfolders
