Moved the logic to check for inline gen.cpp changes out of the if statement that only triggered when the directory changed. Now the sources are always checked properly.

diff --git a/Engine/Source/Programs/UnrealBuildTool/System/TargetMakefile.cs b/Engine/Source/Programs/UnrealBuildTool/System/TargetMakefile.cs
--- a/Engine/Source/Programs/UnrealBuildTool/System/TargetMakefile.cs
+++ b/Engine/Source/Programs/UnrealBuildTool/System/TargetMakefile.cs
@@ -737,16 +737,16 @@ namespace UnrealBuildTool
 								return false;
 							}
 						}
+					}
 
-						// Make sure the inlined gen.cpp files didn't change
-						foreach (TargetMakefileSourceFileInfo SourceFileInfo in SourceFileInfos)
+					// Make sure the inlined gen.cpp files didn't change
+					foreach (TargetMakefileSourceFileInfo SourceFileInfo in SourceFileInfos)
+					{
+						int InlineHash = TargetMakefileSourceFileInfo.CalculateInlineGenCppHash(MetadataCache.GetListOfInlinedGeneratedCppFiles(SourceFileInfo.SourceFileItem!));
+						if (SourceFileInfo.InlineGenCppHash != InlineHash)
 						{
-							int InlineHash = TargetMakefileSourceFileInfo.CalculateInlineGenCppHash(MetadataCache.GetListOfInlinedGeneratedCppFiles(SourceFileInfo.SourceFileItem!));
-							if (SourceFileInfo.InlineGenCppHash != InlineHash)
-							{
-								ReasonNotLoaded = "inlined gen.cpp files changed";
-								return false;
-							}
+							ReasonNotLoaded = "inlined gen.cpp files changed";
+							return false;
 						}
 					}
 				}
