Minor optimizations in UBT to improve time when creating the action graph.

diff --git a/Engine/Source/Programs/UnrealBuildTool/Modes/BuildMode.cs b/Engine/Source/Programs/UnrealBuildTool/Modes/BuildMode.cs
--- a/Engine/Source/Programs/UnrealBuildTool/Modes/BuildMode.cs
+++ b/Engine/Source/Programs/UnrealBuildTool/Modes/BuildMode.cs
@@ -9,6 +9,7 @@ using EpicGames.Core;
 using OpenTracing.Util;
 using UnrealBuildBase;
 using Microsoft.Extensions.Logging;
+using System.Collections.Concurrent;
 
 namespace UnrealBuildTool
 {
@@ -380,7 +381,7 @@ namespace UnrealBuildTool
 				List<LinkedAction> ModuleDependencyActions = PrerequisiteActions.Where(x => x.ActionType == ActionType.GatherModuleDependencies).ToList();
 				if (ModuleDependencyActions.Count > 0)
 				{
-					Dictionary<LinkedAction, bool> PreprocessActionToOutdatedFlag = new Dictionary<LinkedAction, bool>();
+					ConcurrentDictionary<LinkedAction, bool> PreprocessActionToOutdatedFlag = new ConcurrentDictionary<LinkedAction, bool>();
 					ActionGraph.GatherAllOutdatedActions(ModuleDependencyActions, History, PreprocessActionToOutdatedFlag, CppDependencies, BuildConfiguration.bIgnoreOutdatedImportLibraries, Logger);
 
 					List<LinkedAction> PreprocessActions = PreprocessActionToOutdatedFlag.Where(x => x.Value).Select(x => x.Key).ToList();
@@ -397,7 +398,7 @@ namespace UnrealBuildTool
 				}
 
 				// Figure out which actions need to be built
-				Dictionary<LinkedAction, bool> ActionToOutdatedFlag = new Dictionary<LinkedAction, bool>();
+				ConcurrentDictionary<LinkedAction, bool> ActionToOutdatedFlag = new ConcurrentDictionary<LinkedAction, bool>();
 				for (int TargetIdx = 0; TargetIdx < TargetDescriptors.Count; TargetIdx++)
 				{
 					TargetDescriptor TargetDescriptor = TargetDescriptors[TargetIdx];
@@ -530,7 +531,7 @@ namespace UnrealBuildTool
 				{
 					Logger.LogDebug("Re-evaluating action graph");
 					// Re-check the graph to remove any LiveCoding actions added by PatchActionsForTarget() that are already up to date.
-					Dictionary<LinkedAction, bool> LiveActionToOutdatedFlag = new Dictionary<LinkedAction, bool>(MergedActionsToExecute.Count);
+					ConcurrentDictionary<LinkedAction, bool> LiveActionToOutdatedFlag = new ConcurrentDictionary<LinkedAction, bool>(Environment.ProcessorCount, MergedActionsToExecute.Count);
 					ActionGraph.GatherAllOutdatedActions(MergedActionsToExecute, History, LiveActionToOutdatedFlag, CppDependencies, BuildConfiguration.bIgnoreOutdatedImportLibraries, Logger);
 					List<LinkedAction> LiveCodingActionsToExecute = LiveActionToOutdatedFlag.Where(x => x.Value).Select(x => x.Key).ToList();
 					ActionGraph.Link(LiveCodingActionsToExecute, Logger);
diff --git a/Engine/Source/Programs/UnrealBuildTool/System/Action.cs b/Engine/Source/Programs/UnrealBuildTool/System/Action.cs
--- a/Engine/Source/Programs/UnrealBuildTool/System/Action.cs
+++ b/Engine/Source/Programs/UnrealBuildTool/System/Action.cs
@@ -7,6 +7,7 @@ using System.IO;
 using System.Linq;
 using System.Reflection;
 using System.Text;
+using System.Threading;
 using System.Threading.Tasks;
 using EpicGames.Core;
 using UnrealBuildBase;
@@ -563,7 +564,7 @@ namespace UnrealBuildTool
 		{
 			if (VisitedActions.Add(this))
 			{
-				NumTotalDependentActions++;
+				Interlocked.Increment(ref NumTotalDependentActions);
 				foreach (LinkedAction PrerequisiteAction in PrerequisiteActions)
 				{
 					PrerequisiteAction.IncrementDependentCount(VisitedActions);
diff --git a/Engine/Source/Programs/UnrealBuildTool/System/ActionGraph.cs b/Engine/Source/Programs/UnrealBuildTool/System/ActionGraph.cs
--- a/Engine/Source/Programs/UnrealBuildTool/System/ActionGraph.cs
+++ b/Engine/Source/Programs/UnrealBuildTool/System/ActionGraph.cs
@@ -15,6 +15,7 @@ using System.Threading.Tasks;
 using OpenTracing.Util;
 using UnrealBuildBase;
 using Microsoft.Extensions.Logging;
+using System.Collections.Concurrent;
 
 namespace UnrealBuildTool
 {
@@ -258,7 +259,7 @@ namespace UnrealBuildTool
 			using (GlobalTracer.Instance.BuildSpan("ActionGraph.GetActionsToExecute()").StartActive())
 			{
 				// For all targets, build a set of all actions that are outdated.
-				Dictionary<LinkedAction, bool> OutdatedActionDictionary = new Dictionary<LinkedAction, bool>();
+				ConcurrentDictionary<LinkedAction, bool> OutdatedActionDictionary = new ConcurrentDictionary<LinkedAction, bool>();
 				GatherAllOutdatedActions(Actions, History, OutdatedActionDictionary, CppDependencies, bIgnoreOutdatedImportLibraries, Logger);
 
 				// Build a list of actions that are both needed for this target and outdated.
@@ -426,10 +427,10 @@ namespace UnrealBuildTool
 			}
 
 			// Increment all the dependencies
-			foreach (LinkedAction Action in Actions)
+			Actions.AsParallel().ForAll((Action) =>
 			{
 				Action.IncrementDependentCount(new HashSet<LinkedAction>());
-			}
+			});
 
 			// Sort actions by number of actions depending on them, descending. Secondary sort criteria is file size.
 			Actions.Sort(LinkedAction.Compare);
@@ -607,24 +608,20 @@ namespace UnrealBuildTool
 		/// </summary>
 		/// <param name="RootAction">- The action being considered.</param>
 		/// <param name="OutdatedActionDictionary">-</param>
-		/// <param name="OutdatedActionLock"></param>
 		/// <param name="ActionHistory"></param>
 		/// <param name="CppDependencies"></param>
 		/// <param name="bIgnoreOutdatedImportLibraries"></param>
 		/// <param name="Logger"></param>
 		/// <returns>true if outdated</returns>
 		private static void IsIndividualActionOutdated(LinkedAction RootAction,
-			Dictionary<LinkedAction, bool> OutdatedActionDictionary, ReaderWriterLockSlim OutdatedActionLock,
+			ConcurrentDictionary<LinkedAction, bool> OutdatedActionDictionary,
 			ActionHistory ActionHistory, CppDependencyCache CppDependencies, bool bIgnoreOutdatedImportLibraries, ILogger Logger)
 		{
 			// Only compute the outdated-ness for actions that don't aren't cached in the outdated action dictionary.
 			bool bIsOutdated = false;
 			{
 				// OutdatedActionDictionary may have already been populated for RootAction as part of a previously processed target
-				OutdatedActionLock.EnterReadLock();
-				bool bPresent = OutdatedActionDictionary.ContainsKey(RootAction);
-				OutdatedActionLock.ExitReadLock();
-				if (bPresent)
+				if (OutdatedActionDictionary.ContainsKey(RootAction))
 				{
 					return;
 				}
@@ -725,9 +722,7 @@ namespace UnrealBuildTool
 			// We don't yet know that the action is up-to-date - to determine that requires traversal of the graph of prerequisites.
 			if (bIsOutdated)
 			{
-				OutdatedActionLock.EnterWriteLock();
-				OutdatedActionDictionary.Add(RootAction, bIsOutdated);
-				OutdatedActionLock.ExitWriteLock();
+				OutdatedActionDictionary.TryAdd(RootAction, bIsOutdated);
 			}
 		}
 
@@ -740,7 +735,7 @@ namespace UnrealBuildTool
 		/// <param name="bIgnoreOutdatedImportLibraries"></param>
 		/// <param name="Logger">Logger instance</param>
 		/// <returns>true if outdated</returns>
-		private static bool IsActionOutdatedDueToPrerequisites(LinkedAction RootAction,	Dictionary<LinkedAction, bool> OutdatedActionDictionary, bool bIgnoreOutdatedImportLibraries, ILogger Logger)
+		private static bool IsActionOutdatedDueToPrerequisites(LinkedAction RootAction, ConcurrentDictionary<LinkedAction, bool> OutdatedActionDictionary, bool bIgnoreOutdatedImportLibraries, ILogger Logger)
 		{
 			// Only compute the outdated-ness for actions that aren't already cached in the outdated action dictionary.
 			if (OutdatedActionDictionary.TryGetValue(RootAction, out bool bIsOutdated))
@@ -767,7 +762,7 @@ namespace UnrealBuildTool
 			}
 
 			// Cache the outdated-ness of this action.
-			OutdatedActionDictionary.Add(RootAction, bIsOutdated);
+			OutdatedActionDictionary.TryAdd(RootAction, bIsOutdated);
 
 			return bIsOutdated;
 		}
@@ -827,7 +822,7 @@ namespace UnrealBuildTool
 		/// IsActionOutdated.
 		/// </summary>
 		public static void GatherAllOutdatedActions(IReadOnlyList<LinkedAction> Actions, ActionHistory ActionHistory,
-			Dictionary<LinkedAction, bool> OutdatedActions, CppDependencyCache CppDependencies,
+			ConcurrentDictionary<LinkedAction, bool> OutdatedActions, CppDependencyCache CppDependencies,
 			bool bIgnoreOutdatedImportLibraries, ILogger Logger)
 		{
 			using (GlobalTracer.Instance.BuildSpan("Prefetching include dependencies").StartActive())
@@ -846,9 +841,8 @@ namespace UnrealBuildTool
 
 			using (GlobalTracer.Instance.BuildSpan("Cache individual outdated actions").StartActive())
 			{
-				ReaderWriterLockSlim OutdatedActionsLock = new ReaderWriterLockSlim();
 				Parallel.ForEach(Actions,
-					Action => IsIndividualActionOutdated(Action, OutdatedActions, OutdatedActionsLock, 
+					Action => IsIndividualActionOutdated(Action, OutdatedActions,
 						ActionHistory, CppDependencies, bIgnoreOutdatedImportLibraries, Logger));
 			}
 
