UBT: Ensure the Engine/Source directory exists before changing into it.

diff --git a/Engine/Source/Programs/UnrealBuildTool/UnrealBuildTool.cs b/Engine/Source/Programs/UnrealBuildTool/UnrealBuildTool.cs
--- a/Engine/Source/Programs/UnrealBuildTool/UnrealBuildTool.cs
+++ b/Engine/Source/Programs/UnrealBuildTool/UnrealBuildTool.cs
@@ -514,6 +514,7 @@ namespace UnrealBuildTool
 
 				// Change the working directory to be the Engine/Source folder. We are likely running from Engine/Binaries/DotNET
 				// This is critical to be done early so any code that relies on the current directory being Engine/Source will work.
+				DirectoryReference.CreateDirectory(Unreal.EngineSourceDirectory);
 				DirectoryReference.SetCurrentDirectory(Unreal.EngineSourceDirectory);
 
 				// Register encodings from Net FW as this is required when using Ionic as we do in multiple toolchains
