[UnrealBuildTool]
* Added module type to json output

diff --git a/Engine/Source/Programs/UnrealBuildTool/Configuration/UEBuildModule.cs b/Engine/Source/Programs/UnrealBuildTool/Configuration/UEBuildModule.cs
--- a/Engine/Source/Programs/UnrealBuildTool/Configuration/UEBuildModule.cs
+++ b/Engine/Source/Programs/UnrealBuildTool/Configuration/UEBuildModule.cs
@@ -1115,6 +1115,7 @@ namespace UnrealBuildTool
 		public virtual void ExportJson(DirectoryReference? BinaryOutputDir, DirectoryReference? TargetOutputDir, JsonWriter Writer)
 		{
 			Writer.WriteValue("Name", Name);
+			Writer.WriteValue("Type", Rules.Type.ToString());
 			Writer.WriteValue("Directory", ModuleDirectory.FullName);
 			Writer.WriteValue("Rules", RulesFile.FullName);
 			Writer.WriteValue("PCHUsage", Rules.PCHUsage.ToString());
