bStressTestUnity supercedes adaptive unity

diff --git a/Engine/Source/Programs/UnrealBuildTool/System/Unity.cs b/Engine/Source/Programs/UnrealBuildTool/System/Unity.cs
--- a/Engine/Source/Programs/UnrealBuildTool/System/Unity.cs
+++ b/Engine/Source/Programs/UnrealBuildTool/System/Unity.cs
@@ -319,7 +319,7 @@ namespace UnrealBuildTool
 			NormalFiles = new List<FileItem>();
 			AdaptiveFiles = new List<FileItem>();
 
-			if (!Target.bUseAdaptiveUnityBuild)
+			if (!Target.bUseAdaptiveUnityBuild || Target.bStressTestUnity)
 			{
 				NormalFiles = CPPFiles;
 				return;
