Fix code sign error on QAGCookedEditor

diff --git a/Engine/Source/Programs/UnrealBuildTool/ProjectFiles/Xcode/XcodeProjectLegacy.cs b/Engine/Source/Programs/UnrealBuildTool/ProjectFiles/Xcode/XcodeProjectLegacy.cs
--- a/Engine/Source/Programs/UnrealBuildTool/ProjectFiles/Xcode/XcodeProjectLegacy.cs
+++ b/Engine/Source/Programs/UnrealBuildTool/ProjectFiles/Xcode/XcodeProjectLegacy.cs
@@ -1335,7 +1335,7 @@ namespace UnrealBuildTool.XcodeProjectLegacy
 				}
 
 			// #jira UE-143619: Pre Monterey macOS requires this option for a packaged app to run on iOS15 due to new code signature format. Could be removed once Monterey is minimum OS.
-			Content.Append("\t\t\t\tOTHER_CODE_SIGN_FLAGS = \"--generate-entitlement-der\";" + ProjectFileGenerator.NewLine);
+			Content.Append("\t\t\t\tOTHER_CODE_SIGN_FLAGS = \"--generate-entitlement-der --deep\";" + ProjectFileGenerator.NewLine);
 
 			Content.Append("\t\t\t\tMACOSX_DEPLOYMENT_TARGET = " + MacToolChain.Settings.MacOSVersion + ";" + ProjectFileGenerator.NewLine);
 			Content.Append("\t\t\t\tINFOPLIST_OUTPUT_FORMAT = xml;" + ProjectFileGenerator.NewLine);
