Fixed an issue where commented out #include inline.gen.cpp files would still be considered to be included.

diff --git a/Engine/Source/Programs/UnrealBuildTool/System/SourceFileMetadataCache.cs b/Engine/Source/Programs/UnrealBuildTool/System/SourceFileMetadataCache.cs
--- a/Engine/Source/Programs/UnrealBuildTool/System/SourceFileMetadataCache.cs
+++ b/Engine/Source/Programs/UnrealBuildTool/System/SourceFileMetadataCache.cs
@@ -125,9 +125,9 @@ namespace UnrealBuildTool
 		static readonly Regex ReflectionMarkupRegex = new Regex("^\\s*U(CLASS|STRUCT|ENUM|INTERFACE|DELEGATE)\\b", RegexOptions.Compiled | RegexOptions.Multiline);
 
 		/// <summary>
-		/// 
+		/// Regex that matches #include UE_INLINE_GENERATED_CPP_BY_NAME(****) statements.
 		/// </summary>
-		static readonly Regex InlineReflectionMarkupRegex = new Regex(@"UE_INLINE_GENERATED_CPP_BY_NAME\((.+)\)", RegexOptions.Compiled | RegexOptions.Multiline);
+		static readonly Regex InlineReflectionMarkupRegex = new Regex("^\\s*#\\s*include\\s*UE_INLINE_GENERATED_CPP_BY_NAME\\((.+)\\)", RegexOptions.Compiled | RegexOptions.Multiline);
 
 		/// <summary>
 		/// Regex that matches #include statements.
