VulkanRHI: Fix RHI_RAYTRACING not being set on Linux builds.  Fix case on some includes.

#rnx

diff --git a/Engine/Source/Programs/UnrealBuildTool/Platform/Linux/UEBuildLinux.cs b/Engine/Source/Programs/UnrealBuildTool/Platform/Linux/UEBuildLinux.cs
--- a/Engine/Source/Programs/UnrealBuildTool/Platform/Linux/UEBuildLinux.cs
+++ b/Engine/Source/Programs/UnrealBuildTool/Platform/Linux/UEBuildLinux.cs
@@ -82,7 +82,7 @@ namespace UnrealBuildTool
 		/// <summary>
 		/// Enables runtime ray tracing support.
 		/// </summary>
-		[ConfigFile(ConfigHierarchyType.Engine, "/Script/LinuxPlatform.LinuxTargetSettings")]
+		[ConfigFile(ConfigHierarchyType.Engine, "/Script/LinuxTargetPlatform.LinuxTargetSettings")]
 		public bool bEnableRayTracing = false;
 	}
 
diff --git a/Engine/Source/Runtime/Renderer/Private/RayTracing/RayTracingDebug.cpp b/Engine/Source/Runtime/Renderer/Private/RayTracing/RayTracingDebug.cpp
--- a/Engine/Source/Runtime/Renderer/Private/RayTracing/RayTracingDebug.cpp
+++ b/Engine/Source/Runtime/Renderer/Private/RayTracing/RayTracingDebug.cpp
@@ -12,7 +12,7 @@
 #include "RenderGraphBuilder.h"
 #include "SceneUtils.h"
 #include "RaytracingDebugDefinitions.h"
-#include "RaytracingDebugTypes.h"
+#include "RayTracingDebugTypes.h"
 #include "RayTracing/RayTracingLighting.h"
 #include "RayTracing/RaytracingOptions.h"
 #include "RayTracing/RayTracingTraversalStatistics.h"
diff --git a/Engine/Source/Runtime/Renderer/Private/RayTracing/RayTracingScene.cpp b/Engine/Source/Runtime/Renderer/Private/RayTracing/RayTracingScene.cpp
--- a/Engine/Source/Runtime/Renderer/Private/RayTracing/RayTracingScene.cpp
+++ b/Engine/Source/Runtime/Renderer/Private/RayTracing/RayTracingScene.cpp
@@ -9,7 +9,7 @@
 #include "RayTracingDefinitions.h"
 #include "RenderGraphBuilder.h"
 #include "RenderGraphUtils.h"
-#include "RayTracingOptions.h"
+#include "RaytracingOptions.h"
 
 BEGIN_SHADER_PARAMETER_STRUCT(FBuildInstanceBufferPassParams, )
 	SHADER_PARAMETER_RDG_BUFFER_UAV(RWStructuredBuffer, InstanceBuffer)
