UnrealBuildTool: Remove StaticAnalyzerDisabledCheckers.Contains("all") as there is now a boolean bDisableStaticAnalysis

#rnx

diff --git a/Engine/Source/Programs/UnrealBuildTool/Platform/Windows/VCToolChain.cs b/Engine/Source/Programs/UnrealBuildTool/Platform/Windows/VCToolChain.cs
--- a/Engine/Source/Programs/UnrealBuildTool/Platform/Windows/VCToolChain.cs
+++ b/Engine/Source/Programs/UnrealBuildTool/Platform/Windows/VCToolChain.cs
@@ -267,7 +267,7 @@ namespace UnrealBuildTool
 				// compiler version of 19.xx).
 				Arguments.Add($"-fms-compatibility-version=19.{EnvVars.ToolChainVersion.GetComponent(1)}");
 				
-				if (Target.StaticAnalyzer == StaticAnalyzer.Default && CompileEnvironment.PrecompiledHeaderAction != PrecompiledHeaderAction.Create && !CompileEnvironment.StaticAnalyzerDisabledCheckers.Contains("all") && !CompileEnvironment.bDisableStaticAnalysis)
+				if (Target.StaticAnalyzer == StaticAnalyzer.Default && CompileEnvironment.PrecompiledHeaderAction != PrecompiledHeaderAction.Create && !CompileEnvironment.bDisableStaticAnalysis)
 				{
 					Arguments.Add("-Wno-unused-command-line-argument");
 
diff --git a/Engine/Source/Programs/UnrealBuildTool/ToolChain/ClangToolChain.cs b/Engine/Source/Programs/UnrealBuildTool/ToolChain/ClangToolChain.cs
--- a/Engine/Source/Programs/UnrealBuildTool/ToolChain/ClangToolChain.cs
+++ b/Engine/Source/Programs/UnrealBuildTool/ToolChain/ClangToolChain.cs
@@ -320,7 +320,6 @@ namespace UnrealBuildTool
 		protected bool IsAnalyzing(CppCompileEnvironment CompileEnvironment) => 
 				StaticAnalyzer == StaticAnalyzer.Default 
 			&&	CompileEnvironment.PrecompiledHeaderAction != PrecompiledHeaderAction.Create 
-			&& !CompileEnvironment.StaticAnalyzerDisabledCheckers.Contains("all") 
 			&& !CompileEnvironment.bDisableStaticAnalysis;
 
 		protected virtual void GetCppStandardCompileArgument(CppCompileEnvironment CompileEnvironment, List<string> Arguments)
