UnrealBuildTool: Use copy action helper for copying ISPC headers

#rnx

diff --git a/Engine/Source/Programs/UnrealBuildTool/ToolChain/ISPCToolChain.cs b/Engine/Source/Programs/UnrealBuildTool/ToolChain/ISPCToolChain.cs
--- a/Engine/Source/Programs/UnrealBuildTool/ToolChain/ISPCToolChain.cs
+++ b/Engine/Source/Programs/UnrealBuildTool/ToolChain/ISPCToolChain.cs
@@ -487,32 +487,11 @@ namespace UnrealBuildTool
 					);
 
 				// Fix interrupted build issue by copying header after generation completes
-				FileReference SourceFile = ISPCIncludeHeaderFile.Location;
-				FileReference TargetFile = ISPCFinalHeaderFile.Location;
-
-				FileItem SourceFileItem = FileItem.GetItemByFileReference(SourceFile);
-				FileItem TargetFileItem = FileItem.GetItemByFileReference(TargetFile);
-
-				Action CopyAction = Graph.CreateAction(ActionType.BuildProject);
-				CopyAction.CommandDescription = "Copy";
-				CopyAction.CommandPath = BuildHostPlatform.Current.Shell;
-				if (BuildHostPlatform.Current.ShellType == ShellType.Cmd)
-				{
-					CopyAction.CommandArguments = $"/C \"copy /Y \"{SourceFile}\" \"{TargetFile}\" 1>nul\"";
-				}
-				else
-				{
-					CopyAction.CommandArguments = $"-c 'cp -f \"\"{SourceFile}\"\" \"\"{TargetFile}\"'";
-				}
-				CopyAction.WorkingDirectory = Unreal.EngineSourceDirectory;
+				Action CopyAction = Graph.CreateCopyAction(ISPCIncludeHeaderFile, ISPCFinalHeaderFile);
 				CopyAction.PrerequisiteItems.Add(ISPCFile);
-				CopyAction.PrerequisiteItems.Add(SourceFileItem);
-				CopyAction.ProducedItems.Add(TargetFileItem);
-				CopyAction.StatusDescription = TargetFileItem.Location.GetFileName();
-				CopyAction.bCanExecuteRemotely = false;
 				CopyAction.bShouldOutputStatusDescription = false;
 
-				Result.GeneratedHeaderFiles.Add(TargetFileItem);
+				Result.GeneratedHeaderFiles.Add(ISPCFinalHeaderFile);
 
 				Logger.LogDebug("   ISPC Generating Header {StatusDescription}: \"{CommandPath}\" {CommandArguments}", CompileAction.StatusDescription, CompileAction.CommandPath, CompileAction.CommandArguments);
 			}
