Fixed an issue where SNDBS would fail to compile files on agents because the inlined gen.cpp files were not setup as prerequisites.

diff --git a/Engine/Source/Programs/UnrealBuildTool/Configuration/UEBuildModuleCPP.cs b/Engine/Source/Programs/UnrealBuildTool/Configuration/UEBuildModuleCPP.cs
--- a/Engine/Source/Programs/UnrealBuildTool/Configuration/UEBuildModuleCPP.cs
+++ b/Engine/Source/Programs/UnrealBuildTool/Configuration/UEBuildModuleCPP.cs
@@ -580,6 +580,12 @@ namespace UnrealBuildTool
 								if (FoundGenCppFile != null)
 								{
 									GeneratedFiles.Remove(FoundGenCppFile);
+
+									if (!CompileEnvironment.FileInlineGenCPPMap.ContainsKey(CPPFileItem))
+									{
+										CompileEnvironment.FileInlineGenCPPMap[CPPFileItem] = new List<FileItem>();
+									}
+									CompileEnvironment.FileInlineGenCPPMap[CPPFileItem].Add(FileItem.GetItemByPath(FoundGenCppFile));
 								}
 								else
 								{
diff --git a/Engine/Source/Programs/UnrealBuildTool/Platform/Android/AndroidToolChain.cs b/Engine/Source/Programs/UnrealBuildTool/Platform/Android/AndroidToolChain.cs
--- a/Engine/Source/Programs/UnrealBuildTool/Platform/Android/AndroidToolChain.cs
+++ b/Engine/Source/Programs/UnrealBuildTool/Platform/Android/AndroidToolChain.cs
@@ -1302,6 +1302,12 @@ namespace UnrealBuildTool
 					// Add the C++ source file and its included files to the prerequisite item list.
 					CompileAction.PrerequisiteItems.Add(SourceFile);
 
+					List<FileItem>? InlinedFiles;
+					if (CompileEnvironment.FileInlineGenCPPMap.TryGetValue(SourceFile, out InlinedFiles))
+					{
+						CompileAction.PrerequisiteItems.AddRange(InlinedFiles);
+					}
+
 					if (CompileEnvironment.PrecompiledHeaderAction == PrecompiledHeaderAction.Create && !bIsPlainCFile)
 					{
 						// Add the precompiled header file to the produced item list.
diff --git a/Engine/Source/Programs/UnrealBuildTool/Platform/Windows/VCToolChain.cs b/Engine/Source/Programs/UnrealBuildTool/Platform/Windows/VCToolChain.cs
--- a/Engine/Source/Programs/UnrealBuildTool/Platform/Windows/VCToolChain.cs
+++ b/Engine/Source/Programs/UnrealBuildTool/Platform/Windows/VCToolChain.cs
@@ -1431,6 +1431,12 @@ namespace UnrealBuildTool
 					AppendCLArguments_CPP(CompileEnvironment, CompileAction.Arguments);
 				}
 
+				List<FileItem>? InlinedFiles;
+				if (CompileEnvironment.FileInlineGenCPPMap.TryGetValue(SourceFile, out InlinedFiles))
+				{
+					CompileAction.AdditionalPrerequisiteItems.AddRange(InlinedFiles);
+				}
+
 				CompileAction.AdditionalPrerequisiteItems.AddRange(CompileEnvironment.AdditionalPrerequisites);
 
 				if (!String.IsNullOrEmpty(CompileEnvironment.AdditionalArguments))
diff --git a/Engine/Source/Programs/UnrealBuildTool/System/CppCompileEnvironment.cs b/Engine/Source/Programs/UnrealBuildTool/System/CppCompileEnvironment.cs
--- a/Engine/Source/Programs/UnrealBuildTool/System/CppCompileEnvironment.cs
+++ b/Engine/Source/Programs/UnrealBuildTool/System/CppCompileEnvironment.cs
@@ -419,6 +419,11 @@ namespace UnrealBuildTool
 		/// </summary>
 		public List<FileItem> AdditionalPrerequisites = new List<FileItem>();
 
+		/// <summary>
+		/// A dictionary of the source file items and the inlined gen.cpp files contained in it
+		/// </summary>
+		public Dictionary<FileItem, List<FileItem>> FileInlineGenCPPMap = new();
+
 		/// <summary>
 		/// The C++ preprocessor definitions to use.
 		/// </summary>
diff --git a/Engine/Source/Programs/UnrealBuildTool/System/Unity.cs b/Engine/Source/Programs/UnrealBuildTool/System/Unity.cs
--- a/Engine/Source/Programs/UnrealBuildTool/System/Unity.cs
+++ b/Engine/Source/Programs/UnrealBuildTool/System/Unity.cs
@@ -275,6 +275,8 @@ namespace UnrealBuildTool
 				}
 				FileReference UnityCPPFilePath = FileReference.Combine(IntermediateDirectory, UnityCPPFileName);
 
+				List<FileItem> InlinedGenCPPFilesInUnity = new();
+
 				// Add source files to the unity file
 				foreach (FileItem CPPFile in UnityFile.Files)
 				{
@@ -284,12 +286,21 @@ namespace UnrealBuildTool
 						CPPFileString = CPPFile.Location.MakeRelativeTo(Unreal.EngineSourceDirectory);
 					}
 					OutputUnityCPPWriter.WriteLine("#include \"{0}\"", CPPFileString.Replace('\\', '/'));
+
+					List<FileItem>? InlinedGenCPPFiles;
+					if (CompileEnvironment.FileInlineGenCPPMap.TryGetValue(CPPFile, out InlinedGenCPPFiles))
+					{
+						InlinedGenCPPFilesInUnity.AddRange(InlinedGenCPPFiles);
+					}
 				}
 
 				// Write the unity file to the intermediate folder.
 				FileItem UnityCPPFile = Graph.CreateIntermediateTextFile(UnityCPPFilePath, OutputUnityCPPWriter.ToString());
 				NewCPPFiles.Add(UnityCPPFile);
 
+				// Store all the inlined gen.cpp files
+				CompileEnvironment.FileInlineGenCPPMap[UnityCPPFile] = InlinedGenCPPFilesInUnity;
+
 				// Store the mapping of source files to unity files in the makefile
 				foreach(FileItem SourceFile in UnityFile.Files)
 				{
diff --git a/Engine/Source/Programs/UnrealBuildTool/ToolChain/ClangToolChain.cs b/Engine/Source/Programs/UnrealBuildTool/ToolChain/ClangToolChain.cs
--- a/Engine/Source/Programs/UnrealBuildTool/ToolChain/ClangToolChain.cs
+++ b/Engine/Source/Programs/UnrealBuildTool/ToolChain/ClangToolChain.cs
@@ -878,6 +878,12 @@ namespace UnrealBuildTool
 			CompileAction.PrerequisiteItems.AddRange(CompileEnvironment.AdditionalPrerequisites);
 			CompileAction.PrerequisiteItems.Add(SourceFile);
 
+			List<FileItem>? InlinedFiles;
+			if (CompileEnvironment.FileInlineGenCPPMap.TryGetValue(SourceFile, out InlinedFiles))
+			{
+				CompileAction.PrerequisiteItems.AddRange(InlinedFiles);
+			}
+
 			string Extension = Path.GetExtension(SourceFile.AbsolutePath).ToUpperInvariant();
 			if (CompileEnvironment.PrecompiledHeaderAction == PrecompiledHeaderAction.Create)
 			{
