UBT: Projects: Fixed live code edit (edit and continue) on some platforms.

Technically, it did work on the affected platforms, because we could modify a .cpp file and deploy this change, but UBT ran a full build script including compilation, linking and symbol generation, which is unnecessary because only compilation is needed. Thus, it increased the update time from 1-10 sec to 60 seconds (the time is dominated by the link time).

Fixed by adding a special Target to our projects that runs only compilation steps.

#rn fix Sony Live code edit (edit and continue) no longer triggers linking when applying code changes. This way iteration time is significantly reduced.

diff --git a/Engine/Source/Programs/UnrealBuildTool/ProjectFiles/PlatformProjectGenerator.cs b/Engine/Source/Programs/UnrealBuildTool/ProjectFiles/PlatformProjectGenerator.cs
--- a/Engine/Source/Programs/UnrealBuildTool/ProjectFiles/PlatformProjectGenerator.cs
+++ b/Engine/Source/Programs/UnrealBuildTool/ProjectFiles/PlatformProjectGenerator.cs
@@ -10,6 +10,23 @@ using Microsoft.Extensions.Logging;
 
 namespace UnrealBuildTool
 {
+	/// <summary>
+	/// Description of a project build configuration provided during project generation by ProjectFile.
+	/// </summary>
+	/// <seealso cref="PlatformProjectGenerator.GetVisualStudioTargetsString"/>
+	public abstract class ProjectBuildConfiguration
+	{
+		/// <summary>
+		/// Name of the build configuration as displayed in the IDE e.g. Debug_Client.
+		/// </summary>
+		public abstract string ConfigurationName { get; }
+
+		/// <summary>
+		/// Build command associated with this build configuration (as generated by classes derived from ProjectFile).
+		/// </summary>
+		public abstract string BuildCommand { get; }
+	}
+
 	/// <summary>
 	/// Base class for platform-specific project generators
 	/// </summary>
@@ -321,5 +338,28 @@ namespace UnrealBuildTool
 		{
 			return "";
 		}
+
+		/// <summary>
+		/// Indicates whether this generator will be emitting additional build targets.
+		/// Functionally, it's fine to return true and not return any additional text in GetVisualStudioTargetsString, but doing
+		/// so may incur additional cost of preparing data for generation.
+		/// </summary>
+		/// <param name="InPlatform">The platform being added</param>
+		/// <returns>True if Targets will be generate, false otherwise. If true is returned, GetVisualStudioTargetsString will be called next.</returns>
+		/// <seealso cref="GetVisualStudioTargetsString"/>
+		public virtual bool HasVisualStudioTargets(UnrealTargetPlatform InPlatform)
+		{
+			return false;
+		}
+
+		/// <summary>
+		/// Get the text to insert into the Project files as additional build targets for the given platform.
+		/// </summary>
+		/// <param name="InPlatform">The platform being added</param>
+		/// <param name="InPlatformConfigurations">List of build configurations for the platform</param>
+		/// <param name="ProjectFileBuilder">String builder for the project file</param>
+		public virtual void GetVisualStudioTargetsString(UnrealTargetPlatform InPlatform, List<ProjectBuildConfiguration> InPlatformConfigurations, StringBuilder ProjectFileBuilder)
+		{
+		}
 	}
 }
diff --git a/Engine/Source/Programs/UnrealBuildTool/ProjectFiles/Project.cs b/Engine/Source/Programs/UnrealBuildTool/ProjectFiles/Project.cs
--- a/Engine/Source/Programs/UnrealBuildTool/ProjectFiles/Project.cs
+++ b/Engine/Source/Programs/UnrealBuildTool/ProjectFiles/Project.cs
@@ -550,7 +550,7 @@ namespace UnrealBuildTool
 		/// <summary>
 		/// Takes the given path, normalizes it, and quotes it if necessary.
 		/// </summary>
-		public string EscapePath(string InputPath)
+		public static string EscapePath(string InputPath)
 		{
 			string Result = InputPath;
 			if (Result.Contains(' '))
diff --git a/Engine/Source/Programs/UnrealBuildTool/ProjectFiles/VisualStudio/VCProject.cs b/Engine/Source/Programs/UnrealBuildTool/ProjectFiles/VisualStudio/VCProject.cs
--- a/Engine/Source/Programs/UnrealBuildTool/ProjectFiles/VisualStudio/VCProject.cs
+++ b/Engine/Source/Programs/UnrealBuildTool/ProjectFiles/VisualStudio/VCProject.cs
@@ -521,7 +521,7 @@ namespace UnrealBuildTool
 				Configuration = InConfiguration;
 				ProjectPlatformName = InProjectPlatformName;
 				ProjectConfigurationName = InProjectConfigurationName;
-				ProjectTarget = InProjectTarget;				
+				ProjectTarget = InProjectTarget;
 			}
 
 			public string? ProjectConfigurationAndPlatformName
@@ -603,7 +603,7 @@ namespace UnrealBuildTool
 
 			return Version;
 		}
-		
+
 		/// <summary>
 		/// Gets compiler switch for specifying in AdditionalOptions in .vcxproj file for specific C++ version
 		/// </summary>
@@ -1372,6 +1372,9 @@ namespace UnrealBuildTool
 			}
 			VCProjectFileContent.AppendLine("  <ImportGroup Label=\"ExtensionTargets\">");
 			VCProjectFileContent.AppendLine("  </ImportGroup>");
+
+			WriteTargets(ProjectPlatforms, PlatformProjectGenerators, VCProjectFileContent, Logger);
+
 			VCProjectFileContent.AppendLine("</Project>");
 
 			VCFiltersFileContent.AppendLine("</Project>");
@@ -1428,6 +1431,59 @@ namespace UnrealBuildTool
 			return bSuccess;
 		}
 
+		private class ProjectConfigurationForGenerator : ProjectBuildConfiguration
+		{
+			public override string ConfigurationName => Combination.ProjectConfigurationName;
+			public override string BuildCommand      => $"{EscapePath(NormalizeProjectPath(CommandBuilder.BuildScript))} {CommandBuilder.GetBuildArguments()}";
+
+			private ProjectConfigAndTargetCombination Combination;
+			private BuildCommandBuilder               CommandBuilder;
+
+			public ProjectConfigurationForGenerator(ProjectConfigAndTargetCombination InCombination, BuildCommandBuilder InCommandBuilder)
+			{
+				Combination    = InCombination;
+				CommandBuilder = InCommandBuilder;
+			}
+		}
+
+		/// <summary>
+		/// Write additional Target elements if needed by per-platform generators.
+		/// </summary>
+		/// <param name="ProjectPlatforms"></param>
+		/// <param name="PlatformProjectGenerators"></param>
+		/// <param name="VCProjectFileContent"></param>
+		/// <param name="Logger"></param>
+		private void WriteTargets(List<UnrealTargetPlatform> ProjectPlatforms, PlatformProjectGeneratorCollection PlatformProjectGenerators, StringBuilder VCProjectFileContent, ILogger Logger)
+		{
+			foreach (UnrealTargetPlatform Platform in ProjectPlatforms)
+			{
+				PlatformProjectGenerator? ProjGenerator = PlatformProjectGenerators.GetPlatformProjectGenerator(Platform, true);
+				if (ProjGenerator != null && ProjGenerator.HasVisualStudioTargets(Platform))
+				{
+					IEnumerable<ProjectConfigAndTargetCombination> PlatformCombinations = ProjectConfigAndTargetCombinations!.Where(Combination => Combination.Platform == Platform);
+
+					if (PlatformCombinations.Any())
+					{
+						List<ProjectBuildConfiguration> PlatformConfigurations = PlatformCombinations.Select(Combination =>
+							{
+								string UProjectPath = "";
+								if (IsForeignProject)
+								{
+									UProjectPath = String.Format("\"{0}\"", InsertPathVariables(Combination.ProjectTarget!.UnrealProjectFilePath!));
+								}
+
+								BuildCommandBuilder Builder = CreateArgumentsBuilder(Combination, UProjectPath, ProjGenerator);
+
+								return new ProjectConfigurationForGenerator(Combination, Builder) as ProjectBuildConfiguration;
+							}
+						).ToList();
+
+						ProjGenerator.GetVisualStudioTargetsString(Platform, PlatformConfigurations, VCProjectFileContent);
+					}
+				}
+			}
+		}
+
 		private static bool EnsureFilterPathExists(string FilterRelativeSourceDirectory, StringBuilder VCFiltersFileContent, HashSet<string> FilterDirectories)
 		{
 			// We only want each directory to appear once in the filters file
