Fix creating cpp projects cannot run editor correctly

diff --git a/Engine/Source/Programs/UnrealBuildTool/ProjectFiles/Xcode/XcodeProject.cs b/Engine/Source/Programs/UnrealBuildTool/ProjectFiles/Xcode/XcodeProject.cs
--- a/Engine/Source/Programs/UnrealBuildTool/ProjectFiles/Xcode/XcodeProject.cs
+++ b/Engine/Source/Programs/UnrealBuildTool/ProjectFiles/Xcode/XcodeProject.cs
@@ -2332,6 +2332,8 @@ namespace UnrealBuildTool.XcodeProjectXcconfig
 					LegacyProjectFile.ProjectTargets.AddRange(ProjectTargets);
 					LegacyProjectFile.SourceFiles.AddRange(SourceFiles);
 					LegacyProjectFile.IsGeneratedProject = IsGeneratedProject;
+					LegacyProjectFile.IsStubProject = IsStubProject;
+					LegacyProjectFile.IsForeignProject = IsForeignProject;
 				}
 			}
 		}
