Added a script that updates include paths.

diff --git a/Engine/Source/Programs/UnrealBuildTool/Modes/FixIncludePathsMode.cs b/Engine/Source/Programs/UnrealBuildTool/Modes/FixIncludePathsMode.cs
new file mode 100644
--- /dev/null
+++ b/Engine/Source/Programs/UnrealBuildTool/Modes/FixIncludePathsMode.cs
@@ -0,0 +1,245 @@
+// Copyright Epic Games, Inc. All Rights Reserved.
+
+using System;
+using System.Collections.Generic;
+using System.Linq;
+using System.Text;
+using System.Text.RegularExpressions;
+using EpicGames.Core;
+using UnrealBuildBase;
+using Microsoft.Extensions.Logging;
+
+namespace UnrealBuildTool
+{
+	/// <summary>
+	/// Fixes the include paths found in a header and source file
+	/// </summary>
+	[ToolMode("FixIncludePaths", ToolModeOptions.XmlConfig | ToolModeOptions.BuildPlatforms | ToolModeOptions.SingleInstance | ToolModeOptions.StartPrefetchingEngine | ToolModeOptions.ShowExecutionTime)]
+	class FixIncludePathsMode : ToolMode
+	{
+		/// <summary>
+		/// Regex that matches #include statements.
+		/// </summary>
+		static readonly Regex IncludeRegex = new Regex("^[ \t]*#[ \t]*include[ \t]*[\"](?<HeaderFile>[^\"]*)[\"]", RegexOptions.Compiled | RegexOptions.Singleline | RegexOptions.ExplicitCapture);
+
+		static readonly string UnrealRootDirectory = Unreal.RootDirectory.FullName.Replace('\\', '/');
+		static readonly string[] PerferredPaths = { "/Public/", "/Private/", "/Classes/", "/Internal/", "/UHT/", "/VNI/" };
+
+		static readonly string[] PublicDirectories = { "Public", "Classes", };
+
+		[CommandLine("-Filter=", Description = "Set of filters for files to include in the database. Relative to the root directory, or to the project file.")]
+		List<string> FilterRules = new List<string>();
+
+		[CommandLine("-CheckoutWithP4", Description = "Flags that this task should use p4 to check out the file before updating it.")]
+		public bool bCheckoutWithP4 = false;
+
+		[CommandLine("-NoOutput", Description = "Flags that the updated files shouldn't be saved.")]
+		public bool bNoOutput = false;
+
+		private string? FindIncludePath(string FilePath, CppCompileEnvironment env, string HeaderIncludePath)
+		{
+			List<DirectoryReference> EnvPaths = new();
+			EnvPaths.Add(new DirectoryReference(System.IO.Directory.GetParent(FilePath)!));
+			EnvPaths.AddRange(env.UserIncludePaths);
+			//EnvPaths.AddRange(env.SystemIncludePaths);
+
+			// search include paths
+			foreach (var UserIncludePath in EnvPaths)
+			{
+				string Path = System.IO.Path.GetFullPath(System.IO.Path.Combine(UserIncludePath.FullName, HeaderIncludePath));
+				if (System.IO.File.Exists(Path))
+				{
+					return Path.Replace('\\', '/');
+				}
+			}
+
+			return null;
+		}
+
+		private string? GetPerferredInclude(string FullPath)
+		{
+			if (!FullPath.Contains(UnrealRootDirectory))
+			{
+				return null;
+			}
+
+			string? FoundPerferredPath = PerferredPaths.FirstOrDefault(path => FullPath.Contains(path));
+			if (FoundPerferredPath == null)
+			{
+				return null;
+			}
+
+			int end = FullPath.LastIndexOf(FoundPerferredPath) + FoundPerferredPath.Length;
+			return FullPath.Substring(end);
+		}
+
+		/// <summary>
+		/// Execute the command
+		/// </summary>
+		/// <param name="Arguments">Command line arguments</param>
+		/// <returns>Exit code</returns>
+		/// <param name="Logger"></param>
+		public override int Execute(CommandLineArguments Arguments, ILogger Logger)
+		{
+			Arguments.ApplyTo(this);
+
+			// Create the build configuration object, and read the settings
+			BuildConfiguration BuildConfiguration = new BuildConfiguration();
+			XmlConfig.ApplyTo(BuildConfiguration);
+			Arguments.ApplyTo(BuildConfiguration);
+
+			// Parse the filter argument
+			FileFilter? FileFilter = null;
+			if (FilterRules.Count > 0)
+			{
+				FileFilter = new FileFilter(FileFilterType.Exclude);
+				foreach (string FilterRule in FilterRules)
+				{
+					FileFilter.AddRules(FilterRule.Split(';'));
+				}
+			}
+
+			// Force C++ modules to always include their generated code directories
+			UEBuildModuleCPP.bForceAddGeneratedCodeIncludePath = true;
+
+			// Parse all the target descriptors
+			List<TargetDescriptor> TargetDescriptors = TargetDescriptor.ParseCommandLine(Arguments, BuildConfiguration.bUsePrecompiled, BuildConfiguration.bSkipRulesCompile, BuildConfiguration.bForceRulesCompile, Logger);
+
+			// Generate the compile DB for each target
+			using (ISourceFileWorkingSet WorkingSet = new EmptySourceFileWorkingSet())
+			{
+				HashSet<UEBuildModule> ScannedModules = new();
+
+				// Find the compile commands for each file in the target
+				Dictionary<FileReference, string> FileToCommand = new Dictionary<FileReference, string>();
+				foreach (TargetDescriptor TargetDescriptor in TargetDescriptors)
+				{
+					// Create a makefile for the target
+					UEBuildTarget Target = UEBuildTarget.Create(TargetDescriptor, BuildConfiguration.bSkipRulesCompile, BuildConfiguration.bForceRulesCompile, BuildConfiguration.bUsePrecompiled, Logger);
+
+					// Create all the binaries and modules
+					CppCompileEnvironment GlobalCompileEnvironment = Target.CreateCompileEnvironmentForProjectFiles(Logger);
+					foreach (UEBuildBinary Binary in Target.Binaries)
+					{
+						CppCompileEnvironment BinaryCompileEnvironment = Binary.CreateBinaryCompileEnvironment(GlobalCompileEnvironment);
+
+						foreach (UEBuildModuleCPP Module in Binary.Modules.OfType<UEBuildModuleCPP>())
+						{
+							bool IsThirdPartyModule = Module.RulesFile.ContainsName("ThirdParty", Unreal.RootDirectory);
+							if (IsThirdPartyModule)
+								continue;
+
+							UEBuildModuleCPP.InputFileCollection InputFileCollection = Module.FindInputFiles(Target.Platform, new Dictionary<DirectoryItem, FileItem[]>());
+							List<FileItem> InputFiles = new List<FileItem>();
+							InputFiles.AddRange(InputFileCollection.HeaderFiles);
+							InputFiles.AddRange(InputFileCollection.CPPFiles);
+							InputFiles.AddRange(InputFileCollection.CCFiles);
+							InputFiles.AddRange(InputFileCollection.CFiles);
+
+							var FileList = new List<FileReference>();
+							foreach (FileItem InputFile in InputFiles)
+							{
+								if (FileFilter == null || FileFilter.Matches(InputFile.Location.MakeRelativeTo(Unreal.RootDirectory)))
+								{
+									var fileRef = new FileReference(InputFile.AbsolutePath);
+									FileList.Add(fileRef);
+								}
+							}
+
+							if (FileList.Any())
+							{
+								Dictionary<string, string?> PerferredPathCache = new();
+								CppCompileEnvironment env = Module.CreateCompileEnvironmentForIntellisense(Target.Rules, BinaryCompileEnvironment, Logger);
+
+								foreach (var InputFile in FileList)
+								{
+									var Text = FileReference.ReadAllLines(InputFile);
+									var UpdatedText = false;
+
+									for (int i = 0; i < Text.Length; i++)
+									{
+										var Line = Text[i];
+										Match IncludeMatch = IncludeRegex.Match(Line);
+										if (IncludeMatch.Success)
+										{
+											string Include = IncludeMatch.Groups[1].Value;
+
+											if (Include.Contains("/Private/") && PublicDirectories.Any(dir => InputFile.FullName.Contains(System.IO.Path.DirectorySeparatorChar + dir + System.IO.Path.DirectorySeparatorChar)))
+											{
+												Logger.LogError("Can not update #include '{Include}' in the public header '{FileName}' because it will break code including this header.", Include, InputFile.FullName);
+												continue;
+											}
+
+											string? PerferredInclude = null;
+											if (!PerferredPathCache.TryGetValue(Include, out PerferredInclude))
+											{
+												var FullPath = FindIncludePath(InputFile.FullName, env, Include);
+												if (FullPath != null)
+												{
+													// if the include and the source file live in the same directory then it is OK to be relative
+													if (string.Equals(System.IO.Directory.GetParent(FullPath)?.FullName, System.IO.Directory.GetParent(InputFile.FullName)?.FullName, StringComparison.CurrentCultureIgnoreCase) &&
+														string.Equals(Include, System.IO.Path.GetFileName(FullPath), StringComparison.CurrentCultureIgnoreCase))
+													{
+														PerferredInclude = Include;
+													}
+													else
+													{
+														PerferredInclude = GetPerferredInclude(FullPath);
+														PerferredPathCache[Include] = PerferredInclude;
+													}
+												}
+												
+												if (PerferredInclude == null)
+												{
+													Logger.LogWarning("Could not find include path for '{IncludePath}' found in '{FileName}'", Include, InputFile.FullName);
+												}
+											}
+
+											if (PerferredInclude != null && Include != PerferredInclude)
+											{
+												Logger.LogInformation("Updated '{InputFileName}' line {LineNum} -- {OldInclude} -> {NewInclude}", InputFile.FullName, i, Include, PerferredInclude);
+												Text[i] = Line.Replace(Include, PerferredInclude);
+												UpdatedText = true;
+											}
+										}
+									}
+
+									if (UpdatedText)
+									{
+
+										if (!bNoOutput)
+										{
+											Logger.LogInformation("Updating {IncludePath}", InputFile.FullName);
+											try
+											{
+												if (bCheckoutWithP4)
+												{
+													System.Diagnostics.Process Process = new System.Diagnostics.Process();
+													System.Diagnostics.ProcessStartInfo StartInfo = new System.Diagnostics.ProcessStartInfo();
+													Process.StartInfo.WindowStyle = System.Diagnostics.ProcessWindowStyle.Hidden;
+													Process.StartInfo.FileName = "p4.exe";
+													Process.StartInfo.Arguments = $"edit {InputFile.FullName}";
+													Process.Start();
+													Process.WaitForExit();
+												}
+												System.IO.File.WriteAllLines(InputFile.FullName, Text);
+											}
+											catch (Exception ex)
+											{
+												Logger.LogWarning("Failed to write to file: {Exception}", ex);
+											}
+										}
+									}
+								}
+							}
+
+							ScannedModules.Add(Module);
+						}
+					}
+				}
+			}
+
+			return 0;
+		}
+	}
+}
