- Fixed VSMac project generation
#okfrogithub public

diff --git a/Engine/Source/Programs/UnrealBuildTool/ProjectFiles/VisualStudio/VCProjectFileGenerator.cs b/Engine/Source/Programs/UnrealBuildTool/ProjectFiles/VisualStudio/VCProjectFileGenerator.cs
--- a/Engine/Source/Programs/UnrealBuildTool/ProjectFiles/VisualStudio/VCProjectFileGenerator.cs
+++ b/Engine/Source/Programs/UnrealBuildTool/ProjectFiles/VisualStudio/VCProjectFileGenerator.cs
@@ -294,7 +294,7 @@ namespace UnrealBuildTool
 				List<VisualStudioInstallation> Installations = MicrosoftPlatformSDK.FindVisualStudioInstallations(Logger).Where(x => WindowsPlatform.HasCompiler(x.Compiler, WindowsArchitecture.x64, Logger)).ToList();
 
 				// Get the corresponding project file format
-				VCProjectFileFormat Format = VCProjectFileFormat.VisualStudio2019;
+				VCProjectFileFormat Format = VCProjectFileFormat.Default;
 				foreach (VisualStudioInstallation Installation in Installations)
 				{
 					if (Installation.Compiler == WindowsCompiler.VisualStudio2022)
