Added GetModuleDirectory which allows you to get the directory of the passed in module. This makes it easier to include directories from other modules.

diff --git a/Engine/Source/Programs/UnrealBuildTool/Configuration/ModuleRules.cs b/Engine/Source/Programs/UnrealBuildTool/Configuration/ModuleRules.cs
--- a/Engine/Source/Programs/UnrealBuildTool/Configuration/ModuleRules.cs
+++ b/Engine/Source/Programs/UnrealBuildTool/Configuration/ModuleRules.cs
@@ -559,6 +559,11 @@ namespace UnrealBuildTool
 		/// </summary>
 		private List<DirectoryReference> AdditionalModuleDirectories = new List<DirectoryReference>();
 
+		/// <summary>
+		/// The rules assembly to use when searching for modules
+		/// </summary>
+		internal RulesAssembly RulesAssembly;
+
 		/// <summary>
 		/// Plugin containing this module
 		/// </summary>
@@ -1005,6 +1010,21 @@ namespace UnrealBuildTool
 		/// </summary>
 		public List<string> PublicAdditionalLibraries = new List<string>();
 
+		/// <summary>
+		/// Returns the directory of where the passed in module name lives.
+		/// </summary>
+		/// <param name="ModuleName">Name of the module</param>
+		/// <returns>Directory where the module lives</returns>
+		public string GetModuleDirectory(string ModuleName)
+		{
+			FileReference? ModuleFileReference = RulesAssembly.GetModuleFileName(ModuleName);
+			if (ModuleFileReference == null)
+			{
+				throw new BuildException("Could not find a module named '{0}'.", ModuleName);
+			}
+			return ModuleFileReference.Directory.FullName;
+		}
+
 		/// <summary>
 		/// List of additional pre-build libraries (names of the .lib files including extension) - typically used for additional targets which are still built, but using either TargetRules.PreBuildSteps or TargetRules.PreBuildTargets.
 		/// </summary>
diff --git a/Engine/Source/Programs/UnrealBuildTool/System/RulesAssembly.cs b/Engine/Source/Programs/UnrealBuildTool/System/RulesAssembly.cs
--- a/Engine/Source/Programs/UnrealBuildTool/System/RulesAssembly.cs
+++ b/Engine/Source/Programs/UnrealBuildTool/System/RulesAssembly.cs
@@ -481,6 +481,7 @@ namespace UnrealBuildTool
 				}
 				RulesObject.bPrecompile = (RulesObject.bTreatAsEngineModule || ModuleName.Equals("UnrealGame", StringComparison.OrdinalIgnoreCase)) && Target.bPrecompile;
 				RulesObject.bUsePrecompiled = bReadOnly;
+				RulesObject.RulesAssembly = this;
 
 				// go up the type hierarchy (if there is a hierarchy), looking for any extra directories for the module
 				if (RulesObjectType != BaseRulesObjectType && RulesObjectType != typeof(ModuleRules))
