Updated the FixIncludePathsMode to sort the includes after updating them.

diff --git a/Engine/Source/Programs/UnrealBuildTool/Modes/FixIncludePathsMode.cs b/Engine/Source/Programs/UnrealBuildTool/Modes/FixIncludePathsMode.cs
--- a/Engine/Source/Programs/UnrealBuildTool/Modes/FixIncludePathsMode.cs
+++ b/Engine/Source/Programs/UnrealBuildTool/Modes/FixIncludePathsMode.cs
@@ -8,6 +8,7 @@ using System.Text.RegularExpressions;
 using EpicGames.Core;
 using UnrealBuildBase;
 using Microsoft.Extensions.Logging;
+using System.Collections;
 
 namespace UnrealBuildTool
 {
@@ -24,7 +25,6 @@ namespace UnrealBuildTool
 
 		static readonly string UnrealRootDirectory = Unreal.RootDirectory.FullName.Replace('\\', '/');
 		static readonly string[] PerferredPaths = { "/Public/", "/Private/", "/Classes/", "/Internal/", "/UHT/", "/VNI/" };
-
 		static readonly string[] PublicDirectories = { "Public", "Classes", };
 
 		[CommandLine("-Filter=", Description = "Set of filters for files to include in the database. Relative to the root directory, or to the project file.")]
@@ -36,6 +36,9 @@ namespace UnrealBuildTool
 		[CommandLine("-NoOutput", Description = "Flags that the updated files shouldn't be saved.")]
 		public bool bNoOutput = false;
 
+		[CommandLine("-NoIncludeSorting", Description = "Flags that includes should not be sorted.")]
+		public bool bNoIncludeSorting = false;
+
 		private string? FindIncludePath(string FilePath, CppCompileEnvironment env, string HeaderIncludePath)
 		{
 			List<DirectoryReference> EnvPaths = new();
@@ -43,6 +46,7 @@ namespace UnrealBuildTool
 			EnvPaths.AddRange(env.UserIncludePaths);
 			//EnvPaths.AddRange(env.SystemIncludePaths);
 
+			string FoundPath = Path.Replace('\\', '/');
 			// search include paths
 			foreach (var UserIncludePath in EnvPaths)
 			{
@@ -153,8 +157,9 @@ namespace UnrealBuildTool
 
 								foreach (var InputFile in FileList)
 								{
-									var Text = FileReference.ReadAllLines(InputFile);
-									var UpdatedText = false;
+									List<int> LinesUpdated = new();
+									string[] Text = FileReference.ReadAllLines(InputFile);
+									bool UpdatedText = false;
 
 									for (int i = 0; i < Text.Length; i++)
 									{
@@ -200,12 +205,17 @@ namespace UnrealBuildTool
 												Logger.LogInformation("Updated '{InputFileName}' line {LineNum} -- {OldInclude} -> {NewInclude}", InputFile.FullName, i, Include, PerferredInclude);
 												Text[i] = Line.Replace(Include, PerferredInclude);
 												UpdatedText = true;
+												LinesUpdated.Add(i);
 											}
 										}
 									}
 
 									if (UpdatedText)
 									{
+										if (!bNoIncludeSorting)
+										{
+											SortIncludes(InputFile, LinesUpdated, Text);
+										}
 
 										if (!bNoOutput)
 										{
@@ -241,5 +251,90 @@ namespace UnrealBuildTool
 
 			return 0;
 		}
+
+		class HeaderSortComparison : IComparer<string>
+		{
+			private string IWYUFileName;
+
+			public HeaderSortComparison(string IWYUFileName)
+			{
+				this.IWYUFileName = IWYUFileName;
+			}
+
+			public int Compare(string? x, string? y)
+			{
+				if (String.IsNullOrEmpty(x) && String.IsNullOrEmpty(y))
+				{
+					return 0;
+				}
+
+				if (String.IsNullOrEmpty(y))
+				{
+					return -1;
+				}
+
+				if (String.IsNullOrEmpty(x))
+				{
+					return 1;
+				}
+
+				// IWYU header
+				if (x.Contains(IWYUFileName))
+				{
+					return -1;
+				}
+				if (y.Contains(IWYUFileName))
+				{
+					return 1;
+				}
+
+				// generated header
+				if (x.Contains(".generated.h"))
+				{
+					return 1;
+				}
+				if (y.Contains(".generated.h"))
+				{
+					return -1;
+				}
+				return string.Compare(x, y);
+			}
+		}
+
+		private void SortIncludes(FileReference File, List<int> LinesUpdated, string[] Text)
+		{
+			var HeaderSort = new HeaderSortComparison(File.GetFileNameWithoutExtension() + ".h");
+			foreach (var LineIndex in LinesUpdated)
+			{
+				var FirstIncludeIndex = LineIndex;
+				for (int i = LineIndex - 1; i >= 0; i--)
+				{
+					Match IncludeMatch = IncludeRegex.Match(Text[i]);
+					if (IncludeMatch.Success)
+					{
+						FirstIncludeIndex = i;
+					}
+					else
+					{
+						break;
+					}
+				}
+
+				var LastIncludeIndex = LineIndex;
+				for (int i = LineIndex + 1; i < Text.Length; i++)
+				{
+					Match IncludeMatch = IncludeRegex.Match(Text[i]);
+					if (IncludeMatch.Success)
+					{
+						LastIncludeIndex = i;
+					}
+					else
+					{
+						break;
+					}
+				}
+				Array.Sort(Text, FirstIncludeIndex, LastIncludeIndex - FirstIncludeIndex + 1, HeaderSort);
+			}
+		}
 	}
 }
