Minor cleanup.

diff --git a/Engine/Source/Programs/UnrealBuildTool/Modes/FixIncludePathsMode.cs b/Engine/Source/Programs/UnrealBuildTool/Modes/FixIncludePathsMode.cs
--- a/Engine/Source/Programs/UnrealBuildTool/Modes/FixIncludePathsMode.cs
+++ b/Engine/Source/Programs/UnrealBuildTool/Modes/FixIncludePathsMode.cs
@@ -39,44 +39,6 @@ namespace UnrealBuildTool
 		[CommandLine("-NoIncludeSorting", Description = "Flags that includes should not be sorted.")]
 		public bool bNoIncludeSorting = false;
 
-		private string? FindIncludePath(string FilePath, CppCompileEnvironment env, string HeaderIncludePath)
-		{
-			List<DirectoryReference> EnvPaths = new();
-			EnvPaths.Add(new DirectoryReference(System.IO.Directory.GetParent(FilePath)!));
-			EnvPaths.AddRange(env.UserIncludePaths);
-			//EnvPaths.AddRange(env.SystemIncludePaths);
-
-			string FoundPath = Path.Replace('\\', '/');
-			// search include paths
-			foreach (var UserIncludePath in EnvPaths)
-			{
-				string Path = System.IO.Path.GetFullPath(System.IO.Path.Combine(UserIncludePath.FullName, HeaderIncludePath));
-				if (System.IO.File.Exists(Path))
-				{
-					return Path.Replace('\\', '/');
-				}
-			}
-
-			return null;
-		}
-
-		private string? GetPerferredInclude(string FullPath)
-		{
-			if (!FullPath.Contains(UnrealRootDirectory))
-			{
-				return null;
-			}
-
-			string? FoundPerferredPath = PerferredPaths.FirstOrDefault(path => FullPath.Contains(path));
-			if (FoundPerferredPath == null)
-			{
-				return null;
-			}
-
-			int end = FullPath.LastIndexOf(FoundPerferredPath) + FoundPerferredPath.Length;
-			return FullPath.Substring(end);
-		}
-
 		/// <summary>
 		/// Execute the command
 		/// </summary>
@@ -178,8 +140,26 @@ namespace UnrealBuildTool
 											string? PerferredInclude = null;
 											if (!PerferredPathCache.TryGetValue(Include, out PerferredInclude))
 											{
-												var FullPath = FindIncludePath(InputFile.FullName, env, Include);
-												if (FullPath != null)
+												List<DirectoryReference> IncludePaths = new();
+												IncludePaths.Add(new DirectoryReference(System.IO.Directory.GetParent(InputFile.FullName)!));
+												IncludePaths.AddRange(env.UserIncludePaths);
+												//EnvPaths.AddRange(env.SystemIncludePaths);
+
+												// search include paths
+												string? FullPath = null;
+												DirectoryReference? FoundIncludePath = null;
+												foreach (var IncludePath in IncludePaths)
+												{
+													string Path = System.IO.Path.GetFullPath(System.IO.Path.Combine(IncludePath.FullName, Include));
+													if (System.IO.File.Exists(Path))
+													{
+														FullPath = Path.Replace('\\', '/');
+														FoundIncludePath = IncludePath;
+														break;
+													}
+												}
+
+												if (!string.IsNullOrEmpty(FullPath))
 												{
 													// if the include and the source file live in the same directory then it is OK to be relative
 													if (string.Equals(System.IO.Directory.GetParent(FullPath)?.FullName, System.IO.Directory.GetParent(InputFile.FullName)?.FullName, StringComparison.CurrentCultureIgnoreCase) &&
@@ -189,7 +169,17 @@ namespace UnrealBuildTool
 													}
 													else
 													{
-														PerferredInclude = GetPerferredInclude(FullPath);
+														string? FoundPerferredPath = PerferredPaths.FirstOrDefault(path => FullPath.Contains(path));
+														if (!FullPath.Contains(UnrealRootDirectory) || FoundPerferredPath == null)
+														{
+															PerferredInclude = null;
+														}
+														else
+														{
+															int end = FullPath.LastIndexOf(FoundPerferredPath) + FoundPerferredPath.Length;
+															PerferredInclude = FullPath.Substring(end);
+														}
+
 														PerferredPathCache[Include] = PerferredInclude;
 													}
 												}
