TestTarget Resource folders can now exist in any of the module's locations (eg, platform extension directory)

diff --git a/Engine/Source/Programs/UnrealBuildTool/Configuration/TestModuleRules.cs b/Engine/Source/Programs/UnrealBuildTool/Configuration/TestModuleRules.cs
--- a/Engine/Source/Programs/UnrealBuildTool/Configuration/TestModuleRules.cs
+++ b/Engine/Source/Programs/UnrealBuildTool/Configuration/TestModuleRules.cs
@@ -120,11 +120,7 @@ namespace UnrealBuildTool
 				BuildOrderSettings.AddBuildOrderOverride(TeardownFile, SourceFileBuildOrder.Last);
 			}
 
-			string TestResourcesDir = Path.Combine(ModuleDirectory, "Resources");
-			if (System.IO.Directory.Exists(TestResourcesDir))
-			{
-				AdditionalPropertiesForReceipt.Add("ResourcesFolder", TestResourcesDir);
-			}
+			SetResourcesFolder("Resources");
 
 			if (!PublicDependencyModuleNames.Contains("Catch2"))
 			{
@@ -149,19 +145,12 @@ namespace UnrealBuildTool
 		/// </summary>
 		protected void SetResourcesFolder(string ResourcesRelativeFolder)
 		{
-			string TestResourcesDir = Path.Combine(ModuleDirectory, ResourcesRelativeFolder);
-			if (!System.IO.Directory.Exists(TestResourcesDir))
-			{
-				throw new Exception($"Resources folder not found {TestResourcesDir}");
-			}
-			else
+			AdditionalPropertiesForReceipt.Inner.RemoveAll(Prop => Prop.Name == "ResourcesFolder");
+
+			foreach (DirectoryReference Directory in GetAllModuleDirectories())
 			{
-				ReceiptProperty? Existing = AdditionalPropertiesForReceipt.Inner.Find(Prop => Prop.Name == "ResourcesFolder");
-				if (Existing != null)
-				{
-					Existing.Value = TestResourcesDir;
-				}
-				else
+				string TestResourcesDir = Path.Combine(Directory.FullName, ResourcesRelativeFolder);
+				if (System.IO.Directory.Exists(TestResourcesDir))
 				{
 					AdditionalPropertiesForReceipt.Add("ResourcesFolder", TestResourcesDir);
 				}
diff --git a/Engine/Source/Programs/UnrealBuildTool/Configuration/UEBuildDeploy.cs b/Engine/Source/Programs/UnrealBuildTool/Configuration/UEBuildDeploy.cs
--- a/Engine/Source/Programs/UnrealBuildTool/Configuration/UEBuildDeploy.cs
+++ b/Engine/Source/Programs/UnrealBuildTool/Configuration/UEBuildDeploy.cs
@@ -67,19 +67,19 @@ namespace UnrealBuildTool
 		/// <returns>True if successful, false if not</returns>
 		public virtual bool PrepTargetForDeployment(TargetReceipt Receipt)
 		{
-			if (Receipt.IsTestTarget)
+			bool Result = true;
+			if (Receipt.IsTestTarget && Receipt.Launch != null)
 			{
-				ReceiptProperty? TestResourcesFolder = Receipt.AdditionalProperties.FirstOrDefault(Prop => Prop.Name == "ResourcesFolder");
-				if (TestResourcesFolder != null && Receipt.Launch != null)
+				foreach (ReceiptProperty TestResourcesFolder in Receipt.AdditionalProperties.FindAll(Prop => Prop.Name == "ResourcesFolder"))
 				{
 					DirectoryInfo ResourceDirectoryInfo = new DirectoryInfo(TestResourcesFolder.Value);
 					if (ResourceDirectoryInfo.Exists)
 					{
-						return CopyFilesRecursively(TestResourcesFolder.Value, Path.Combine(Receipt.Launch.Directory.FullName, ResourceDirectoryInfo.Name));
+						Result &= CopyFilesRecursively(TestResourcesFolder.Value, Path.Combine(Receipt.Launch.Directory.FullName, ResourceDirectoryInfo.Name));
 					}
 				}
 			}
-			return true;
+			return Result;
 		}
 	}
 }
