- Add concept of trusted and untrusted server code/targets, so we can include/exclude code from certain trusted/untrusted build targets.

diff --git a/Engine/Source/Programs/UnrealBuildTool/Configuration/TargetRules.cs b/Engine/Source/Programs/UnrealBuildTool/Configuration/TargetRules.cs
--- a/Engine/Source/Programs/UnrealBuildTool/Configuration/TargetRules.cs
+++ b/Engine/Source/Programs/UnrealBuildTool/Configuration/TargetRules.cs
@@ -989,6 +989,30 @@ namespace UnrealBuildTool
 		}
 		private bool? bWithServerCodeOverride;
 
+		/// <summary>
+		/// Compile trusted-server-only code.
+		/// </summary>
+		[RequiresUniqueBuildEnvironment]
+		public bool bWithServerCodeTrusted
+		{
+			get { return bWithServerCode && bWithServerCodeTrustedPrivate; }
+			set { bWithServerCodeTrustedPrivate = value; }
+		}
+
+		/// <summary>
+		/// Compile untrusted-server-only code.
+		/// </summary>
+		[RequiresUniqueBuildEnvironment]
+		public bool bWithServerCodeUntrusted
+		{
+			get { return bWithServerCode && !bWithServerCodeTrustedPrivate; }
+			set { bWithServerCodeTrustedPrivate = !value; }
+		}
+
+		[CommandLine("-TrustedServer", Value = "true")]
+		[CommandLine("-NoTrustedServer", Value = "false")]
+		private bool bWithServerCodeTrustedPrivate = true;
+
 		/// <summary>
 		/// Compile with FName storing the number part in the name table. 
 		/// Saves memory when most names are not numbered and those that are are referenced multiple times.
@@ -2223,6 +2247,14 @@ namespace UnrealBuildTool
 			{
 				GlobalDefinitions.Add("UE_SERVER=1");
 				GlobalDefinitions.Add("USE_NULL_RHI=1");
+				if(bWithServerCodeTrusted)
+                {
+					GlobalDefinitions.Add("UE_SERVER_TRUSTED=1");
+                }
+                else
+                {
+					GlobalDefinitions.Add("UE_SERVER_UNTRUSTED=1");
+				}
 			}
 		}
 
@@ -2859,6 +2891,16 @@ namespace UnrealBuildTool
 			get { return Inner.bWithServerCode; }
 		}
 
+		public bool bWithServerCodeTrusted
+		{
+			get { return Inner.bWithServerCodeTrusted; }
+		}
+
+		public bool bWithServerCodeUntrusted
+		{
+			get { return Inner.bWithServerCodeUntrusted; }
+		}
+
 		public bool bFNameOutlineNumber
 		{
 			get { return Inner.bFNameOutlineNumber;  }
diff --git a/Engine/Source/Programs/UnrealBuildTool/Configuration/UEBuildTarget.cs b/Engine/Source/Programs/UnrealBuildTool/Configuration/UEBuildTarget.cs
--- a/Engine/Source/Programs/UnrealBuildTool/Configuration/UEBuildTarget.cs
+++ b/Engine/Source/Programs/UnrealBuildTool/Configuration/UEBuildTarget.cs
@@ -1,4 +1,4 @@
-﻿// Copyright Epic Games, Inc. All Rights Reserved.
+// Copyright Epic Games, Inc. All Rights Reserved.
 
 using System;
 using System.Collections.Generic;
@@ -4172,14 +4172,9 @@ namespace UnrealBuildTool
 			}
 
 			// Check if server-only code should be compiled out.
-			if (Rules.bWithServerCode == true)
-			{
-				GlobalCompileEnvironment.Definitions.Add("WITH_SERVER_CODE=1");
-			}
-			else
-			{
-				GlobalCompileEnvironment.Definitions.Add("WITH_SERVER_CODE=0");
-			}
+			GlobalCompileEnvironment.Definitions.AddFormat("WITH_SERVER_CODE={0}", Rules.bWithServerCode ? "1" : "0");
+			GlobalCompileEnvironment.Definitions.AddFormat("WITH_SERVER_CODE_TRUSTED={0}", Rules.bWithServerCodeTrusted ? "1" : "0");
+			GlobalCompileEnvironment.Definitions.AddFormat("WITH_SERVER_CODE_UNTRUSTED={0}", Rules.bWithServerCodeUntrusted ? "1" : "0");
 
 			GlobalCompileEnvironment.Definitions.Add(String.Format("UE_FNAME_OUTLINE_NUMBER={0}", Rules.bFNameOutlineNumber? 1 : 0));
 
diff --git a/Engine/Source/Runtime/Core/Public/Misc/Build.h b/Engine/Source/Runtime/Core/Public/Misc/Build.h
--- a/Engine/Source/Runtime/Core/Public/Misc/Build.h
+++ b/Engine/Source/Runtime/Core/Public/Misc/Build.h
@@ -39,6 +39,30 @@
 	#define UE_SERVER					0
 #endif
 
+/**
+ *   Whether compiling for a trusted dedicated server or not.
+ */
+#ifndef UE_SERVER_TRUSTED
+	#define UE_SERVER_TRUSTED			0
+#endif
+
+/**
+ *   Whether compiling for an untrusted dedicated server or not.
+ */
+#ifndef UE_SERVER_UNTRUSTED
+	#define UE_SERVER_UNTRUSTED			0
+#endif
+
+#if UE_SERVER
+	#if UE_SERVER_TRUSTED + UE_SERVER_UNTRUSTED != 1
+		#error Exactly one of [UE_SERVER_TRUSTED UE_SERVER_UNTRUSTED] should be defined to be 1 when UE_SERVER=1
+	#endif
+#else
+	#if UE_SERVER_TRUSTED || UE_SERVER_UNTRUSTED
+		#error All of [UE_SERVER_TRUSTED UE_SERVER_UNTRUSTED] should be defined to be 0 when UE_SERVER=0
+	#endif
+#endif
+
 // Ensure that we have one, and only one build config coming from UBT
 #if UE_BUILD_DEBUG + UE_BUILD_DEVELOPMENT + UE_BUILD_TEST + UE_BUILD_SHIPPING != 1
 	#error Exactly one of [UE_BUILD_DEBUG UE_BUILD_DEVELOPMENT UE_BUILD_TEST UE_BUILD_SHIPPING] should be defined to be 1
