UBT: Remove dependency on HoloLens code from UEBuildTarget.cs.

diff --git a/Engine/Source/Programs/UnrealBuildTool/Configuration/UEBuildTarget.cs b/Engine/Source/Programs/UnrealBuildTool/Configuration/UEBuildTarget.cs
--- a/Engine/Source/Programs/UnrealBuildTool/Configuration/UEBuildTarget.cs
+++ b/Engine/Source/Programs/UnrealBuildTool/Configuration/UEBuildTarget.cs
@@ -4172,9 +4172,9 @@ namespace UnrealBuildTool
 			}
 
 			// Check if server-only code should be compiled out.
-			GlobalCompileEnvironment.Definitions.AddFormat("WITH_SERVER_CODE={0}", Rules.bWithServerCode ? "1" : "0");
-			GlobalCompileEnvironment.Definitions.AddFormat("WITH_SERVER_CODE_TRUSTED={0}", Rules.bWithServerCodeTrusted ? "1" : "0");
-			GlobalCompileEnvironment.Definitions.AddFormat("WITH_SERVER_CODE_UNTRUSTED={0}", Rules.bWithServerCodeUntrusted ? "1" : "0");
+			GlobalCompileEnvironment.Definitions.Add(String.Format("WITH_SERVER_CODE={0}", Rules.bWithServerCode ? "1" : "0"));
+			GlobalCompileEnvironment.Definitions.Add(String.Format("WITH_SERVER_CODE_TRUSTED={0}", Rules.bWithServerCodeTrusted ? "1" : "0"));
+			GlobalCompileEnvironment.Definitions.Add(String.Format("WITH_SERVER_CODE_UNTRUSTED={0}", Rules.bWithServerCodeUntrusted ? "1" : "0"));
 
 			GlobalCompileEnvironment.Definitions.Add(String.Format("UE_FNAME_OUTLINE_NUMBER={0}", Rules.bFNameOutlineNumber? 1 : 0));
 
