Updated FixIncludePathsMode to support thirdparty directories and use the shortened include paths if possible.

diff --git a/Engine/Source/Programs/UnrealBuildTool/Modes/FixIncludePathsMode.cs b/Engine/Source/Programs/UnrealBuildTool/Modes/FixIncludePathsMode.cs
--- a/Engine/Source/Programs/UnrealBuildTool/Modes/FixIncludePathsMode.cs
+++ b/Engine/Source/Programs/UnrealBuildTool/Modes/FixIncludePathsMode.cs
@@ -24,7 +24,7 @@ namespace UnrealBuildTool
 		static readonly Regex IncludeRegex = new Regex("^[ \t]*#[ \t]*include[ \t]*[\"](?<HeaderFile>[^\"]*)[\"]", RegexOptions.Compiled | RegexOptions.Singleline | RegexOptions.ExplicitCapture);
 
 		static readonly string UnrealRootDirectory = Unreal.RootDirectory.FullName.Replace('\\', '/');
-		static readonly string[] PerferredPaths = { "/Public/", "/Private/", "/Classes/", "/Internal/", "/UHT/", "/VNI/" };
+		static readonly string[] PreferredPaths = { "/Public/", "/Private/", "/Classes/", "/Internal/", "/UHT/", "/VNI/" };
 		static readonly string[] PublicDirectories = { "Public", "Classes", };
 
 		[CommandLine("-Filter=", Description = "Set of filters for files to include in the database. Relative to the root directory, or to the project file.")]
@@ -114,7 +114,7 @@ namespace UnrealBuildTool
 
 							if (FileList.Any())
 							{
-								Dictionary<string, string?> PerferredPathCache = new();
+								Dictionary<string, string?> PreferredPathCache = new();
 								CppCompileEnvironment env = Module.CreateCompileEnvironmentForIntellisense(Target.Rules, BinaryCompileEnvironment, Logger);
 
 								foreach (var InputFile in FileList)
@@ -137,8 +137,8 @@ namespace UnrealBuildTool
 												continue;
 											}
 
-											string? PerferredInclude = null;
-											if (!PerferredPathCache.TryGetValue(Include, out PerferredInclude))
+											string? PreferredInclude = null;
+											if (!PreferredPathCache.TryGetValue(Include, out PreferredInclude))
 											{
 												List<DirectoryReference> IncludePaths = new();
 												IncludePaths.Add(new DirectoryReference(System.IO.Directory.GetParent(InputFile.FullName)!));
@@ -161,39 +161,54 @@ namespace UnrealBuildTool
 
 												if (!string.IsNullOrEmpty(FullPath))
 												{
+													if (FullPath.Contains("ThirdParty"))
+													{
+														Logger.LogDebug("Skipping '{Include}' because it is a third party header", Include);
+														PreferredInclude = Include;
+														PreferredPathCache[Include] = PreferredInclude;
+														continue;
+													}
+
 													// if the include and the source file live in the same directory then it is OK to be relative
 													if (string.Equals(System.IO.Directory.GetParent(FullPath)?.FullName, System.IO.Directory.GetParent(InputFile.FullName)?.FullName, StringComparison.CurrentCultureIgnoreCase) &&
 														string.Equals(Include, System.IO.Path.GetFileName(FullPath), StringComparison.CurrentCultureIgnoreCase))
 													{
-														PerferredInclude = Include;
+														PreferredInclude = Include;
 													}
 													else
 													{
-														string? FoundPerferredPath = PerferredPaths.FirstOrDefault(path => FullPath.Contains(path));
-														if (!FullPath.Contains(UnrealRootDirectory) || FoundPerferredPath == null)
+														string? FoundPreferredPath = PreferredPaths.FirstOrDefault(path => FullPath.Contains(path));
+														if (!FullPath.Contains(UnrealRootDirectory) || FoundPreferredPath == null)
 														{
-															PerferredInclude = null;
+															PreferredInclude = null;
 														}
 														else
 														{
-															int end = FullPath.LastIndexOf(FoundPerferredPath) + FoundPerferredPath.Length;
-															PerferredInclude = FullPath.Substring(end);
+															int end = FullPath.LastIndexOf(FoundPreferredPath) + FoundPreferredPath.Length;
+															PreferredInclude = FullPath.Substring(end);
+
+															// Is the current include a shortened version of the preferred include path?
+															if (PreferredInclude != Include && PreferredInclude.Contains(Include))
+															{
+																Logger.LogDebug("Using '{Include}' because it is shorter than '{PreferredInclude}'", Include, PreferredInclude);
+																PreferredInclude = Include;
+															}
 														}
 
-														PerferredPathCache[Include] = PerferredInclude;
+														PreferredPathCache[Include] = PreferredInclude;
 													}
 												}
 												
-												if (PerferredInclude == null)
+												if (PreferredInclude == null)
 												{
 													Logger.LogWarning("Could not find include path for '{IncludePath}' found in '{FileName}'", Include, InputFile.FullName);
 												}
 											}
 
-											if (PerferredInclude != null && Include != PerferredInclude)
+											if (PreferredInclude != null && Include != PreferredInclude)
 											{
-												Logger.LogInformation("Updated '{InputFileName}' line {LineNum} -- {OldInclude} -> {NewInclude}", InputFile.FullName, i, Include, PerferredInclude);
-												Text[i] = Line.Replace(Include, PerferredInclude);
+												Logger.LogInformation("Updated '{InputFileName}' line {LineNum} -- {OldInclude} -> {NewInclude}", InputFile.FullName, i, Include, PreferredInclude);
+												Text[i] = Line.Replace(Include, PreferredInclude);
 												UpdatedText = true;
 												LinesUpdated.Add(i);
 											}
