Updated GenerateClangDatabase to update more relative paths to be absolute.

diff --git a/Engine/Source/Programs/UnrealBuildTool/Modes/GenerateClangDatabase.cs b/Engine/Source/Programs/UnrealBuildTool/Modes/GenerateClangDatabase.cs
--- a/Engine/Source/Programs/UnrealBuildTool/Modes/GenerateClangDatabase.cs
+++ b/Engine/Source/Programs/UnrealBuildTool/Modes/GenerateClangDatabase.cs
@@ -17,7 +17,7 @@ namespace UnrealBuildTool
 	[ToolMode("GenerateClangDatabase", ToolModeOptions.XmlConfig | ToolModeOptions.BuildPlatforms | ToolModeOptions.SingleInstance | ToolModeOptions.StartPrefetchingEngine | ToolModeOptions.ShowExecutionTime)]
 	class GenerateClangDatabase : ToolMode
 	{
-		static Regex LineRegex = new Regex(@"^(\-include|\-I|\/I|\/imsvc|\-isystem|\/FI)\s*(.*)");
+		static Regex ArgumentRegex = new Regex(@"^(\-include|\-I|\/I|\/imsvc|\-isystem|\/FI|\/Fo)\s*(.*)");
 
 		/// <summary>
 		/// Set of filters for files to include in the database. Relative to the root directory, or to the project file.
@@ -288,24 +288,23 @@ namespace UnrealBuildTool
 				// The file that is going to compile
 				else if (!Line.StartsWith("-") && !Line.StartsWith("/"))
 				{
-					var OldPath = Line.Replace("\"", "");
-					var FileReference = new FileReference(OldPath);
-					Line = Line.Replace(OldPath, FileReference.FullName.Replace("\\", "/"));
+					Line = ConvertPath(Line, Line);
 				}
 				// Arguments
 				else
 				{
-					Match LineMatch = LineRegex.Match(Line);
-					if (LineMatch.Success)
+					int StrIndex = Line.IndexOf("..\\");
+					if (StrIndex != -1)
 					{
-						var OldPath = LineMatch.Groups[2].Value.Replace("\"", "");
-						if (OldPath[^1] == '\\' || OldPath[^1] == '/')
+						Line = ConvertPath(Line, Line.Substring(StrIndex));
+					}
+					else
+					{
+						Match LineMatch = ArgumentRegex.Match(Line);
+						if (LineMatch.Success)
 						{
-							OldPath = OldPath.Remove(OldPath.Length - 1, 1);
+							Line = ConvertPath(Line, LineMatch.Groups[2].Value);
 						}
-
-						var FileReference = new FileReference(OldPath);
-						Line = Line.Replace(OldPath, FileReference.FullName.Replace("\\", "/"));
 					}
 				}
 
@@ -316,5 +315,17 @@ namespace UnrealBuildTool
 
 			return NewFileItem.AbsolutePath;
 		}
+
+		static private string ConvertPath(string Line, string OldPath)
+		{
+			OldPath = OldPath.Replace("\"", "");
+			if (OldPath[^1] == '\\' || OldPath[^1] == '/')
+			{
+				OldPath = OldPath.Remove(OldPath.Length - 1, 1);
+			}
+
+			var FileReference = new FileReference(OldPath);
+			return Line.Replace(OldPath, FileReference.FullName.Replace("\\", "/"));
+		}
 	}
 }
