Removed "-Compiler=Clang" because we found we got better results from the clang tools when they did the msvc to clang conversion for windows than UBT.

diff --git a/Engine/Source/Programs/UnrealBuildTool/Modes/GenerateClangDatabase.cs b/Engine/Source/Programs/UnrealBuildTool/Modes/GenerateClangDatabase.cs
--- a/Engine/Source/Programs/UnrealBuildTool/Modes/GenerateClangDatabase.cs
+++ b/Engine/Source/Programs/UnrealBuildTool/Modes/GenerateClangDatabase.cs
@@ -139,7 +139,7 @@ namespace UnrealBuildTool
 				foreach (TargetDescriptor TargetDescriptor in TargetDescriptors)
 				{
 					// Disable PCHs and unity builds for the target
-					TargetDescriptor.AdditionalArguments = TargetDescriptor.AdditionalArguments.Append(new string[] { "-NoPCH", "-DisableUnity", "-Compiler=Clang" });
+					TargetDescriptor.AdditionalArguments = TargetDescriptor.AdditionalArguments.Append(new string[] { "-NoPCH", "-DisableUnity" });
 
 					// Create a makefile for the target
 					UEBuildTarget Target = UEBuildTarget.Create(TargetDescriptor, BuildConfiguration.bSkipRulesCompile, BuildConfiguration.bForceRulesCompile, BuildConfiguration.bUsePrecompiled, Logger);
