Adding and moving to the Unreal 5.2 EngineIncludeOrderVersion.

diff --git a/Engine/Source/Programs/UnrealBuildTool/Configuration/TargetRules.cs b/Engine/Source/Programs/UnrealBuildTool/Configuration/TargetRules.cs
--- a/Engine/Source/Programs/UnrealBuildTool/Configuration/TargetRules.cs
+++ b/Engine/Source/Programs/UnrealBuildTool/Configuration/TargetRules.cs
@@ -147,12 +147,17 @@ namespace UnrealBuildTool
 		/// </summary>
 		Unreal5_1,
 
+		/// <summary>
+		/// Include order used in Unreal 5.2
+		/// </summary>
+		Unreal5_2,
+
 		// *** When adding new entries here, be sure to update UEBuildModuleCPP.CurrentIncludeOrderDefine to ensure that the correct guard is used. ***
 
 		/// <summary>
 		/// Always use the latest version of include order.
 		/// </summary>
-		Latest = Unreal5_1,
+		Latest = Unreal5_2,
 
 		/// <summary>
 		/// Contains the oldest version of include order that the engine supports.
@@ -266,6 +271,7 @@ namespace UnrealBuildTool
 			return new List<string>()
 			{
 				"UE_ENABLE_INCLUDE_ORDER_DEPRECATED_IN_5_1",
+				"UE_ENABLE_INCLUDE_ORDER_DEPRECATED_IN_5_2",
 			};
 		}
 
@@ -279,6 +285,7 @@ namespace UnrealBuildTool
 			return new List<string>()
 			{
 				"UE_ENABLE_INCLUDE_ORDER_DEPRECATED_IN_5_1=" + (InVersion < EngineIncludeOrderVersion.Unreal5_1 ? "1" : "0"),
+				"UE_ENABLE_INCLUDE_ORDER_DEPRECATED_IN_5_2=" + (InVersion < EngineIncludeOrderVersion.Unreal5_2 ? "1" : "0"),
 			};
 		}
 	}
