- Remove unused orientation parameters from GeneratePList, etc functions

diff --git a/Engine/Source/Programs/AutomationTool/IOS/IOSPlatform.Automation.cs b/Engine/Source/Programs/AutomationTool/IOS/IOSPlatform.Automation.cs
--- a/Engine/Source/Programs/AutomationTool/IOS/IOSPlatform.Automation.cs
+++ b/Engine/Source/Programs/AutomationTool/IOS/IOSPlatform.Automation.cs
@@ -603,10 +603,10 @@ public class IOSPlatform : ApplePlatform
 		IOSExports.GetProvisioningData(InProject, bDistribution, out MobileProvision, out SigningCertificate, out TeamUUID, out bAutomaticSigning);
 	}
 
-	public virtual bool DeployGeneratePList(FileReference ProjectFile, UnrealTargetConfiguration Config, DirectoryReference ProjectDirectory, bool bIsUEGame, string GameName, bool bIsClient, string ProjectName, DirectoryReference InEngineDir, DirectoryReference AppDirectory, string InExecutablePath, out bool bSupportsPortrait, out bool bSupportsLandscape)
+	public virtual bool DeployGeneratePList(FileReference ProjectFile, UnrealTargetConfiguration Config, DirectoryReference ProjectDirectory, bool bIsUEGame, string GameName, bool bIsClient, string ProjectName, DirectoryReference InEngineDir, DirectoryReference AppDirectory, string InExecutablePath)
 	{
 		FileReference TargetReceiptFileName = GetTargetReceiptFileName(Config, InExecutablePath, InEngineDir, ProjectDirectory, bIsUEGame);
-		return IOSExports.GeneratePList(ProjectFile, Config, ProjectDirectory, bIsUEGame, GameName, bIsClient, ProjectName, InEngineDir, AppDirectory, TargetReceiptFileName, Log.Logger, out bSupportsPortrait, out bSupportsLandscape);
+		return IOSExports.GeneratePList(ProjectFile, Config, ProjectDirectory, bIsUEGame, GameName, bIsClient, ProjectName, InEngineDir, AppDirectory, TargetReceiptFileName, Log.Logger);
 	}
 
 	protected string MakeIPAFileName(UnrealTargetConfiguration TargetConfiguration, ProjectParams Params, DeploymentContext SC, bool bAllowDistroPrefix)
@@ -1350,8 +1350,6 @@ public class IOSPlatform : ApplePlatform
 					}
 
 					var TargetConfiguration = SC.StageTargetConfigurations[0];
-					bool bSupportsPortrait = false;
-					bool bSupportsLandscape = false;
 
 					DeployGeneratePList(
 							SC.RawProjectPath,
@@ -1362,9 +1360,7 @@ public class IOSPlatform : ApplePlatform
 							SC.IsCodeBasedProject ? false : Params.Client, // Code based projects will have Client in their executable name already
 							SC.ShortProjectName, DirectoryReference.Combine(SC.LocalRoot, "Engine"),
 							DirectoryReference.Combine((SC.IsCodeBasedProject ? SC.ProjectRoot : DirectoryReference.Combine(SC.LocalRoot, "Engine")), "Binaries", PlatformName, "Payload", (SC.IsCodeBasedProject ? SC.ShortProjectName : "UnrealGame") + ".app"),
-							SC.StageExecutables[0],
-							out bSupportsPortrait,
-							out bSupportsLandscape);
+							SC.StageExecutables[0]);
 
 					// copy the plist to the stage dir
 					SC.StageFile(StagedFileType.SystemNonUFS, TargetPListFile, new StagedFileReference("Info.plist"));
diff --git a/Engine/Source/Programs/AutomationTool/TVOS/TVOSPlatform.Automation.cs b/Engine/Source/Programs/AutomationTool/TVOS/TVOSPlatform.Automation.cs
--- a/Engine/Source/Programs/AutomationTool/TVOS/TVOSPlatform.Automation.cs
+++ b/Engine/Source/Programs/AutomationTool/TVOS/TVOSPlatform.Automation.cs
@@ -37,7 +37,7 @@ public class TVOSPlatform : IOSPlatform
 		TVOSExports.GetProvisioningData(InProject, bDistribution, out MobileProvision, out SigningCertificate, out Team, out bAutomaticSigning);
     }
 
-	public override bool DeployGeneratePList(FileReference ProjectFile, UnrealTargetConfiguration Config, DirectoryReference ProjectDirectory, bool bIsUEGame, string GameName, bool bIsClient, string ProjectName, DirectoryReference InEngineDir, DirectoryReference AppDirectory, string InExecutablePath, out bool bSupportsPortrait, out bool bSupportsLandscape)
+	public override bool DeployGeneratePList(FileReference ProjectFile, UnrealTargetConfiguration Config, DirectoryReference ProjectDirectory, bool bIsUEGame, string GameName, bool bIsClient, string ProjectName, DirectoryReference InEngineDir, DirectoryReference AppDirectory, string InExecutablePath)
 	{
 		string TargetName = Path.GetFileNameWithoutExtension(InExecutablePath).Split("-".ToCharArray())[0];
 		FileReference TargetReceiptFileName;
@@ -49,7 +49,7 @@ public class TVOSPlatform : IOSPlatform
 		{
 			TargetReceiptFileName = TargetReceipt.GetDefaultPath(ProjectDirectory, TargetName, UnrealTargetPlatform.TVOS, Config, "");
 		}
-		return TVOSExports.GeneratePList(ProjectFile, Config, ProjectDirectory, bIsUEGame, GameName, bIsClient, ProjectName, InEngineDir, AppDirectory, TargetReceiptFileName, Log.Logger, out bSupportsPortrait, out bSupportsLandscape);
+		return TVOSExports.GeneratePList(ProjectFile, Config, ProjectDirectory, bIsUEGame, GameName, bIsClient, ProjectName, InEngineDir, AppDirectory, TargetReceiptFileName, Log.Logger);
 	}
 
     public override string GetCookPlatform(bool bDedicatedServer, bool bIsClientOnly)
@@ -111,8 +111,6 @@ public class TVOSPlatform : IOSPlatform
 
                     var TargetConfiguration = SC.StageTargetConfigurations[0];
 
-                    bool bSupportsPortrait = false;
-                    bool bSupportsLandscape = false;
                     DeployGeneratePList(
 						SC.RawProjectPath,
 						TargetConfiguration,
@@ -127,9 +125,7 @@ public class TVOSPlatform : IOSPlatform
 						"TVOS",
 						"Payload",
 						(SC.IsCodeBasedProject ? SC.ShortProjectName : "UnrealGame") + ".app"),
-						SC.StageExecutables[0],
-						out bSupportsPortrait,
-						out bSupportsLandscape);
+						SC.StageExecutables[0]);
                 }
 
 
diff --git a/Engine/Source/Programs/UnrealBuildTool/Platform/IOS/IOSExports.cs b/Engine/Source/Programs/UnrealBuildTool/Platform/IOS/IOSExports.cs
--- a/Engine/Source/Programs/UnrealBuildTool/Platform/IOS/IOSExports.cs
+++ b/Engine/Source/Programs/UnrealBuildTool/Platform/IOS/IOSExports.cs
@@ -102,13 +102,11 @@ namespace UnrealBuildTool
 		/// <param name="AppDirectory"></param>
 		/// <param name="BuildReceiptFileName"></param>
 		/// <param name="Logger">Logger for output</param>
-		/// <param name="bSupportsPortrait"></param>
-		/// <param name="bSupportsLandscape"></param>
 		/// <returns></returns>
-		public static bool GeneratePList(FileReference ProjectFile, UnrealTargetConfiguration Config, DirectoryReference ProjectDirectory, bool bIsUnrealGame, string GameName, bool bIsClient, string ProjectName, DirectoryReference InEngineDir, DirectoryReference AppDirectory, FileReference BuildReceiptFileName, ILogger Logger, out bool bSupportsPortrait, out bool bSupportsLandscape)
+		public static bool GeneratePList(FileReference ProjectFile, UnrealTargetConfiguration Config, DirectoryReference ProjectDirectory, bool bIsUnrealGame, string GameName, bool bIsClient, string ProjectName, DirectoryReference InEngineDir, DirectoryReference AppDirectory, FileReference BuildReceiptFileName, ILogger Logger)
 		{
 			TargetReceipt Receipt = TargetReceipt.Read(BuildReceiptFileName);
-			return new UEDeployIOS(Logger).GeneratePList(ProjectFile, Config, ProjectDirectory.FullName, bIsUnrealGame, GameName, bIsClient, ProjectName, InEngineDir.FullName, AppDirectory.FullName, Receipt, out bSupportsPortrait, out bSupportsLandscape);
+			return new UEDeployIOS(Logger).GeneratePList(ProjectFile, Config, ProjectDirectory.FullName, bIsUnrealGame, GameName, bIsClient, ProjectName, InEngineDir.FullName, AppDirectory.FullName, Receipt);
 		}
 
 		/// <summary>
diff --git a/Engine/Source/Programs/UnrealBuildTool/Platform/IOS/UEDeployIOS.cs b/Engine/Source/Programs/UnrealBuildTool/Platform/IOS/UEDeployIOS.cs
--- a/Engine/Source/Programs/UnrealBuildTool/Platform/IOS/UEDeployIOS.cs
+++ b/Engine/Source/Programs/UnrealBuildTool/Platform/IOS/UEDeployIOS.cs
@@ -1,4 +1,4 @@
-// Copyright Epic Games, Inc. All Rights Reserved.
+﻿// Copyright Epic Games, Inc. All Rights Reserved.
 
 using System;
 using System.Collections;
@@ -217,7 +217,7 @@ namespace UnrealBuildTool
 			return MinVersionToReturn;
 		}
 
-		public static bool GenerateIOSPList(FileReference? ProjectFile, UnrealTargetConfiguration Config, string ProjectDirectory, bool bIsUnrealGame, string GameName, bool bIsClient, string ProjectName, string InEngineDir, string AppDirectory, UnrealPluginLanguage? UPL, string? BundleID, bool bBuildAsFramework, ILogger Logger, out bool bSupportsPortrait, out bool bSupportsLandscape)
+		public static bool GenerateIOSPList(FileReference? ProjectFile, UnrealTargetConfiguration Config, string ProjectDirectory, bool bIsUnrealGame, string GameName, bool bIsClient, string ProjectName, string InEngineDir, string AppDirectory, UnrealPluginLanguage? UPL, string? BundleID, bool bBuildAsFramework, ILogger Logger)
 		{
 			// generate the Info.plist for future use
 			string BuildDirectory = ProjectDirectory + "/Build/IOS";
@@ -244,7 +244,7 @@ namespace UnrealBuildTool
 			bool bSupported = true;
 			Ini.GetBool("/Script/IOSRuntimeSettings.IOSRuntimeSettings", "bSupportsPortraitOrientation", out bSupported);
 			SupportedOrientations += bSupported ? "\t\t<string>UIInterfaceOrientationPortrait</string>\n" : "";
-			bSupportsPortrait = bSupported;
+			bool bSupportsPortrait = bSupported;
 
 			Ini.GetBool("/Script/IOSRuntimeSettings.IOSRuntimeSettings", "bSupportsUpsideDownOrientation", out bSupported);
 			SupportedOrientations += bSupported ? "\t\t<string>UIInterfaceOrientationPortraitUpsideDown</string>\n" : "";
@@ -254,7 +254,7 @@ namespace UnrealBuildTool
 			Ini.GetBool("/Script/IOSRuntimeSettings.IOSRuntimeSettings", "bSupportsLandscapeLeftOrientation", out bSupportsLandscapeLeft);
 			bool bSupportsLandscapeRight = false;
 			Ini.GetBool("/Script/IOSRuntimeSettings.IOSRuntimeSettings", "bSupportsLandscapeRightOrientation", out bSupportsLandscapeRight);
-			bSupportsLandscape = bSupportsLandscapeLeft || bSupportsLandscapeRight;
+			bool bSupportsLandscape = bSupportsLandscapeLeft || bSupportsLandscapeRight;
 
 			if (bSupportsLandscapeLeft && bSupportsLandscapeRight)
 			{
@@ -275,8 +275,8 @@ namespace UnrealBuildTool
 			else
 			{
 				// max one landscape orientation is supported
-				SupportedOrientations += bSupportsLandscapeRight? "\t\t<string>UIInterfaceOrientationLandscapeRight</string>\n": "";
-				SupportedOrientations += bSupportsLandscapeLeft? "\t\t<string>UIInterfaceOrientationLandscapeLeft</string>\n" : "";
+				SupportedOrientations += bSupportsLandscapeRight ? "\t\t<string>UIInterfaceOrientationLandscapeRight</string>\n" : "";
+				SupportedOrientations += bSupportsLandscapeLeft ? "\t\t<string>UIInterfaceOrientationLandscapeLeft</string>\n" : "";
 			}
 
 			// ITunes file sharing
@@ -362,9 +362,9 @@ namespace UnrealBuildTool
 			bool bRemoteNotificationsSupported = false;
 			Ini.GetBool("/Script/IOSRuntimeSettings.IOSRuntimeSettings", "bEnableRemoteNotificationsSupport", out bRemoteNotificationsSupported);
 
-            // Add background fetch as background mode
-            bool bBackgroundFetch = false;
-            Ini.GetBool("/Script/IOSRuntimeSettings.IOSRuntimeSettings", "bEnableBackgroundFetch", out bBackgroundFetch);
+			// Add background fetch as background mode
+			bool bBackgroundFetch = false;
+			Ini.GetBool("/Script/IOSRuntimeSettings.IOSRuntimeSettings", "bEnableBackgroundFetch", out bBackgroundFetch);
 
 			// Get any Location Services permission descriptions added
 			string LocationAlwaysUsageDescription = "";
@@ -552,12 +552,12 @@ namespace UnrealBuildTool
 				Text.AppendLine(string.Format("\t<string>{0}</string>", FacebookAppID));
 				Text.AppendLine("\t<key>FacebookDisplayName</key>");
 				Text.AppendLine(string.Format("\t<string>{0}</string>", FacebookDisplayName));
-				
+
 				if (!bEnableAutomaticLogging)
 				{
 					Text.AppendLine("<key>FacebookAutoLogAppEventsEnabled</key><false/>");
 				}
-				
+
 				if (!bEnableAdvertisingId)
 				{
 					Text.AppendLine("<key>FacebookAdvertiserIDCollectionEnabled</key><false/>");
@@ -590,14 +590,14 @@ namespace UnrealBuildTool
 			{
 				Text.AppendLine("\t<key>UIBackgroundModes</key>");
 				Text.AppendLine("\t<array>");
-                if (bRemoteNotificationsSupported)
-                {
-				    Text.AppendLine("\t\t<string>remote-notification</string>");
-                }
-                if (bBackgroundFetch)
-                {
-                    Text.AppendLine("\t\t<string>fetch</string>");
-                }
+				if (bRemoteNotificationsSupported)
+				{
+					Text.AppendLine("\t\t<string>remote-notification</string>");
+				}
+				if (bBackgroundFetch)
+				{
+					Text.AppendLine("\t\t<string>fetch</string>");
+				}
 				Text.AppendLine("\t</array>");
 			}
 
@@ -686,15 +686,15 @@ namespace UnrealBuildTool
 			return false;
 		}
 
- 		public bool GeneratePList(FileReference ProjectFile, UnrealTargetConfiguration Config, string ProjectDirectory, bool bIsUnrealGame, string GameName, bool bIsClient, string ProjectName, string InEngineDir, string AppDirectory, TargetReceipt Receipt, out bool bSupportsPortrait, out bool bSupportsLandscape)
+ 		public bool GeneratePList(FileReference ProjectFile, UnrealTargetConfiguration Config, string ProjectDirectory, bool bIsUnrealGame, string GameName, bool bIsClient, string ProjectName, string InEngineDir, string AppDirectory, TargetReceipt Receipt)
 		{
 			List<string> UPLScripts = CollectPluginDataPaths(Receipt.AdditionalProperties, Logger);
 			VersionNumber? SdkVersion = GetSdkVersion(Receipt);
 			bool bBuildAsFramework = GetCompileAsDll(Receipt);
-			return GeneratePList(ProjectFile, Config, ProjectDirectory, bIsUnrealGame, GameName, bIsClient, ProjectName, InEngineDir, AppDirectory, UPLScripts, "", bBuildAsFramework, out bSupportsPortrait, out bSupportsLandscape);
+			return GeneratePList(ProjectFile, Config, ProjectDirectory, bIsUnrealGame, GameName, bIsClient, ProjectName, InEngineDir, AppDirectory, UPLScripts, "", bBuildAsFramework);
 		}
 
-		public virtual bool GeneratePList(FileReference? ProjectFile, UnrealTargetConfiguration Config, string ProjectDirectory, bool bIsUnrealGame, string GameName, bool bIsClient, string ProjectName, string InEngineDir, string AppDirectory, List<string> UPLScripts, string? BundleID, bool bBuildAsFramework, out bool bSupportsPortrait, out bool bSupportsLandscape)
+		public virtual bool GeneratePList(FileReference? ProjectFile, UnrealTargetConfiguration Config, string ProjectDirectory, bool bIsUnrealGame, string GameName, bool bIsClient, string ProjectName, string InEngineDir, string AppDirectory, List<string> UPLScripts, string? BundleID, bool bBuildAsFramework)
 		{
 			// remember name with -IOS-Shipping, etc
 			// string ExeName = GameName;
@@ -726,7 +726,7 @@ namespace UnrealBuildTool
 			// Passing in true for distribution is not ideal here but given the way that ios packaging happens and this call chain it seems unavoidable for now, maybe there is a way to correctly pass it in that I can't find?
 			UPL.Init(ProjectArches, true, RelativeEnginePath, BundlePath, ProjectDirectory, Config.ToString(), false);
 
-			return GenerateIOSPList(ProjectFile, Config, ProjectDirectory, bIsUnrealGame, GameName, bIsClient, ProjectName, InEngineDir, AppDirectory, UPL, BundleID, bBuildAsFramework, Logger, out bSupportsPortrait, out bSupportsLandscape);
+			return GenerateIOSPList(ProjectFile, Config, ProjectDirectory, bIsUnrealGame, GameName, bIsClient, ProjectName, InEngineDir, AppDirectory, UPL, BundleID, bBuildAsFramework, Logger);
 		}
 
 		protected virtual void CopyCloudResources(string InEngineDir, string AppDirectory)
@@ -978,9 +978,7 @@ namespace UnrealBuildTool
 				DestFileInfo.Attributes = DestFileInfo.Attributes & ~FileAttributes.ReadOnly;
 			}
 
-			bool bSupportsPortrait = true;
-			bool bSupportsLandscape = false;
-			GeneratePList(ProjectFile, Config, InProjectDirectory, bIsUnrealGame, GameExeName, false, InProjectName, InEngineDir, AppDirectory, UPLScripts, BundleID, bBuildAsFramework, out bSupportsPortrait, out bSupportsLandscape);
+			GeneratePList(ProjectFile, Config, InProjectDirectory, bIsUnrealGame, GameExeName, false, InProjectName, InEngineDir, AppDirectory, UPLScripts, BundleID, bBuildAsFramework);
 
 			// ensure the destination is writable
 			if (File.Exists(AppDirectory + "/" + GameName))
@@ -1065,8 +1063,7 @@ namespace UnrealBuildTool
 			else
 			{
 				// @todo tvos merge: This used to copy the bundle back - where did that code go? It needs to be fixed up for TVOS directories
-				bool bSupportPortrait, bSupportLandscape;
-				GeneratePList(ProjectFile, Configuration, ProjectDirectory, bIsUnrealGame, GameName, false, (ProjectFile == null) ? "" : Path.GetFileNameWithoutExtension(ProjectFile.FullName), "../../Engine", "", UPLScripts, BundleID, bBuildAsFramework, out bSupportPortrait, out bSupportLandscape);
+				GeneratePList(ProjectFile, Configuration, ProjectDirectory, bIsUnrealGame, GameName, false, (ProjectFile == null) ? "" : Path.GetFileNameWithoutExtension(ProjectFile.FullName), "../../Engine", "", UPLScripts, BundleID, bBuildAsFramework);
 			}
 			return true;
 		}
diff --git a/Engine/Source/Programs/UnrealBuildTool/Platform/TVOS/TVOSExports.cs b/Engine/Source/Programs/UnrealBuildTool/Platform/TVOS/TVOSExports.cs
--- a/Engine/Source/Programs/UnrealBuildTool/Platform/TVOS/TVOSExports.cs
+++ b/Engine/Source/Programs/UnrealBuildTool/Platform/TVOS/TVOSExports.cs
@@ -98,13 +98,11 @@ namespace UnrealBuildTool
 		/// <param name="AppDirectory"></param>
 		/// <param name="BuildReceiptFileName"></param>
 		/// <param name="Logger"></param>
-		/// <param name="bSupportsPortrait"></param>
-		/// <param name="bSupportsLandscape"></param>
 		/// <returns></returns>
-		public static bool GeneratePList(FileReference ProjectFile, UnrealTargetConfiguration Config, DirectoryReference ProjectDirectory, bool bIsUnrealGame, string GameName, bool bIsClient, string ProjectName, DirectoryReference InEngineDir, DirectoryReference AppDirectory, FileReference BuildReceiptFileName, ILogger Logger, out bool bSupportsPortrait, out bool bSupportsLandscape)
+		public static bool GeneratePList(FileReference ProjectFile, UnrealTargetConfiguration Config, DirectoryReference ProjectDirectory, bool bIsUnrealGame, string GameName, bool bIsClient, string ProjectName, DirectoryReference InEngineDir, DirectoryReference AppDirectory, FileReference BuildReceiptFileName, ILogger Logger)
 		{
 			TargetReceipt Receipt = TargetReceipt.Read(BuildReceiptFileName);
-			return new UEDeployTVOS(Logger).GeneratePList(ProjectFile, Config, ProjectDirectory.FullName, bIsUnrealGame, GameName, bIsClient, ProjectName, InEngineDir.FullName, AppDirectory.FullName, Receipt, out bSupportsPortrait, out bSupportsLandscape);
+			return new UEDeployTVOS(Logger).GeneratePList(ProjectFile, Config, ProjectDirectory.FullName, bIsUnrealGame, GameName, bIsClient, ProjectName, InEngineDir.FullName, AppDirectory.FullName, Receipt);
 		}
 	}
 }
diff --git a/Engine/Source/Programs/UnrealBuildTool/Platform/TVOS/UEDeployTVOS.cs b/Engine/Source/Programs/UnrealBuildTool/Platform/TVOS/UEDeployTVOS.cs
--- a/Engine/Source/Programs/UnrealBuildTool/Platform/TVOS/UEDeployTVOS.cs
+++ b/Engine/Source/Programs/UnrealBuildTool/Platform/TVOS/UEDeployTVOS.cs
@@ -203,9 +203,8 @@ namespace UnrealBuildTool
 			return bSkipDefaultPNGs;
 		}
 
-		public override bool GeneratePList(FileReference? ProjectFile, UnrealTargetConfiguration Config, string ProjectDirectory, bool bIsUnrealGame, string GameName, bool bIsClient, string ProjectName, string InEngineDir, string AppDirectory, List<string> UPLScripts, string? BundleID, bool bBuildAsFramework, out bool bSupportsPortrait, out bool bSupportsLandscape)
+		public override bool GeneratePList(FileReference? ProjectFile, UnrealTargetConfiguration Config, string ProjectDirectory, bool bIsUnrealGame, string GameName, bool bIsClient, string ProjectName, string InEngineDir, string AppDirectory, List<string> UPLScripts, string? BundleID, bool bBuildAsFramework)
 		{
-			bSupportsLandscape = bSupportsPortrait = true;
 			return GenerateTVOSPList(ProjectDirectory, bIsUnrealGame, GameName, bIsClient, ProjectName, InEngineDir, AppDirectory, null, BundleID, Logger);
 		}
 	}
diff --git a/Engine/Source/Programs/UnrealBuildTool/ProjectFiles/Xcode/XcodeProjectLegacy.cs b/Engine/Source/Programs/UnrealBuildTool/ProjectFiles/Xcode/XcodeProjectLegacy.cs
--- a/Engine/Source/Programs/UnrealBuildTool/ProjectFiles/Xcode/XcodeProjectLegacy.cs
+++ b/Engine/Source/Programs/UnrealBuildTool/ProjectFiles/Xcode/XcodeProjectLegacy.cs
@@ -1320,11 +1320,10 @@ namespace UnrealBuildTool.XcodeProjectLegacy
 								ReceiptFilename = TargetReceipt.GetDefaultPath(ProjectPath, GameName, UnrealTargetPlatform.IOS, Config.BuildConfig, "");
 							}
 							Directory.CreateDirectory(Path.GetDirectoryName(IOSInfoPlistPath)!);
-							bool bSupportPortrait, bSupportLandscape;
 							TargetReceipt? Receipt;
 							TargetReceipt.TryRead(ReceiptFilename, out Receipt);
 							bool bBuildAsFramework = UEDeployIOS.GetCompileAsDll(Receipt);
-							UEDeployIOS.GenerateIOSPList(ProjectFile, Config.BuildConfig, ProjectPath.FullName, bIsUnrealGame, GameName, bIsClient, Config.BuildTarget, EngineDir.FullName, ProjectPath + "/Binaries/IOS/Payload", null, BundleIdentifier, bBuildAsFramework, Logger, out bSupportPortrait, out bSupportLandscape);
+							UEDeployIOS.GenerateIOSPList(ProjectFile, Config.BuildConfig, ProjectPath.FullName, bIsUnrealGame, GameName, bIsClient, Config.BuildTarget, EngineDir.FullName, ProjectPath + "/Binaries/IOS/Payload", null, BundleIdentifier, bBuildAsFramework, Logger);
 						}
 						if (bCreateTVOSInfoPlist)
 						{
