Updated the output to allow you to quickly jump to the file in VS's output window.
Updated the logic to support modules that don't use private/public folders.

diff --git a/Engine/Source/Programs/UnrealBuildTool/Modes/FixIncludePathsMode.cs b/Engine/Source/Programs/UnrealBuildTool/Modes/FixIncludePathsMode.cs
--- a/Engine/Source/Programs/UnrealBuildTool/Modes/FixIncludePathsMode.cs
+++ b/Engine/Source/Programs/UnrealBuildTool/Modes/FixIncludePathsMode.cs
@@ -126,6 +126,7 @@ namespace UnrealBuildTool
 									for (int i = 0; i < Text.Length; i++)
 									{
 										var Line = Text[i];
+										int LineNumber = i + 1;
 										Match IncludeMatch = IncludeRegex.Match(Line);
 										if (IncludeMatch.Success)
 										{
@@ -133,7 +134,7 @@ namespace UnrealBuildTool
 
 											if (Include.Contains("/Private/") && PublicDirectories.Any(dir => InputFile.FullName.Contains(System.IO.Path.DirectorySeparatorChar + dir + System.IO.Path.DirectorySeparatorChar)))
 											{
-												Logger.LogError("Can not update #include '{Include}' in the public header '{FileName}' because it will break code including this header.", Include, InputFile.FullName);
+												Logger.LogError("{FileName}({LineNumber}): Can not update #include '{Include}' in the public file because it may break external code that uses it.", InputFile.FullName, LineNumber, Include);
 												continue;
 											}
 
@@ -143,7 +144,7 @@ namespace UnrealBuildTool
 												List<DirectoryReference> IncludePaths = new();
 												IncludePaths.Add(new DirectoryReference(System.IO.Directory.GetParent(InputFile.FullName)!));
 												IncludePaths.AddRange(env.UserIncludePaths);
-												//EnvPaths.AddRange(env.SystemIncludePaths);
+												IncludePaths.AddRange(env.SystemIncludePaths);
 
 												// search include paths
 												string? FullPath = null;
@@ -163,7 +164,7 @@ namespace UnrealBuildTool
 												{
 													if (FullPath.Contains("ThirdParty"))
 													{
-														Logger.LogDebug("Skipping '{Include}' because it is a third party header", Include);
+														Logger.LogDebug("{FileName}({LineNumber}): Skipping '{Include}' because it is a third party header.", InputFile.FullName, LineNumber, Include);
 														PreferredInclude = Include;
 														PreferredPathCache[Include] = PreferredInclude;
 														continue;
@@ -173,25 +174,54 @@ namespace UnrealBuildTool
 													if (string.Equals(System.IO.Directory.GetParent(FullPath)?.FullName, System.IO.Directory.GetParent(InputFile.FullName)?.FullName, StringComparison.CurrentCultureIgnoreCase) &&
 														string.Equals(Include, System.IO.Path.GetFileName(FullPath), StringComparison.CurrentCultureIgnoreCase))
 													{
+														Logger.LogDebug("{FileName}({LineNumber}): Using '{Include}' because it is in the same directory.", InputFile.FullName, LineNumber, Include);
 														PreferredInclude = Include;
 													}
 													else
 													{
-														string? FoundPreferredPath = PreferredPaths.FirstOrDefault(path => FullPath.Contains(path));
-														if (!FullPath.Contains(UnrealRootDirectory) || FoundPreferredPath == null)
+														if (!FullPath.Contains(UnrealRootDirectory))
 														{
-															PreferredInclude = null;
+															Logger.LogDebug("{FileName}({LineNumber}): Skipping '{Include}' because it isn't under the Unreal root directory.", InputFile.FullName, LineNumber, Include);
 														}
 														else
 														{
-															int end = FullPath.LastIndexOf(FoundPreferredPath) + FoundPreferredPath.Length;
-															PreferredInclude = FullPath.Substring(end);
-
-															// Is the current include a shortened version of the preferred include path?
-															if (PreferredInclude != Include && PreferredInclude.Contains(Include))
+															string? FoundPreferredPath = PreferredPaths.FirstOrDefault(path => FullPath.Contains(path));
+															if (FoundPreferredPath != null)
+															{
+																int end = FullPath.LastIndexOf(FoundPreferredPath) + FoundPreferredPath.Length;
+																PreferredInclude = FullPath.Substring(end);
+
+																// Is the current include a shortened version of the preferred include path?
+																if (PreferredInclude != Include && PreferredInclude.Contains(Include))
+																{
+																	Logger.LogDebug("{FileName}({LineNumber}): Using '{Include}' because it is shorter than '{PreferredInclude}'.", InputFile.FullName, LineNumber, Include, PreferredInclude);
+																	PreferredInclude = Include;
+																}
+															}
+															else
 															{
-																Logger.LogDebug("Using '{Include}' because it is shorter than '{PreferredInclude}'", Include, PreferredInclude);
-																PreferredInclude = Include;
+																PreferredInclude = null;
+
+																string ModulePath = FullPath;
+																FileReference IncludeFileReference = FileReference.FromString(FullPath);
+																DirectoryReference? TempDirectory = IncludeFileReference.Directory;
+																DirectoryReference? FoundDirectory = null;
+																// find the module this include is part of
+																while (TempDirectory != null)
+																{
+																	if (DirectoryReference.EnumerateFiles(TempDirectory, $"*.build.cs").Any())
+																	{
+																		FoundDirectory = TempDirectory;
+																		break;
+																	}
+
+																	TempDirectory = TempDirectory.ParentDirectory;
+																}
+
+																if (FoundDirectory != null)
+																{
+																	PreferredInclude = FullPath.Substring(FoundDirectory.FullName.Length + 1);
+																}
 															}
 														}
 
@@ -201,13 +231,13 @@ namespace UnrealBuildTool
 												
 												if (PreferredInclude == null)
 												{
-													Logger.LogWarning("Could not find include path for '{IncludePath}' found in '{FileName}'", Include, InputFile.FullName);
+													Logger.LogDebug("{FileName}({LineNumber}): Could not find path to '{IncludePath}'", InputFile.FullName, LineNumber, Include);
 												}
 											}
 
 											if (PreferredInclude != null && Include != PreferredInclude)
 											{
-												Logger.LogInformation("Updated '{InputFileName}' line {LineNum} -- {OldInclude} -> {NewInclude}", InputFile.FullName, i, Include, PreferredInclude);
+												Logger.LogInformation("{FileName}({LineNumber}): Updated '{OldInclude}' -> '{NewInclude}'", InputFile.FullName, LineNumber, Include, PreferredInclude);
 												Text[i] = Line.Replace(Include, PreferredInclude);
 												UpdatedText = true;
 												LinesUpdated.Add(i);
