[UnrealBuildTool]
* Added support for forcing module to use inlined .gen.cpp files

diff --git a/Engine/Source/Programs/UnrealBuildTool/Configuration/ModuleRules.cs b/Engine/Source/Programs/UnrealBuildTool/Configuration/ModuleRules.cs
--- a/Engine/Source/Programs/UnrealBuildTool/Configuration/ModuleRules.cs
+++ b/Engine/Source/Programs/UnrealBuildTool/Configuration/ModuleRules.cs
@@ -889,6 +889,11 @@ namespace UnrealBuildTool
 		/// </summary>
 		public bool bBuildLocallyWithSNDBS = false;
 
+		/// <summary>
+		/// Enable warnings for when there are .gen.cpp files that could be inlined in a matching handwritten cpp file
+		/// </summary>
+		public bool bEnableNonInlinedGenCppWarnings = false;
+
 		/// <summary>
 		/// Redistribution override flag for this module.
 		/// </summary>
diff --git a/Engine/Source/Programs/UnrealBuildTool/Configuration/UEBuildModuleCPP.cs b/Engine/Source/Programs/UnrealBuildTool/Configuration/UEBuildModuleCPP.cs
--- a/Engine/Source/Programs/UnrealBuildTool/Configuration/UEBuildModuleCPP.cs
+++ b/Engine/Source/Programs/UnrealBuildTool/Configuration/UEBuildModuleCPP.cs
@@ -543,14 +543,22 @@ namespace UnrealBuildTool
 			// Compile all the generated CPP files
 			if (GeneratedCppDirectories != null && !CompileEnvironment.bHackHeaderGenerator)
 			{
-				List<string> GeneratedFiles = new List<string>();
+				var GeneratedFiles = new Dictionary<string, string>();
 				if (SpecificFilesToCompile.Count == 0)
 				{
 					foreach (string GeneratedDir in GeneratedCppDirectories)
 					{
 						if (Directory.Exists(GeneratedDir))
 						{
-							GeneratedFiles.AddRange(Directory.GetFiles(GeneratedDir, "*.gen.cpp"));
+							string[] Files = Directory.GetFiles(GeneratedDir, "*.gen.cpp");
+							GeneratedFiles.EnsureCapacity(GeneratedFiles.Count + Files.Length);
+							foreach (var File in Files)
+							{
+								// Can't use GetFileNameWithoutAnyExtensions because of the .init.gen.cpp files
+								string FileName = Path.GetFileName(File);
+								FileName = FileName.Substring(0, FileName.Length - ".gen.cpp".Length); 
+								GeneratedFiles.Add(FileName, File);
+							}
 						}
 					}
 				}
@@ -560,7 +568,7 @@ namespace UnrealBuildTool
 					{
 						if (GeneratedCppDirectories.Any(x => FileToCompile.IsUnderDirectory(new DirectoryReference(x))))
 						{
-							GeneratedFiles.Add(FileToCompile.FullName);
+							GeneratedFiles.Add(FileToCompile.GetFileNameWithoutAnyExtensions(), FileToCompile.FullName);
 						}
 					}
 				}
@@ -575,12 +583,8 @@ namespace UnrealBuildTool
 							var ListOfInlinedGenCpps = ModuleCompileEnvironment.MetadataCache.GetListOfInlinedGeneratedCppFiles(CPPFileItem);
 							foreach (string ListOfInlinedGenCppsItem in ListOfInlinedGenCpps)
 							{
-								string GeneratedCppName = ListOfInlinedGenCppsItem + ".gen.cpp";
-								string? FoundGenCppFile = GeneratedFiles.FirstOrDefault((file) => String.Compare(Path.GetFileName(file), GeneratedCppName, StringComparison.InvariantCultureIgnoreCase) == 0);
-								if (FoundGenCppFile != null)
+								if (GeneratedFiles.Remove(ListOfInlinedGenCppsItem, out string? FoundGenCppFile))
 								{
-									GeneratedFiles.Remove(FoundGenCppFile);
-
 									if (!CompileEnvironment.FileInlineGenCPPMap.ContainsKey(CPPFileItem))
 									{
 										CompileEnvironment.FileInlineGenCPPMap[CPPFileItem] = new List<FileItem>();
@@ -589,12 +593,28 @@ namespace UnrealBuildTool
 								}
 								else
 								{
-									Logger.LogError("'{CPPFileItem}' is looking for a generated cpp with named '{HeaderFile}'", CPPFileItem.AbsolutePath, GeneratedCppName);
+									Logger.LogError("'{CPPFileItem}' is looking for a generated cpp with named '{HeaderFile}.gen.cpp'", CPPFileItem.AbsolutePath, ListOfInlinedGenCppsItem);
 								}
 							}
 						}
 					}
 
+					if (Rules.bEnableNonInlinedGenCppWarnings)
+					{
+						var CPPFilesLookup = new Dictionary<string, FileItem>();
+						foreach (var CPPFile in CPPFiles)
+						{
+							CPPFilesLookup.Add(Utils.GetFilenameWithoutAnyExtensions(CPPFile.Name), CPPFile);
+						}
+						foreach (var Name in GeneratedFiles.Keys)
+						{
+							if (CPPFilesLookup.TryGetValue(Name, out FileItem? Item))
+							{
+								Logger.LogWarning("'{0}' .gen.cpp not inlined. Add '#include UE_INLINE_GENERATED_CPP_BY_NAME({1})'", Item.Name, Name);
+							}
+						}
+					}
+
 					bool bMergeUnityFiles = Target.bMergeModuleAndGeneratedUnityFiles && Rules.bMergeUnityFiles;
 
 					// Create a compile environment for the generated files. We can disable creating debug info here to improve link times.
@@ -622,7 +642,7 @@ namespace UnrealBuildTool
 
 					// Compile all the generated files
 					List<FileItem> GeneratedFileItems = new List<FileItem>();
-					foreach (string GeneratedFilename in GeneratedFiles)
+					foreach (string GeneratedFilename in GeneratedFiles.Values)
 					{
 						FileItem GeneratedCppFileItem = FileItem.GetItemByPath(GeneratedFilename);
 						if (SpecificFilesToCompile.Count == 0 || SpecificFilesToCompile.Contains(GeneratedCppFileItem.Location))
