Added a bool setting to allow us to turn off the forcing of PCH files for gen.cpp files when PCH is disabled. This improves NoPCH/DisableUnity compile times.

diff --git a/Engine/Source/Programs/UnrealBuildTool/Configuration/TargetRules.cs b/Engine/Source/Programs/UnrealBuildTool/Configuration/TargetRules.cs
--- a/Engine/Source/Programs/UnrealBuildTool/Configuration/TargetRules.cs
+++ b/Engine/Source/Programs/UnrealBuildTool/Configuration/TargetRules.cs
@@ -1518,6 +1518,12 @@ namespace UnrealBuildTool
 		[XmlConfigFile(Category = "BuildConfiguration")]
 		public bool bForceDeterministic = false;
 
+		/// <summary>
+		/// Whether PCH headers should be force included for gen.cpp files when PCH is disabled.
+		/// </summary>
+		[XmlConfigFile(Category = "BuildConfiguration")]
+		public bool bForceIncludePCHHeadersForGenCppFilesWhenPCHIsDisabled = true;
+
 		/// <summary>
 		/// Whether to just preprocess source files for this target, and skip compilation
 		/// </summary>
@@ -3219,6 +3225,11 @@ namespace UnrealBuildTool
 			get { return Inner.bForceDeterministic; }
 		}
 
+		public bool bForceIncludePCHHeadersForGenCppFilesWhenPCHIsDisabled
+		{
+			get { return Inner.bForceIncludePCHHeadersForGenCppFilesWhenPCHIsDisabled; }
+		}
+
 		public bool bPreprocessOnly
 		{
 			get { return Inner.bPreprocessOnly; }
diff --git a/Engine/Source/Programs/UnrealBuildTool/Configuration/UEBuildModuleCPP.cs b/Engine/Source/Programs/UnrealBuildTool/Configuration/UEBuildModuleCPP.cs
--- a/Engine/Source/Programs/UnrealBuildTool/Configuration/UEBuildModuleCPP.cs
+++ b/Engine/Source/Programs/UnrealBuildTool/Configuration/UEBuildModuleCPP.cs
@@ -627,7 +627,7 @@ namespace UnrealBuildTool
 					}
 
 					// Always force include the PCH, even if PCHs are disabled, for generated code. Legacy code can rely on PCHs being included to compile correctly, and this used to be done by UHT manually including it.
-					if (GeneratedCPPCompileEnvironment.PrecompiledHeaderFile == null && Rules.PrivatePCHHeaderFile != null && Rules.PCHUsage != ModuleRules.PCHUsageMode.UseExplicitOrSharedPCHs)
+					if (Target.bForceIncludePCHHeadersForGenCppFilesWhenPCHIsDisabled && GeneratedCPPCompileEnvironment.PrecompiledHeaderFile == null && Rules.PrivatePCHHeaderFile != null && Rules.PCHUsage != ModuleRules.PCHUsageMode.UseExplicitOrSharedPCHs)
 					{
 						FileItem PrivatePchFileItem = FileItem.GetItemByFileReference(FileReference.Combine(ModuleDirectory, Rules.PrivatePCHHeaderFile));
 						if (!PrivatePchFileItem.Exists)
