PR #9636: UnrealBuildTool - .gexf export done with invariant culture. (Contributed by [REDACTED_PERSON])

diff --git a/Engine/Source/Programs/UnrealBuildTool/System/GraphVisualization.cs b/Engine/Source/Programs/UnrealBuildTool/System/GraphVisualization.cs
--- a/Engine/Source/Programs/UnrealBuildTool/System/GraphVisualization.cs
+++ b/Engine/Source/Programs/UnrealBuildTool/System/GraphVisualization.cs
@@ -4,6 +4,7 @@ using EpicGames.Core;
 using System;
 using System.Collections.Generic;
 using System.Drawing;
+using System.Globalization;
 using System.Linq;
 using System.Text;
 using System.Xml;
@@ -84,187 +85,200 @@ namespace UnrealBuildTool
 		/// <param name="GraphEdges">List of all graph edges.  Index order is important and must match with the individual edge Id members!</param>
 		public static void WriteGraphFile(FileReference Filename, string Description, List<GraphNode> GraphNodes, List<GraphEdge> GraphEdges)
 		{
-			XmlWriterSettings Settings = new XmlWriterSettings();
-			Settings.Indent = true;
-			Settings.IndentChars = "    ";
+			CultureInfo OriginalCulture = CultureInfo.CurrentCulture;
 
-			// Figure out all of the custom attribute types we're dealing with
-			Dictionary<string, GraphAttribute> AllAttributes = new Dictionary<string, GraphAttribute>(StringComparer.InvariantCultureIgnoreCase);
-			foreach (GraphNode GraphNode in GraphNodes)
+			try
 			{
-				foreach ((string AttributeName, object AttributeValue) in GraphNode.Attributes)
-				{
-					string AttributeType = GetAttributeType(AttributeValue);
+				// export graph using invariant culture so "." is used as a decimal separator
+				CultureInfo.CurrentCulture = CultureInfo.InvariantCulture;
 
-					GraphAttribute? Attribute;
-					if (!AllAttributes.TryGetValue(AttributeName, out Attribute))
-					{
-						AllAttributes[AttributeName] = new GraphAttribute(AllAttributes.Count, AttributeName, AttributeType);
-					}
-					else if (!Attribute.TypeName.Equals(AttributeType))
+				XmlWriterSettings Settings = new XmlWriterSettings();
+				Settings.Indent = true;
+				Settings.IndentChars = "    ";
+
+				// Figure out all of the custom attribute types we're dealing with
+				Dictionary<string, GraphAttribute> AllAttributes = new Dictionary<string, GraphAttribute>(StringComparer.InvariantCultureIgnoreCase);
+				foreach (GraphNode GraphNode in GraphNodes)
+				{
+					foreach ((string AttributeName, object AttributeValue) in GraphNode.Attributes)
 					{
-						throw new InvalidOperationException("Multiple graph nodes with the same attribute name but different types encountered!");
+						string AttributeType = GetAttributeType(AttributeValue);
+
+						GraphAttribute? Attribute;
+						if (!AllAttributes.TryGetValue(AttributeName, out Attribute))
+						{
+							AllAttributes[AttributeName] = new GraphAttribute(AllAttributes.Count, AttributeName, AttributeType);
+						}
+						else if (!Attribute.TypeName.Equals(AttributeType))
+						{
+							throw new InvalidOperationException("Multiple graph nodes with the same attribute name but different types encountered!");
+						}
 					}
 				}
-			}
-
-			using (XmlWriter Writer = XmlWriter.Create(Filename.FullName, Settings))
-			{
-				// NOTE: The GEXF XML format is defined here:  http://gexf.net/1.2draft/gexf-12draft-primer.pdf
 
-				string GEXFNamespace = "http://www.gexf.net/1.2-draft";
-				string SchemaNamespace = "http://www.w3.org/2001/XMLSchema-instance";
-				string VizNamespace = "http://www.gexf.net/1.2draft/viz";
+				using (XmlWriter Writer = XmlWriter.Create(Filename.FullName, Settings))
+				{
+					// NOTE: The GEXF XML format is defined here:  http://gexf.net/1.2draft/gexf-12draft-primer.pdf
 
-				Writer.WriteStartElement("gexf", GEXFNamespace);
-				Writer.WriteAttributeString("xmlns", "xsi", null, SchemaNamespace);
-				Writer.WriteAttributeString("schemaLocation", SchemaNamespace, "http://www.gexf.net/1.2draft http://www.gexf.net/1.2draft/gexf.xsd");
-				Writer.WriteAttributeString("xmlns", "viz", null, VizNamespace);
-				Writer.WriteAttributeString("version", "1.2");
+					string GEXFNamespace = "http://www.gexf.net/1.2-draft";
+					string SchemaNamespace = "http://www.w3.org/2001/XMLSchema-instance";
+					string VizNamespace = "http://www.gexf.net/1.2draft/viz";
 
-				Writer.WriteStartElement("meta");
-				{
-					Writer.WriteAttributeString("creator", "UnrealBuildTool");
-					Writer.WriteAttributeString("description", Description);
-				}
-				Writer.WriteEndElement();   // meta
+					Writer.WriteStartElement("gexf", GEXFNamespace);
+					Writer.WriteAttributeString("xmlns", "xsi", null, SchemaNamespace);
+					Writer.WriteAttributeString("schemaLocation", SchemaNamespace, "http://www.gexf.net/1.2draft http://www.gexf.net/1.2draft/gexf.xsd");
+					Writer.WriteAttributeString("xmlns", "viz", null, VizNamespace);
+					Writer.WriteAttributeString("version", "1.2");
 
-				Dictionary<GraphNode, int> NodeToId = new Dictionary<GraphNode, int>();
-				{
-					Writer.WriteStartElement("graph");
+					Writer.WriteStartElement("meta");
 					{
-						Writer.WriteAttributeString("mode", "static");
-						Writer.WriteAttributeString("defaultedgetype", "directed");
+						Writer.WriteAttributeString("creator", "UnrealBuildTool");
+						Writer.WriteAttributeString("description", Description);
+					}
+					Writer.WriteEndElement();   // meta
 
-						if (AllAttributes.Count > 0)
+					Dictionary<GraphNode, int> NodeToId = new Dictionary<GraphNode, int>();
+					{
+						Writer.WriteStartElement("graph");
 						{
-							Writer.WriteStartElement("attributes");
-							{
-								// @todo: Add support for edge attributes, not just node attributes
-								Writer.WriteAttributeString("class", "node");   // Node attributes, not edges!
+							Writer.WriteAttributeString("mode", "static");
+							Writer.WriteAttributeString("defaultedgetype", "directed");
 
-								foreach (var Attribute in AllAttributes.Values)
+							if (AllAttributes.Count > 0)
+							{
+								Writer.WriteStartElement("attributes");
 								{
-									Writer.WriteStartElement("attribute");
+									// @todo: Add support for edge attributes, not just node attributes
+									Writer.WriteAttributeString("class", "node");   // Node attributes, not edges!
+
+									foreach (var Attribute in AllAttributes.Values)
 									{
-										Writer.WriteAttributeString("id", Attribute.Id.ToString());
-										Writer.WriteAttributeString("title", Attribute.Name);
-										Writer.WriteAttributeString("type", Attribute.TypeName);
+										Writer.WriteStartElement("attribute");
+										{
+											Writer.WriteAttributeString("id", Attribute.Id.ToString());
+											Writer.WriteAttributeString("title", Attribute.Name);
+											Writer.WriteAttributeString("type", Attribute.TypeName);
+										}
+										Writer.WriteEndElement();   // attribute
 									}
-									Writer.WriteEndElement();   // attribute
+
+									// @todo: Add support for attribute type default values
 								}
+								Writer.WriteEndElement();   // attributes
 
-								// @todo: Add support for attribute type default values
 							}
-							Writer.WriteEndElement();   // attributes
-
-						}
 
-						Writer.WriteStartElement("nodes");
-						{
-							foreach (GraphNode GraphNode in GraphNodes)
+							Writer.WriteStartElement("nodes");
 							{
-								Writer.WriteStartElement("node");
+								foreach (GraphNode GraphNode in GraphNodes)
 								{
-									int Id = GetNodeId(GraphNode, NodeToId);
-									Writer.WriteAttributeString("id", Id.ToString());
-									Writer.WriteAttributeString("label", GraphNode.Label);
-
-									Writer.WriteStartElement("color", VizNamespace);
+									Writer.WriteStartElement("node");
 									{
-										Writer.WriteAttributeString("r", GraphNode.Color.R.ToString());
-										Writer.WriteAttributeString("g", GraphNode.Color.G.ToString());
-										Writer.WriteAttributeString("b", GraphNode.Color.B.ToString());
-										Writer.WriteAttributeString("a", (GraphNode.Color.A / 255.0f).ToString());
-									}
-									Writer.WriteEndElement();   // viz:color
+										int Id = GetNodeId(GraphNode, NodeToId);
+										Writer.WriteAttributeString("id", Id.ToString());
+										Writer.WriteAttributeString("label", GraphNode.Label);
 
-									Writer.WriteStartElement("size", VizNamespace);
-									{
-										Writer.WriteAttributeString("value", GraphNode.Size.ToString());
-									}
-									Writer.WriteEndElement();   // viz:size
+										Writer.WriteStartElement("color", VizNamespace);
+										{
+											Writer.WriteAttributeString("r", GraphNode.Color.R.ToString());
+											Writer.WriteAttributeString("g", GraphNode.Color.G.ToString());
+											Writer.WriteAttributeString("b", GraphNode.Color.B.ToString());
+											Writer.WriteAttributeString("a", (GraphNode.Color.A / 255.0f).ToString());
+										}
+										Writer.WriteEndElement();   // viz:color
 
-									Writer.WriteStartElement("shape", VizNamespace);
-									{
-										// NOTE: Valid shapes are:  disc, square, triangle, diamond, image
-										Writer.WriteAttributeString("value", "disc");
-									}
-									Writer.WriteEndElement();   // viz:shape
+										Writer.WriteStartElement("size", VizNamespace);
+										{
+											Writer.WriteAttributeString("value", GraphNode.Size.ToString());
+										}
+										Writer.WriteEndElement();   // viz:size
 
-									if (GraphNode.Attributes.Count > 0)
-									{
-										Writer.WriteStartElement("attvalues");
+										Writer.WriteStartElement("shape", VizNamespace);
 										{
-											foreach (var AttributeHashEntry in GraphNode.Attributes)
+											// NOTE: Valid shapes are:  disc, square, triangle, diamond, image
+											Writer.WriteAttributeString("value", "disc");
+										}
+										Writer.WriteEndElement();   // viz:shape
+
+										if (GraphNode.Attributes.Count > 0)
+										{
+											Writer.WriteStartElement("attvalues");
 											{
-												var AttributeName = AttributeHashEntry.Key;
-												var AttributeValue = AttributeHashEntry.Value;
+												foreach (var AttributeHashEntry in GraphNode.Attributes)
+												{
+													var AttributeName = AttributeHashEntry.Key;
+													var AttributeValue = AttributeHashEntry.Value;
 
-												var Attribute = AllAttributes[AttributeName];
+													var Attribute = AllAttributes[AttributeName];
 
-												Writer.WriteStartElement("attvalue");
-												{
-													Writer.WriteAttributeString("for", Attribute.Id.ToString());
-													Writer.WriteAttributeString("value", AttributeValue.ToString());
+													Writer.WriteStartElement("attvalue");
+													{
+														Writer.WriteAttributeString("for", Attribute.Id.ToString());
+														Writer.WriteAttributeString("value", AttributeValue.ToString());
+													}
+													Writer.WriteEndElement();
 												}
-												Writer.WriteEndElement();
 											}
+											Writer.WriteEndElement();   // attvalues
 										}
-										Writer.WriteEndElement();   // attvalues
 									}
+									Writer.WriteEndElement();   // node
 								}
-								Writer.WriteEndElement();   // node
 							}
-						}
-						Writer.WriteEndElement();   // nodes
+							Writer.WriteEndElement();   // nodes
 
 
-						Writer.WriteStartElement("edges");
-						{
-							for(int EdgeId = 0; EdgeId < GraphEdges.Count; EdgeId++)
+							Writer.WriteStartElement("edges");
 							{
-								GraphEdge GraphEdge = GraphEdges[EdgeId];
-								Writer.WriteStartElement("edge");
+								for (int EdgeId = 0; EdgeId < GraphEdges.Count; EdgeId++)
 								{
-									Writer.WriteAttributeString("id", EdgeId.ToString());
-									Writer.WriteAttributeString("source", GetNodeId(GraphEdge.Source, NodeToId).ToString());
-									Writer.WriteAttributeString("target", GetNodeId(GraphEdge.Target, NodeToId).ToString());
-									Writer.WriteAttributeString("weight", GraphEdge.Weight.ToString());
-
-									Writer.WriteStartElement("color", VizNamespace);
+									GraphEdge GraphEdge = GraphEdges[EdgeId];
+									Writer.WriteStartElement("edge");
 									{
-										Writer.WriteAttributeString("r", GraphEdge.Color.R.ToString());
-										Writer.WriteAttributeString("g", GraphEdge.Color.G.ToString());
-										Writer.WriteAttributeString("b", GraphEdge.Color.B.ToString());
-										Writer.WriteAttributeString("a", (GraphEdge.Color.A / 255.0f).ToString());
-									}
-									Writer.WriteEndElement();   // viz:color
+										Writer.WriteAttributeString("id", EdgeId.ToString());
+										Writer.WriteAttributeString("source", GetNodeId(GraphEdge.Source, NodeToId).ToString());
+										Writer.WriteAttributeString("target", GetNodeId(GraphEdge.Target, NodeToId).ToString());
+										Writer.WriteAttributeString("weight", GraphEdge.Weight.ToString());
 
-									Writer.WriteStartElement("thickness", VizNamespace);
-									{
-										Writer.WriteAttributeString("value", GraphEdge.Thickness.ToString());
-									}
-									Writer.WriteEndElement();   // viz:thickness
+										Writer.WriteStartElement("color", VizNamespace);
+										{
+											Writer.WriteAttributeString("r", GraphEdge.Color.R.ToString());
+											Writer.WriteAttributeString("g", GraphEdge.Color.G.ToString());
+											Writer.WriteAttributeString("b", GraphEdge.Color.B.ToString());
+											Writer.WriteAttributeString("a", (GraphEdge.Color.A / 255.0f).ToString());
+										}
+										Writer.WriteEndElement();   // viz:color
 
-									Writer.WriteStartElement("shape", VizNamespace);
-									{
-										// NOTE: Valid shapes are:  solid, dotted, dashed, double
-										Writer.WriteAttributeString("value", "solid");
+										Writer.WriteStartElement("thickness", VizNamespace);
+										{
+											Writer.WriteAttributeString("value", GraphEdge.Thickness.ToString());
+										}
+										Writer.WriteEndElement();   // viz:thickness
+
+										Writer.WriteStartElement("shape", VizNamespace);
+										{
+											// NOTE: Valid shapes are:  solid, dotted, dashed, double
+											Writer.WriteAttributeString("value", "solid");
+										}
+										Writer.WriteEndElement();   // viz:shape
 									}
-									Writer.WriteEndElement();   // viz:shape
+									Writer.WriteEndElement();   // edge
 								}
-								Writer.WriteEndElement();   // edge
 							}
+							Writer.WriteEndElement();   // nodes
 						}
-						Writer.WriteEndElement();   // nodes
+
+						Writer.WriteEndElement();   // graph
 					}
+					Writer.WriteEndElement();   // gexf
 
-					Writer.WriteEndElement();   // graph
+					Writer.Flush();
 				}
-				Writer.WriteEndElement();   // gexf
 
-				Writer.Flush();
+			}
+			finally
+			{
+				CultureInfo.CurrentCulture = OriginalCulture;
 			}
 		}
 
