UnrealBuildTool: Add additional configuration for various cpp conformance arguments, disabled by default

#rnx

diff --git a/Engine/Source/Programs/UnrealBuildTool/Platform/Windows/UEBuildWindows.cs b/Engine/Source/Programs/UnrealBuildTool/Platform/Windows/UEBuildWindows.cs
--- a/Engine/Source/Programs/UnrealBuildTool/Platform/Windows/UEBuildWindows.cs
+++ b/Engine/Source/Programs/UnrealBuildTool/Platform/Windows/UEBuildWindows.cs
@@ -294,11 +294,42 @@ namespace UnrealBuildTool
 		public string? ManifestFile;
 
 		/// <summary>
-		/// Enables strict standard conformance mode (/permissive-) in VS2017+.
+		/// Enables strict standard conformance mode (/permissive-).
 		/// </summary>
 		[XmlConfigFile(Category = "WindowsPlatform")]
 		[CommandLine("-Strict")]
-		public bool bStrictConformanceMode = false; 
+		public bool bStrictConformanceMode = false;
+
+		/// <summary>
+		/// Enables updated __cplusplus macro (/Zc:__cplusplus).
+		/// </summary>
+		[RequiresUniqueBuildEnvironment]
+		[XmlConfigFile(Category = "WindowsPlatform")]
+		[CommandLine("-UpdatedCPPMacro")]
+		public bool bUpdatedCPPMacro = false;
+
+		/// <summary>
+		/// Enables inline conformance (Remove unreferenced COMDAT) (/Zc:inline).
+		/// </summary>
+		[RequiresUniqueBuildEnvironment]
+		[XmlConfigFile(Category = "WindowsPlatform")]
+		[CommandLine("-StrictInline")]
+		public bool bStrictInlineConformance = false;
+
+		/// <summary>
+		/// Enables new preprocessor conformance (/Zc:preprocessor).
+		/// </summary>
+		[XmlConfigFile(Category = "WindowsPlatform")]
+		[CommandLine("-StrictPreprocessor")]
+		public bool bStrictPreprocessorConformance = false;
+
+		/// <summary>
+		/// Enables enum types conformance (/Zc:enumTypes) in VS2022 17.4 Preview 4.0+.
+		/// </summary>
+		[RequiresUniqueBuildEnvironment]
+		[XmlConfigFile(Category = "WindowsPlatform")]
+		[CommandLine("-StrictEnumTypes")]
+		public bool bStrictEnumTypesConformance = false;
 
 		/// VS2015 updated some of the CRT definitions but not all of the Windows SDK has been updated to match.
 		/// Microsoft provides legacy_stdio_definitions library to enable building with VS2015 until they fix everything up.
@@ -636,6 +667,26 @@ namespace UnrealBuildTool
 			get { return Inner.bStrictConformanceMode; }
 		}
 
+		public bool bUpdatedCPPMacro
+		{
+			get { return Inner.bUpdatedCPPMacro; }
+		}
+
+		public bool bStrictInlineConformance
+		{
+			get { return Inner.bStrictInlineConformance; }
+		}
+
+		public bool bStrictPreprocessorConformance
+		{
+			get { return Inner.bStrictPreprocessorConformance; }
+		}
+
+		public bool bStrictEnumTypesConformance
+		{
+			get { return Inner.bStrictEnumTypesConformance; }
+		}
+
 		public int DefaultStackSize
 		{
 			get { return Inner.DefaultStackSize; }
diff --git a/Engine/Source/Programs/UnrealBuildTool/Platform/Windows/VCToolChain.cs b/Engine/Source/Programs/UnrealBuildTool/Platform/Windows/VCToolChain.cs
--- a/Engine/Source/Programs/UnrealBuildTool/Platform/Windows/VCToolChain.cs
+++ b/Engine/Source/Programs/UnrealBuildTool/Platform/Windows/VCToolChain.cs
@@ -416,16 +416,39 @@ namespace UnrealBuildTool
 			Arguments.Add("/wd4819");
 
 			// Disable Microsoft extensions on VS2017+ for improved standards compliance.
-			if (Target.WindowsPlatform.Compiler.IsMSVC() && Target.WindowsPlatform.bStrictConformanceMode)
+			if (Target.WindowsPlatform.Compiler.IsMSVC())
 			{
-				// This define is needed to ensure that MSVC static analysis mode doesn't declare attributes that are incompatible with strict conformance mode
-				AddDefinition(Arguments, "SAL_NO_ATTRIBUTE_DECLARATIONS=1");
-				
-				Arguments.Add("/permissive-");
-				Arguments.Add("/Zc:strictStrings-"); // Have to disable strict const char* semantics due to Windows headers not being compliant.
-				if (EnvVars.CompilerVersion >= new VersionNumber(14, 32) && EnvVars.CompilerVersion < new VersionNumber(14, 33, 31629))
+				if (Target.WindowsPlatform.bStrictConformanceMode)
+				{
+					// This define is needed to ensure that MSVC static analysis mode doesn't declare attributes that are incompatible with strict conformance mode
+					AddDefinition(Arguments, "SAL_NO_ATTRIBUTE_DECLARATIONS=1");
+
+					Arguments.Add("/permissive-");
+					Arguments.Add("/Zc:strictStrings-"); // Have to disable strict const char* semantics due to Windows headers not being compliant.
+					if (EnvVars.CompilerVersion >= new VersionNumber(14, 32) && EnvVars.CompilerVersion < new VersionNumber(14, 33, 31629))
+					{
+						Arguments.Add("/Zc:lambda-");
+					}
+				}
+
+				if (Target.WindowsPlatform.bUpdatedCPPMacro)
+				{
+					Arguments.Add("/Zc:__cplusplus");
+				}
+
+				if (Target.WindowsPlatform.bStrictInlineConformance)
+				{
+					Arguments.Add("/Zc:inline");
+				}
+
+				if (Target.WindowsPlatform.bStrictPreprocessorConformance)
+				{
+					Arguments.Add("/Zc:preprocessor");
+				}
+
+				if (Target.WindowsPlatform.bStrictEnumTypesConformance && EnvVars.CompilerVersion >= new VersionNumber(14, 34, 31931))
 				{
-					Arguments.Add("/Zc:lambda-");
+					Arguments.Add("/Zc:enumTypes");
 				}
 			}
 
