- Sped up Mac project generation by caching which libraries have been validated

diff --git a/Engine/Source/Programs/UnrealBuildTool/Platform/Mac/UEBuildMac.cs b/Engine/Source/Programs/UnrealBuildTool/Platform/Mac/UEBuildMac.cs
--- a/Engine/Source/Programs/UnrealBuildTool/Platform/Mac/UEBuildMac.cs
+++ b/Engine/Source/Programs/UnrealBuildTool/Platform/Mac/UEBuildMac.cs
@@ -191,6 +191,7 @@ namespace UnrealBuildTool
 			Target.bUsePCHFiles = Target.bUsePCHFiles && !bCompilingMultipleArchitectures;
 		}
 
+		static HashSet<FileReference> ValidatedLibs = new();
 		public override void ValidateModule(UEBuildModule Module, ReadOnlyTargetRules Target)
 		{ 
 			if (BuildHostPlatform.Current.Platform == UnrealTargetPlatform.Mac && !Target.MacPlatform.bSkipClangValidation)
@@ -199,6 +200,12 @@ namespace UnrealBuildTool
 				// Validate the added public libraries
 				foreach (FileReference LibLoc in Module.PublicLibraries)
 				{
+					if (ValidatedLibs.Contains(LibLoc))
+					{
+						continue;
+					}
+					ValidatedLibs.Add(LibLoc);
+
 					switch (LibLoc.GetExtension())
 					{
 						case ".a":
