bug fix : link step depends on the embedded manifest. If the .manifest changes, the application needs to be relinked.

#rnx

diff --git a/Engine/Source/Programs/UnrealBuildTool/Platform/Windows/VCToolChain.cs b/Engine/Source/Programs/UnrealBuildTool/Platform/Windows/VCToolChain.cs
--- a/Engine/Source/Programs/UnrealBuildTool/Platform/Windows/VCToolChain.cs
+++ b/Engine/Source/Programs/UnrealBuildTool/Platform/Windows/VCToolChain.cs
@@ -2141,6 +2141,11 @@ namespace UnrealBuildTool
 				{
 					InputFileNames.Add(string.Format("\"{0}\"", SystemLibrary));
 				}
+
+				if (Target.WindowsPlatform.ManifestFile != null)
+				{
+					PrerequisiteItems.Add(FileItem.GetItemByPath(Target.WindowsPlatform.ManifestFile));
+				}
 			}
 
 			Arguments.AddRange(InputFileNames);
