Handle macOS zombie process that can occur when Ctrl-C is pressed in a terminal or when cancelling an Xcode build.

#rnx

diff --git a/Engine/Source/Programs/AutomationTool/Program.cs b/Engine/Source/Programs/AutomationTool/Program.cs
--- a/Engine/Source/Programs/AutomationTool/Program.cs
+++ b/Engine/Source/Programs/AutomationTool/Program.cs
@@ -382,7 +382,7 @@ namespace AutomationToolDriver
 			{
 				Console.CancelKeyPress += delegate
 				{
-					Console.WriteLine("Dying....");
+					Console.WriteLine("AutomationTool: Ctrl-C pressed. Exiting...");
 				};
 			}
 
diff --git a/Engine/Source/Programs/UnrealBuildTool/UnrealBuildTool.cs b/Engine/Source/Programs/UnrealBuildTool/UnrealBuildTool.cs
--- a/Engine/Source/Programs/UnrealBuildTool/UnrealBuildTool.cs
+++ b/Engine/Source/Programs/UnrealBuildTool/UnrealBuildTool.cs
@@ -458,6 +458,23 @@ namespace UnrealBuildTool
 			JsonTracer? Tracer = null;
 
 			ILogger Logger = Log.Logger;
+
+			// When running RunUBT.sh on a Mac we need to install a Ctrl-C handler, or hitting Ctrl-C from a terminal
+			// or from cancelling a build within Xcode, can leave a dotnet process in a zombie state. 
+			// By putting this in, the Ctrl-C may not be handled immediately, but it shouldn't leave a blocking zombie process
+			if (OperatingSystem.IsMacOS())
+			{
+				Console.CancelKeyPress += delegate
+				{
+					Console.WriteLine("UnrealBuildTool: Ctrl-C pressed. Exiting...");
+
+					// While the Ctrl-C handler fixes most instances of a zombie process, we still need to 
+					// force an exit from the process to handle _all_ cases.  Ctrl-C should not be a regular event! 
+					// Note: this could be a dotnet (6.0.302) on macOS issue.  Recheck with next release if this is still required.
+					Environment.Exit(-1);
+				};
+			}
+
 			try
 			{
 				// Start capturing performance info
