rename VCManifestGenerator to AppXManifestGenerator to prevent confusion with the embedded link.exe .manifest file.

#rnx

diff --git a/Engine/Source/Programs/UnrealBuildTool/Platform/Windows/AppXManifestGeneratorBase.cs b/Engine/Source/Programs/UnrealBuildTool/Platform/Windows/AppXManifestGeneratorBase.cs
new file mode 100644
--- /dev/null
+++ b/Engine/Source/Programs/UnrealBuildTool/Platform/Windows/AppXManifestGeneratorBase.cs
@@ -0,0 +1,1039 @@
+// Copyright Epic Games, Inc. All Rights Reserved.
+
+using System;
+using System.Collections.Generic;
+using System.Linq;
+using System.IO;
+using System.Text.RegularExpressions;
+using System.Xml;
+using EpicGames.Core;
+using System.Xml.Linq;
+using System.Text;
+using System.Diagnostics;
+using UnrealBuildBase;
+using System.Diagnostics.CodeAnalysis;
+using Microsoft.Extensions.Logging;
+
+namespace UnrealBuildTool
+{
+	/// <summary>
+	/// Base class for VC appx manifest generation
+	/// </summary>
+	abstract public class AppXManifestGeneratorBase
+	{
+		/// default schema namespace
+        protected virtual string Schema2010NS { get { return "http://schemas.microsoft.com/appx/2010/manifest"; } }
+
+		/// config section for platform-specific target settings
+		protected virtual string IniSection_PlatformTargetSettings { get { return string.Format( "/Script/{0}PlatformEditor.{0}TargetSettings", Platform.ToString() ); } }
+		
+		/// config section for general target settings
+		protected virtual string IniSection_GeneralProjectSettings { get { return "/Script/EngineSettings.GeneralProjectSettings"; } }
+
+		/// default subdirectory for build resources
+		protected const string BuildResourceSubPath = "Resources";
+
+		/// default subdirectory for engine resources
+		protected const string EngineResourceSubPath = "DefaultImages";
+
+		/// platform to use for reading configuration
+		protected virtual UnrealTargetPlatform ConfigPlatform { get { return Platform; } }
+
+        /// Manifest compliance values
+        protected const int MaxResourceEntries = 200;
+
+		/// cached engine ini
+		protected ConfigHierarchy? EngineIni;
+
+		/// cached game ini
+		protected ConfigHierarchy? GameIni;
+
+		/// the default culture to use
+		protected string? DefaultCulture;
+
+		/// all cultures that will be staged
+		protected List<string>? CulturesToStage;
+
+		/// resource writer for the default culture
+		protected UEResXWriter? DefaultResourceWriter;
+
+		/// resource writers for each culture
+		protected Dictionary<string, UEResXWriter>? PerCultureResourceWriters;
+
+		/// the platform to generate the manifest for
+		protected UnrealTargetPlatform Platform;
+
+		/// project file to use
+		protected FileReference? ProjectFile;
+
+		/// directory containing the project
+		protected string? ProjectPath;
+
+		/// output path - where the manifest will be created
+		protected string? OutputPath;
+
+		/// intermediate path - used for temporary files
+		protected string? IntermediatePath;
+
+		/// files that should be included in the manifest output folder and where to source them from
+		protected Dictionary<string,string>? ManifestFiles; // Dst, Src
+
+		/// <summary>
+		/// Logger for output
+		/// </summary>
+		protected readonly ILogger Logger;
+
+		/// <summary>
+		/// Create a manifest generator for the given platform variant.
+		/// </summary>
+		public AppXManifestGeneratorBase( UnrealTargetPlatform InPlatform, ILogger InLogger )
+		{
+			Platform = InPlatform;
+			Logger = InLogger;
+		}
+
+		/// <summary>
+		/// Lookup a switch in a dictionary
+		/// </summary>
+        protected static bool SafeGetBool(IDictionary<string, string> InDictionary, string Key, bool DefaultValue = false)
+		{
+			if (InDictionary.ContainsKey(Key))
+			{
+				var Value = InDictionary[Key].Trim().ToLower();
+				return Value == "true" || Value == "1" || Value == "yes";
+			}
+
+			return DefaultValue;
+		}
+
+		/// <summary>
+		/// Attempts to create the given directory
+		/// </summary>
+		protected static bool CreateCheckDirectory(string TargetDirectory, ILogger Logger)
+		{
+			if (!Directory.Exists(TargetDirectory))
+			{
+				try
+				{
+					Directory.CreateDirectory(TargetDirectory);
+				}
+				catch (Exception)
+				{
+					Logger.LogError("Could not create directory {TargetDir}.", TargetDirectory);
+					return false;
+				}
+				if (!Directory.Exists(TargetDirectory))
+				{
+					Logger.LogError("Path {TargetDir} does not exist or is not a directory.", TargetDirectory);
+					return false;
+				}
+			}
+			return true;
+		}
+
+
+        /// <summary>
+        /// Runs Makepri. Blocking
+        /// </summary>
+        /// <param name="CommandLine">Commandline</param>
+        /// <returns>bool    Application ran successfully</returns>
+        protected bool RunMakePri(string CommandLine)
+        {
+			string PriExecutable = GetMakePriBinaryPath();
+
+			if (File.Exists(PriExecutable) == false)
+			{
+				throw new BuildException("BUILD FAILED: Couldn't find the makepri executable: {0}", PriExecutable);
+			}
+
+			StringBuilder ProcessOutput = new StringBuilder();
+			void LocalProcessOutput(DataReceivedEventArgs Args)
+			{
+				if (Args != null && Args.Data != null)
+				{
+					ProcessOutput.AppendLine(Args.Data.TrimEnd());
+				}
+			}
+
+			ProcessStartInfo StartInfo = new ProcessStartInfo(PriExecutable, CommandLine);			
+			StartInfo.UseShellExecute = false;
+			StartInfo.CreateNoWindow = true;
+			StartInfo.StandardOutputEncoding = Encoding.Unicode;
+			StartInfo.StandardErrorEncoding = Encoding.Unicode;
+
+			Process LocalProcess = new Process();
+			LocalProcess.StartInfo = StartInfo;
+			LocalProcess.OutputDataReceived += (Sender, Args) => { LocalProcessOutput(Args); };
+			LocalProcess.ErrorDataReceived += (Sender, Args) => { LocalProcessOutput(Args); };
+			int ExitCode = Utils.RunLocalProcess(LocalProcess);
+
+			if (ExitCode == 0)
+			{
+				Logger.LogDebug("Output", ProcessOutput.ToString());
+				return true;
+			}
+			else
+			{
+				Logger.LogInformation("Output", ProcessOutput.ToString());
+				Logger.LogError("{File} returned an error.", Path.GetFileName(PriExecutable));
+				Logger.LogError("Exit code: {Code}", ExitCode);
+				return false;
+			}
+        }
+
+
+		/// <summary>
+		/// Returns a valid version of the given package version string
+		/// </summary>
+		protected string? ValidatePackageVersion(string InVersionNumber)
+		{
+			string WorkingVersionNumber = Regex.Replace(InVersionNumber, "[^.0-9]", "");
+			string CompletedVersionString = "";
+			if (WorkingVersionNumber != null)
+			{
+				string[] SplitVersionString = WorkingVersionNumber.Split(new char[] { '.' });
+				int NumVersionElements = Math.Min(4, SplitVersionString.Length);
+				for (int VersionElement = 0; VersionElement < NumVersionElements; VersionElement++)
+				{
+					string QuadElement = SplitVersionString[VersionElement];
+					int QuadValue = 0;
+					if (QuadElement.Length == 0 || !int.TryParse(QuadElement, out QuadValue))
+					{
+						CompletedVersionString += "0";
+					}
+					else
+					{
+						if (QuadValue < 0)
+						{
+							QuadValue = 0;
+						}
+						if (QuadValue > 65535)
+						{
+							QuadValue = 65535;
+						}
+						CompletedVersionString += QuadValue;
+					}
+					if (VersionElement < 3)
+					{
+						CompletedVersionString += ".";
+					}
+				}
+				for (int VersionElement = NumVersionElements; VersionElement < 4; VersionElement++)
+				{
+					CompletedVersionString += "0";
+					if (VersionElement < 3)
+					{
+						CompletedVersionString += ".";
+					}
+				}
+			}
+			if (CompletedVersionString == null || CompletedVersionString.Length <= 0)
+			{
+				Logger.LogError("Invalid package version {Ver}. Package versions must be in the format #.#.#.# where # is a number 0-65535.", InVersionNumber);
+				Logger.LogError("Consider setting [{IniSection}]:PackageVersion to provide a specific value.", IniSection_PlatformTargetSettings);
+			}
+			return CompletedVersionString;
+		}
+
+		/// <summary>
+		/// Returns a valid version of the given application id
+		/// </summary>
+		protected string? ValidateProjectBaseName(string InApplicationId)
+		{
+			string ReturnVal = Regex.Replace(InApplicationId, "[^A-Za-z0-9]", "");
+			if (ReturnVal != null)
+			{
+				// Remove any leading numbers (must start with a letter)
+				ReturnVal = Regex.Replace(ReturnVal, "^[0-9]*", "");
+			}
+			if (ReturnVal == null || ReturnVal.Length <= 0)
+			{
+				Logger.LogError("Invalid application ID {AppId}. Application IDs must only contain letters and numbers. And they must begin with a letter.", InApplicationId);
+			}
+			return ReturnVal;
+		}
+
+		/// <summary>
+		/// Reads an integer from the cached ini files
+		/// </summary>
+		[return: NotNullIfNotNull("DefaultValue")]
+        protected string? ReadIniString(string? Key, string Section, string? DefaultValue = null)
+		{
+			if (Key == null)
+				return DefaultValue;
+
+			string Value;
+			if (GameIni!.GetString(Section, Key, out Value) && !string.IsNullOrWhiteSpace(Value))
+				return Value;
+
+			if (EngineIni!.GetString(Section, Key, out Value) && !string.IsNullOrWhiteSpace(Value))
+				return Value;
+
+			return DefaultValue;
+		}
+
+		/// <summary>
+		/// Reads a string from the cached ini files
+		/// </summary>
+		[return: NotNullIfNotNull("DefaultValue")]
+        protected string? GetConfigString(string PlatformKey, string? GenericKey, string? DefaultValue = null)
+		{
+			string? GenericValue = ReadIniString(GenericKey, IniSection_GeneralProjectSettings, DefaultValue);
+			return ReadIniString(PlatformKey, IniSection_PlatformTargetSettings, GenericValue);
+		}
+
+		/// <summary>
+		/// Reads a bool from the cached ini files
+		/// </summary>
+		protected bool GetConfigBool(string PlatformKey, string? GenericKey, bool DefaultValue = false)
+		{
+			var GenericValue = ReadIniString(GenericKey, IniSection_GeneralProjectSettings, null);
+			var ResultStr = ReadIniString(PlatformKey, IniSection_PlatformTargetSettings, GenericValue);
+
+			if (ResultStr == null)
+				return DefaultValue;
+
+			ResultStr = ResultStr.Trim().ToLower();
+
+			return ResultStr == "true" || ResultStr == "1" || ResultStr == "yes";
+		}
+
+		/// <summary>
+		/// Reads a color from the cached ini files
+		/// </summary>
+		protected string GetConfigColor(string PlatformConfigKey, string DefaultValue)
+		{
+			var ConfigValue = GetConfigString(PlatformConfigKey, null, null);
+			if (ConfigValue == null)
+				return DefaultValue;
+
+			Dictionary<string, string>? Pairs;
+			int R, G, B;
+			if (ConfigHierarchy.TryParse(ConfigValue, out Pairs) &&
+				int.TryParse(Pairs["R"], out R) &&
+				int.TryParse(Pairs["G"], out G) &&
+				int.TryParse(Pairs["B"], out B))
+			{
+				return "#" + R.ToString("X2") + G.ToString("X2") + B.ToString("X2");
+			}
+
+			Logger.LogWarning("Failed to parse color config value. Using default.");
+			return DefaultValue;
+		}
+
+
+
+
+		private bool RemoveStaleResourceFiles()
+		{
+			// remove all resource files that should not be included
+			string TargetResourceDir = Path.Combine(OutputPath!, BuildResourceSubPath);
+			var TargetResourceInstances = Directory.EnumerateFiles(TargetResourceDir, "*.*", SearchOption.AllDirectories);
+
+			var StaleResourceFiles = TargetResourceInstances.Where(X => !ManifestFiles!.ContainsKey(X)).ToList();
+			if (StaleResourceFiles.Any())
+			{
+				Logger.LogDebug("Removing stale manifest resource files...");
+				foreach (string StaleResourceFile in StaleResourceFiles)
+				{
+					// try to delete the file & the directory that contains it
+					try
+					{
+						Logger.LogDebug("    removing {Path}", Utils.MakePathRelativeTo(StaleResourceFile, OutputPath!));
+						FileUtils.ForceDeleteFile(StaleResourceFile);
+						if (!Directory.EnumerateFileSystemEntries(Path.GetDirectoryName(StaleResourceFile)!).Any())
+						{
+							Directory.Delete(Path.GetDirectoryName(StaleResourceFile)!, false);
+						}
+					}
+					catch (Exception E)
+					{
+						Logger.LogError("    Could not remove {StaleResourceFile} - {Message}.", StaleResourceFile, E.Message);
+					}
+				}
+			}
+
+			return StaleResourceFiles.Any();
+		}
+
+
+		/// <summary>
+		/// Attempts to locate the given resource binary file in several known folder locations
+		/// </summary>
+		protected virtual bool FindResourceBinaryFile( out string SourcePath, string ResourceFileName, bool AllowEngineFallback = true)
+		{
+			// look in project normal Build location
+			SourcePath = Path.Combine(ProjectPath!, "Build", Platform.ToString(), BuildResourceSubPath);
+			bool bFileExists = File.Exists(Path.Combine(SourcePath, ResourceFileName));
+
+			// look in Platform Extensions next
+			if (!bFileExists)
+			{
+				SourcePath = Path.Combine(ProjectPath!, "Platforms", Platform.ToString(), "Build", BuildResourceSubPath);
+				bFileExists = File.Exists(Path.Combine(SourcePath, ResourceFileName));
+			}
+
+			// look in Engine, if allowed
+			if (!bFileExists && AllowEngineFallback)
+			{
+				SourcePath = Path.Combine(Unreal.EngineDirectory.FullName, "Build", Platform.ToString(), EngineResourceSubPath);
+				bFileExists = File.Exists(Path.Combine(SourcePath, ResourceFileName));
+
+				// look in Platform extensions too
+				if (!bFileExists)
+				{
+					SourcePath = Path.Combine(Unreal.EngineDirectory.FullName, "Platforms", Platform.ToString(), "Build", EngineResourceSubPath);
+					bFileExists = File.Exists(Path.Combine(SourcePath, ResourceFileName));
+				}
+			}
+
+			return bFileExists;
+		}
+
+		/// <summary>
+		/// Determines whether the given resource binary file can be found in one of the known folder locations
+		/// </summary>
+		protected bool DoesResourceBinaryFileExist(string ResourceFileName, bool AllowEngineFallback = true)
+		{
+			string SourcePath;
+			return FindResourceBinaryFile( out SourcePath, ResourceFileName, AllowEngineFallback );
+		}
+
+
+		/// <summary>
+		/// Adds the given resource binary file(s) to the manifest files
+		/// </summary>
+		protected bool AddResourceBinaryFileReference(string ResourceFileName, bool AllowEngineFallback = true)
+		{
+			string TargetPath = Path.Combine(OutputPath!, BuildResourceSubPath);
+			string SourcePath;
+			bool bFileExists = FindResourceBinaryFile(out SourcePath, ResourceFileName, AllowEngineFallback);
+
+			// At least the default culture entry for any resource binary must always exist
+			if (!bFileExists)
+			{
+				return false;
+			}
+
+			// If the target resource folder doesn't exist yet, create it
+			if (!CreateCheckDirectory(TargetPath, Logger))
+			{
+				return false;
+			}
+
+			// Find all copies of the resource file in the source directory (could be up to one for each culture and the default).
+			List<string> SourceResourceInstances = new List<string>();
+			SourceResourceInstances.Add( Path.Combine( SourcePath, ResourceFileName) );
+			foreach( string CultureId in CulturesToStage!)
+			{
+				string CultureResourceFile = Path.Combine(SourcePath, CultureId, ResourceFileName);
+				if (File.Exists(CultureResourceFile))
+				{
+					SourceResourceInstances.Add( CultureResourceFile );
+				}
+			}
+
+			// Copy new resource files
+			foreach (string SourceResourceFile in SourceResourceInstances)
+			{
+				string TargetResourcePath = Path.Combine(TargetPath, SourceResourceFile.Substring(SourcePath.Length + 1));
+				if (!CreateCheckDirectory(Path.GetDirectoryName(TargetResourcePath)!, Logger))
+				{
+					Logger.LogError("Unable to create intermediate directory {IntDir}.", Path.GetDirectoryName(TargetResourcePath));
+					continue;
+				}
+				AddFileReference(SourceResourceFile, TargetResourcePath, bIsGeneratedFile:false);
+			}
+
+			return true;
+		}
+
+
+		/// <summary>
+		/// Adds the given file to the manifest files
+		/// </summary>
+		protected void AddFileReference(string SourcePath, string TargetPath, bool bIsGeneratedFile)
+		{
+			// check if the file is unchanged
+			if (File.Exists(TargetPath) && !string.IsNullOrEmpty(SourcePath) )
+			{
+				FileInfo SrcFileInfo = new FileInfo(SourcePath);
+				FileInfo DstFileInfo = new FileInfo(TargetPath);
+
+				bool bFileIsUnchanged = (SrcFileInfo.Length == DstFileInfo.Length);
+				if (bFileIsUnchanged)
+				{
+					if (bIsGeneratedFile)
+					{
+						// this file is auto-generated - we need to compare the contents
+						byte[] OriginalContents = File.ReadAllBytes(TargetPath);
+						byte[] NewContents = File.ReadAllBytes(SourcePath);
+						bFileIsUnchanged = Enumerable.SequenceEqual(OriginalContents, NewContents);
+					}
+					else
+					{
+						// this file is not generated, just copied from somewhere else - can check the time to confirm if they're different
+						bFileIsUnchanged = (SrcFileInfo.CreationTime == DstFileInfo.CreationTime);
+					}
+				}
+
+				if (bFileIsUnchanged)
+				{
+					// use an empty source string if the file doesn't need changing
+					SourcePath = "";
+				}
+			}
+
+			ManifestFiles!.Add(TargetPath, SourcePath);
+		}
+
+
+		/// <summary>
+		/// Copies all of the generated files to the output folder
+		/// </summary>
+		private List<string> CopyFilesToOutput()
+		{
+			List<string> UpdatedFiles = new List<string>();
+
+			// early out if there's nothing to copy
+			if (ManifestFiles == null || !ManifestFiles.Any(X => !string.IsNullOrEmpty(X.Value)))
+			{
+				return UpdatedFiles;
+			}
+
+			Logger.LogDebug("Updating manifest resource files...");
+
+			// copy over any new or updated files
+			foreach ( var ManifestFilePair in ManifestFiles! )
+			{
+				string TargetPath = ManifestFilePair.Key;
+				string SourcePath = ManifestFilePair.Value;
+				if (string.IsNullOrEmpty(SourcePath))
+				{
+					// if source is an empty string then the file is up-to-date
+					continue;
+				}
+				
+				// remove old version, if any
+				bool bFileExists = File.Exists(TargetPath);
+				if (bFileExists)
+				{
+					try
+					{
+						FileUtils.ForceDeleteFile(TargetPath);
+					}
+					catch (Exception E)
+					{
+						Logger.LogError("    Could not replace file {TargetPath} - {Message}", TargetPath, E.Message);
+					}
+				}
+
+				// copy new version
+				try
+				{
+					if (bFileExists)
+					{
+						Logger.LogDebug("    updating {Path}", Utils.MakePathRelativeTo(TargetPath, OutputPath!));
+					}
+					else
+					{
+						Logger.LogDebug("    adding {Path}", Utils.MakePathRelativeTo(TargetPath, OutputPath!));
+					}
+
+					Directory.CreateDirectory(Path.GetDirectoryName(TargetPath)!);
+					File.Copy(SourcePath, TargetPath);
+					File.SetAttributes(TargetPath, FileAttributes.Normal);
+					File.SetCreationTime(TargetPath, File.GetCreationTime(SourcePath));
+
+					UpdatedFiles.Add(TargetPath);
+				}
+				catch (Exception E)
+				{
+					Logger.LogError("    Unable to copy file {TargetPath} - {Message}", TargetPath, E.Message);
+				}
+			}
+
+			return UpdatedFiles;
+		}
+
+
+		/// <summary>
+		/// Adds the given string to the culture string writers
+		/// </summary>
+		protected string AddResourceEntry(string ResourceEntryName, string ConfigKey, string GenericINISection, string GenericINIKey, string DefaultValue, string ValueSuffix = "")
+		{
+			string? ConfigScratchValue = null;
+
+			// Get the default culture value
+			string DefaultCultureScratchValue;
+			if (EngineIni!.GetString(IniSection_PlatformTargetSettings, "CultureStringResources", out DefaultCultureScratchValue))
+			{
+				Dictionary<string, string>? Values;
+				if (!ConfigHierarchy.TryParse(DefaultCultureScratchValue, out Values))
+				{
+					Logger.LogError("Invalid default culture string resources: \"{Culture}\". Unable to add resource entry.", DefaultCultureScratchValue);
+					return "";
+				}
+
+				ConfigScratchValue = Values[ConfigKey];
+			}
+
+			if (string.IsNullOrEmpty(ConfigScratchValue))
+			{
+				// No platform specific value is provided. Use the generic config or default value
+				ConfigScratchValue = ReadIniString(GenericINIKey, GenericINISection, DefaultValue);
+			}
+
+			DefaultResourceWriter!.AddResource(ResourceEntryName, ConfigScratchValue + ValueSuffix);
+
+			// Find the default value
+			List<string>? PerCultureValues;
+			if (EngineIni.GetArray(IniSection_PlatformTargetSettings, "PerCultureResources", out PerCultureValues))
+			{
+				foreach (string CultureCombinedValues in PerCultureValues)
+				{
+					Dictionary<string, string>? SeparatedCultureValues;
+					if (!ConfigHierarchy.TryParse(CultureCombinedValues, out SeparatedCultureValues)
+						|| !SeparatedCultureValues.ContainsKey("CultureStringResources")
+						|| !SeparatedCultureValues.ContainsKey("CultureId"))
+					{
+						Logger.LogError("Invalid per-culture resource: \"{Culture}\". Unable to add resource entry.", CultureCombinedValues);
+						continue;
+					}
+
+					var CultureId = SeparatedCultureValues["CultureId"];
+					if (CulturesToStage!.Contains(CultureId))
+					{
+						Dictionary<string, string>? CultureStringResources;
+						if (!ConfigHierarchy.TryParse(SeparatedCultureValues["CultureStringResources"], out CultureStringResources))
+						{
+							Logger.LogError("Invalid culture string resources: \"{Culture}\". Unable to add resource entry.", CultureCombinedValues);
+							continue;
+						}
+
+						var Value = CultureStringResources[ConfigKey];
+
+						if (CulturesToStage.Contains(CultureId) && !string.IsNullOrEmpty(Value))
+						{
+							var Writer = PerCultureResourceWriters![CultureId];
+							Writer.AddResource(ResourceEntryName, Value + ValueSuffix);
+						}
+					}
+				}
+			}
+
+			return "ms-resource:" + ResourceEntryName;
+		}
+
+		/// <summary>
+		/// Adds an additional string to all culture resource writers
+		/// </summary>
+		protected string AddExternalResourceEntry(string ResourceEntryName, string DefaultValue, Dictionary<string,string> CultureIdToCultureValues)
+		{
+			DefaultResourceWriter!.AddResource(ResourceEntryName, DefaultValue);
+
+			foreach( KeyValuePair<string,string> CultureIdToCultureValue in CultureIdToCultureValues)
+			{
+				var Writer = PerCultureResourceWriters![CultureIdToCultureValue.Key];
+				Writer.AddResource(ResourceEntryName, CultureIdToCultureValue.Value);
+			}
+
+			return "ms-resource:" + ResourceEntryName;
+		}
+
+		/// <summary>
+		/// Adds a debug-only string to all resource writers
+		/// </summary>
+		protected string AddDebugResourceString(string ResourceEntryName, string Value)
+		{
+			DefaultResourceWriter!.AddResource(ResourceEntryName, Value);
+
+			foreach (var CultureId in CulturesToStage!)
+			{
+				var Writer = PerCultureResourceWriters![CultureId];
+				Writer.AddResource(ResourceEntryName, Value);
+			}
+
+			return "ms-resource:" + ResourceEntryName;
+		}
+
+		/// <summary>
+		/// Get the XName from a given string and schema pair
+		/// </summary>
+		protected virtual XName GetName( string BaseName, string SchemaName )
+		{
+			return XName.Get(BaseName);
+		}
+
+		/// <summary>
+		/// Get the resources element
+		/// </summary>
+		protected XElement GetResources()
+		{
+			List<string> ResourceCulturesList = new List<string>(CulturesToStage!);
+			// Move the default culture to the front of the list
+			ResourceCulturesList.Remove(DefaultCulture!);
+			ResourceCulturesList.Insert(0, DefaultCulture!);
+
+			// Check that we have a valid number of cultures
+			if (CulturesToStage!.Count < 1 || CulturesToStage.Count >= MaxResourceEntries)
+			{
+				Logger.LogWarning("Incorrect number of cultures to stage. There must be between 1 and {MaxCultures} cultures selected.", MaxResourceEntries);
+			}
+
+			// Create the culture list. This list is unordered except that the default language must be first which we already took care of above.
+			var CultureElements = ResourceCulturesList.Select(c =>
+				new XElement(GetName("Resource", Schema2010NS), new XAttribute("Language", c)));
+
+			return new XElement(GetName("Resources", Schema2010NS), CultureElements);
+		}
+
+		/// <summary>
+		/// Get the package identity name string
+		/// </summary>
+		protected string GetIdentityPackageName(string? TargetName)
+		{
+            // Read the PackageName from config
+			var DefaultName = (ProjectFile != null) ? ProjectFile.GetFileNameWithoutAnyExtensions() : (TargetName ?? "DefaultUEProject");
+            var PackageName = Regex.Replace(GetConfigString("PackageName", "ProjectName", DefaultName), "[^-.A-Za-z0-9]", "");
+            if (string.IsNullOrWhiteSpace(PackageName))
+            {
+                Logger.LogError("Invalid package name {Name}. Package names must only contain letters, numbers, dash, and period and must be at least one character long.", PackageName);
+                Logger.LogError("Consider using the setting [{IniSection}]:PackageName to provide a specific value.", IniSection_PlatformTargetSettings);
+            }
+
+			// If specified in the project settings append the users machine name onto the package name to allow sharing of devkits without stomping of deploys
+			bool bPackageNameUseMachineName;
+			if (EngineIni!.GetBool(IniSection_PlatformTargetSettings, "bPackageNameUseMachineName", out bPackageNameUseMachineName) && bPackageNameUseMachineName)
+			{
+				var MachineName = Regex.Replace(Environment.MachineName.ToString(), "[^-.A-Za-z0-9]", "");
+				PackageName = PackageName + ".NOT.SHIPPABLE." + MachineName;
+			}
+
+			return PackageName;
+		}
+
+		/// <summary>
+		/// Get the publisher name string
+		/// </summary>
+		protected string GetIdentityPublisherName()
+		{
+            var PublisherName = GetConfigString("PublisherName", "CompanyDistinguishedName", "CN=NoPublisher");
+			return PublisherName;
+		}
+
+		/// <summary>
+		/// Get the package version string
+		/// </summary>
+		protected string? GetIdentityVersionNumber()
+		{
+            var VersionNumber = GetConfigString("PackageVersion", "ProjectVersion", "1.0.0.0");
+            VersionNumber = ValidatePackageVersion(VersionNumber);
+
+            // If specified in the project settings attempt to retrieve the current build number and increment the version number by that amount, accounting for overflows
+            bool bIncludeEngineVersionInPackageVersion;
+            if (EngineIni!.GetBool(IniSection_PlatformTargetSettings, "bIncludeEngineVersionInPackageVersion", out bIncludeEngineVersionInPackageVersion) && bIncludeEngineVersionInPackageVersion)
+            {
+				VersionNumber = IncludeBuildVersionInPackageVersion(VersionNumber);
+            }
+
+			return VersionNumber;
+		}
+
+		/// <summary>
+		/// Get the package identity element
+		/// </summary>
+		protected XElement GetIdentity(string? TargetName, out string IdentityName)
+        {
+            string PackageName = GetIdentityPackageName(TargetName);
+            string PublisherName = GetIdentityPublisherName();
+            string? VersionNumber = GetIdentityVersionNumber();
+
+            IdentityName = PackageName;
+
+            return new XElement(GetName("Identity", Schema2010NS),
+                new XAttribute("Name", PackageName),
+                new XAttribute("Publisher", PublisherName),
+                new XAttribute("Version", VersionNumber!));
+        }
+
+		/// <summary>
+		/// Updates the given package version to include the engine build version, if requested
+		/// </summary>
+		protected virtual string? IncludeBuildVersionInPackageVersion(string? VersionNumber)
+		{
+			BuildVersion? BuildVersionForPackage;
+			if (VersionNumber != null && BuildVersion.TryRead(BuildVersion.GetDefaultFileName(), out BuildVersionForPackage) && BuildVersionForPackage.Changelist != 0)
+			{
+				// Break apart the version number into individual elements
+				string[] SplitVersionString = VersionNumber.Split('.');
+				VersionNumber = string.Format("{0}.{1}.{2}.{3}",
+					SplitVersionString[0],
+					SplitVersionString[1],
+					BuildVersionForPackage.Changelist / 10000,
+					BuildVersionForPackage.Changelist % 10000);
+			}
+
+			return VersionNumber;
+		}
+
+		/// <summary>
+		/// Get the path to the makepri.exe tool
+		/// </summary>
+		protected abstract string GetMakePriBinaryPath();
+
+		/// <summary>
+		/// Get any additional platform-specific parameters for makepri.exe
+		/// </summary>
+		protected virtual string GetMakePriExtraCommandLine()
+		{
+			return "";
+		}
+
+		/// <summary>
+		/// Return the entire manifest element
+		/// </summary>
+		protected abstract XElement GetManifest(List<UnrealTargetConfiguration> TargetConfigs, List<string> Executables, string? TargetName, out string IdentityName);
+
+		/// <summary>
+		/// Perform any platform-specific processing on the manifest before it is saved
+		/// </summary>
+		protected virtual void ProcessManifest(List<UnrealTargetConfiguration> TargetConfigs, List<string> Executables, string ManifestName, string ManifestTargetPath, string ManifestIntermediatePath)
+        {
+		}
+
+		/// <summary>
+		/// Create a manifest and return the list of modified files
+		/// </summary>
+		public List<string>? CreateManifest(string InManifestName, string InOutputPath, string InIntermediatePath, string? InTargetName, FileReference? InProjectFile, string InProjectDirectory, List<UnrealTargetConfiguration> InTargetConfigs, List<string> InExecutables)
+		{
+			// Check parameter values are valid.
+			if (InTargetConfigs.Count != InExecutables.Count)
+			{
+				Logger.LogError("The number of target configurations ({NumConfigs}) and executables ({NumExes}) passed to manifest generation do not match.", InTargetConfigs.Count, InExecutables.Count);
+				return null;
+			}
+			if (InTargetConfigs.Count < 1)
+			{
+				Logger.LogError("The number of target configurations is zero, so we cannot generate a manifest.");
+				return null;
+			}
+
+			if (!CreateCheckDirectory(InOutputPath, Logger))
+			{
+				Logger.LogError("Failed to create output directory \"{OutputDir}\".", InOutputPath);
+				return null;
+			}
+			if (!CreateCheckDirectory(InIntermediatePath, Logger))
+			{
+				Logger.LogError("Failed to create intermediate directory \"{IntDir}\".", InIntermediatePath);
+				return null;
+			}
+
+			OutputPath = InOutputPath;
+			IntermediatePath = InIntermediatePath;
+			ProjectFile = InProjectFile;
+			ProjectPath = InProjectDirectory;
+			ManifestFiles = new Dictionary<string, string>();
+
+			// Load up INI settings. We'll use engine settings to retrieve the manifest configuration, but these may reference
+			// values in either game or engine settings, so we'll keep both.
+			GameIni = ConfigCache.ReadHierarchy(ConfigHierarchyType.Game, DirectoryReference.FromFile(InProjectFile), ConfigPlatform);
+			EngineIni = ConfigCache.ReadHierarchy(ConfigHierarchyType.Engine, DirectoryReference.FromFile(InProjectFile), ConfigPlatform);
+
+			// Load and verify/clean culture list
+			{
+				List<string>? CulturesToStageWithDuplicates;
+				GameIni.GetArray("/Script/UnrealEd.ProjectPackagingSettings", "CulturesToStage", out CulturesToStageWithDuplicates);
+				GameIni.GetString("/Script/UnrealEd.ProjectPackagingSettings", "DefaultCulture", out DefaultCulture);
+				if (CulturesToStageWithDuplicates == null || CulturesToStageWithDuplicates.Count < 1)
+				{
+					Logger.LogError("At least one culture must be selected to stage.");
+					return null;
+				}
+
+				CulturesToStage = CulturesToStageWithDuplicates.Distinct().ToList();
+			}
+			if (DefaultCulture == null || DefaultCulture.Length < 1)
+			{
+				DefaultCulture = CulturesToStage[0];
+				Logger.LogWarning("A default culture must be selected to stage. Using {DefaultCulture}.", DefaultCulture);
+			}
+			if (!CulturesToStage.Contains(DefaultCulture))
+			{
+				DefaultCulture = CulturesToStage[0];
+				Logger.LogWarning("The default culture must be one of the staged cultures. Using {DefaultCulture}.", DefaultCulture);
+			}
+
+			List<string>? PerCultureValues;
+			if (EngineIni.GetArray(IniSection_PlatformTargetSettings, "PerCultureResources", out PerCultureValues))
+			{
+				foreach (string CultureCombinedValues in PerCultureValues)
+				{
+					Dictionary<string, string>? SeparatedCultureValues;
+					if (!ConfigHierarchy.TryParse(CultureCombinedValues, out SeparatedCultureValues))
+					{
+						Logger.LogWarning("Invalid per-culture resource value: {Culture}", CultureCombinedValues);
+						continue;
+					}
+
+					string StageId = SeparatedCultureValues["StageId"];
+					int CultureIndex = CulturesToStage.FindIndex(x => x == StageId);
+					if (CultureIndex >= 0)
+					{
+						CulturesToStage[CultureIndex] = SeparatedCultureValues["CultureId"];
+						if (DefaultCulture == StageId)
+						{
+							DefaultCulture = SeparatedCultureValues["CultureId"];
+						}
+					}
+				}
+			}
+			// Only warn if shipping, we can run without translated cultures they're just needed for cert
+			else if (InTargetConfigs.Contains(UnrealTargetConfiguration.Shipping))
+			{
+				Logger.LogInformation("Staged culture mappings not setup in the editor. See Per Culture Resources in the {Platform} Target Settings.", Platform.ToString() );
+			}
+
+			// Clean out the resources intermediate path so that we know there are no stale binary files.
+			string IntermediateResourceDirectory = Path.Combine(IntermediatePath, BuildResourceSubPath);
+			FileUtils.ForceDeleteDirectory(IntermediateResourceDirectory);
+			if (!CreateCheckDirectory(IntermediateResourceDirectory, Logger))
+			{
+				Logger.LogError("Could not create directory {IntDir}.", IntermediateResourceDirectory);
+				return null;
+			}
+
+			// Construct a single resource writer for the default (no-culture) values
+			string DefaultResourceIntermediatePath = Path.Combine(IntermediateResourceDirectory, "resources.resw");
+			DefaultResourceWriter = new UEResXWriter(DefaultResourceIntermediatePath);
+
+			// Construct the ResXWriters for each culture
+			PerCultureResourceWriters = new Dictionary<string, UEResXWriter>();
+			foreach (string Culture in CulturesToStage)
+			{
+				string IntermediateStringResourcePath = Path.Combine(IntermediateResourceDirectory, Culture);
+				string IntermediateStringResourceFile = Path.Combine(IntermediateStringResourcePath, "resources.resw");
+				if (!CreateCheckDirectory(IntermediateStringResourcePath, Logger))
+				{
+					Logger.LogWarning("Culture {Culture} resources not staged.", Culture);
+					CulturesToStage.Remove(Culture);
+					if (Culture == DefaultCulture)
+					{
+						DefaultCulture = CulturesToStage[0];
+						Logger.LogWarning("Default culture skipped. Using {DefaultCulture} as default culture.", DefaultCulture);
+					}
+					continue;
+				}
+				PerCultureResourceWriters.Add(Culture, new UEResXWriter(IntermediateStringResourceFile));
+			}
+
+
+
+			// Create the manifest document
+			string? IdentityName = null;
+			var ManifestXmlDocument = new XDocument(GetManifest(InTargetConfigs, InExecutables, InTargetName, out IdentityName));
+
+			// Export manifest to the intermediate directory then compare the contents to any existing target manifest
+			// and replace if there are differences.
+			string ManifestIntermediatePath = Path.Combine(IntermediatePath, InManifestName);
+			string ManifestTargetPath = Path.Combine(OutputPath, InManifestName);
+			ManifestXmlDocument.Save(ManifestIntermediatePath);
+			AddFileReference(ManifestIntermediatePath, ManifestTargetPath, bIsGeneratedFile: true);
+			ProcessManifest(InTargetConfigs, InExecutables, InManifestName, ManifestTargetPath, ManifestIntermediatePath);
+
+			// Export the resource tables starting with the default culture
+			string DefaultResourceTargetPath = Path.Combine(OutputPath, BuildResourceSubPath, "resources.resw");
+			DefaultResourceWriter.Close();
+			AddFileReference(DefaultResourceIntermediatePath, DefaultResourceTargetPath, bIsGeneratedFile: true);
+
+			foreach (var Writer in PerCultureResourceWriters)
+			{
+				Writer.Value.Close();
+
+				string IntermediateStringResourceFile = Path.Combine(IntermediateResourceDirectory, Writer.Key, "resources.resw");
+				string TargetStringResourceFile = Path.Combine(OutputPath, BuildResourceSubPath, Writer.Key, "resources.resw");
+
+				AddFileReference(IntermediateStringResourceFile, TargetStringResourceFile, bIsGeneratedFile: true);
+			}
+
+
+			// include a reference to the Package Resource Index so it isn't removed by RemoveStaleResourceFiles
+			string TargetResourceIndexFile = Path.Combine(OutputPath, "resources.pri");
+			AddFileReference("", TargetResourceIndexFile, bIsGeneratedFile: true);
+
+			// The resource database is dependent on everything else calculated here (manifest, resource string tables, binary resources).
+			// So if any file has been updated we'll need to run the config.
+			bool bHadStaleResources = RemoveStaleResourceFiles();
+			List<string> UpdatedFilePaths = CopyFilesToOutput();
+
+			if (bHadStaleResources || UpdatedFilePaths.Any())
+			{
+				// Create resource index configuration
+				string ResourceConfigFile = Path.Combine(IntermediatePath, "priconfig.xml");
+				RunMakePri("createconfig /cf \"" + ResourceConfigFile + "\" /dq " + DefaultCulture + " /o " + GetMakePriExtraCommandLine());
+
+				// Load the new resource index configuration
+				XmlDocument PriConfig = new XmlDocument();
+				PriConfig.Load(ResourceConfigFile);
+
+				// remove the packaging node - we do not want to split the pri & only want one .pri file
+				XmlNode PackagingNode = PriConfig.SelectSingleNode("/resources/packaging")!;
+				PackagingNode.ParentNode!.RemoveChild(PackagingNode);
+
+				// all required resources are explicitly listed in resources.resfiles, rather than relying on makepri to discover them
+				string ResourcesResFile = Path.Combine(IntermediatePath, "resources.resfiles");
+				XmlNode PriIndexNode = PriConfig.SelectSingleNode("/resources/index")!;
+				XmlAttribute PriStartIndex = PriIndexNode.Attributes!["startIndexAt"]!;
+				PriStartIndex.Value = ResourcesResFile;
+
+				// swap the folder indexer-config to a RESFILES indexer-config.
+				XmlElement FolderIndexerConfigNode = (XmlElement)PriConfig.SelectSingleNode("/resources/index/indexer-config[@type='folder']")!;
+				FolderIndexerConfigNode.SetAttribute("type", "RESFILES");
+				FolderIndexerConfigNode.RemoveAttribute("foldernameAsQualifier");
+				FolderIndexerConfigNode.RemoveAttribute("filenameAsQualifier");
+
+				PriConfig.Save(ResourceConfigFile);
+
+				// generate resources.resfiles
+				IEnumerable<string> Resources = Directory.EnumerateFiles(Path.Combine(OutputPath, BuildResourceSubPath), "*.*", SearchOption.AllDirectories);
+				System.Text.StringBuilder ResourcesList = new System.Text.StringBuilder();
+				foreach (string Resource in Resources)
+				{
+					ResourcesList.AppendLine( Utils.MakePathRelativeTo( Resource, OutputPath ) );
+				}
+				File.WriteAllText(ResourcesResFile, ResourcesList.ToString());
+
+				// remove old Package Resource Index
+				try
+				{
+					FileUtils.ForceDeleteFile(TargetResourceIndexFile);
+				}
+				catch (Exception E)
+				{
+					Logger.LogError("cannot remove old pri file: {TargetResourceIndexFile} - {Message}", TargetResourceIndexFile, E.Message);
+				}
+
+				// Generate the Package Resource Index
+				string ResourceLogFile = Path.Combine(IntermediatePath, "ResIndexLog.xml");
+				string MakePriCommandLine = "new /pr \"" + OutputPath + "\" /cf \"" + ResourceConfigFile + "\" /mn \"" + ManifestTargetPath + "\" /il \"" + ResourceLogFile + "\" /of \"" + TargetResourceIndexFile + "\" /o";
+				if (IdentityName != null)
+				{
+					MakePriCommandLine += " /indexName \"" + IdentityName + "\"";
+				}
+
+				Logger.LogDebug("    generating {Path}", Utils.MakePathRelativeTo(TargetResourceIndexFile, OutputPath!));
+				RunMakePri(MakePriCommandLine);
+				UpdatedFilePaths.Add(TargetResourceIndexFile);
+			}
+
+			// Report if nothing was changed
+			if (!bHadStaleResources && !UpdatedFilePaths.Any())
+			{
+				Logger.LogDebug($"Manifest resource files are up to date");
+			}
+
+			return UpdatedFilePaths;
+		}
+	}
+}
diff --git a/Engine/Source/Programs/UnrealBuildTool/Platform/Windows/VCManifestGenerator.cs b/Engine/Source/Programs/UnrealBuildTool/Platform/Windows/VCManifestGenerator.cs
deleted file mode 100644
