Disable Iris compilation while we fix static analysis warnings.

diff --git a/Engine/Source/Programs/UnrealBuildTool/Configuration/TargetRules.cs b/Engine/Source/Programs/UnrealBuildTool/Configuration/TargetRules.cs
--- a/Engine/Source/Programs/UnrealBuildTool/Configuration/TargetRules.cs
+++ b/Engine/Source/Programs/UnrealBuildTool/Configuration/TargetRules.cs
@@ -773,7 +773,7 @@ namespace UnrealBuildTool
 		[RequiresUniqueBuildEnvironment]
 		[CommandLine("-NoUseIris", Value = "false")]
 		[CommandLine("-UseIris", Value = "true")]
-		public bool bUseIris = true;
+		public bool bUseIris = false;
 
 		/// <summary>
 		/// Whether we are compiling editor code or not. Prefer the more explicit bCompileAgainstEditor instead.
