Fix missing frame pointer when using TSAN

#rnx

diff --git a/Engine/Source/Programs/UnrealBuildTool/Platform/Linux/LinuxToolChain.cs b/Engine/Source/Programs/UnrealBuildTool/Platform/Linux/LinuxToolChain.cs
--- a/Engine/Source/Programs/UnrealBuildTool/Platform/Linux/LinuxToolChain.cs
+++ b/Engine/Source/Programs/UnrealBuildTool/Platform/Linux/LinuxToolChain.cs
@@ -405,7 +405,7 @@ namespace UnrealBuildTool
 			}
 
 			bool bRetainFramePointers = CompileEnvironment.bRetainFramePointers
-				|| Options.HasFlag(ClangToolChainOptions.EnableAddressSanitizer) || Options.HasFlag(ClangToolChainOptions.EnableMemorySanitizer)
+				|| Options.HasFlag(ClangToolChainOptions.EnableAddressSanitizer) || Options.HasFlag(ClangToolChainOptions.EnableMemorySanitizer) || Options.HasFlag(ClangToolChainOptions.EnableThreadSanitizer)
 				|| CompileEnvironment.Configuration == CppConfiguration.Debug;
 
 			if (CompileEnvironment.Configuration == CppConfiguration.Shipping)
