[UBT]
* Fixed bug where name collision happened if both VNI and UHT files were created for same file

Fixes VerseExample UBT exception

diff --git a/Engine/Source/Programs/UnrealBuildTool/Configuration/UEBuildModuleCPP.cs b/Engine/Source/Programs/UnrealBuildTool/Configuration/UEBuildModuleCPP.cs
--- a/Engine/Source/Programs/UnrealBuildTool/Configuration/UEBuildModuleCPP.cs
+++ b/Engine/Source/Programs/UnrealBuildTool/Configuration/UEBuildModuleCPP.cs
@@ -548,17 +548,21 @@ namespace UnrealBuildTool
 				{
 					foreach (string GeneratedDir in GeneratedCppDirectories)
 					{
-						if (Directory.Exists(GeneratedDir))
+						if (!Directory.Exists(GeneratedDir))
 						{
-							string[] Files = Directory.GetFiles(GeneratedDir, "*.gen.cpp");
-							GeneratedFiles.EnsureCapacity(GeneratedFiles.Count + Files.Length);
-							foreach (var File in Files)
-							{
-								// Can't use GetFileNameWithoutAnyExtensions because of the .init.gen.cpp files
-								string FileName = Path.GetFileName(File);
-								FileName = FileName.Substring(0, FileName.Length - ".gen.cpp".Length); 
-								GeneratedFiles.Add(FileName, File);
-							}
+							continue;
+						}
+
+						string Prefix = Path.GetFileName(GeneratedDir) + '/'; // "UHT/" or "VNI/"
+
+						string[] Files = Directory.GetFiles(GeneratedDir, "*.gen.cpp");
+						GeneratedFiles.EnsureCapacity(GeneratedFiles.Count + Files.Length);
+						foreach (var File in Files)
+						{
+							// Can't use GetFileNameWithoutAnyExtensions because of the .init.gen.cpp files
+							string FileName = Path.GetFileName(File);
+							string Key = Prefix + FileName.Substring(0, FileName.Length - ".gen.cpp".Length); 
+							GeneratedFiles.Add(Key, File);
 						}
 					}
 				}
@@ -583,7 +587,9 @@ namespace UnrealBuildTool
 							var ListOfInlinedGenCpps = ModuleCompileEnvironment.MetadataCache.GetListOfInlinedGeneratedCppFiles(CPPFileItem);
 							foreach (string ListOfInlinedGenCppsItem in ListOfInlinedGenCpps)
 							{
-								if (GeneratedFiles.Remove(ListOfInlinedGenCppsItem, out string? FoundGenCppFile))
+								string Prefix = "UHT/";
+								string Key = Prefix + ListOfInlinedGenCppsItem;
+								if (GeneratedFiles.Remove(Key, out string? FoundGenCppFile))
 								{
 									if (!CompileEnvironment.FileInlineGenCPPMap.ContainsKey(CPPFileItem))
 									{
@@ -608,7 +614,12 @@ namespace UnrealBuildTool
 						}
 						foreach (var Name in GeneratedFiles.Keys)
 						{
-							if (CPPFilesLookup.TryGetValue(Name, out FileItem? Item))
+							if (!Name.StartsWith("UHT/"))
+							{
+								continue;
+							}
+							string NameWithoutPrefix = Name.Substring(4);
+							if (CPPFilesLookup.TryGetValue(NameWithoutPrefix, out FileItem? Item))
 							{
 								Logger.LogWarning("'{0}' .gen.cpp not inlined. Add '#include UE_INLINE_GENERATED_CPP_BY_NAME({1})'", Item.Name, Name);
 							}
