Enable physics macro deprecation in all modules for non-engine modules that relied on inheriting from dependencies on engine modules.

diff --git a/Engine/Source/Programs/UnrealBuildTool/Configuration/ModuleRules.cs b/Engine/Source/Programs/UnrealBuildTool/Configuration/ModuleRules.cs
--- a/Engine/Source/Programs/UnrealBuildTool/Configuration/ModuleRules.cs
+++ b/Engine/Source/Programs/UnrealBuildTool/Configuration/ModuleRules.cs
@@ -1399,33 +1399,31 @@ namespace UnrealBuildTool
 
 			PublicDefinitions.Add("WITH_CLOTH_COLLISION_DETECTION=1");
 
-			if(!bTreatAsEngineModule)
-			{
-				// Non-engine modules may still be relying on appropriate definitions for physics.
-				// Nothing in engine should use these anymore as they were all deprecated and 
-				// assumed to be in the following configuration from 5.1
+			// Modules may still be relying on appropriate definitions for physics.
+			// Nothing in engine should use these anymore as they were all deprecated and 
+			// assumed to be in the following configuration from 5.1, this will cause
+			// deprecation warning to fire in any module still relying on these macros
 
-				Func<string, string, string, string> GetDeprecatedPhysicsMacro = (string Macro, string Value, string Version) =>
-				{
-					return Macro + "=UE_DEPRECATED_MACRO(" + Version + ", \"" + Macro + " is deprecated and should always be considered " + Value + ".\") " + Value;
-				};
-
-				PublicDefinitions.AddRange(
-					new string[]{
-						GetDeprecatedPhysicsMacro("INCLUDE_CHAOS", "1", "5.1"),
-						GetDeprecatedPhysicsMacro("WITH_CHAOS", "1", "5.1"),
-						GetDeprecatedPhysicsMacro("WITH_CHAOS_CLOTHING", "1", "5.1"),
-						GetDeprecatedPhysicsMacro("WITH_CHAOS_NEEDS_TO_BE_FIXED", "1", "5.1"),
-						GetDeprecatedPhysicsMacro("WITH_PHYSX", "1", "5.1"),
-						GetDeprecatedPhysicsMacro("WITH_PHYSX_COOKING", "0", "5.1"),
-						GetDeprecatedPhysicsMacro("PHYSICS_INTERFACE_PHYSX", "0", "5.1"),
-						GetDeprecatedPhysicsMacro("WITH_APEX", "0", "5.1"),
-						GetDeprecatedPhysicsMacro("WITH_APEX_CLOTHING", "0", "5.1"),
-						GetDeprecatedPhysicsMacro("WITH_NVCLOTH", "0", "5.1"),
-						GetDeprecatedPhysicsMacro("WITH_IMMEDIATE_PHYSX", "0", "5.1"),
-						GetDeprecatedPhysicsMacro("WITH_CUSTOM_SQ_STRUCTURE", "0", "5.1")
-					});
-			}
+			Func<string, string, string, string> GetDeprecatedPhysicsMacro = (string Macro, string Value, string Version) =>
+			{
+				return Macro + "=UE_DEPRECATED_MACRO(" + Version + ", \"" + Macro + " is deprecated and should always be considered " + Value + ".\") " + Value;
+			};
+
+			PublicDefinitions.AddRange(
+				new string[]{
+					GetDeprecatedPhysicsMacro("INCLUDE_CHAOS", "1", "5.1"),
+					GetDeprecatedPhysicsMacro("WITH_CHAOS", "1", "5.1"),
+					GetDeprecatedPhysicsMacro("WITH_CHAOS_CLOTHING", "1", "5.1"),
+					GetDeprecatedPhysicsMacro("WITH_CHAOS_NEEDS_TO_BE_FIXED", "1", "5.1"),
+					GetDeprecatedPhysicsMacro("WITH_PHYSX", "1", "5.1"),
+					GetDeprecatedPhysicsMacro("WITH_PHYSX_COOKING", "0", "5.1"),
+					GetDeprecatedPhysicsMacro("PHYSICS_INTERFACE_PHYSX", "0", "5.1"),
+					GetDeprecatedPhysicsMacro("WITH_APEX", "0", "5.1"),
+					GetDeprecatedPhysicsMacro("WITH_APEX_CLOTHING", "0", "5.1"),
+					GetDeprecatedPhysicsMacro("WITH_NVCLOTH", "0", "5.1"),
+					GetDeprecatedPhysicsMacro("WITH_IMMEDIATE_PHYSX", "0", "5.1"),
+					GetDeprecatedPhysicsMacro("WITH_CUSTOM_SQ_STRUCTURE", "0", "5.1")
+				});
 		}
 
 		/// <summary>
