potential fix for PGO data directory

diff --git a/Engine/Source/Programs/UnrealBuildTool/ToolChain/ClangToolChain.cs b/Engine/Source/Programs/UnrealBuildTool/ToolChain/ClangToolChain.cs
--- a/Engine/Source/Programs/UnrealBuildTool/ToolChain/ClangToolChain.cs
+++ b/Engine/Source/Programs/UnrealBuildTool/ToolChain/ClangToolChain.cs
@@ -975,6 +975,11 @@ namespace UnrealBuildTool
 			Arguments.Add(GetOutputFileArgument(OutputFile));
 		}
 
+		protected virtual List<string> ExpandResponseFileContents(List<string> ResponseFileContents)
+        {
+			return ResponseFileContents.Select(x => Utils.ExpandVariables(x)).ToList();
+		}
+
 		protected virtual Action CompileCPPFile(CppCompileEnvironment CompileEnvironment, FileItem SourceFile, DirectoryReference OutputDir, string ModuleName, IActionGraphBuilder Graph, IReadOnlyCollection<string> GlobalArguments, CPPOutput Result)
 		{
 			Action CompileAction = Graph.CreateAction(ActionType.Compile);
@@ -992,7 +997,7 @@ namespace UnrealBuildTool
 			List<string> ResponseFileContents = new();
 			ResponseFileContents.AddRange(GlobalArguments);
 			ResponseFileContents.AddRange(FileArguments);
-			ResponseFileContents = ResponseFileContents.Select(x => Utils.ExpandVariables(x)).ToList();
+			ResponseFileContents = ExpandResponseFileContents(ResponseFileContents);
 
 			if (RuntimePlatform.IsWindows)
 			{
