Add support for sorting system includes.

diff --git a/Engine/Source/Programs/UnrealBuildTool/Modes/FixIncludePathsMode.cs b/Engine/Source/Programs/UnrealBuildTool/Modes/FixIncludePathsMode.cs
--- a/Engine/Source/Programs/UnrealBuildTool/Modes/FixIncludePathsMode.cs
+++ b/Engine/Source/Programs/UnrealBuildTool/Modes/FixIncludePathsMode.cs
@@ -21,7 +21,7 @@ namespace UnrealBuildTool
 		/// <summary>
 		/// Regex that matches #include statements.
 		/// </summary>
-		static readonly Regex IncludeRegex = new Regex("^[ \t]*#[ \t]*include[ \t]*[\"](?<HeaderFile>[^\"]*)[\"]", RegexOptions.Compiled | RegexOptions.Singleline | RegexOptions.ExplicitCapture);
+		static readonly Regex IncludeRegex = new Regex("^[ \t]*#[ \t]*include[ \t]*[<\"](?<HeaderFile>[^\">]*)[\">]", RegexOptions.Compiled | RegexOptions.Singleline | RegexOptions.ExplicitCapture);
 
 		static readonly string UnrealRootDirectory = Unreal.RootDirectory.FullName.Replace('\\', '/');
 		static readonly string[] PreferredPaths = { "/Public/", "/Private/", "/Classes/", "/Internal/", "/UHT/", "/VNI/" };
@@ -323,6 +323,16 @@ namespace UnrealBuildTool
 					return 1;
 				}
 
+				// system includes
+				if (x.Contains('<'))
+				{
+					return 1;
+				}
+				if (y.Contains('<'))
+				{
+					return -1;
+				}
+
 				// generated header
 				if (x.Contains(".generated.h"))
 				{
@@ -332,6 +342,7 @@ namespace UnrealBuildTool
 				{
 					return -1;
 				}
+
 				return string.Compare(x, y);
 			}
 		}
