- Remove WITH_SERVER_CODE_UNTRUSTED define to avoid confusion

diff --git a/Engine/Source/Programs/UnrealBuildTool/Configuration/TargetRules.cs b/Engine/Source/Programs/UnrealBuildTool/Configuration/TargetRules.cs
--- a/Engine/Source/Programs/UnrealBuildTool/Configuration/TargetRules.cs
+++ b/Engine/Source/Programs/UnrealBuildTool/Configuration/TargetRules.cs
@@ -1015,16 +1015,6 @@ namespace UnrealBuildTool
 			set { bWithServerCodeTrustedPrivate = value; }
 		}
 
-		/// <summary>
-		/// Compile untrusted-server-only code.
-		/// </summary>
-		[RequiresUniqueBuildEnvironment]
-		public bool bWithServerCodeUntrusted
-		{
-			get { return bWithServerCode && !bWithServerCodeTrustedPrivate; }
-			set { bWithServerCodeTrustedPrivate = !value; }
-		}
-
 		[CommandLine("-TrustedServer", Value = "true")]
 		[CommandLine("-NoTrustedServer", Value = "false")]
 		private bool bWithServerCodeTrustedPrivate = true;
@@ -2949,11 +2939,6 @@ namespace UnrealBuildTool
 			get { return Inner.bWithServerCodeTrusted; }
 		}
 
-		public bool bWithServerCodeUntrusted
-		{
-			get { return Inner.bWithServerCodeUntrusted; }
-		}
-
 		public bool bFNameOutlineNumber
 		{
 			get { return Inner.bFNameOutlineNumber;  }
diff --git a/Engine/Source/Programs/UnrealBuildTool/Configuration/UEBuildTarget.cs b/Engine/Source/Programs/UnrealBuildTool/Configuration/UEBuildTarget.cs
--- a/Engine/Source/Programs/UnrealBuildTool/Configuration/UEBuildTarget.cs
+++ b/Engine/Source/Programs/UnrealBuildTool/Configuration/UEBuildTarget.cs
@@ -4174,7 +4174,6 @@ namespace UnrealBuildTool
 			// Check if server-only code should be compiled out.
 			GlobalCompileEnvironment.Definitions.Add(String.Format("WITH_SERVER_CODE={0}", Rules.bWithServerCode ? 1 : 0));
 			GlobalCompileEnvironment.Definitions.Add(String.Format("WITH_SERVER_CODE_TRUSTED={0}", Rules.bWithServerCodeTrusted ? 1 : 0));
-			GlobalCompileEnvironment.Definitions.Add(String.Format("WITH_SERVER_CODE_UNTRUSTED={0}", Rules.bWithServerCodeUntrusted ? 1 : 0));
 
 			GlobalCompileEnvironment.Definitions.Add(String.Format("UE_FNAME_OUTLINE_NUMBER={0}", Rules.bFNameOutlineNumber? 1 : 0));
 
