Server mode for UBT for use by VS Code extension. Provides compile configuration on demand, responding to commands over stdin in json format.
VS Code project file generator argument to skip compile commands json files.
VS Code project file generator adds VS Code extension id as configuration provider to configurations.
Add optional implementable funcion for toolchains to return the path to default C++ compiler.
Add rootdirectory= argument to provide root directory if UBT is published to a different directory than standard.

diff --git a/Engine/Source/Programs/Shared/EpicGames.Build/Unreal.cs b/Engine/Source/Programs/Shared/EpicGames.Build/Unreal.cs
--- a/Engine/Source/Programs/Shared/EpicGames.Build/Unreal.cs
+++ b/Engine/Source/Programs/Shared/EpicGames.Build/Unreal.cs
@@ -20,6 +20,13 @@ namespace UnrealBuildBase
 				return DirectoryReference.FindCorrectCase(LocationOverride.RootDirectory);
 			}
 
+			string? OverrideArg = Environment.GetCommandLineArgs().FirstOrDefault(x => x?.StartsWith("-rootdirectory=") ?? false, null);
+			if (OverrideArg != null)
+			{
+				string[] Parts = OverrideArg.Split('=', 2);
+				return new DirectoryReference(Path.GetFullPath(Parts[1]));
+			}
+
 			// This base library may be used - and so be launched - from more than one location (at time of writing, UnrealBuildTool and AutomationTool)
 			// Programs that use this assembly must be located under "Engine/Binaries/DotNET" and so we look for that sequence of directories in that path of the executing assembly
 
diff --git a/Engine/Source/Programs/UnrealBuildTool/Configuration/UEBuildTarget.cs b/Engine/Source/Programs/UnrealBuildTool/Configuration/UEBuildTarget.cs
--- a/Engine/Source/Programs/UnrealBuildTool/Configuration/UEBuildTarget.cs
+++ b/Engine/Source/Programs/UnrealBuildTool/Configuration/UEBuildTarget.cs
@@ -3874,7 +3874,7 @@ namespace UnrealBuildTool
 		/// <summary>
 		/// Sets up the global compile and link environment for the target.
 		/// </summary>
-		private void SetupGlobalEnvironment(UEToolChain ToolChain, CppCompileEnvironment GlobalCompileEnvironment, LinkEnvironment GlobalLinkEnvironment)
+		public void SetupGlobalEnvironment(UEToolChain ToolChain, CppCompileEnvironment GlobalCompileEnvironment, LinkEnvironment GlobalLinkEnvironment)
 		{
 			UEBuildPlatform BuildPlatform = UEBuildPlatform.GetBuildPlatform(Platform);
 
@@ -4308,7 +4308,7 @@ namespace UnrealBuildTool
 			BuildPlatform.SetUpConfigurationEnvironment(Rules, GlobalCompileEnvironment, GlobalLinkEnvironment);
 		}
 
-		static CppConfiguration GetCppConfiguration(UnrealTargetConfiguration Configuration)
+		public static CppConfiguration GetCppConfiguration(UnrealTargetConfiguration Configuration)
 		{
 			switch (Configuration)
 			{
diff --git a/Engine/Source/Programs/UnrealBuildTool/Modes/GenerateProjectFilesMode.cs b/Engine/Source/Programs/UnrealBuildTool/Modes/GenerateProjectFilesMode.cs
--- a/Engine/Source/Programs/UnrealBuildTool/Modes/GenerateProjectFilesMode.cs
+++ b/Engine/Source/Programs/UnrealBuildTool/Modes/GenerateProjectFilesMode.cs
@@ -233,6 +233,7 @@ namespace UnrealBuildTool
 			foreach(ProjectFileGenerator Generator in Generators)
 			{
 				ProjectFileGenerator.Current = Generator;
+				Arguments.ApplyTo(Generator);
 				bool bGenerateSuccess = Generator.GenerateProjectFiles(PlatformProjectGenerators, Arguments.GetRawArray(), Logger);
 				ProjectFileGenerator.Current = null;
 
diff --git a/Engine/Source/Programs/UnrealBuildTool/Modes/ServerMode.cs b/Engine/Source/Programs/UnrealBuildTool/Modes/ServerMode.cs
new file mode 100644
--- /dev/null
+++ b/Engine/Source/Programs/UnrealBuildTool/Modes/ServerMode.cs
@@ -0,0 +1,516 @@
+using EpicGames.Core;
+using Microsoft.Extensions.Logging;
+using System;
+using System.Linq;
+using System.Text.Json;
+using System.Text.Json.Serialization;
+using System.Collections.Generic;
+using OpenTracing.Util;
+using UnrealBuildTool;
+using UnrealBuildBase;
+
+namespace UnrealBuildTool
+{
+	// TODO: Interface for commands/results?
+	enum CommandType
+	{
+		Quit,
+		DiscoverTargets,
+		SetTarget,
+		QueryFile,
+		GetBrowseConfiguration,
+		GetCompileSettings,
+	}
+
+	interface ICommand
+	{
+		public CommandType Type { get; }
+		public long Sequence { get; set; }
+	}
+	interface ICommandResponse
+	{
+		public CommandType Type { get; }
+		public long Sequence { get; set; }
+	}
+
+	struct Command : ICommand
+	{
+		[JsonConverter(typeof(JsonStringEnumConverter))]
+		public CommandType Type { get; set; }
+		public long Sequence { get; set; }
+	};
+
+	struct SetTargetCommand : ICommand
+	{
+		[JsonConverter(typeof(JsonStringEnumConverter))]
+		public CommandType Type { get; set; }
+		public long Sequence { get; set; }
+		public string Target { get; set; }
+		public string Platform { get; set; }
+		public string Configuration { get; set; }
+	}
+
+	struct QueryFileCommand : ICommand
+	{
+		[JsonConverter(typeof(JsonStringEnumConverter))]
+		public CommandType Type { get; set; }
+		public long Sequence { get; set; }
+		public string AbsolutePath { get; set; }
+	}
+
+	struct GetBrowseConfigurationCommand : ICommand
+	{
+		[JsonConverter(typeof(JsonStringEnumConverter))]
+		public CommandType Type { get; set; }
+		public long Sequence { get; set; }
+		public string? AbsolutePath { get; set; }
+	}
+
+	struct GetCompileSettingsCommand : ICommand
+	{
+		[JsonConverter(typeof(JsonStringEnumConverter))]
+		public CommandType Type { get; set; }
+		public long Sequence { get; set; }
+		public List<string> AbsolutePaths { get; set; }
+	}
+
+	struct DiscoverTargetsResponse : ICommandResponse
+	{
+		[JsonConverter(typeof(JsonStringEnumConverter))]
+		public CommandType Type { get { return CommandType.DiscoverTargets; } }
+		public long Sequence { get; set; }
+		public List<string> Configurations { get; set; }
+		public List<string> Platforms { get; set; }
+		public List<string> Targets { get; set; }
+
+		public string DefaultConfiguration { get; set; }
+		public string DefaultPlatform { get; set; }
+		public string DefaultTarget { get; set; }
+	}
+
+	struct SetTargetResponse : ICommandResponse
+	{
+		[JsonConverter(typeof(JsonStringEnumConverter))]
+		public CommandType Type { get { return CommandType.SetTarget; } }
+		public long Sequence { get; set; }
+	}
+	struct QueryFileResponse : ICommandResponse 
+	{
+		[JsonConverter(typeof(JsonStringEnumConverter))]
+		public CommandType Type { get { return CommandType.QueryFile; } }
+		public long Sequence { get; set; }
+		public bool Found { get; set; }
+	}
+
+	struct GetBrowseConfigurationResponse : ICommandResponse
+	{
+		[JsonConverter(typeof(JsonStringEnumConverter))]
+		public CommandType Type { get { return CommandType.GetBrowseConfiguration; } }
+		public long Sequence { get; set; }
+		public bool Success { get; set; }
+		
+		public List<string> Paths { get; set; }
+		public string? Standard { get; set; }
+		public string? WindowsSdkVersion { get; set; } 
+	}
+
+	class CompileSettings
+	{
+		public List<string> IncludePaths { get; set; } = new();
+		public List<string> Defines { get; set; } = new();
+		public string? Standard { get; set; }
+		public List<string> ForcedIncludes { get; set; } = new();
+		public string? CompilerPath { get; set; }
+		public List<string> CompilerArgs { get; set; } = new();
+		public string? WindowsSdkVersion { get; set; }
+
+	}
+
+	class GetCompileSettingsResponse : ICommandResponse
+	{
+		[JsonConverter(typeof(JsonStringEnumConverter))]
+		public CommandType Type { get { return CommandType.GetCompileSettings; } }
+		public long Sequence { get; set; }
+		public List<CompileSettings?> Settings { get; set; } = new();
+	}
+
+	internal class TargetIntellisenseInfo
+	{
+		public Dictionary<UEBuildModule, CompileSettings> ModuleToCompileSettings = new();
+		public Dictionary<DirectoryReference, UEBuildModule> DirToModule = new();
+
+		public UEBuildModule? FindModuleForFile(FileReference File)
+		{
+			DirectoryReference? Dir = File.Directory;
+			while (Dir != null)
+			{
+				if (DirToModule.TryGetValue(Dir, out UEBuildModule? Module))
+				{
+					return Module;
+				}
+				Dir = Dir.ParentDirectory;
+			}
+			return null;
+		}
+		public UEBuildModule? FindModuleForDirectory(DirectoryReference Directory)
+		{
+			DirectoryReference? Dir = Directory;
+			while (Dir != null)
+			{
+				if (DirToModule.TryGetValue(Dir, out UEBuildModule? Module))
+				{
+					return Module;
+				}
+				Dir = Dir.ParentDirectory;
+			}
+			return null;
+		}
+	}
+
+	[ToolMode("Server", ToolModeOptions.BuildPlatforms | ToolModeOptions.XmlConfig | ToolModeOptions.UseStartupTraceListener)]
+	class ServerMode : ToolMode
+	{
+		[CommandLine("-LogDirectory=")]
+		public DirectoryReference? LogDirectory = null;
+
+		private bool KeepRunning { get; set; } = true;
+
+		private BuildConfiguration BuildConfiguration = new();
+		private UEBuildTarget? CurrentTarget = null;
+		private TargetIntellisenseInfo? CurrentTargetIntellisenseInfo = null;
+		private GetBrowseConfigurationResponse CurrentBrowseConfiguration;
+
+		public override int Execute(CommandLineArguments Arguments, ILogger Logger)
+		{
+			Arguments.ApplyTo(this);
+
+			if (LogDirectory == null)
+			{
+				LogDirectory = DirectoryReference.Combine(UnrealBuildTool.EngineProgramSavedDirectory, "UnrealBuildTool");
+			}
+
+			FileReference LogFile = FileReference.Combine(LogDirectory, "Log_Server.txt");
+			Log.AddFileWriter("DefaultLogTraceListener", LogFile);
+			Logger.LogInformation("Server mode started");
+
+			XmlConfig.ApplyTo(BuildConfiguration);
+			Arguments.ApplyTo(BuildConfiguration);
+
+			ProjectFileGenerator.bGenerateProjectFiles = true;
+
+			while (Arguments.TryGetPositionalArgument(out string? PositionalArg))
+			{
+				HandleCommand(PositionalArg!, Logger);
+			}
+
+			while (KeepRunning)
+			{
+				var Line = Console.ReadLine();
+				if (Line == null) { continue; }
+				HandleCommand(Line, Logger);
+			}
+			return 0;
+		}
+
+		private void HandleCommand(string CommandString, ILogger Logger)
+		{
+			var jsonOptions = new JsonSerializerOptions
+			{
+				PropertyNameCaseInsensitive = true
+			};
+			var responseOptions = new JsonSerializerOptions
+			{
+				PropertyNamingPolicy = JsonNamingPolicy.CamelCase
+			};
+			try
+			{
+				Command? MaybeCommand = JsonSerializer.Deserialize<Command>(CommandString, jsonOptions);
+				if (MaybeCommand is null)
+				{
+					Console.Error.WriteLine($"Failed to parse command {CommandString}");
+					return;
+				}
+				Logger.LogInformation($"Got command {CommandString}");
+				var command = MaybeCommand.Value;
+				switch (command.Type)
+				{
+					case CommandType.Quit:
+						KeepRunning = false;
+						break;
+					case CommandType.DiscoverTargets:
+						{
+							var response = Handle_DiscoverTargets(Logger);
+							response.Sequence = command.Sequence;
+							Console.WriteLine(JsonSerializer.Serialize(response, responseOptions));
+						}
+						break;
+					case CommandType.SetTarget:
+						{
+							var response = Handle_SetTarget(JsonSerializer.Deserialize<SetTargetCommand>(CommandString, jsonOptions), Logger);
+							response.Sequence = command.Sequence;
+							Console.WriteLine(JsonSerializer.Serialize(response, responseOptions));
+						}
+						break;
+					case CommandType.QueryFile:
+						{
+							var response = Handle_QueryFile(JsonSerializer.Deserialize<QueryFileCommand>(CommandString, jsonOptions), Logger);
+							response.Sequence = command.Sequence;
+							Console.WriteLine(JsonSerializer.Serialize(response, responseOptions));
+						}
+						break;
+					case CommandType.GetBrowseConfiguration:
+						{
+							var response = Handle_GetBrowseConfiguration(JsonSerializer.Deserialize<GetBrowseConfigurationCommand>(CommandString, jsonOptions), Logger);
+							response.Sequence = command.Sequence;
+							Console.WriteLine(JsonSerializer.Serialize(response, responseOptions));
+						}
+						break;
+					case CommandType.GetCompileSettings:
+						{
+							var response = Handle_GetCompileSettings(JsonSerializer.Deserialize<GetCompileSettingsCommand>(CommandString, jsonOptions), Logger);
+							response.Sequence = command.Sequence;
+							Console.WriteLine(JsonSerializer.Serialize(response, responseOptions));
+						}
+						break;
+				}
+			}
+			catch
+			{
+				Console.Error.WriteLine($"Failed to parse command {CommandString}");
+			}
+		}
+
+		// TODO: make async 
+		private DiscoverTargetsResponse Handle_DiscoverTargets(ILogger Logger)
+		{
+			List<UnrealTargetPlatform> Platforms = new();
+			foreach (UnrealTargetPlatform Platform in UnrealTargetPlatform.GetValidPlatforms())
+			{
+				// Only include desktop platforms for now 
+				if (ProjectFileGenerator.IsValidDesktopPlatform(Platform))
+				{
+					// If there is a build platform present, add it to the SupportedPlatforms list
+					UEBuildPlatform? BuildPlatform;
+					if (UEBuildPlatform.TryGetBuildPlatform(Platform, out BuildPlatform))
+					{
+						if (InstalledPlatformInfo.IsValidPlatform(Platform, EProjectType.Code))
+						{
+							Platforms.Add(Platform);
+						}
+					}
+				}
+			}
+
+			List<UnrealTargetConfiguration> AllowedTargetConfigurations = new List<UnrealTargetConfiguration>();
+			AllowedTargetConfigurations = Enum.GetValues(typeof(UnrealTargetConfiguration)).Cast<UnrealTargetConfiguration>().ToList();
+
+			List<string> Configurations = new();
+			foreach (UnrealTargetConfiguration CurConfiguration in AllowedTargetConfigurations)
+			{
+				if (CurConfiguration != UnrealTargetConfiguration.Unknown)
+				{
+					if (InstalledPlatformInfo.IsValidConfiguration(CurConfiguration, EProjectType.Code))
+					{
+						Configurations.Add(CurConfiguration.ToString());
+					}
+				}
+			}
+
+			List<FileReference> Projects = NativeProjects.EnumerateProjectFiles(Logger).ToList();
+			List<FileReference> AllTargetFiles = ProjectFileGenerator.DiscoverTargets(Projects, Logger, null, Platforms, bIncludeEngineSource: true, bIncludeTempTargets: false);
+			Projects.RemoveAll(x => !AllTargetFiles.Any(y => y.IsUnderDirectory(x.Directory)));
+
+			List<string> TargetNames = new();
+			foreach (FileReference TargetFilePath in AllTargetFiles)
+			{
+				string TargetName = TargetFilePath.GetFileNameWithoutAnyExtensions();
+				TargetNames.Add(TargetName);
+			}
+
+			// TODO: Handle 0 target names
+
+			var Reply = new DiscoverTargetsResponse
+			{
+				Configurations = Configurations,
+				Platforms = Platforms.Select(x => x.ToString()).ToList(),
+				Targets = TargetNames,
+
+				// TODO: Allow selecting default from project file 
+				DefaultTarget = TargetNames.FirstOrDefault(x => x == "UnrealEditor", TargetNames[0]),
+				DefaultPlatform = Platforms[0].ToString(),
+				DefaultConfiguration = UnrealTargetConfiguration.Development.ToString(),
+			};
+			return Reply;
+		}
+
+		private SetTargetResponse Handle_SetTarget(SetTargetCommand command, ILogger Logger)
+		{
+			CommandLineArguments Args = new CommandLineArguments(new string[] { command.Target, command.Configuration, command.Platform });
+			List<TargetDescriptor> TargetDescriptors = new();
+
+			TargetDescriptor.ParseSingleCommandLine(Args, false, false, false, TargetDescriptors, Logger);
+			if (TargetDescriptors.Count == 0)
+			{
+				// TOOD: Error 
+				return new SetTargetResponse { };
+			}
+
+			HashSet<string> BrowseConfigurationFolders = new HashSet<string>();
+			try
+			{
+				using (GlobalTracer.Instance.BuildSpan("UEBuildTarget.Create()").StartActive())
+				{
+					bool bUsePrecompiled = false;
+					CurrentTarget = UEBuildTarget.Create(TargetDescriptors[0], false, false, bUsePrecompiled, Logger);
+				}
+
+				CurrentTargetIntellisenseInfo = new TargetIntellisenseInfo();
+				CurrentBrowseConfiguration = new GetBrowseConfigurationResponse { Success = true };
+
+				// Partially duplicated from UEBuildTarget.Build because we just want to get C++ compile actions without running UHT
+				// or generating link actions / full dependency graph 
+				CppConfiguration CppConfiguration = UEBuildTarget.GetCppConfiguration(CurrentTarget.Configuration);
+				SourceFileMetadataCache MetadataCache = SourceFileMetadataCache.CreateHierarchy(null, Logger);
+				CppCompileEnvironment GlobalCompileEnvironment = new CppCompileEnvironment(CurrentTarget.Platform, CppConfiguration, CurrentTarget.Architecture, MetadataCache);
+				LinkEnvironment GlobalLinkEnvironment = new LinkEnvironment(GlobalCompileEnvironment.Platform, GlobalCompileEnvironment.Configuration, GlobalCompileEnvironment.Architecture);
+
+				UEToolChain TargetToolChain = CurrentTarget.CreateToolchain(CurrentTarget.Platform);
+				TargetToolChain.SetEnvironmentVariables();
+				CurrentTarget.SetupGlobalEnvironment(TargetToolChain, GlobalCompileEnvironment, GlobalLinkEnvironment);
+
+				// TODO: For installed builds, filter out all the binaries that aren't in mods
+				foreach (UEBuildBinary Binary in CurrentTarget.Binaries)
+				{
+					HashSet<UEBuildModule> LinkEnvironmentVisitedModules = new HashSet<UEBuildModule>();
+					CppCompileEnvironment BinaryCompileEnvironment = Binary.CreateBinaryCompileEnvironment(GlobalCompileEnvironment);
+					CurrentBrowseConfiguration.Standard = BinaryCompileEnvironment.CppStandard.ToString();
+					CurrentBrowseConfiguration.WindowsSdkVersion = CurrentTarget.Rules.WindowsPlatform.WindowsSdkVersion;
+
+					foreach (UEBuildModuleCPP Module in Binary.Modules.OfType<UEBuildModuleCPP>())
+					{
+						if (Module.Binary != null && Module.Binary != Binary) 
+						{
+							continue;
+						}
+
+						CppCompileEnvironment ModuleCompileEnvironment = Module.CreateModuleCompileEnvironment(CurrentTarget.Rules, BinaryCompileEnvironment);
+						foreach (DirectoryReference Dir in Module.ModuleDirectories)
+						{
+							BrowseConfigurationFolders.Add(Dir.ToString());
+							CurrentTargetIntellisenseInfo.DirToModule.Add(Dir, Module);
+						}
+						if (Module.GeneratedCodeDirectory != null)
+						{
+							CurrentTargetIntellisenseInfo.DirToModule.Add(Module.GeneratedCodeDirectory, Module);
+						}
+
+						foreach (DirectoryReference Dir in ModuleCompileEnvironment.SystemIncludePaths)
+						{
+							BrowseConfigurationFolders.Add(Dir.ToString());
+						}
+						foreach (DirectoryReference Dir in ModuleCompileEnvironment.UserIncludePaths)
+						{
+							BrowseConfigurationFolders.Add(Dir.ToString());
+						}
+
+						CompileSettings Settings = new CompileSettings();
+						Settings.IncludePaths.AddRange(ModuleCompileEnvironment.SystemIncludePaths.Select(x => x.ToString()));
+						Settings.IncludePaths.AddRange(ModuleCompileEnvironment.UserIncludePaths.Select(x => x.ToString()));
+						Settings.Defines = ModuleCompileEnvironment.Definitions;
+						Settings.Standard = ModuleCompileEnvironment.CppStandard.ToString();
+						Settings.ForcedIncludes = ModuleCompileEnvironment.ForceIncludeFiles.Select(x => x.ToString()).ToList();
+						Settings.CompilerPath = TargetToolChain.GetCppCompilerPath()?.ToString();
+						Settings.WindowsSdkVersion = CurrentTarget.Rules.WindowsPlatform.WindowsSdkVersion;
+						CurrentTargetIntellisenseInfo.ModuleToCompileSettings.Add(Module, Settings);
+					}
+				}
+			}
+			catch (Exception e)
+			{
+				Logger.LogError("Caught exception setting up target: {0}", e);
+			}
+
+			CurrentBrowseConfiguration.Paths = BrowseConfigurationFolders.ToList();
+
+			var Reply = new SetTargetResponse
+			{
+			};
+			return Reply;
+		}
+
+		private QueryFileResponse Handle_QueryFile(QueryFileCommand Command, ILogger Logger)
+		{
+			var FailResponse = new QueryFileResponse { Found = false };
+			if (CurrentTargetIntellisenseInfo == null )
+			{
+				return FailResponse;
+			}
+			try 
+			{
+				var Path = System.IO.Path.GetFullPath(Command.AbsolutePath); 
+				if (CurrentTargetIntellisenseInfo.FindModuleForFile(new FileReference(Path)) != null)
+				{
+					return new QueryFileResponse { Found = true };
+				}
+				return FailResponse;
+			}
+			catch (Exception )
+			{
+				return FailResponse;
+			}
+		}
+		private GetBrowseConfigurationResponse Handle_GetBrowseConfiguration(GetBrowseConfigurationCommand Command, ILogger Logger)
+		{
+			GetBrowseConfigurationResponse Response = new GetBrowseConfigurationResponse { Success = false, Paths = new List<string>() };
+			if (CurrentTargetIntellisenseInfo == null)
+			{
+				return Response;
+			}
+			if (Command.AbsolutePath != null)
+			{
+				var Path = System.IO.Path.GetFullPath(Command.AbsolutePath); 
+				UEBuildModule? Module = CurrentTargetIntellisenseInfo.FindModuleForDirectory(new DirectoryReference(Path));
+				if (Module != null)
+				{
+					Response.Success = true;
+					Response.Paths.AddRange(Module.ModuleDirectories.Select(x => x.ToString()));
+					if (CurrentTargetIntellisenseInfo.ModuleToCompileSettings.TryGetValue(Module, out var Settings))
+					{
+						Response.Paths.AddRange(Settings.IncludePaths);
+						Response.WindowsSdkVersion = Settings.WindowsSdkVersion;
+						Response.Standard = Settings.Standard;
+					}
+				}
+				return Response;
+			}
+
+			Response = CurrentBrowseConfiguration;
+			return Response;
+		}
+
+		private GetCompileSettingsResponse Handle_GetCompileSettings(GetCompileSettingsCommand Command, ILogger Logger)
+		{
+			GetCompileSettingsResponse FailResponse = new GetCompileSettingsResponse { }; 
+			try 
+			{
+				GetCompileSettingsResponse Response = new GetCompileSettingsResponse { };
+				foreach (string AbsolutePath in Command.AbsolutePaths)
+				{
+					CompileSettings? Settings = null;
+					var Path = System.IO.Path.GetFullPath(AbsolutePath); // convert absolute path to dos? 
+					UEBuildModule? Module = CurrentTargetIntellisenseInfo!.FindModuleForFile(new FileReference(Path));
+					if (Module is not null)
+					{
+						CurrentTargetIntellisenseInfo!.ModuleToCompileSettings.TryGetValue(Module, out Settings);
+					}
+					Response.Settings.Add(Settings);
+				}
+				return Response;
+			}
+			catch (Exception)
+			{
+				return FailResponse;
+			}
+		}
+	}
+}
diff --git a/Engine/Source/Programs/UnrealBuildTool/Platform/Windows/VCToolChain.cs b/Engine/Source/Programs/UnrealBuildTool/Platform/Windows/VCToolChain.cs
--- a/Engine/Source/Programs/UnrealBuildTool/Platform/Windows/VCToolChain.cs
+++ b/Engine/Source/Programs/UnrealBuildTool/Platform/Windows/VCToolChain.cs
@@ -54,6 +54,11 @@ namespace UnrealBuildTool
 				}
 			}
 		}
+		
+		public override FileReference? GetCppCompilerPath()
+		{
+			return EnvVars.CompilerPath;
+		}
 
 		/// <summary>
 		/// Prepares the environment for building
diff --git a/Engine/Source/Programs/UnrealBuildTool/ProjectFiles/ProjectFileGenerator.cs b/Engine/Source/Programs/UnrealBuildTool/ProjectFiles/ProjectFileGenerator.cs
--- a/Engine/Source/Programs/UnrealBuildTool/ProjectFiles/ProjectFileGenerator.cs
+++ b/Engine/Source/Programs/UnrealBuildTool/ProjectFiles/ProjectFileGenerator.cs
@@ -1,4 +1,4 @@
-﻿// Copyright Epic Games, Inc. All Rights Reserved.
+// Copyright Epic Games, Inc. All Rights Reserved.
 
 using System;
 using System.Collections.Generic;
@@ -314,7 +314,7 @@ namespace UnrealBuildTool
 		/// Whether to include temporary targets generated by UAT to support content only projects with non-default settings.
 		/// </summary>
 		[XmlConfigFile]
-		bool bIncludeTempTargets = false;
+		protected bool bIncludeTempTargets = false;
 
 		/// <summary>
 		/// True if we should reflect "Source" sub-directories on disk in the primary project as project directories.
@@ -864,7 +864,7 @@ namespace UnrealBuildTool
 
 			// Find all of the target files.  This will filter out any modules or targets that don't
 			// belong to platforms we're generating project files for.
-			List<FileReference> AllTargetFiles = DiscoverTargets(AllGameProjects, Logger);
+			List<FileReference> AllTargetFiles = DiscoverTargets(AllGameProjects, Logger, OnlyGameProject, SupportedPlatforms, bIncludeEngineSource, bIncludeTempTargets);
 
 			// Sort the targets by name. When we have multiple targets of a given type for a project, we'll use the order to determine which goes in the primary project file (so that client names with a suffix will go into their own project).
 			AllTargetFiles = AllTargetFiles.OrderBy(x => x.FullName, StringComparer.OrdinalIgnoreCase).ToList();
@@ -1575,7 +1575,7 @@ namespace UnrealBuildTool
 		/// Finds all target files (filtering by platform)
 		/// </summary>
 		/// <returns>Filtered list of target files</returns>
-		protected List<FileReference> DiscoverTargets(List<FileReference> AllGameProjects, ILogger Logger)
+		public static List<FileReference> DiscoverTargets(List<FileReference> AllGameProjects, ILogger Logger, FileReference? OnlyGameProject, List<UnrealTargetPlatform> SupportedPlatforms, bool bIncludeEngineSource, bool bIncludeTempTargets)
 		{
 			List<FileReference> AllTargetFiles = new List<FileReference>();
 
diff --git a/Engine/Source/Programs/UnrealBuildTool/ProjectFiles/Rider/RiderProjectFileGenerator.cs b/Engine/Source/Programs/UnrealBuildTool/ProjectFiles/Rider/RiderProjectFileGenerator.cs
--- a/Engine/Source/Programs/UnrealBuildTool/ProjectFiles/Rider/RiderProjectFileGenerator.cs
+++ b/Engine/Source/Programs/UnrealBuildTool/ProjectFiles/Rider/RiderProjectFileGenerator.cs
@@ -189,7 +189,7 @@ namespace UnrealBuildTool
 
 			// Find all of the target files.  This will filter out any modules or targets that don't
 			// belong to platforms we're generating project files for.
-			List<FileReference> AllTargetFiles = DiscoverTargets(AllGames, Logger);
+			List<FileReference> AllTargetFiles = DiscoverTargets(AllGames, Logger, OnlyGameProject, SupportedPlatforms, bIncludeEngineSource: bIncludeEngineSource, bIncludeTempTargets: bIncludeTempTargets);
 
 			// Sort the targets by name. When we have multiple targets of a given type for a project, we'll use the order to determine which goes in the primary project file (so that client names with a suffix will go into their own project).
 			AllTargetFiles = AllTargetFiles.OrderBy(x => x.FullName, StringComparer.OrdinalIgnoreCase).ToList();
diff --git a/Engine/Source/Programs/UnrealBuildTool/ProjectFiles/VisualStudioCode/VSCodeProjectFileGenerator.cs b/Engine/Source/Programs/UnrealBuildTool/ProjectFiles/VisualStudioCode/VSCodeProjectFileGenerator.cs
--- a/Engine/Source/Programs/UnrealBuildTool/ProjectFiles/VisualStudioCode/VSCodeProjectFileGenerator.cs
+++ b/Engine/Source/Programs/UnrealBuildTool/ProjectFiles/VisualStudioCode/VSCodeProjectFileGenerator.cs
@@ -53,6 +53,14 @@ namespace UnrealBuildTool
 		[XmlConfigFile(Name = "AddDebugCoreConfig")]
 		private bool bAddDebugCoreConfig = false;
 
+		/// <summary>
+		/// Do not create compile commands json files with compiler arguments for each file; works better with VS Code extension using
+		/// UBT server mode.
+		/// </summary>
+		[XmlConfigFile(Name = "NoCompileCommands")]
+		[CommandLine("-NoCompileCommands")]
+		private bool bNoCompileCommands = false;
+
 		private enum EPathType
 		{
 			Absolute,
@@ -238,7 +246,7 @@ namespace UnrealBuildTool
 
 		public override bool ShouldGenerateIntelliSenseData()
 		{
-			return true;
+			return !bNoCompileCommands;
 		}
 
 		protected override ProjectFile AllocateProjectFile(FileReference InitFilePath, DirectoryReference BaseDir)
@@ -667,7 +675,10 @@ namespace UnrealBuildTool
 			OutFile.BeginObject();
 
 			OutFile.AddField("name", Name);
-			OutFile.AddField("compilerPath", CompilerPath.FullName);
+			if (CompilerPath != null)
+			{
+				OutFile.AddField("compilerPath", CompilerPath.FullName);
+			}
 
 			if (HostPlatform == UnrealTargetPlatform.Mac)
 			{
@@ -723,8 +734,10 @@ namespace UnrealBuildTool
 				OutFile.AddField("intelliSenseMode", "clang-x64");
 			}
 
+			OutFile.AddField("configurationProvider", "epic.ue");
+
 			FileReference CompileCommands = FileReference.Combine(OutputDirectory, string.Format("compileCommands_{0}.json", ProjectName));
-			WriteCompileCommands(CompileCommands, SourceFiles, CompilerPath, ModuleCommandLines);
+			WriteCompileCommands(CompileCommands, SourceFiles, CompilerPath!, ModuleCommandLines);
 			OutFile.AddField("compileCommands", MakePathString(CompileCommands, bInAbsolute: true, bForceSkipQuotes: true));
 
 			OutFile.EndObject();
@@ -1619,6 +1632,8 @@ namespace UnrealBuildTool
 				// Adding this section aids discovery of extensions which are helpful to have installed for Unreal development.
 				WorkspaceFile.BeginArray("recommendations");
 				{
+					// Add when/if published to marketplace. 
+					// WorkspaceFile.AddUnnamedField("epic.vscode-ue");
 					WorkspaceFile.AddUnnamedField("ms-vscode.cpptools");
 					WorkspaceFile.AddUnnamedField("ms-dotnettools.csharp");
 
diff --git a/Engine/Source/Programs/UnrealBuildTool/ToolChain/UEToolChain.cs b/Engine/Source/Programs/UnrealBuildTool/ToolChain/UEToolChain.cs
--- a/Engine/Source/Programs/UnrealBuildTool/ToolChain/UEToolChain.cs
+++ b/Engine/Source/Programs/UnrealBuildTool/ToolChain/UEToolChain.cs
@@ -35,6 +35,12 @@ namespace UnrealBuildTool
 			return DirectoryReference.Combine(OutputDir, "Ifc");
 		}
 
+		// Return the path to the cpp compiler that will be used by this toolchain.
+		public virtual FileReference? GetCppCompilerPath()
+		{
+			return null;
+		}
+
 		public abstract CPPOutput CompileCPPFiles(CppCompileEnvironment CompileEnvironment, List<FileItem> InputFiles, DirectoryReference OutputDir, string ModuleName, IActionGraphBuilder Graph);
 
 		public virtual CPPOutput CompileRCFiles(CppCompileEnvironment Environment, List<FileItem> InputFiles, DirectoryReference OutputDir, IActionGraphBuilder Graph)
