Add missing copyright.

diff --git a/Engine/Source/Programs/UnrealBuildTool/Modes/ServerMode.cs b/Engine/Source/Programs/UnrealBuildTool/Modes/ServerMode.cs
--- a/Engine/Source/Programs/UnrealBuildTool/Modes/ServerMode.cs
+++ b/Engine/Source/Programs/UnrealBuildTool/Modes/ServerMode.cs
@@ -1,3 +1,5 @@
+// Copyright Epic Games, Inc. All Rights Reserved.
+
 using EpicGames.Core;
 using Microsoft.Extensions.Logging;
 using System;
