UnrealBuildTool: Ensure Action list is ordered before running any executor. Pass in IEnumerable instead of List

#rnx

diff --git a/Engine/Source/Programs/UnrealBuildTool/Executors/Experimental/FASTBuild.cs b/Engine/Source/Programs/UnrealBuildTool/Executors/Experimental/FASTBuild.cs
--- a/Engine/Source/Programs/UnrealBuildTool/Executors/Experimental/FASTBuild.cs
+++ b/Engine/Source/Programs/UnrealBuildTool/Executors/Experimental/FASTBuild.cs
@@ -427,22 +427,27 @@ namespace UnrealBuildTool
 		}
 
 		[SupportedOSPlatform("windows")]
-		public override bool ExecuteActions(List<LinkedAction> Actions, ILogger Logger)
+		public override bool ExecuteActions(IEnumerable<LinkedAction> Actions, ILogger Logger)
 		{
-			if (Actions.Count <= 0)
+			if (!Actions.Any())
+			{
 				return true;
+			}
+
 
 			IEnumerable<LinkedAction> CompileActions = Actions.Where(Action => Action.ActionType == ActionType.Compile && Action.bCanExecuteRemotely && Action.bCanExecuteRemotelyWithSNDBS);
 			if (CompileActions.Any() && DetectBuildType(CompileActions, Logger))
 			{
 				string FASTBuildFilePath = Path.Combine(Unreal.EngineDirectory.FullName, "Intermediate", "Build", "fbuild.bff");
 				if (!CreateBffFile(Actions, FASTBuildFilePath, Logger))
+				{
 					return false;
+				}
 
 				return ExecuteBffFile(FASTBuildFilePath, Logger);
 			}
 
-			return LocalExecutor.ExecuteActions(Actions.Distinct().ToList(), Logger); ;
+			return LocalExecutor.ExecuteActions(Actions, Logger);
 		}
 
 		private void AddText(string StringToWrite)
diff --git a/Engine/Source/Programs/UnrealBuildTool/Executors/Experimental/HordeExecutor.cs b/Engine/Source/Programs/UnrealBuildTool/Executors/Experimental/HordeExecutor.cs
--- a/Engine/Source/Programs/UnrealBuildTool/Executors/Experimental/HordeExecutor.cs
+++ b/Engine/Source/Programs/UnrealBuildTool/Executors/Experimental/HordeExecutor.cs
@@ -421,10 +421,15 @@ namespace UnrealBuildTool
 		/// Executes the specified actions locally.
 		/// </summary>
 		/// <returns>True if all the tasks successfully executed, or false if any of them failed.</returns>
-		public override bool ExecuteActions(List<LinkedAction> InputActions, ILogger Logger)
+		public override bool ExecuteActions(IEnumerable<LinkedAction> InputActions, ILogger Logger)
 		{
+			if (!InputActions.Any())
+			{
+				return true;
+			}
+
 			int NumCompletedActions = 0;
-			int TotalActions = InputActions.Count;
+			int TotalActions = InputActions.Count();
 			int ActualNumParallelProcesses = Math.Min(TotalActions, NumParallelProcesses);
 
 			using ManagedProcessGroup ProcessGroup = new ManagedProcessGroup();
diff --git a/Engine/Source/Programs/UnrealBuildTool/Executors/HybridExecutor.cs b/Engine/Source/Programs/UnrealBuildTool/Executors/HybridExecutor.cs
--- a/Engine/Source/Programs/UnrealBuildTool/Executors/HybridExecutor.cs
+++ b/Engine/Source/Programs/UnrealBuildTool/Executors/HybridExecutor.cs
@@ -65,8 +65,13 @@ namespace UnrealBuildTool
 		/// <param name="ActionsToExecute">Actions to be executed</param>
 		/// <param name="Logger">Logger for output</param>
 		/// <returns>True if the build succeeded, false otherwise</returns>
-		public override bool ExecuteActions(List<LinkedAction> ActionsToExecute, ILogger Logger)
+		public override bool ExecuteActions(IEnumerable<LinkedAction> ActionsToExecute, ILogger Logger)
 		{
+			if (!ActionsToExecute.Any())
+			{
+				return true;
+			}
+
 			// Find the number of dependants for each action
 			Dictionary<LinkedAction, int> ActionToNumDependents = ActionsToExecute.ToDictionary(x => x, x => 0);
 			foreach (LinkedAction Action in ActionsToExecute)
@@ -112,7 +117,7 @@ namespace UnrealBuildTool
 
 			// Split the list of actions into those which should be executed locally and remotely
 			List<LinkedAction> LocalActionsToExecute = new List<LinkedAction>(LeafActions.Count);
-			List<LinkedAction> RemoteActionsToExecute = new List<LinkedAction>(ActionsToExecute.Count - LeafActions.Count);
+			List<LinkedAction> RemoteActionsToExecute = new List<LinkedAction>(ActionsToExecute.Count() - LeafActions.Count);
 			foreach (LinkedAction ActionToExecute in ActionsToExecute)
 			{
 				if (LeafActions.Contains(ActionToExecute))
diff --git a/Engine/Source/Programs/UnrealBuildTool/Executors/ParallelExecutor.cs b/Engine/Source/Programs/UnrealBuildTool/Executors/ParallelExecutor.cs
--- a/Engine/Source/Programs/UnrealBuildTool/Executors/ParallelExecutor.cs
+++ b/Engine/Source/Programs/UnrealBuildTool/Executors/ParallelExecutor.cs
@@ -164,10 +164,15 @@ namespace UnrealBuildTool
 		/// Executes the specified actions locally.
 		/// </summary>
 		/// <returns>True if all the tasks successfully executed, or false if any of them failed.</returns>
-		public override bool ExecuteActions(List<LinkedAction> InputActions, ILogger Logger)
+		public override bool ExecuteActions(IEnumerable<LinkedAction> InputActions, ILogger Logger)
 		{
+			if (!InputActions.Any())
+			{
+				return true;
+			}
+
 			int NumCompletedActions = 0;
-			int TotalActions = InputActions.Count;
+			int TotalActions = InputActions.Count();
 			int ActualNumParallelProcesses = Math.Min(TotalActions, NumParallelProcesses);
 
 			using ManagedProcessGroup ProcessGroup = new ManagedProcessGroup();
@@ -206,7 +211,7 @@ namespace UnrealBuildTool
 			return ExecuteTasks.Values.All(x => x.Result.ExitCode == 0);
 		}
 
-		protected static Task<ExecuteResults> CreateExecuteTask(LinkedAction Action, List<LinkedAction> InputActions, Dictionary<LinkedAction, Task<ExecuteResults>> ExecuteTasks, ManagedProcessGroup ProcessGroup, SemaphoreSlim MaxProcessSemaphore, CancellationToken CancellationToken)
+		protected static Task<ExecuteResults> CreateExecuteTask(LinkedAction Action, IEnumerable<LinkedAction> InputActions, Dictionary<LinkedAction, Task<ExecuteResults>> ExecuteTasks, ManagedProcessGroup ProcessGroup, SemaphoreSlim MaxProcessSemaphore, CancellationToken CancellationToken)
 		{
 			List<LinkedAction> PrerequisiteActions = Action.PrerequisiteActions.Where(x => InputActions.Contains(x)).ToList();
 			if (PrerequisiteActions.Count == 0)
@@ -214,7 +219,7 @@ namespace UnrealBuildTool
 				return Task.Factory.StartNew(
 					() => ExecuteAction(new Task<ExecuteResults>[0], Action, ProcessGroup, MaxProcessSemaphore, CancellationToken),
 					CancellationToken,
-					TaskCreationOptions.LongRunning,
+					TaskCreationOptions.PreferFairness | TaskCreationOptions.LongRunning,
 					TaskScheduler.Current
 				).Unwrap();
 			}
@@ -234,7 +239,7 @@ namespace UnrealBuildTool
 				PrerequisiteTasks.ToArray(),
 				(AntecedentTasks) => ExecuteAction(AntecedentTasks, Action, ProcessGroup, MaxProcessSemaphore, CancellationToken),
 				CancellationToken,
-				TaskContinuationOptions.LongRunning,
+				TaskContinuationOptions.PreferFairness | TaskContinuationOptions.LongRunning,
 				TaskScheduler.Current
 			).Unwrap();
 		}
diff --git a/Engine/Source/Programs/UnrealBuildTool/Executors/SNDBS.cs b/Engine/Source/Programs/UnrealBuildTool/Executors/SNDBS.cs
--- a/Engine/Source/Programs/UnrealBuildTool/Executors/SNDBS.cs
+++ b/Engine/Source/Programs/UnrealBuildTool/Executors/SNDBS.cs
@@ -197,10 +197,12 @@ namespace UnrealBuildTool
 			return true;
 		}
 
-		public override bool ExecuteActions(List<LinkedAction> Actions, ILogger Logger)
+		public override bool ExecuteActions(IEnumerable<LinkedAction> Actions, ILogger Logger)
 		{
-			if (Actions.Count == 0)
+			if (!Actions.Any())
+			{
 				return true;
+			}
 
 			// Clean the intermediate directory in case there are any leftovers from previous builds
 			if (DirectoryReference.Exists(IntermediateDir))
@@ -265,7 +267,7 @@ namespace UnrealBuildTool
 			{
 				StartInfo.Arguments += " -k";
 			}
-			return ExecuteProcessWithProgressMarkup(StartInfo, Actions.Count, Logger);
+			return ExecuteProcessWithProgressMarkup(StartInfo, Actions.Count(), Logger);
 		}
 
 		/// <summary>
diff --git a/Engine/Source/Programs/UnrealBuildTool/Executors/XGE.cs b/Engine/Source/Programs/UnrealBuildTool/Executors/XGE.cs
--- a/Engine/Source/Programs/UnrealBuildTool/Executors/XGE.cs
+++ b/Engine/Source/Programs/UnrealBuildTool/Executors/XGE.cs
@@ -390,52 +390,30 @@ namespace UnrealBuildTool
 		}
 
 		[SupportedOSPlatform("windows")]
-		public override bool ExecuteActions(List<LinkedAction> ActionsToExecute, ILogger Logger)
+		public override bool ExecuteActions(IEnumerable<LinkedAction> Actions, ILogger Logger)
 		{
-			bool XGEResult = true;
-
-			// Batch up XGE execution by actions with the same output event handler.
-			List<LinkedAction> ActionBatch = new List<LinkedAction>();
-			ActionBatch.Add(ActionsToExecute[0]);
-			for (int ActionIndex = 1; ActionIndex < ActionsToExecute.Count && XGEResult; ++ActionIndex)
-			{
-				LinkedAction CurrentAction = ActionsToExecute[ActionIndex];
-				ActionBatch.Add(CurrentAction);
-			}
-			if (ActionBatch.Count > 0 && XGEResult)
+			if (!Actions.Any())
 			{
-				XGEResult = ExecuteActionBatch(ActionBatch, Logger);
-				ActionBatch.Clear();
+				return true;
 			}
 
-			return XGEResult;
-		}
+			// Write the actions to execute to a XGE task file.
+			string XGETaskFilePath = FileReference.Combine(Unreal.EngineDirectory, "Intermediate", "Build", "XGETasks.xml").FullName;
+			WriteTaskFile(Actions, XGETaskFilePath, true, false, Logger);
 
-		[SupportedOSPlatform("windows")]
-		bool ExecuteActionBatch(List<LinkedAction> Actions, ILogger Logger)
-		{
-			bool XGEResult = true;
-			if (Actions.Count > 0)
-			{
-				// Write the actions to execute to a XGE task file.
-				string XGETaskFilePath = FileReference.Combine(Unreal.EngineDirectory, "Intermediate", "Build", "XGETasks.xml").FullName;
-				WriteTaskFile(Actions, XGETaskFilePath, true, false, Logger);
-
-				XGEResult = ExecuteTaskFileWithProgressMarkup(XGETaskFilePath, Actions.Count, Logger);
-			}
-			return XGEResult;
+			return ExecuteTaskFileWithProgressMarkup(XGETaskFilePath, Actions.Count(), Logger);
 		}
 
 		/// <summary>
 		/// Writes a XGE task file containing the specified actions to the specified file path.
 		/// </summary>
-		static void WriteTaskFile(List<LinkedAction> InActions, string TaskFilePath, bool bProgressMarkup, bool bXGEExport, ILogger Logger)
+		static void WriteTaskFile(IEnumerable<LinkedAction> InActions, string TaskFilePath, bool bProgressMarkup, bool bXGEExport, ILogger Logger)
 		{
 			bool HostOnVpn = TryGetCoordinatorHost(out string? CoordinatorHost) && IsHostOnVpn(CoordinatorHost, Logger);
 
 			Dictionary<string, string> ExportEnv = new Dictionary<string, string>();
 
-			List<LinkedAction> Actions = InActions;
+			List<LinkedAction> Actions = InActions.ToList();
 			if (bXGEExport)
 			{
 				IDictionary CurrentEnvironment = Environment.GetEnvironmentVariables();
diff --git a/Engine/Source/Programs/UnrealBuildTool/System/ActionExecutor.cs b/Engine/Source/Programs/UnrealBuildTool/System/ActionExecutor.cs
--- a/Engine/Source/Programs/UnrealBuildTool/System/ActionExecutor.cs
+++ b/Engine/Source/Programs/UnrealBuildTool/System/ActionExecutor.cs
@@ -29,7 +29,7 @@ namespace UnrealBuildTool
 			MemoryPerActionBytesOverride = Math.Max(MemoryPerActionBytesOverride, MemoryPerActionOverrideGB * 1024 * 1024 * 1024);
 		}
 
-		public abstract bool ExecuteActions(List<LinkedAction> ActionsToExecute, ILogger Logger);
+		public abstract bool ExecuteActions(IEnumerable<LinkedAction> ActionsToExecute, ILogger Logger);
 	}
 
 }
diff --git a/Engine/Source/Programs/UnrealBuildTool/System/ActionGraph.cs b/Engine/Source/Programs/UnrealBuildTool/System/ActionGraph.cs
--- a/Engine/Source/Programs/UnrealBuildTool/System/ActionGraph.cs
+++ b/Engine/Source/Programs/UnrealBuildTool/System/ActionGraph.cs
@@ -16,6 +16,7 @@ using OpenTracing.Util;
 using UnrealBuildBase;
 using Microsoft.Extensions.Logging;
 using System.Collections.Concurrent;
+using System.Collections.Immutable;
 
 namespace UnrealBuildTool
 {
@@ -380,9 +381,12 @@ namespace UnrealBuildTool
 				// Figure out which executor to use
 				ActionExecutor Executor = SelectExecutor(BuildConfiguration, ActionsToExecute.Count, TargetDescriptors, Logger);
 
+				// Ensure actions are ordered
+				IOrderedEnumerable<LinkedAction> OrderedActionsToExecute = ActionsToExecute.OrderByDescending(x => x.NumTotalDependentActions).ThenByDescending(x => x.PrerequisiteItems.Count());
+
 				// Execute the build
 				Stopwatch Timer = Stopwatch.StartNew();
-				if (!Executor.ExecuteActions(ActionsToExecute, Logger))
+				if (!Executor.ExecuteActions(OrderedActionsToExecute, Logger))
 				{
 					throw new CompilationResultException(CompilationResult.OtherCompilationError);
 				}
