UnrealHeaderTool: Update #include paths to be relative to the module's first public include path (or the generated UHT path if no public include paths are available) rather than the module's parent directory

diff --git a/Engine/Source/Programs/UnrealBuildTool/System/ExternalExecution.cs b/Engine/Source/Programs/UnrealBuildTool/System/ExternalExecution.cs
--- a/Engine/Source/Programs/UnrealBuildTool/System/ExternalExecution.cs
+++ b/Engine/Source/Programs/UnrealBuildTool/System/ExternalExecution.cs
@@ -114,6 +114,11 @@ namespace UnrealBuildTool
 		/// </summary>
 		public DirectoryReference[] ModuleDirectories;
 
+		/// <summary>
+		/// The include search path for generated headers to include other headers
+		/// </summary>
+		public DirectoryReference ModuleIncudeBase;
+
 		/// <summary>
 		/// Module type
 		/// </summary>
@@ -169,11 +174,12 @@ namespace UnrealBuildTool
 		/// </summary>
 		public bool bIsReadOnly;
 
-		public UHTModuleInfo(string ModuleName, FileReference ModuleRulesFile, DirectoryReference[] ModuleDirectories, UHTModuleType ModuleType, DirectoryItem GeneratedCodeDirectory, EGeneratedCodeVersion GeneratedCodeVersion, bool bIsReadOnly, ModuleRules.PackageOverrideType OverrideType)
+		public UHTModuleInfo(string ModuleName, FileReference ModuleRulesFile, DirectoryReference[] ModuleDirectories, DirectoryReference ModuleIncudeBase, UHTModuleType ModuleType, DirectoryItem GeneratedCodeDirectory, EGeneratedCodeVersion GeneratedCodeVersion, bool bIsReadOnly, ModuleRules.PackageOverrideType OverrideType)
 		{
 			this.ModuleName = ModuleName;
 			this.ModuleRulesFile = ModuleRulesFile;
 			this.ModuleDirectories = ModuleDirectories;
+			this.ModuleIncudeBase = ModuleIncudeBase;
 			this.ModuleType = ModuleType.ToString();
 			this.OverrideModuleType = OverrideType.ToString();
 			this.PublicUObjectClassesHeaders = new List<FileItem>();
@@ -191,6 +197,7 @@ namespace UnrealBuildTool
 			ModuleName = Reader.ReadString()!;
 			ModuleRulesFile = Reader.ReadFileReference();
 			ModuleDirectories = Reader.ReadArray<DirectoryReference>(Reader.ReadDirectoryReferenceNotNull)!;
+			ModuleIncudeBase = Reader.ReadDirectoryReference()!;
 			ModuleType = Reader.ReadString()!;
 			OverrideModuleType = Reader.ReadString()!;
 			PublicUObjectClassesHeaders = Reader.ReadList(() => Reader.ReadFileItem())!;
@@ -209,6 +216,7 @@ namespace UnrealBuildTool
 			Writer.WriteString(ModuleName);
 			Writer.WriteFileReference(ModuleRulesFile);
 			Writer.WriteArray<DirectoryReference>(ModuleDirectories, Writer.WriteDirectoryReference);
+			Writer.WriteDirectoryReference(ModuleIncudeBase);
 			Writer.WriteString(ModuleType);
 			Writer.WriteString(OverrideModuleType);
 			Writer.WriteList(PublicUObjectClassesHeaders, Item => Writer.WriteFileItem(Item));
@@ -472,7 +480,10 @@ namespace UnrealBuildTool
 				{
 					UEBuildModuleCPP Module = ModulesSortedByType[Idx];
 
-					UHTModuleInfo Info = new UHTModuleInfo(Module.Name, Module.RulesFile, Module.ModuleDirectories, ModuleToType[Module], DirectoryItem.GetItemByDirectoryReference(Module.GeneratedCodeDirectoryUHT!), GeneratedCodeVersion, Module.Rules.bUsePrecompiled, Module.Rules.OverridePackageType);
+					DirectoryReference ModuleIncludeBase = Module.PublicIncludePaths.FirstOrDefault() ?? Module.GeneratedCodeDirectoryUHT!;
+					DirectoryItem GeneratedCodeDirectory = DirectoryItem.GetItemByDirectoryReference(Module.GeneratedCodeDirectoryUHT!);
+
+					UHTModuleInfo Info = new UHTModuleInfo(Module.Name, Module.RulesFile, Module.ModuleDirectories, ModuleIncludeBase, ModuleToType[Module], GeneratedCodeDirectory, GeneratedCodeVersion, Module.Rules.bUsePrecompiled, Module.Rules.OverridePackageType);
 					ModuleInfoArray[Idx] = Info;
 
 					Queue.Enqueue(() => SetupUObjectModule(Info, ExcludedFolders, MetadataCache, Queue));
@@ -1007,7 +1018,7 @@ namespace UnrealBuildTool
 						ModuleType = (UHTModuleType)Enum.Parse(typeof(UHTModuleType), UObjectModule.ModuleType),
 						OverrideModuleType = (EPackageOverrideType)Enum.Parse(typeof(EPackageOverrideType), UObjectModule.OverrideModuleType),
 						BaseDirectory = UObjectModule.ModuleDirectories[0].FullName,
-						IncludeBase = UObjectModule.ModuleDirectories[0].ParentDirectory!.FullName,
+						IncludeBase = UObjectModule.ModuleIncudeBase.FullName,
 						OutputDirectory = Path.GetDirectoryName(UObjectModule.GeneratedCPPFilenameBase)!,
 						ClassesHeaders = UObjectModule.PublicUObjectClassesHeaders.Select((Header) => Header.AbsolutePath).ToList(),
 						PublicHeaders = UObjectModule.PublicUObjectHeaders.Select((Header) => Header.AbsolutePath).ToList(),
