Add try catch to read ConfigRemapFile to that ConfigFile can be used with assemblies outside of the Engine DotNet Path

diff --git a/Engine/Source/Programs/UnrealBuildTool/System/ConfigFile.cs b/Engine/Source/Programs/UnrealBuildTool/System/ConfigFile.cs
--- a/Engine/Source/Programs/UnrealBuildTool/System/ConfigFile.cs
+++ b/Engine/Source/Programs/UnrealBuildTool/System/ConfigFile.cs
@@ -169,10 +169,18 @@ namespace UnrealBuildTool
 		/// </summary>
 		static ConfigFile()
 		{
-			// read the special ConfigRemap.ini file into sections
-			FileReference ConfigRemapFile = FileReference.Combine(UnrealBuildBase.Unreal.EngineDirectory, "Config", "ConfigRedirects.ini");
 			Dictionary<string, ConfigFileSection> Sections = new(StringComparer.InvariantCultureIgnoreCase);
-			ReadIntoSections(ConfigRemapFile, Sections, ConfigLineAction.Set);
+			try
+			{
+				// read the special ConfigRemap.ini file into sections
+				FileReference ConfigRemapFile = FileReference.Combine(UnrealBuildBase.Unreal.EngineDirectory, "Config", "ConfigRedirects.ini");
+				ReadIntoSections(ConfigRemapFile, Sections, ConfigLineAction.Set);
+			}
+			catch (Exception)
+			{
+				// Make ConfigFile when EngineDirectory is unknown a warning since ConfigRemapFile cannot be read in this case; e.g. Assemblies outside Engine that depend on ConfigFile
+				Log.Logger.LogWarning("Failed to read ConfigRemapFile into Sections");
+			}
 			
 			// walk over the sections, where all but the special SectionNameRemap section is a section of keys to remap in that same section
 			foreach (KeyValuePair<string, ConfigFileSection> Pair in Sections)
