UnrealBuildTool: Use BuildSettings.Latest as defaults, so changing the latest version will also update these variables. Latest is still set to V2.

#rnx

diff --git a/Engine/Source/Programs/UnrealBuildTool/Configuration/TestTargetRules.cs b/Engine/Source/Programs/UnrealBuildTool/Configuration/TestTargetRules.cs
--- a/Engine/Source/Programs/UnrealBuildTool/Configuration/TestTargetRules.cs
+++ b/Engine/Source/Programs/UnrealBuildTool/Configuration/TestTargetRules.cs
@@ -178,7 +178,7 @@ namespace UnrealBuildTool
 
 			VSTestRunSettingsFile = FileReference.Combine(Unreal.EngineDirectory, "Source", "Programs", "LowLevelTests", "vstest.runsettings");
 
-			DefaultBuildSettings = BuildSettingsVersion.V2;
+			DefaultBuildSettings = BuildSettingsVersion.Latest;
 
 			Type = TargetType.Program;
 			LinkType = TargetLinkType.Monolithic;
diff --git a/Engine/Source/Programs/UnrealBuildTool/System/RulesCompiler.cs b/Engine/Source/Programs/UnrealBuildTool/System/RulesCompiler.cs
--- a/Engine/Source/Programs/UnrealBuildTool/System/RulesCompiler.cs
+++ b/Engine/Source/Programs/UnrealBuildTool/System/RulesCompiler.cs
@@ -25,7 +25,7 @@ namespace UnrealBuildTool
 		/// <summary>
 		/// 
 		/// </summary>
-		const BuildSettingsVersion DefaultEngineBuildSettingsVersion = BuildSettingsVersion.V2;
+		const BuildSettingsVersion DefaultEngineBuildSettingsVersion = BuildSettingsVersion.Latest;
 
 		/// <summary>
 		/// Find all the module rules files under a given directory
