Fix postbuild action for win64 low level tests: correct path to build binary
#rnx

diff --git a/Engine/Source/Programs/UnrealBuildTool/Configuration/TestTargetRules.cs b/Engine/Source/Programs/UnrealBuildTool/Configuration/TestTargetRules.cs
--- a/Engine/Source/Programs/UnrealBuildTool/Configuration/TestTargetRules.cs
+++ b/Engine/Source/Programs/UnrealBuildTool/Configuration/TestTargetRules.cs
@@ -110,7 +110,8 @@ namespace UnrealBuildTool
 				{
 					OutputName = OutputName + "-" + Target.Platform + "-" + Target.Configuration;
 				}
-				string OutputFileName = $"$(EngineDir)\\Binaries\\Win64\\{ExeBinariesSubFolder}\\{OutputName}.exe.is_unreal_test";
+				string OutputDirectory = UEBuildTarget.GetOutputDirectoryForExecutable(Unreal.EngineDirectory, this.File!).FullName;
+				string OutputFileName = $"{OutputDirectory}\\Binaries\\{Target.Platform}\\{ExeBinariesSubFolder}\\{OutputName}.exe.is_unreal_test";
 				PostBuildSteps.Add("echo > " + OutputFileName);
 			}
 		}
