UnrealBuildTool: support clang-cl x64 from visual studio install

#rnx

diff --git a/Engine/Source/Programs/UnrealBuildTool/Platform/Windows/MicrosoftPlatformSDK.cs b/Engine/Source/Programs/UnrealBuildTool/Platform/Windows/MicrosoftPlatformSDK.cs
--- a/Engine/Source/Programs/UnrealBuildTool/Platform/Windows/MicrosoftPlatformSDK.cs
+++ b/Engine/Source/Programs/UnrealBuildTool/Platform/Windows/MicrosoftPlatformSDK.cs
@@ -662,12 +662,14 @@ namespace UnrealBuildTool
 						foreach (VisualStudioInstallation Installation in FindVisualStudioInstallations(WindowsCompiler.VisualStudio2019, Logger))
 						{
 							AddClangToolChain(DirectoryReference.Combine(Installation.BaseDir, "VC", "Tools", "Llvm"), ToolChains, IsAutoSdk: false, Logger);
+							AddClangToolChain(DirectoryReference.Combine(Installation.BaseDir, "VC", "Tools", "Llvm", "x64"), ToolChains, IsAutoSdk: false, Logger);
 						}
 
 						// Check for installations bundled with Visual Studio 2022
 						foreach (VisualStudioInstallation Installation in FindVisualStudioInstallations(WindowsCompiler.VisualStudio2022, Logger))
 						{
 							AddClangToolChain(DirectoryReference.Combine(Installation.BaseDir, "VC", "Tools", "Llvm"), ToolChains, IsAutoSdk: false, Logger);
+							AddClangToolChain(DirectoryReference.Combine(Installation.BaseDir, "VC", "Tools", "Llvm", "x64"), ToolChains, IsAutoSdk: false, Logger);
 						}
 
 						// Check for AutoSDK paths
@@ -960,6 +962,7 @@ namespace UnrealBuildTool
 			{
 				FileVersionInfo VersionInfo = FileVersionInfo.GetVersionInfo(CompilerFile.FullName);
 				VersionNumber Version = new VersionNumber(VersionInfo.FileMajorPart, VersionInfo.FileMinorPart, VersionInfo.FileBuildPart);
+				VersionNumber Family = new VersionNumber(VersionInfo.FileMajorPart);
 
 				int Rank = PreferredClangVersions.TakeWhile(x => !x.Contains(Version)).Count();
 				bool Is64Bit = Is64BitExecutable(CompilerFile);
@@ -971,7 +974,7 @@ namespace UnrealBuildTool
 				}
 
 				Logger.LogDebug("Found Clang toolchain: {ToolChainDir} (Version={Version}, Is64Bit={Is64Bit}, Rank={Rank})", ToolChainDir, Version, Is64Bit, Rank);
-				ToolChains.Add(new ToolChainInstallation(Version, Rank, Version, Is64Bit, false, WindowsArchitecture.x64, null, ToolChainDir, null, IsAutoSdk));
+				ToolChains.Add(new ToolChainInstallation(Family, Rank, Version, Is64Bit, false, WindowsArchitecture.x64, null, ToolChainDir, null, IsAutoSdk));
 			}
 		}
 
@@ -990,6 +993,7 @@ namespace UnrealBuildTool
 			{
 				FileVersionInfo VersionInfo = FileVersionInfo.GetVersionInfo(CompilerFile.FullName);
 				VersionNumber Version = new VersionNumber(VersionInfo.FileMajorPart, VersionInfo.FileMinorPart, VersionInfo.FileBuildPart);
+				VersionNumber Family = new VersionNumber(VersionInfo.FileMajorPart);
 
 				// The icx.exe version may be lower than the toolchain folder version, so use that instead if available
 				if (VersionNumber.TryParse(ToolChainDir.GetDirectoryName(), out VersionNumber? FolderVersion) && FolderVersion > Version)
@@ -1007,7 +1011,7 @@ namespace UnrealBuildTool
 				}
 
 				Logger.LogDebug("Found Intel OneAPI toolchain: {ToolChainDir} (Version={Version}, Is64Bit={Is64Bit}, Rank={Rank})", ToolChainDir, Version, Is64Bit, Rank);
-				ToolChains.Add(new ToolChainInstallation(Version, Rank, Version, Is64Bit, false, WindowsArchitecture.x64, Error, ToolChainDir, null, IsAutoSdk));
+				ToolChains.Add(new ToolChainInstallation(Family, Rank, Version, Is64Bit, false, WindowsArchitecture.x64, Error, ToolChainDir, null, IsAutoSdk));
 			}
 		}
 
