UnrealBuildTool: Remove Hostx86 compiler support, clean up unneeded function parameters

#rnx

diff --git a/Engine/Platforms/Hololens/Source/Programs/UnrealBuildTool/UEBuildHoloLens.cs b/Engine/Platforms/Hololens/Source/Programs/UnrealBuildTool/UEBuildHoloLens.cs
--- a/Engine/Platforms/Hololens/Source/Programs/UnrealBuildTool/UEBuildHoloLens.cs
+++ b/Engine/Platforms/Hololens/Source/Programs/UnrealBuildTool/UEBuildHoloLens.cs
@@ -178,7 +178,7 @@ namespace UnrealBuildTool
 		/// <summary>
 		/// Return the standard Visual C++ library path for the given platform in this toolchain
 		/// </summary>
-		protected override DirectoryReference GetToolChainLibsDir(UnrealTargetPlatform Platform)
+		protected override DirectoryReference GetToolChainLibsDir()
 		{
 			string ArchFolder = WindowsExports.GetArchitectureSubpath(Architecture);
 
diff --git a/Engine/Source/Programs/UnrealBuildTool/Platform/Windows/MicrosoftPlatformSDK.cs b/Engine/Source/Programs/UnrealBuildTool/Platform/Windows/MicrosoftPlatformSDK.cs
--- a/Engine/Source/Programs/UnrealBuildTool/Platform/Windows/MicrosoftPlatformSDK.cs
+++ b/Engine/Source/Programs/UnrealBuildTool/Platform/Windows/MicrosoftPlatformSDK.cs
@@ -973,8 +973,8 @@ namespace UnrealBuildTool
 					Error = $"UnrealBuildTool requires at minimum the Clang {MinimumClangVersion} toolchain. Please install a later toolchain such as {PreferredClangVersions.Select(x => x.Min).Max()} from LLVM.";
 				}
 
-				Logger.LogDebug("Found Clang toolchain: {ToolChainDir} (Version={Version}, Is64Bit={Is64Bit}, Rank={Rank})", ToolChainDir, Version, Is64Bit, Rank);
-				ToolChains.Add(new ToolChainInstallation(Family, Rank, Version, Is64Bit, false, WindowsArchitecture.x64, null, ToolChainDir, null, IsAutoSdk));
+				Logger.LogDebug("Found Clang toolchain: {ToolChainDir} (Version={Version}, Is64Bit={Is64Bit}, Rank={Rank}, Error={Error})", ToolChainDir, Version, Is64Bit, Rank, Error != null);
+				ToolChains.Add(new ToolChainInstallation(Family, Rank, Version, Is64Bit, false, WindowsArchitecture.x64, Error, ToolChainDir, null, IsAutoSdk));
 			}
 		}
 
@@ -1010,7 +1010,7 @@ namespace UnrealBuildTool
 					Error = $"UnrealBuildTool requires at minimum the Intel OneAPI {MinimumIntelOneApiVersion} toolchain. Please install a later toolchain such as {PreferredIntelOneApiVersions.Select(x => x.Min).Max()} from Intel.";
 				}
 
-				Logger.LogDebug("Found Intel OneAPI toolchain: {ToolChainDir} (Version={Version}, Is64Bit={Is64Bit}, Rank={Rank})", ToolChainDir, Version, Is64Bit, Rank);
+				Logger.LogDebug("Found Intel OneAPI toolchain: {ToolChainDir} (Version={Version}, Is64Bit={Is64Bit}, Rank={Rank}, Error={Error})", ToolChainDir, Version, Is64Bit, Rank, Error != null);
 				ToolChains.Add(new ToolChainInstallation(Family, Rank, Version, Is64Bit, false, WindowsArchitecture.x64, Error, ToolChainDir, null, IsAutoSdk));
 			}
 		}
@@ -1076,16 +1076,6 @@ namespace UnrealBuildTool
 		}
 
 
-		/// <summary>
-		/// Checks if the given directory contains a valid Clang toolchain
-		/// </summary>
-		/// <param name="ToolChainDir">Directory to check</param>
-		/// <returns>True if the given directory is valid</returns>
-		static bool IsValidToolChainDirClang(DirectoryReference ToolChainDir)
-		{
-			return FileReference.Exists(FileReference.Combine(ToolChainDir, "bin", "clang-cl.exe"));
-		}
-
 		/// <summary>
 		/// Determines if the given path is a valid Visual C++ version number
 		/// </summary>
@@ -1094,15 +1084,11 @@ namespace UnrealBuildTool
 		/// <returns>True if the path is a valid version</returns>
 		static bool IsValidToolChainDirMSVC(DirectoryReference ToolChainDir, [NotNullWhen(true)] out VersionNumber? Version)
 		{
-			FileReference CompilerExe = FileReference.Combine(ToolChainDir, "bin", "Hostx86", "x64", "cl.exe");
+			FileReference CompilerExe = FileReference.Combine(ToolChainDir, "bin", "Hostx64", "x64", "cl.exe");
 			if (!FileReference.Exists(CompilerExe))
 			{
-				CompilerExe = FileReference.Combine(ToolChainDir, "bin", "Hostx64", "x64", "cl.exe");
-				if (!FileReference.Exists(CompilerExe))
-				{
-					Version = null;
-					return false;
-				}
+				Version = null;
+				return false;
 			}
 
 			FileVersionInfo VersionInfo = FileVersionInfo.GetVersionInfo(CompilerExe.FullName);
@@ -1122,11 +1108,11 @@ namespace UnrealBuildTool
 		/// <returns>True if the given directory contains a 64-bit toolchain</returns>
 		static bool Has64BitToolChain(DirectoryReference ToolChainDir)
 		{
-			return FileReference.Exists(FileReference.Combine(ToolChainDir, "bin", "amd64", "cl.exe")) || FileReference.Exists(FileReference.Combine(ToolChainDir, "bin", "Hostx64", "x64", "cl.exe"));
+			return FileReference.Exists(FileReference.Combine(ToolChainDir, "bin", "Hostx64", "x64", "cl.exe"));
 		}
 
 		/// <summary>
-		/// Checks if the given directory contains a arm64 toolchain.  Used to require arm64, which is an optional install item, when that is our target architecture.
+		/// Checks if the given directory contains an arm64 toolchain.  Used to require arm64, which is an optional install item, when that is our target architecture.
 		/// </summary>
 		/// <param name="ToolChainDir">Directory to check</param>
 		/// <returns>True if the given directory contains the arm64 toolchain</returns>
diff --git a/Engine/Source/Programs/UnrealBuildTool/Platform/Windows/VCEnvironment.cs b/Engine/Source/Programs/UnrealBuildTool/Platform/Windows/VCEnvironment.cs
--- a/Engine/Source/Programs/UnrealBuildTool/Platform/Windows/VCEnvironment.cs
+++ b/Engine/Source/Programs/UnrealBuildTool/Platform/Windows/VCEnvironment.cs
@@ -131,22 +131,16 @@ namespace UnrealBuildTool
 			this.bUseCPPWinRT = Params.bUseCPPWinRT;
 			this.bAllowClangLinker = Params.bAllowClangLinker;
 
-			// Get the standard VC paths
-			DirectoryReference VCToolPath = GetVCToolPath(ToolChain, ToolChainDir, Architecture);
-
-            // Compile using 64 bit tools for 64 bit targets, and 32 for 32.
-			CompilerPath = GetCompilerToolPath(Params.Platform, Compiler, Architecture, CompilerDir);
-
-			// Regardless of the target, if we're linking on a 64 bit machine, we want to use the 64 bit linker (it's faster than the 32 bit linker and can handle large linking jobs)
-			DirectoryReference DefaultLinkerDir = VCToolPath;
-			LinkerPath = GetLinkerToolPath(Params.Platform, Compiler, CompilerDir, DefaultLinkerDir);
-			LibraryManagerPath = GetLibraryLinkerToolPath(Params.Platform, Compiler, CompilerDir, DefaultLinkerDir);
+			// Get the compiler and linker paths from the Toolchain directory
+			CompilerPath = GetCompilerToolPath(Compiler, Architecture, CompilerDir);
+			LinkerPath = GetLinkerToolPath(Compiler, Architecture, CompilerDir, ToolChainDir);
+			LibraryManagerPath = GetLibraryLinkerToolPath(Compiler, Architecture, CompilerDir, ToolChainDir);
 
 			// Get the resource compiler path from the Windows SDK
-			ResourceCompilerPath = GetResourceCompilerToolPath(Params.Platform, WindowsSdkDir, WindowsSdkVersion, Logger);
+			ResourceCompilerPath = GetResourceCompilerToolPath(WindowsSdkDir, WindowsSdkVersion, Logger);
 
 			// Get all the system include paths
-			SetupEnvironment(Params.Platform, Logger);
+			SetupEnvironment(Logger);
 		}
 
 		/// <summary>
@@ -158,12 +152,12 @@ namespace UnrealBuildTool
 			Environment.SetEnvironmentVariable("VC_COMPILER_PATH", CompilerPath.FullName, EnvironmentVariableTarget.Process);
 			Environment.SetEnvironmentVariable("VC_COMPILER_DIR", CompilerPath.Directory.FullName, EnvironmentVariableTarget.Process);
 
-			AddDirectoryToPath(GetVCToolPath(ToolChain, ToolChainDir, Architecture));
+			AddDirectoryToPath(GetVCToolPath(ToolChainDir, Architecture));
 			if (Architecture == WindowsArchitecture.ARM64)
 			{
 				// Add both toolchain paths to the PATH environment variable. There are some support DLLs which are only added to one of the paths, but which the toolchain in the other directory
 				// needs to run (eg. mspdbcore.dll).
-				AddDirectoryToPath(GetVCToolPath(ToolChain, ToolChainDir, WindowsArchitecture.x64));
+				AddDirectoryToPath(GetVCToolPath(ToolChainDir, WindowsArchitecture.x64));
 			}
 
 			// Add the Windows SDK directory to the path too, for mt.exe.
@@ -190,31 +184,24 @@ namespace UnrealBuildTool
 		/// <summary>
 		/// Gets the path to the tool binaries.
 		/// </summary>
-		/// <param name="Compiler">The compiler version</param>
 		/// <param name="VCToolChainDir">Base directory for the VC toolchain</param>
 		/// <param name="Architecture">Target Architecture</param>
-		/// <returns>Directory containing the 32-bit toolchain binaries</returns>
-		public static DirectoryReference GetVCToolPath(WindowsCompiler Compiler, DirectoryReference VCToolChainDir, WindowsArchitecture Architecture)
+		/// <returns>Directory containing the 64-bit toolchain binaries</returns>
+		public static DirectoryReference GetVCToolPath(DirectoryReference VCToolChainDir, WindowsArchitecture Architecture)
 		{
-			FileReference NativeCompilerPath = FileReference.Combine(VCToolChainDir, "bin", "HostX64", WindowsExports.GetArchitectureSubpath(Architecture), "cl.exe");
-			if (FileReference.Exists(NativeCompilerPath))
-			{
-				return NativeCompilerPath.Directory;
-			}
-
-			FileReference CrossCompilerPath = FileReference.Combine(VCToolChainDir, "bin", "HostX86", WindowsExports.GetArchitectureSubpath(Architecture), "cl.exe");
-			if (FileReference.Exists(CrossCompilerPath))
+			FileReference CompilerPath = FileReference.Combine(VCToolChainDir, "bin", "Hostx64", WindowsExports.GetArchitectureSubpath(Architecture), "cl.exe");
+			if (FileReference.Exists(CompilerPath))
 			{
-				return CrossCompilerPath.Directory;
+				return CompilerPath.Directory;
 			}
 
-			throw new BuildException("No required compiler toolchain found in {0}", VCToolChainDir);
+			throw new BuildException("No required {0} compiler toolchain found in {1}", Architecture, VCToolChainDir);
 		}
 
 		/// <summary>
 		/// Gets the path to the compiler.
 		/// </summary>
-		static FileReference GetCompilerToolPath(UnrealTargetPlatform Platform, WindowsCompiler Compiler, WindowsArchitecture Architecture, DirectoryReference CompilerDir)
+		static FileReference GetCompilerToolPath(WindowsCompiler Compiler, WindowsArchitecture Architecture, DirectoryReference CompilerDir)
 		{
 			if (Compiler == WindowsCompiler.Clang)
 			{
@@ -224,15 +211,14 @@ namespace UnrealBuildTool
 			{
 				return FileReference.Combine(CompilerDir, "windows", "bin", "icx.exe");
 			}
-			return FileReference.Combine(GetVCToolPath(Compiler, CompilerDir, Architecture), "cl.exe");
+			return FileReference.Combine(GetVCToolPath(CompilerDir, Architecture), "cl.exe");
 		}
 
 		/// <summary>
 		/// Gets the path to the linker.
 		/// </summary>
-		FileReference GetLinkerToolPath(UnrealTargetPlatform Platform, WindowsCompiler Compiler, DirectoryReference CompilerDir, DirectoryReference DefaultLinkerDir)
+		FileReference GetLinkerToolPath(WindowsCompiler Compiler, WindowsArchitecture Architecture, DirectoryReference CompilerDir, DirectoryReference ToochainDir)
 		{
-			// Regardless of the target, if we're linking on a 64 bit machine, we want to use the 64 bit linker (it's faster than the 32 bit linker)
 			if (Compiler == WindowsCompiler.Clang && bAllowClangLinker)
 			{
 				return FileReference.Combine(CompilerDir, "bin", "lld-link.exe");
@@ -241,16 +227,15 @@ namespace UnrealBuildTool
 			{
 				return FileReference.Combine(CompilerDir, "windows", "bin", "intel64", "xilink.exe");
 			}
-			return FileReference.Combine(DefaultLinkerDir, "link.exe");
+			return FileReference.Combine(GetVCToolPath(ToochainDir, Architecture), "link.exe");
 
 		}
 
 		/// <summary>
 		/// Gets the path to the library linker.
 		/// </summary>
-		FileReference GetLibraryLinkerToolPath(UnrealTargetPlatform Platform, WindowsCompiler Compiler, DirectoryReference CompilerDir, DirectoryReference DefaultLinkerDir)
+		FileReference GetLibraryLinkerToolPath(WindowsCompiler Compiler, WindowsArchitecture Architecture, DirectoryReference CompilerDir, DirectoryReference ToochainDir)
 		{
-			// Regardless of the target, if we're linking on a 64 bit machine, we want to use the 64 bit linker (it's faster than the 32 bit linker)
 			if (Compiler == WindowsCompiler.Clang && bAllowClangLinker)
 			{
 				// @todo: lld-link is not currently working for building .lib
@@ -260,15 +245,14 @@ namespace UnrealBuildTool
 			{
 				return FileReference.Combine(CompilerDir, "windows", "bin", "intel64", "xilib.exe");
 			}
-			return FileReference.Combine(DefaultLinkerDir, "lib.exe");
+			return FileReference.Combine(GetVCToolPath(ToochainDir, Architecture), "lib.exe");
 		}
 
 		/// <summary>
-		/// Gets the path to the resource compiler's rc.exe for the specified platform.
+		/// Gets the path to the resource compiler.
 		/// </summary>
-		virtual protected FileReference GetResourceCompilerToolPath(UnrealTargetPlatform Platform, DirectoryReference WindowsSdkDir, VersionNumber WindowsSdkVersion, ILogger Logger)
+		virtual protected FileReference GetResourceCompilerToolPath(DirectoryReference WindowsSdkDir, VersionNumber WindowsSdkVersion, ILogger Logger)
 		{
-			// 64 bit -- we can use the 32 bit version to target 64 bit on 32 bit OS.
 			FileReference ResourceCompilerPath = FileReference.Combine(WindowsSdkDir, "bin", WindowsSdkVersion.ToString(), "x64", "rc.exe");
 			if(FileReference.Exists(ResourceCompilerPath))
 			{
@@ -285,9 +269,9 @@ namespace UnrealBuildTool
 		}
 
 		/// <summary>
-		/// Return the standard Visual C++ library path for the given platform in this toolchain
+		/// Return the standard Visual C++ library path.
 		/// </summary>
-		protected virtual DirectoryReference GetToolChainLibsDir(UnrealTargetPlatform Platform)
+		protected virtual DirectoryReference GetToolChainLibsDir()
 		{
 			string ArchFolder = WindowsExports.GetArchitectureSubpath(Architecture);
 
@@ -313,7 +297,7 @@ namespace UnrealBuildTool
 		/// Sets up the standard compile environment for the toolchain
 		/// </summary>
 		[SupportedOSPlatform("windows")]
-		private void SetupEnvironment(UnrealTargetPlatform Platform, ILogger Logger)
+		private void SetupEnvironment(ILogger Logger)
 		{
 			string ArchFolder = WindowsExports.GetArchitectureSubpath(Architecture);
 
@@ -321,7 +305,7 @@ namespace UnrealBuildTool
 			IncludePaths.Add(DirectoryReference.Combine(ToolChainDir, "INCLUDE"));
 
 			// Add the standard Visual C++ library paths
-			LibraryPaths.Add( GetToolChainLibsDir(Platform));
+			LibraryPaths.Add(GetToolChainLibsDir());
 
 			// If we're on >= Visual Studio 2015 and using pre-Windows 10 SDK, we need to find a Windows 10 SDK and add the UCRT include paths
 			if(ToolChain.IsMSVC() && WindowsSdkVersion < new VersionNumber(10))
