UnrealBuildTool: Remove property that disables generating cpp dependency files

#rnx

diff --git a/Engine/Source/Programs/UnrealBuildTool/Platform/Android/AndroidToolChain.cs b/Engine/Source/Programs/UnrealBuildTool/Platform/Android/AndroidToolChain.cs
--- a/Engine/Source/Programs/UnrealBuildTool/Platform/Android/AndroidToolChain.cs
+++ b/Engine/Source/Programs/UnrealBuildTool/Platform/Android/AndroidToolChain.cs
@@ -1376,13 +1376,10 @@ namespace UnrealBuildTool
 					}
 
 					// Generate the included header dependency list
-					if (CompileEnvironment.bGenerateDependenciesFile)
-					{
-						FileItem DependencyListFile = FileItem.GetItemByFileReference(FileReference.Combine(OutputDir, InlineArchName(Path.GetFileName(SourceFile.AbsolutePath) + ".d", Arch, true)));
-						FileArguments += string.Format(" -MD -MF\"{0}\"", DependencyListFile.AbsolutePath.Replace('\\', '/'));
-						CompileAction.DependencyListFile = DependencyListFile;
-						CompileAction.ProducedItems.Add(DependencyListFile);
-					}
+					FileItem DependencyListFile = FileItem.GetItemByFileReference(FileReference.Combine(OutputDir, InlineArchName(Path.GetFileName(SourceFile.AbsolutePath) + ".d", Arch, true)));
+					FileArguments += string.Format(" -MD -MF\"{0}\"", DependencyListFile.AbsolutePath.Replace('\\', '/'));
+					CompileAction.DependencyListFile = DependencyListFile;
+					CompileAction.ProducedItems.Add(DependencyListFile);
 
 					// Build a full argument list
 					string AllArguments = Arguments + FileArguments + CompileEnvironment.AdditionalArguments;
@@ -1553,13 +1550,10 @@ namespace UnrealBuildTool
 					}
 
 					// Generate the included header dependency list
-					if (CompileEnvironment.bGenerateDependenciesFile)
-					{
-						FileItem DependencyListFile = FileItem.GetItemByFileReference(FileReference.Combine(OutputDir, InlineArchName(Path.GetFileName(ISPCFile.AbsolutePath) + ".d", Arch, true)));
-						Arguments.Add(String.Format("-M -MF \"{0}\"", DependencyListFile.AbsolutePath.Replace('\\', '/')));
-						CompileAction.DependencyListFile = DependencyListFile;
-						CompileAction.ProducedItems.Add(DependencyListFile);
-					}
+					FileItem DependencyListFile = FileItem.GetItemByFileReference(FileReference.Combine(OutputDir, InlineArchName(Path.GetFileName(ISPCFile.AbsolutePath) + ".d", Arch, true)));
+					Arguments.Add(String.Format("-M -MF \"{0}\"", DependencyListFile.AbsolutePath.Replace('\\', '/')));
+					CompileAction.DependencyListFile = DependencyListFile;
+					CompileAction.ProducedItems.Add(DependencyListFile);
 
 					CompileAction.ProducedItems.Add(ISPCIncludeHeaderFile);
 
@@ -1757,12 +1751,9 @@ namespace UnrealBuildTool
 					}
 
 					// Consume the included header dependency list
-					if (CompileEnvironment.bGenerateDependenciesFile)
-					{
-						FileItem DependencyListFile = FileItem.GetItemByFileReference(FileReference.Combine(OutputDir, InlineArchName(Path.GetFileName(ISPCFile.AbsolutePath) + ".d", Arch, true)));
-						CompileAction.DependencyListFile = DependencyListFile;
-						CompileAction.PrerequisiteItems.Add(DependencyListFile);
-					}
+					FileItem DependencyListFile = FileItem.GetItemByFileReference(FileReference.Combine(OutputDir, InlineArchName(Path.GetFileName(ISPCFile.AbsolutePath) + ".d", Arch, true)));
+					CompileAction.DependencyListFile = DependencyListFile;
+					CompileAction.PrerequisiteItems.Add(DependencyListFile);
 
 					CompileAction.ProducedItems.AddRange(CompiledISPCObjFiles);
 
diff --git a/Engine/Source/Programs/UnrealBuildTool/Platform/Windows/VCToolChain.cs b/Engine/Source/Programs/UnrealBuildTool/Platform/Windows/VCToolChain.cs
--- a/Engine/Source/Programs/UnrealBuildTool/Platform/Windows/VCToolChain.cs
+++ b/Engine/Source/Programs/UnrealBuildTool/Platform/Windows/VCToolChain.cs
@@ -1511,21 +1511,18 @@ namespace UnrealBuildTool
 					GenerateParseTimingInfoAction(SourceFile, CompileAction.TimingFile, Graph);
 				}
 
-				if (CompileEnvironment.bGenerateDependenciesFile)
+				if (Target.WindowsPlatform.Compiler.IsMSVC() && !CompileAction.ForceClFilter)
 				{
-					if (Target.WindowsPlatform.Compiler.IsMSVC() && !CompileAction.ForceClFilter)
-					{
-						CompileAction.DependencyListFile = FileItem.GetItemByFileReference(FileReference.Combine(OutputDir, $"{SourceFile.Location.GetFileName()}.dep.json"));
-					}
-					else if (Target.WindowsPlatform.Compiler.IsClang())
-					{
-						CompileAction.DependencyListFile = FileItem.GetItemByFileReference(FileReference.Combine(OutputDir, $"{SourceFile.Location.GetFileName()}.d"));
-					}
-					else
-					{
-						CompileAction.DependencyListFile = FileItem.GetItemByFileReference(FileReference.Combine(OutputDir, $"{SourceFile.Location.GetFileName()}.txt"));
-						CompileAction.bShowIncludes = Target.WindowsPlatform.bShowIncludes;
-					}
+					CompileAction.DependencyListFile = FileItem.GetItemByFileReference(FileReference.Combine(OutputDir, $"{SourceFile.Location.GetFileName()}.dep.json"));
+				}
+				else if (Target.WindowsPlatform.Compiler.IsClang())
+				{
+					CompileAction.DependencyListFile = FileItem.GetItemByFileReference(FileReference.Combine(OutputDir, $"{SourceFile.Location.GetFileName()}.d"));
+				}
+				else
+				{
+					CompileAction.DependencyListFile = FileItem.GetItemByFileReference(FileReference.Combine(OutputDir, $"{SourceFile.Location.GetFileName()}.txt"));
+					CompileAction.bShowIncludes = Target.WindowsPlatform.bShowIncludes;
 				}
 
 				// Allow derived toolchains to make further changes
diff --git a/Engine/Source/Programs/UnrealBuildTool/System/CppCompileEnvironment.cs b/Engine/Source/Programs/UnrealBuildTool/System/CppCompileEnvironment.cs
--- a/Engine/Source/Programs/UnrealBuildTool/System/CppCompileEnvironment.cs
+++ b/Engine/Source/Programs/UnrealBuildTool/System/CppCompileEnvironment.cs
@@ -378,11 +378,6 @@ namespace UnrealBuildTool
 		/// </summary>
 		public bool bPrintTimingInfo;
 
-		/// <summary>
-		/// Whether to output a dependencies file along with the output build products
-		/// </summary>
-		public bool bGenerateDependenciesFile = true;
-
 		/// <summary>
 		/// When enabled, allows XGE to compile pre-compiled header files on remote machines.  Otherwise, PCHs are always generated locally.
 		/// </summary>
@@ -561,7 +556,6 @@ namespace UnrealBuildTool
 			PGOFilenamePrefix = Other.PGOFilenamePrefix;
 			PGODirectory = Other.PGODirectory;
 			bPrintTimingInfo = Other.bPrintTimingInfo;
-			bGenerateDependenciesFile = Other.bGenerateDependenciesFile;
 			bAllowRemotelyCompiledPCHs = Other.bAllowRemotelyCompiledPCHs;
 			UserIncludePaths = new HashSet<DirectoryReference>(Other.UserIncludePaths);
 			SystemIncludePaths = new HashSet<DirectoryReference>(Other.SystemIncludePaths);
diff --git a/Engine/Source/Programs/UnrealBuildTool/ToolChain/ClangToolChain.cs b/Engine/Source/Programs/UnrealBuildTool/ToolChain/ClangToolChain.cs
--- a/Engine/Source/Programs/UnrealBuildTool/ToolChain/ClangToolChain.cs
+++ b/Engine/Source/Programs/UnrealBuildTool/ToolChain/ClangToolChain.cs
@@ -958,7 +958,7 @@ namespace UnrealBuildTool
 				}
 
 				// Generate the included header dependency list
-				if (!PreprocessDepends && CompileEnvironment.bGenerateDependenciesFile)
+				if (!PreprocessDepends)
 				{
 					FileItem DependencyListFile = FileItem.GetItemByFileReference(FileReference.Combine(OutputDir, Path.GetFileName(SourceFile.AbsolutePath) + ".d"));
 					Arguments.Add(GetDepencenciesListFileArgument(DependencyListFile));
@@ -1034,7 +1034,7 @@ namespace UnrealBuildTool
 				CompileEnvironment.bAllowRemotelyCompiledPCHs;
 
 			// Two-pass compile where the preprocessor is run first to output the dependency list
-			if (PreprocessDepends && CompileEnvironment.bGenerateDependenciesFile)
+			if (PreprocessDepends)
 			{
 				Action PrepassAction = Graph.CreateAction(ActionType.Compile);
 				PrepassAction.PrerequisiteItems.AddRange(CompileAction.PrerequisiteItems);
diff --git a/Engine/Source/Programs/UnrealBuildTool/ToolChain/ISPCToolChain.cs b/Engine/Source/Programs/UnrealBuildTool/ToolChain/ISPCToolChain.cs
--- a/Engine/Source/Programs/UnrealBuildTool/ToolChain/ISPCToolChain.cs
+++ b/Engine/Source/Programs/UnrealBuildTool/ToolChain/ISPCToolChain.cs
@@ -459,13 +459,10 @@ namespace UnrealBuildTool
 				Arguments.Add($"-h \"{NormalizeCommandLinePath(ISPCIncludeHeaderFile)}\"");
 
 				// Generate the included header dependency list
-				if (CompileEnvironment.bGenerateDependenciesFile)
-				{
-					FileItem DependencyListFile = FileItem.GetItemByFileReference(FileReference.Combine(OutputDir, Path.GetFileName(ISPCFile.AbsolutePath) + ".txt"));
-					Arguments.Add($"-MMM \"{NormalizeCommandLinePath(DependencyListFile)}\"");
-					CompileAction.DependencyListFile = DependencyListFile;
-					CompileAction.ProducedItems.Add(DependencyListFile);
-				}
+				FileItem DependencyListFile = FileItem.GetItemByFileReference(FileReference.Combine(OutputDir, Path.GetFileName(ISPCFile.AbsolutePath) + ".txt"));
+				Arguments.Add($"-MMM \"{NormalizeCommandLinePath(DependencyListFile)}\"");
+				CompileAction.DependencyListFile = DependencyListFile;
+				CompileAction.ProducedItems.Add(DependencyListFile);
 
 				Arguments.AddRange(GlobalArguments);
 
@@ -671,12 +668,9 @@ namespace UnrealBuildTool
 				Arguments.AddRange(GlobalArguments);
 
 				// Consume the included header dependency list
-				if (CompileEnvironment.bGenerateDependenciesFile)
-				{
-					FileItem DependencyListFile = FileItem.GetItemByFileReference(FileReference.Combine(OutputDir, Path.GetFileName(ISPCFile.AbsolutePath) + ".txt"));
-					CompileAction.DependencyListFile = DependencyListFile;
-					CompileAction.PrerequisiteItems.Add(DependencyListFile);
-				}
+				FileItem DependencyListFile = FileItem.GetItemByFileReference(FileReference.Combine(OutputDir, Path.GetFileName(ISPCFile.AbsolutePath) + ".txt"));
+				CompileAction.DependencyListFile = DependencyListFile;
+				CompileAction.PrerequisiteItems.Add(DependencyListFile);
 
 				CompileAction.ProducedItems.AddRange(CompiledISPCObjFiles);
 
