Check for File nullable when tagging test builds
#rnx

diff --git a/Engine/Source/Programs/UnrealBuildTool/Configuration/TestTargetRules.cs b/Engine/Source/Programs/UnrealBuildTool/Configuration/TestTargetRules.cs
--- a/Engine/Source/Programs/UnrealBuildTool/Configuration/TestTargetRules.cs
+++ b/Engine/Source/Programs/UnrealBuildTool/Configuration/TestTargetRules.cs
@@ -110,9 +110,12 @@ namespace UnrealBuildTool
 				{
 					OutputName = OutputName + "-" + Target.Platform + "-" + Target.Configuration;
 				}
-				string OutputDirectory = UEBuildTarget.GetOutputDirectoryForExecutable(Unreal.EngineDirectory, this.File!).FullName;
-				string OutputFileName = $"{OutputDirectory}\\Binaries\\{Target.Platform}\\{ExeBinariesSubFolder}\\{OutputName}.exe.is_unreal_test";
-				PostBuildSteps.Add("echo > " + OutputFileName);
+				if (File != null)
+				{
+					string OutputDirectory = UEBuildTarget.GetOutputDirectoryForExecutable(Unreal.EngineDirectory, File).FullName;
+					string OutputFileName = $"{OutputDirectory}\\Binaries\\{Target.Platform}\\{ExeBinariesSubFolder}\\{OutputName}.exe.is_unreal_test";
+					PostBuildSteps.Add("echo > " + OutputFileName);
+				}
 			}
 		}
 
