UnrealBuildTool: Move some util properties to EpicGames.Build

#rnx

diff --git a/Engine/Source/Programs/Shared/EpicGames.Build/Unreal.cs b/Engine/Source/Programs/Shared/EpicGames.Build/Unreal.cs
--- a/Engine/Source/Programs/Shared/EpicGames.Build/Unreal.cs
+++ b/Engine/Source/Programs/Shared/EpicGames.Build/Unreal.cs
@@ -135,6 +135,16 @@ namespace UnrealBuildBase
 		/// </summary>
 		public static readonly DirectoryReference EngineSourceDirectory = DirectoryReference.Combine(EngineDirectory, "Source");
 
+		/// <summary>
+		/// Writable engine directory. Uses the user's settings folder for installed builds.
+		/// </summary>
+		public static readonly DirectoryReference WritableEngineDirectory = IsEngineInstalled() && UserSettingDirectory != null ? DirectoryReference.Combine(UserSettingDirectory, "UnrealEngine") : EngineDirectory;
+
+		/// <summary>
+		/// The engine saved programs directory
+		/// </summary>
+		public static readonly DirectoryReference EngineProgramSavedDirectory = IsEngineInstalled() && UserSettingDirectory != null ? UserSettingDirectory : DirectoryReference.Combine(EngineDirectory, "Programs");
+
 		/// <summary>
 		/// The path to UBT
 		/// </summary>
@@ -156,6 +166,13 @@ namespace UnrealBuildBase
 		/// </summary>
 		static public readonly FileReference DotnetPath = FileReference.Combine(DotnetDirectory, "dotnet" + RuntimePlatform.ExeExtension);
 
+		/// <summary>
+		/// Returns the User Settings Directory path. This matches FPlatformProcess::UserSettingsDir().
+		/// NOTE: This function may return null. Some accounts (eg. the SYSTEM account on Windows) do not have a personal folder, and Jenkins
+		/// runs using this account by default.
+		/// </summary>
+		static public readonly DirectoryReference? UserSettingDirectory = GetUserSettingDirectory();
+
 		/// <summary>
 		/// Whether we're running with engine installed
 		/// </summary>
@@ -314,5 +331,35 @@ namespace UnrealBuildBase
 		{
 			return GetExtensionDirs(BaseDir, bIncludePlatformDirectories, bIncludeRestrictedDirectories, bIncludeBaseDirectory).Select(x => DirectoryReference.Combine(x, SubDir)).Where(x => DirectoryReference.Exists(x)).ToList();
 		}
+
+		/// <summary>
+		/// Returns the User Settings Directory path. This matches FPlatformProcess::UserSettingsDir().
+		/// NOTE: This function may return null. Some accounts (eg. the SYSTEM account on Windows) do not have a personal folder, and Jenkins
+		/// runs using this account by default.
+		/// </summary>
+		private static DirectoryReference? GetUserSettingDirectory()
+		{
+			if (RuntimePlatform.IsMac)
+			{
+				return new DirectoryReference(Path.Combine(Environment.GetFolderPath(Environment.SpecialFolder.Personal), "Library", "Application Support", "Epic"));
+			}
+			else if (RuntimePlatform.IsLinux)
+			{
+				return new DirectoryReference(Path.Combine(Environment.GetFolderPath(Environment.SpecialFolder.ApplicationData), "Epic"));
+			}
+			else
+			{
+				// Not all user accounts have a local application data directory (eg. SYSTEM, used by Jenkins for builds).
+				string DirectoryName = Environment.GetFolderPath(Environment.SpecialFolder.LocalApplicationData);
+				if (String.IsNullOrEmpty(DirectoryName))
+				{
+					return DirectoryReference.Combine(EngineDirectory, "Saved");
+				}
+				else
+				{
+					return new DirectoryReference(DirectoryName);
+				}
+			}
+		}
 	}
 }
diff --git a/Engine/Source/Programs/UnrealBuildTool/Modes/BuildMode.cs b/Engine/Source/Programs/UnrealBuildTool/Modes/BuildMode.cs
--- a/Engine/Source/Programs/UnrealBuildTool/Modes/BuildMode.cs
+++ b/Engine/Source/Programs/UnrealBuildTool/Modes/BuildMode.cs
@@ -121,7 +121,7 @@ namespace UnrealBuildTool
 			// Fixup the log path if it wasn't overridden by a config file
 			if (BaseLogFileName == null)
 			{
-				BaseLogFileName = FileReference.Combine(UnrealBuildTool.EngineProgramSavedDirectory, "UnrealBuildTool", "Log.txt").FullName;
+				BaseLogFileName = FileReference.Combine(Unreal.EngineProgramSavedDirectory, "UnrealBuildTool", "Log.txt").FullName;
 			}
 
 			// Create the log file, and flush the startup listener to it
@@ -777,6 +777,7 @@ namespace UnrealBuildTool
 				// Save the makefile for next time
 				if(MakefileLocation != null)
 				{
+					Debugger.Launch();
 					using (GlobalTracer.Instance.BuildSpan("TargetMakefile.Save()").StartActive())
 					{
 						Makefile.Save(MakefileLocation);
diff --git a/Engine/Source/Programs/UnrealBuildTool/Modes/GenerateProjectFilesMode.cs b/Engine/Source/Programs/UnrealBuildTool/Modes/GenerateProjectFilesMode.cs
--- a/Engine/Source/Programs/UnrealBuildTool/Modes/GenerateProjectFilesMode.cs
+++ b/Engine/Source/Programs/UnrealBuildTool/Modes/GenerateProjectFilesMode.cs
@@ -70,7 +70,7 @@ namespace UnrealBuildTool
 			XmlConfig.ApplyTo(this);
 
 			// set up logging (taken from BuildMode)
-			FileReference LogFile = FileReference.Combine(UnrealBuildTool.EngineProgramSavedDirectory, "UnrealBuildTool", "Log_GPF.txt");
+			FileReference LogFile = FileReference.Combine(Unreal.EngineProgramSavedDirectory, "UnrealBuildTool", "Log_GPF.txt");
 			Log.AddFileWriter("DefaultLogTraceListener", LogFile);
 
 			// Parse rocket-specific arguments.
diff --git a/Engine/Source/Programs/UnrealBuildTool/Modes/UnrealHeaderToolMode.cs b/Engine/Source/Programs/UnrealBuildTool/Modes/UnrealHeaderToolMode.cs
--- a/Engine/Source/Programs/UnrealBuildTool/Modes/UnrealHeaderToolMode.cs
+++ b/Engine/Source/Programs/UnrealBuildTool/Modes/UnrealHeaderToolMode.cs
@@ -933,7 +933,7 @@ namespace UnrealBuildTool.Modes
 				{
 					if (Unreal.IsEngineInstalled())
 					{
-						CachedEngineProgramSavedDirectory = Utils.GetUserSettingDirectory() ?? DirectoryReference.Combine(Unreal.EngineDirectory, "Programs");
+						CachedEngineProgramSavedDirectory = Unreal.UserSettingDirectory ?? DirectoryReference.Combine(Unreal.EngineDirectory, "Programs");
 					}
 					else
 					{
diff --git a/Engine/Source/Programs/UnrealBuildTool/System/ConfigHierarchy.cs b/Engine/Source/Programs/UnrealBuildTool/System/ConfigHierarchy.cs
--- a/Engine/Source/Programs/UnrealBuildTool/System/ConfigHierarchy.cs
+++ b/Engine/Source/Programs/UnrealBuildTool/System/ConfigHierarchy.cs
@@ -983,7 +983,7 @@ namespace UnrealBuildTool
 		{
 			string OutString = InString.Replace("{TYPE}", BaseIniName);
 
-			DirectoryReference? UserSettingsDir = Utils.GetUserSettingDirectory();
+			DirectoryReference? UserSettingsDir = Unreal.UserSettingDirectory;
 			if (UserSettingsDir != null)
 			{
 				OutString = OutString.Replace("{USERSETTINGS}", UserSettingsDir.FullName);
@@ -1178,7 +1178,7 @@ namespace UnrealBuildTool
 		{
 			if(Unreal.IsEngineInstalled())
 			{
-				DirectoryReference? UserSettingDir = Utils.GetUserSettingDirectory();
+				DirectoryReference? UserSettingDir = Unreal.UserSettingDirectory;
 				if (UserSettingDir != null)
 				{
 					return DirectoryReference.Combine(UserSettingDir, "UnrealEngine", String.Format("{0}.{1}", ReadOnlyBuildVersion.Current.MajorVersion, ReadOnlyBuildVersion.Current.MinorVersion), "Saved");
diff --git a/Engine/Source/Programs/UnrealBuildTool/System/RulesCompiler.cs b/Engine/Source/Programs/UnrealBuildTool/System/RulesCompiler.cs
--- a/Engine/Source/Programs/UnrealBuildTool/System/RulesCompiler.cs
+++ b/Engine/Source/Programs/UnrealBuildTool/System/RulesCompiler.cs
@@ -214,7 +214,7 @@ namespace UnrealBuildTool
 			RulesAssembly Result = Parent;
 			if (ModuleFileToContext.Count > 0)
 			{
-				FileReference AssemblyFileName = FileReference.Combine(UnrealBuildTool.WritableEngineDirectory, "Intermediate", "Build", "BuildRules", "MarketplaceRules.dll");
+				FileReference AssemblyFileName = FileReference.Combine(Unreal.WritableEngineDirectory, "Intermediate", "Build", "BuildRules", "MarketplaceRules.dll");
 				Result = new RulesAssembly(MarketplaceScope, new List<DirectoryReference> { DirectoryReference.Combine(Unreal.EngineDirectory, "Plugins", "Marketplace") }, Plugins, ModuleFileToContext, new List<FileReference>(), AssemblyFileName, bContainsEngineModules: true, DefaultBuildSettings: DefaultEngineBuildSettingsVersion, bReadOnly: bReadOnly, bSkipCompile: bSkipCompile, bForceCompile: bForceCompile, Parent: Parent, Logger: Logger);
 			}
 			return Result;
diff --git a/Engine/Source/Programs/UnrealBuildTool/System/UnrealPluginLanguage.cs b/Engine/Source/Programs/UnrealBuildTool/System/UnrealPluginLanguage.cs
--- a/Engine/Source/Programs/UnrealBuildTool/System/UnrealPluginLanguage.cs
+++ b/Engine/Source/Programs/UnrealBuildTool/System/UnrealPluginLanguage.cs
@@ -3498,7 +3498,7 @@ namespace UnrealBuildTool
 				yield return FileReference.Combine(ProjectDirectory, "Config", PlatformName, PlatformName + BaseIniName + ".ini");
 			}
 
-			DirectoryReference UserSettingsFolder = Utils.GetUserSettingDirectory(); // Match FPlatformProcess::UserSettingsDir()
+			DirectoryReference? UserSettingsFolder = Unreal.UserSettingDirectory; // Match FPlatformProcess::UserSettingsDir()
 			DirectoryReference? PersonalFolder = null; // Match FPlatformProcess::UserDir()
 			if (RuntimePlatform.IsMac || RuntimePlatform.IsLinux)
 			{
diff --git a/Engine/Source/Programs/UnrealBuildTool/System/XmlConfig.cs b/Engine/Source/Programs/UnrealBuildTool/System/XmlConfig.cs
--- a/Engine/Source/Programs/UnrealBuildTool/System/XmlConfig.cs
+++ b/Engine/Source/Programs/UnrealBuildTool/System/XmlConfig.cs
@@ -92,7 +92,7 @@ namespace UnrealBuildTool
 				CacheFile = FileReference.Combine(Unreal.EngineDirectory, "Intermediate", "Build", "XmlConfigCache.bin");
 				if(Unreal.IsEngineInstalled())
 				{
-					DirectoryReference? UserSettingsDir = Utils.GetUserSettingDirectory();
+					DirectoryReference? UserSettingsDir = Unreal.UserSettingDirectory;
 					if(UserSettingsDir != null)
 					{
 						CacheFile = FileReference.Combine(UserSettingsDir, "UnrealEngine", String.Format("XmlConfigCache-{0}.bin", Unreal.RootDirectory.FullName.Replace(":", "").Replace(Path.DirectorySeparatorChar, '+')));
diff --git a/Engine/Source/Programs/UnrealBuildTool/UnrealBuildTool.cs b/Engine/Source/Programs/UnrealBuildTool/UnrealBuildTool.cs
--- a/Engine/Source/Programs/UnrealBuildTool/UnrealBuildTool.cs
+++ b/Engine/Source/Programs/UnrealBuildTool/UnrealBuildTool.cs
@@ -40,21 +40,11 @@ namespace UnrealBuildTool
 		/// </summary>
 		static FileReference? InstalledProjectFile;
 
-		/// <summary>
-		/// Directory for saved application settings (typically Engine/Programs)
-		/// </summary>
-		static DirectoryReference? CachedEngineProgramSavedDirectory;
-
 		/// <summary>
 		/// The full name of the Engine/Source directory
 		/// </summary>
 		[Obsolete("Replace with Unreal.EngineSourceDirectory")]
-		public static readonly DirectoryReference EngineSourceDirectory = DirectoryReference.Combine(Unreal.EngineDirectory, "Source");
-
-		/// <summary>
-		/// Cached copy of the writable engine directory
-		/// </summary>
-		static DirectoryReference? CachedWritableEngineDirectory;
+		public static readonly DirectoryReference EngineSourceDirectory = Unreal.EngineSourceDirectory;
 
 		/// <summary>
 		/// Cached copy of the source root directory that was used to compile the installed engine
@@ -65,51 +55,14 @@ namespace UnrealBuildTool
 		/// <summary>
 		/// Writable engine directory. Uses the user's settings folder for installed builds.
 		/// </summary>
-		public static DirectoryReference WritableEngineDirectory
-		{
-			get
-			{
-				if (CachedWritableEngineDirectory == null)
-				{
-					DirectoryReference? UserDir = null;
-					if (Unreal.IsEngineInstalled())
-					{
-						UserDir = Utils.GetUserSettingDirectory();
-					}
-					if (UserDir == null)
-					{
-						CachedWritableEngineDirectory = Unreal.EngineDirectory;
-					}
-					else
-					{
-						CachedWritableEngineDirectory = DirectoryReference.Combine(UserDir, "UnrealEngine");
-					}
-				}
-				return CachedWritableEngineDirectory;
-			}
-		}
+		[Obsolete("Replace with Unreal.WritableEngineDirectory")]
+		public static DirectoryReference WritableEngineDirectory = Unreal.WritableEngineDirectory;
 
 		/// <summary>
 		/// The engine programs directory
 		/// </summary>
-		public static DirectoryReference EngineProgramSavedDirectory
-		{
-			get
-			{
-				if (CachedEngineProgramSavedDirectory == null)
-				{
-					if (Unreal.IsEngineInstalled())
-					{
-						CachedEngineProgramSavedDirectory = Utils.GetUserSettingDirectory() ?? DirectoryReference.Combine(Unreal.EngineDirectory, "Programs");
-					}
-					else
-					{
-						CachedEngineProgramSavedDirectory = DirectoryReference.Combine(Unreal.EngineDirectory, "Programs");
-					}
-				}
-				return CachedEngineProgramSavedDirectory;
-			}
-		}
+		[Obsolete("Replace with Unreal.EngineProgramSavedDirectory")]
+		public static DirectoryReference EngineProgramSavedDirectory = Unreal.EngineProgramSavedDirectory;
 
 		/// <summary>
 		/// The original root directory that was used to compile the installed engine
@@ -227,10 +180,7 @@ namespace UnrealBuildTool
 		/// </summary>
 		/// <returns>A string containing the path to the UBT assembly.</returns>
 		[Obsolete("Deprecated in UE5.1 - use UnrealBuildTool.DotnetPath Unreal.UnrealBuildToolDllPath")]
-		static public FileReference GetUBTPath()
-		{
-			return Unreal.UnrealBuildToolPath;
-		}
+		static public FileReference GetUBTPath() => Unreal.UnrealBuildToolPath;
 
 		/// <summary>
 		/// The Unreal remote tool ini directory.  This should be valid if compiling using a remote server
