Fixing UnrealBuildTool deprecation warnings.

diff --git a/Engine/Source/Programs/UnrealBuildTool/Modes/ServerMode.cs b/Engine/Source/Programs/UnrealBuildTool/Modes/ServerMode.cs
--- a/Engine/Source/Programs/UnrealBuildTool/Modes/ServerMode.cs
+++ b/Engine/Source/Programs/UnrealBuildTool/Modes/ServerMode.cs
@@ -188,7 +188,7 @@ namespace UnrealBuildTool
 
 			if (LogDirectory == null)
 			{
-				LogDirectory = DirectoryReference.Combine(UnrealBuildTool.EngineProgramSavedDirectory, "UnrealBuildTool");
+				LogDirectory = DirectoryReference.Combine(Unreal.EngineProgramSavedDirectory, "UnrealBuildTool");
 			}
 
 			FileReference LogFile = FileReference.Combine(LogDirectory, "Log_Server.txt");
