Removed diagnostic code that was launching the debugger

diff --git a/Engine/Source/Programs/UnrealBuildTool/Modes/BuildMode.cs b/Engine/Source/Programs/UnrealBuildTool/Modes/BuildMode.cs
--- a/Engine/Source/Programs/UnrealBuildTool/Modes/BuildMode.cs
+++ b/Engine/Source/Programs/UnrealBuildTool/Modes/BuildMode.cs
@@ -777,7 +777,6 @@ namespace UnrealBuildTool
 				// Save the makefile for next time
 				if(MakefileLocation != null)
 				{
-					Debugger.Launch();
 					using (GlobalTracer.Instance.BuildSpan("TargetMakefile.Save()").StartActive())
 					{
 						Makefile.Save(MakefileLocation);
