UnrealBuildTool: Obsolete some properties moved to EpicGames.Build

#rnx

diff --git a/Engine/Source/Programs/UnrealBuildTool/System/Utils.cs b/Engine/Source/Programs/UnrealBuildTool/System/Utils.cs
--- a/Engine/Source/Programs/UnrealBuildTool/System/Utils.cs
+++ b/Engine/Source/Programs/UnrealBuildTool/System/Utils.cs
@@ -1013,30 +1013,8 @@ namespace UnrealBuildTool
 		/// NOTE: This function may return null. Some accounts (eg. the SYSTEM account on Windows) do not have a personal folder, and Jenkins
 		/// runs using this account by default.
 		/// </summary>
-		public static DirectoryReference GetUserSettingDirectory()
-		{
-			if (RuntimePlatform.IsMac)
-			{
-				return new DirectoryReference(Path.Combine(Environment.GetFolderPath(Environment.SpecialFolder.Personal), "Library", "Application Support", "Epic"));
-			}
-			else if (RuntimePlatform.IsLinux)
-			{
-				return new DirectoryReference(Path.Combine(Environment.GetFolderPath(Environment.SpecialFolder.ApplicationData), "Epic"));
-			}
-			else
-			{
-				// Not all user accounts have a local application data directory (eg. SYSTEM, used by Jenkins for builds).
-				string DirectoryName = Environment.GetFolderPath(Environment.SpecialFolder.LocalApplicationData);
-				if(String.IsNullOrEmpty(DirectoryName))
-				{
-					return DirectoryReference.Combine(Unreal.EngineDirectory, "Saved");
-				}
-				else
-				{
-					return new DirectoryReference(DirectoryName);
-				}
-			}
-		}
+		[Obsolete("Replace with Unreal.UserSettingDirectory")]
+		public static DirectoryReference? GetUserSettingDirectory() => Unreal.UserSettingDirectory;
 
 		enum LOGICAL_PROCESSOR_RELATIONSHIP
 		{
