Disable fma contraction on macs for ddc consistency between arm/x64

diff --git a/Engine/Source/Programs/UnrealBuildTool/Platform/Mac/MacToolChain.cs b/Engine/Source/Programs/UnrealBuildTool/Platform/Mac/MacToolChain.cs
--- a/Engine/Source/Programs/UnrealBuildTool/Platform/Mac/MacToolChain.cs
+++ b/Engine/Source/Programs/UnrealBuildTool/Platform/Mac/MacToolChain.cs
@@ -270,6 +270,12 @@ namespace UnrealBuildTool
 
 			Arguments.Add("-fasm-blocks");
 
+			// Disable FMA contraction on arm builds in order to preserve consistency with x64
+			// derived data builds.
+			// This is also added to the x64 builds in case someone ever adds -mfma.
+			Arguments.Add("-ffp-contract=off");
+			
+
 			if (CompileEnvironment.bEnableOSX109Support)
 			{
 				Arguments.Add("-faligned-new"); // aligned operator new is supported only on macOS 10.14 and above
