Fix bailing on target setup when the same path is included in multiple modules.

diff --git a/Engine/Source/Programs/UnrealBuildTool/Modes/ServerMode.cs b/Engine/Source/Programs/UnrealBuildTool/Modes/ServerMode.cs
--- a/Engine/Source/Programs/UnrealBuildTool/Modes/ServerMode.cs
+++ b/Engine/Source/Programs/UnrealBuildTool/Modes/ServerMode.cs
@@ -191,6 +191,8 @@ namespace UnrealBuildTool
 				LogDirectory = DirectoryReference.Combine(Unreal.EngineProgramSavedDirectory, "UnrealBuildTool");
 			}
 
+			DirectoryReference.CreateDirectory(LogDirectory);
+
 			FileReference LogFile = FileReference.Combine(LogDirectory, "Log_Server.txt");
 			Log.AddFileWriter("DefaultLogTraceListener", LogFile);
 			Logger.LogInformation("Server mode started");
@@ -399,11 +401,11 @@ namespace UnrealBuildTool
 						foreach (DirectoryReference Dir in Module.ModuleDirectories)
 						{
 							BrowseConfigurationFolders.Add(Dir.ToString());
-							CurrentTargetIntellisenseInfo.DirToModule.Add(Dir, Module);
+							CurrentTargetIntellisenseInfo.DirToModule.TryAdd(Dir, Module);
 						}
 						if (Module.GeneratedCodeDirectory != null)
 						{
-							CurrentTargetIntellisenseInfo.DirToModule.Add(Module.GeneratedCodeDirectory, Module);
+							CurrentTargetIntellisenseInfo.DirToModule.TryAdd(Module.GeneratedCodeDirectory, Module);
 						}
 
 						foreach (DirectoryReference Dir in ModuleCompileEnvironment.SystemIncludePaths)
