[UnrealBuildTool]
* Added so if strict conformance is disabled at least hiddenfriend conformance is enabled. This should have a minimal impact (if any) on licensees and enable us to submit huge compile time optimizations

diff --git a/Engine/Source/Programs/UnrealBuildTool/Platform/Windows/VCToolChain.cs b/Engine/Source/Programs/UnrealBuildTool/Platform/Windows/VCToolChain.cs
--- a/Engine/Source/Programs/UnrealBuildTool/Platform/Windows/VCToolChain.cs
+++ b/Engine/Source/Programs/UnrealBuildTool/Platform/Windows/VCToolChain.cs
@@ -435,6 +435,10 @@ namespace UnrealBuildTool
 						Arguments.Add("/Zc:lambda-");
 					}
 				}
+				else
+				{
+					Arguments.Add("/Zc:hiddenFriend");
+				}
 
 				if (Target.WindowsPlatform.bUpdatedCPPMacro)
 				{
