UnrealBuildTool: Throw an error if multiple modules are loaded that have conflicting ShortNames. Suggestion from UDN

#rnx

diff --git a/Engine/Source/Programs/UnrealBuildTool/Configuration/UEBuildTarget.cs b/Engine/Source/Programs/UnrealBuildTool/Configuration/UEBuildTarget.cs
--- a/Engine/Source/Programs/UnrealBuildTool/Configuration/UEBuildTarget.cs
+++ b/Engine/Source/Programs/UnrealBuildTool/Configuration/UEBuildTarget.cs
@@ -4552,6 +4552,20 @@ namespace UnrealBuildTool
 
 				// Now, go ahead and create the module builder instance
 				Module = InstantiateModule(RulesObject, GeneratedCodeDirectory, Logger);
+
+				// Check if this module conflicts with the ShortName of any other existing modules
+				UEBuildModule? ConflictingModule = Modules.Values.FirstOrDefault(x => (x.Rules.ShortName ?? x.Name) == (Module.Rules.ShortName ?? Module.Name));
+				if (ConflictingModule != null)
+				{
+					Logger.LogError("Conflicting module short names found, unable to create module {Name}", Module.Name);
+					Logger.LogInformation(" * Name={Name} ShortName={ShortName} Path={Path}", ConflictingModule.Name, ConflictingModule.Rules.ShortName ?? ConflictingModule.Name, ConflictingModule.RulesFile);
+					Logger.LogInformation(" * Name={Name} ShortName={ShortName} Path={Path}", Module.Name, Module.Rules.ShortName ?? Module.Name, Module.RulesFile);
+
+					throw new BuildException("Unable to create module {0} (ShortName={1}) as it conflicts with existing module {2} (ShortName={3})",
+						Module.Name, Module.Rules.ShortName ?? Module.Name,
+						ConflictingModule.Name, ConflictingModule.Rules.ShortName ?? ConflictingModule.Name);
+				}
+
 				Modules.Add(Module.Name, Module);
 
 				// Module must not have Verse if Verse is not enabled
