[Backout] - CL23048035
Original CL Desc
-----------------------------------------------------------------
Fixing UnrealBuildTool deprecation warnings.

Cherry pick from UE5-main

diff --git a/Engine/Source/Programs/UnrealBuildTool/Modes/ServerMode.cs b/Engine/Source/Programs/UnrealBuildTool/Modes/ServerMode.cs
--- a/Engine/Source/Programs/UnrealBuildTool/Modes/ServerMode.cs
+++ b/Engine/Source/Programs/UnrealBuildTool/Modes/ServerMode.cs
@@ -188,7 +188,7 @@ namespace UnrealBuildTool
 
 			if (LogDirectory == null)
 			{
-				LogDirectory = DirectoryReference.Combine(Unreal.EngineProgramSavedDirectory, "UnrealBuildTool");
+				LogDirectory = DirectoryReference.Combine(UnrealBuildTool.EngineProgramSavedDirectory, "UnrealBuildTool");
 			}
 
 			DirectoryReference.CreateDirectory(LogDirectory);
