- If CommandLineTools is detected as the build tools location (xcode-select -p), throw an exception with instructions on how to fix
- Stop logging out the xcode-select output every run

diff --git a/Engine/Source/Programs/UnrealBuildTool/ToolChain/AppleToolChain.cs b/Engine/Source/Programs/UnrealBuildTool/ToolChain/AppleToolChain.cs
--- a/Engine/Source/Programs/UnrealBuildTool/ToolChain/AppleToolChain.cs
+++ b/Engine/Source/Programs/UnrealBuildTool/ToolChain/AppleToolChain.cs
@@ -33,8 +33,14 @@ namespace UnrealBuildTool
 			{
 				Reason = "xcode-select";
 
-				// on the Mac, run xcode-select directly
-				DeveloperDir = Utils.RunLocalProcessAndReturnStdOut("xcode-select", "--print-path", Logger);
+				// on the Mac, run xcode-select directly.
+				int ReturnCode;
+				DeveloperDir = Utils.RunLocalProcessAndReturnStdOut("xcode-select", "--print-path", null, out ReturnCode);
+				if (ReturnCode != 0)
+				{
+					string? MinVersion = UEBuildPlatform.GetSDK(UnrealTargetPlatform.Mac)!.GetSDKInfo("Sdk")!.Min;
+					throw new BuildException($"We were unable to find your build tools (via 'xcode-select --print-path'). Please install Xcode, version {MinVersion} or later");
+				}
 
 				// make sure we get a full path
 				if (Directory.Exists(DeveloperDir) == false)
@@ -42,6 +48,14 @@ namespace UnrealBuildTool
 					throw new BuildException("Selected Xcode ('{0}') doesn't exist, cannot continue.", DeveloperDir);
 				}
 
+				if (DeveloperDir.Contains("CommandLineTools", StringComparison.InvariantCultureIgnoreCase))
+				{
+					throw new BuildException($"Your Mac is set to use CommandLineTools for its build tools ({DeveloperDir}). Unreal expects Xcode as the build tools. Please install Xcode if it's not already, then do one of the following:\n" +
+						"  - Run Xcode, go to Settings, and in the Locations tab, choose your Xcode in Command Line Tools dropdown.\n" +
+						"  - In Terminal, run 'sudo xcode-select -s /Applications/Xcode.app' (or an alternate location if you installed Xcode to a non-standard location)\n" + 
+						"Either way, you will need to enter your Mac password.");
+				}
+
 				if (DeveloperDir.EndsWith("/") == false)
 				{
 					// we expect this to end with a slash
