[UnrealBuildTool]
* Enabled _cplusplus preprocessor define representing correct number

diff --git a/Engine/Source/Programs/UnrealBuildTool/Platform/Windows/UEBuildWindows.cs b/Engine/Source/Programs/UnrealBuildTool/Platform/Windows/UEBuildWindows.cs
--- a/Engine/Source/Programs/UnrealBuildTool/Platform/Windows/UEBuildWindows.cs
+++ b/Engine/Source/Programs/UnrealBuildTool/Platform/Windows/UEBuildWindows.cs
@@ -306,7 +306,7 @@ namespace UnrealBuildTool
 		[RequiresUniqueBuildEnvironment]
 		[XmlConfigFile(Category = "WindowsPlatform")]
 		[CommandLine("-UpdatedCPPMacro")]
-		public bool bUpdatedCPPMacro = false;
+		public bool bUpdatedCPPMacro = true;
 
 		/// <summary>
 		/// Enables inline conformance (Remove unreferenced COMDAT) (/Zc:inline).
