Undo changelist 23056435

Will fix a file first and then resubmit

diff --git a/Engine/Source/Programs/UnrealBuildTool/Platform/Windows/UEBuildWindows.cs b/Engine/Source/Programs/UnrealBuildTool/Platform/Windows/UEBuildWindows.cs
--- a/Engine/Source/Programs/UnrealBuildTool/Platform/Windows/UEBuildWindows.cs
+++ b/Engine/Source/Programs/UnrealBuildTool/Platform/Windows/UEBuildWindows.cs
@@ -306,7 +306,7 @@ namespace UnrealBuildTool
 		[RequiresUniqueBuildEnvironment]
 		[XmlConfigFile(Category = "WindowsPlatform")]
 		[CommandLine("-UpdatedCPPMacro")]
-		public bool bUpdatedCPPMacro = true;
+		public bool bUpdatedCPPMacro = false;
 
 		/// <summary>
 		/// Enables inline conformance (Remove unreferenced COMDAT) (/Zc:inline).
