UnrealBuildTool: Update preferred vs2022 msvc version to [REDACTED_PHONE]

#rnx

diff --git a/Engine/Source/Programs/UnrealBuildTool/Platform/Windows/MicrosoftPlatformSDK.Versions.cs b/Engine/Source/Programs/UnrealBuildTool/Platform/Windows/MicrosoftPlatformSDK.Versions.cs
--- a/Engine/Source/Programs/UnrealBuildTool/Platform/Windows/MicrosoftPlatformSDK.Versions.cs
+++ b/Engine/Source/Programs/UnrealBuildTool/Platform/Windows/MicrosoftPlatformSDK.Versions.cs
@@ -33,7 +33,7 @@ namespace UnrealBuildTool
 		static readonly VersionNumberRange[] PreferredVisualCppVersions = new VersionNumberRange[]
 		{
 			VersionNumberRange.Parse("14.29.30133", "14.29.99999"), // VS2019 16.11.x
-			VersionNumberRange.Parse("14.33.31629", "14.33.99999"), // VS2022 17.3.x
+			VersionNumberRange.Parse("14.34.31933", "14.34.99999"), // VS2022 17.4.x
 		};
 
 		/// <summary>
@@ -85,8 +85,8 @@ namespace UnrealBuildTool
 		{
 			"Microsoft.VisualStudio.Workload.ManagedDesktop",
 			"Microsoft.VisualStudio.Workload.NativeCrossPlat",
-			"Microsoft.VisualStudio.Component.VC.14.33.17.3.ARM64",
-			"Microsoft.VisualStudio.Component.VC.14.33.17.3.x86.x64",
+			"Microsoft.VisualStudio.Component.VC.14.34.17.4.ARM64",
+			"Microsoft.VisualStudio.Component.VC.14.34.17.4.x86.x64",
 			"Microsoft.Net.Component.4.6.2.TargetingPack",
 		};
 
