- Fix for bad bundle id for editor builds in installed build projects, related to the plist changes needed for Xcode14

diff --git a/Engine/Source/Programs/UnrealBuildTool/Platform/Mac/UEDeployMac.cs b/Engine/Source/Programs/UnrealBuildTool/Platform/Mac/UEDeployMac.cs
--- a/Engine/Source/Programs/UnrealBuildTool/Platform/Mac/UEDeployMac.cs
+++ b/Engine/Source/Programs/UnrealBuildTool/Platform/Mac/UEDeployMac.cs
@@ -55,7 +55,17 @@ namespace UnrealBuildTool
             Ini.GetString("/Script/IOSRuntimeSettings.IOSRuntimeSettings", "BundleIdentifier", out BundleIdentifier);
 
             string BundleVersion = MacToolChain.LoadEngineDisplayVersion();
-			PListData = PListData.Replace("${EXECUTABLE_NAME}", ExeName).Replace("${APP_NAME}", BundleIdentifier.Replace("[PROJECT_NAME]", ProjectName).Replace("_", "")).Replace("${ICON_NAME}", GameName).Replace("${MACOSX_DEPLOYMENT_TARGET}", MacToolChain.Settings.MinMacOSVersion).Replace("${BUNDLE_VERSION}", BundleVersion);
+			// duplicating some logic in MacToolchain for the BundleID
+			string[] ExeNameParts = ExeName.Split('-');
+			bool bBuildingEditor = ExeNameParts[0].EndsWith("Editor");
+			string FinalBundleID = bBuildingEditor ? $"com.epicgames.{ExeNameParts[0]}" : BundleIdentifier.Replace("[PROJECT_NAME]", ProjectName);
+			FinalBundleID = FinalBundleID.Replace("_", "");
+
+			PListData = PListData.Replace("${EXECUTABLE_NAME}", ExeName).
+				Replace("${APP_NAME}", FinalBundleID).
+				Replace("${ICON_NAME}", GameName).
+				Replace("${MACOSX_DEPLOYMENT_TARGET}", MacToolChain.Settings.MinMacOSVersion).
+				Replace("${BUNDLE_VERSION}", BundleVersion);
 
 			if (!Directory.Exists(IntermediateDirectory))
 			{
