UBT: XmlConfigFile: Fixed potential exceptions handling deprecated properties and improved diagnostic messages when using renamed members (NewAttributeName).

#rnx

diff --git a/Engine/Source/Programs/UnrealBuildTool/System/XmlConfig.cs b/Engine/Source/Programs/UnrealBuildTool/System/XmlConfig.cs
--- a/Engine/Source/Programs/UnrealBuildTool/System/XmlConfig.cs
+++ b/Engine/Source/Programs/UnrealBuildTool/System/XmlConfig.cs
@@ -348,13 +348,19 @@ namespace UnrealBuildTool
 							// Check if setting has been deprecated
 							if (FieldValue.XmlConfigAttribute.Deprecated)
 							{
+								string CurrentSettingName = FieldValue.XmlConfigAttribute.Name != null ? FieldValue.XmlConfigAttribute.Name : FieldValue.Target.MemberInfo.Name;
+
 								Logger.LogWarning("Deprecated setting found in \"{SourceFile}\":", FieldValue.SourceFile);
-								Logger.LogWarning("The setting \"{Setting}\" is deprecated. Support for this setting will be removed in a future version of Unreal Engine.", FieldValue.Target.MemberInfo.Name);
+								Logger.LogWarning("The setting \"{Setting}\" is deprecated. Support for this setting will be removed in a future version of Unreal Engine.", CurrentSettingName);
 
 								if (FieldValue.XmlConfigAttribute.NewAttributeName != null)
 								{
-									Logger.LogWarning("Use \"{NewAttributeName}\" in place of \"{OldAttributeName}\"", FieldValue.XmlConfigAttribute.NewAttributeName, FieldValue.Target.MemberInfo.Name);
-									Logger.LogInformation("The value provided for deprecated setting \"{OldName}\" will be applied to \"{NewName}\"", FieldValue.Target.MemberInfo.Name, FieldValue.XmlConfigAttribute.NewAttributeName);
+									// NewAttributeName is the name of a member in code. However, the log messages below are written from the XML's perspective,
+									// so we need to check if the new target member is not exposed under a custom name in the config.
+									string NewSettingName = GetMemberConfigAttributeName(TargetType, FieldValue.XmlConfigAttribute.NewAttributeName);
+
+									Logger.LogWarning("Use \"{NewAttributeName}\" in place of \"{OldAttributeName}\"", NewSettingName, CurrentSettingName);
+									Logger.LogInformation("The value provided for deprecated setting \"{OldName}\" will be applied to \"{NewName}\"", CurrentSettingName, NewSettingName);
 
 									TargetToWrite = GetTargetMember(TargetType, FieldValue.XmlConfigAttribute.NewAttributeName) ?? TargetToWrite;
 								}
@@ -368,11 +374,27 @@ namespace UnrealBuildTool
 			}
 		}
 
+		private static string GetMemberConfigAttributeName(Type TargetType, string MemberName)
+		{
+			MemberInfo? MemberInfo = TargetType.GetRuntimeFields()
+				.FirstOrDefault(x => x.Name == MemberName);
+
+			if (MemberInfo == null)
+			{
+				MemberInfo = TargetType.GetRuntimeProperties()
+					.FirstOrDefault(x => x.Name == MemberName);
+			}
+
+			XmlConfigFileAttribute? Attribute = MemberInfo?.GetCustomAttributes<XmlConfigFileAttribute>().FirstOrDefault();
+
+			return Attribute?.Name ?? MemberName;
+		}
+
 		private static XmlConfigData.TargetMember? GetTargetMember(Type TargetType, string MemberName)
 		{
 			// First, try to find the new field to which the setting should be actually applied.
 			FieldInfo? FieldInfo = TargetType.GetRuntimeFields()
-				.First(x => x.Name == MemberName);
+				.FirstOrDefault(x => x.Name == MemberName);
 
 			if (FieldInfo != null)
 			{
@@ -381,7 +403,7 @@ namespace UnrealBuildTool
 
 			// If not found, try to find the new property to which the setting should be actually applied.
 			PropertyInfo? PropertyInfo = TargetType.GetRuntimeProperties()
-				.First(x => x.Name == MemberName);
+				.FirstOrDefault(x => x.Name == MemberName);
 
 			if (PropertyInfo != null)
 			{
