UBT: Edit and Continue: Unified the property names used to enable edit and continue.

On some platforms we used bUseLiveCodeEdit, which was expected to be present in the per-platform configuration block. At the same time, on Windows we rely on bSupportEditAndContinue from BuildConfiguration.

Modified the console platforms to use BuildConfiguration.bSupportEditAndContinue for consistency. However, there is some value in being able to control this setting per-platform, so I've left the existing properties and just renamed them bUseLiveCodeEdit -> bSupportEditAndContinue.

As a result, all platforms that support Edit and Continue are affected by BuildConfiguration.bSupportEditAndContinue. However, if XXPlatform.bSupportEditAndContinue is presetnt, it overrides the global.

Removed ClangToolChainOptions.EnableLiveCodeEditing and related code as it turns out it wasn't really used.

diff --git a/Engine/Source/Programs/UnrealBuildTool/Configuration/TargetRules.cs b/Engine/Source/Programs/UnrealBuildTool/Configuration/TargetRules.cs
--- a/Engine/Source/Programs/UnrealBuildTool/Configuration/TargetRules.cs
+++ b/Engine/Source/Programs/UnrealBuildTool/Configuration/TargetRules.cs
@@ -1613,7 +1613,7 @@ namespace UnrealBuildTool
 		public bool bPGOOptimize = false;
 
 		/// <summary>
-		/// Whether to support edit and continue.  Only works on Microsoft compilers.
+		/// Whether to support edit and continue.
 		/// </summary>
 		[XmlConfigFile(Category = "BuildConfiguration")]
 		public bool bSupportEditAndContinue = false;
diff --git a/Engine/Source/Programs/UnrealBuildTool/ToolChain/ClangToolChain.cs b/Engine/Source/Programs/UnrealBuildTool/ToolChain/ClangToolChain.cs
--- a/Engine/Source/Programs/UnrealBuildTool/ToolChain/ClangToolChain.cs
+++ b/Engine/Source/Programs/UnrealBuildTool/ToolChain/ClangToolChain.cs
@@ -94,20 +94,15 @@ namespace UnrealBuildTool
 		/// </summary>
 		GenerateSymbols = 1 << 13,
 
-		/// <summary>
-		/// Enables live code editing
-		/// </summary>
-		EnableLiveCodeEditing = 1 << 14,
-
 		/// <summary>
 		/// Enables dead code/data stripping and common code folding.
 		/// </summary>
-		EnableDeadStripping = 1 << 15,
+		EnableDeadStripping = 1 << 14,
 
 		/// <summary>
 		/// Indicates that the target is a moduler build i.e. Target.LinkType == TargetLinkType.Modular
 		/// </summary>
-		ModularBuild = 1 << 16,
+		ModularBuild = 1 << 15,
 	}
 
 	abstract class ClangToolChain : ISPCToolChain
