Added a mode that profiles different unity sizes and returns the best size for each module.

diff --git a/Engine/Source/Programs/UnrealBuildTool/Modes/ProfileUnitySizesMode.cs b/Engine/Source/Programs/UnrealBuildTool/Modes/ProfileUnitySizesMode.cs
new file mode 100644
--- /dev/null
+++ b/Engine/Source/Programs/UnrealBuildTool/Modes/ProfileUnitySizesMode.cs
@@ -0,0 +1,371 @@
+// Copyright Epic Games, Inc. All Rights Reserved.
+
+using System;
+using System.Collections.Generic;
+using System.Text.RegularExpressions;
+using EpicGames.Core;
+using UnrealBuildBase;
+using Microsoft.Extensions.Logging;
+using System.Linq;
+
+namespace UnrealBuildTool
+{
+	/// <summary>
+	/// Profiles different unity sizes and prints out the different size and its timings
+	/// </summary>
+	[ToolMode("ProfileUnitySizes", ToolModeOptions.XmlConfig | ToolModeOptions.BuildPlatforms | ToolModeOptions.SingleInstance | ToolModeOptions.StartPrefetchingEngine | ToolModeOptions.ShowExecutionTime)]
+	class ProfileUnitySizesMode : ToolMode
+	{
+		/// <summary>
+		/// Set of filters for files to include in the database. Relative to the root directory, or to the project file.
+		/// </summary>
+		[CommandLine("-Filter=")]
+		List<string> FilterRules = new List<string>();
+
+		class TimingData
+		{
+			public Double ExecutorTiming = 0;
+			public Double CPUTiming = 0;
+			public int UnitySize = 0;
+			public int NumFiles = 0;
+
+			public bool IsValid()
+			{
+				return ExecutorTiming != 0;
+			}
+		}
+
+		class TimingLogger : ILogger
+		{
+			static Regex ExecutorTimingRegex = new Regex(@"executor.*\s(\d+\.*\d*)\sseconds");
+			static Regex NumFilesRegex = new Regex(@"\[\d+/(\d+)\]");
+			static Regex CPUTimingRegex = new Regex(@"CPU Time:\s(\d+\.*\d*)");
+			private readonly ILogger Inner;
+			private readonly ILogger OldLogger;
+
+			public TimingData TimingData = new();
+
+			public TimingLogger(ILogger inner)
+			{
+				Inner = inner;
+				OldLogger = EpicGames.Core.Log.EventParser.Logger;
+				EpicGames.Core.Log.EventParser.Logger = this;
+			}
+
+			/// <inheritdoc/>
+			public IDisposable BeginScope<TState>(TState State)
+			{
+				return Inner.BeginScope(State);
+			}
+
+			/// <inheritdoc/>
+			public bool IsEnabled(LogLevel LogLevel)
+			{
+				return Inner.IsEnabled(LogLevel);
+			}
+
+			/// <inheritdoc/>
+			public void Log<TState>(LogLevel LogLevel, EventId EventId, TState State, Exception? Exception, Func<TState, Exception?, string> Formatter)
+			{
+				if (State != null)
+				{
+					string? LogText = State.ToString();
+					if (!string.IsNullOrEmpty(LogText))
+					{
+						// Console.WriteLine(LogText);
+						Match ExecutorTimingMatch = ExecutorTimingRegex.Match(LogText);
+						if (ExecutorTimingMatch.Success)
+						{
+							if (!Double.TryParse(ExecutorTimingMatch.Groups[1].Value, out TimingData.ExecutorTiming))
+							{
+								Console.WriteLine($"Failed to parse '{LogText}'");
+							}
+						}
+
+						Match CPUTimingMatch = CPUTimingRegex.Match(LogText);
+						if (CPUTimingMatch.Success)
+						{
+							if (!Double.TryParse(CPUTimingMatch.Groups[1].Value, out TimingData.CPUTiming))
+							{
+								Console.WriteLine($"Failed to parse '{LogText}'");
+							}
+						}
+
+						Match NumFilesMatch = NumFilesRegex.Match(LogText);
+						if (NumFilesMatch.Success)
+						{
+							if (!int.TryParse(NumFilesMatch.Groups[1].Value, out TimingData.NumFiles))
+							{
+								Console.WriteLine($"Failed to parse '{LogText}'");
+							}
+						}
+					}
+				}
+			}
+		}
+
+		/// <summary>
+		/// Execute the command
+		/// </summary>
+		/// <param name="Arguments">Command line arguments</param>
+		/// <returns>Exit code</returns>
+		/// <param name="Logger"></param>
+		public override int Execute(CommandLineArguments Arguments, ILogger Logger)
+		{
+			Arguments.ApplyTo(this);
+
+			// Create the build configuration object, and read the settings
+			BuildConfiguration BuildConfiguration = new BuildConfiguration();
+			XmlConfig.ApplyTo(BuildConfiguration);
+			Arguments.ApplyTo(BuildConfiguration);
+
+			// Parse the filter argument
+			FileFilter? FileFilter = null;
+			if (FilterRules.Count > 0)
+			{
+				FileFilter = new FileFilter(FileFilterType.Exclude);
+				foreach (string FilterRule in FilterRules)
+				{
+					FileFilter.AddRules(FilterRule.Split(';'));
+				}
+			}
+
+			// Force C++ modules to always include their generated code directories
+			UEBuildModuleCPP.bForceAddGeneratedCodeIncludePath = true;
+
+			// Parse all the target descriptors
+			List<TargetDescriptor> TargetDescriptors = TargetDescriptor.ParseCommandLine(Arguments, BuildConfiguration.bUsePrecompiled, BuildConfiguration.bSkipRulesCompile, BuildConfiguration.bForceRulesCompile, Logger);
+
+			foreach (TargetDescriptor TargetDescriptor in TargetDescriptors)
+			{
+				List<UEBuildModule> ModuleList = new();
+
+				// Create a makefile for the target
+				TimingLogger TimingLogger = new(Logger);
+				UEBuildTarget Target = UEBuildTarget.Create(TargetDescriptor, BuildConfiguration.bSkipRulesCompile, BuildConfiguration.bForceRulesCompile, BuildConfiguration.bUsePrecompiled, TimingLogger);
+				UEToolChain TargetToolChain = Target.CreateToolchain(Target.Platform);
+
+				CppCompileEnvironment GlobalCompileEnvironment = Target.CreateCompileEnvironmentForProjectFiles(TimingLogger);
+				foreach (UEBuildBinary Binary in Target.Binaries)
+				{
+					CppCompileEnvironment BinaryCompileEnvironment = Binary.CreateBinaryCompileEnvironment(GlobalCompileEnvironment);
+					foreach (UEBuildModule Module in Binary.Modules)
+					{
+						if (FileFilter == null || FileFilter.Matches(Module.RulesFile.MakeRelativeTo(Unreal.RootDirectory)))
+						{
+							if (Module.Rules.Type != ModuleRules.ModuleType.External &&
+								Module.Rules.bUseUnity)
+							{
+								ModuleList.Add(Module);
+							}
+						}
+					}
+				}
+
+				// build each Module
+				ModuleList.SortBy(module => module.Name);
+				foreach (UEBuildModule Module in ModuleList)
+				{					
+					CompileModule(BuildConfiguration, TargetDescriptor, Target, Module, Logger);
+				}
+			}
+
+			return 0;
+		}
+
+		/// <summary>
+		/// Compile the module multiple times looking for the best unity size 
+		/// </summary>
+		private void CompileModule(BuildConfiguration BuildConfiguration, TargetDescriptor TargetDescriptor, UEBuildTarget Target, UEBuildModule Module, ILogger Logger)
+		{
+			TargetDescriptor.OnlyModuleNames.Clear();
+			TargetDescriptor.OnlyModuleNames.Add(Module.Name);
+
+			const int UnitySizeDivision = 8;
+			const int TotalBuilds = UnitySizeDivision + 3;
+
+			Logger.LogInformation($"{Module.Name}:");
+
+			int CurrentModuleUnitySize = (Module.Rules.NumIncludedBytesPerUnityCPPOverride != 0) ? Module.Rules.NumIncludedBytesPerUnityCPPOverride : Target.Rules.NumIncludedBytesPerUnityCPP;
+			int TargetUnitySize = Target.Rules.NumIncludedBytesPerUnityCPP;
+
+			int BuildNum = 1;
+			CompileModule($"  [{BuildNum++}/{TotalBuilds}] ", BuildConfiguration, TargetDescriptor, Module, Logger, CurrentModuleUnitySize, true, false);
+			
+			TimingData CurrentCompileTime = GetBestCompileModuleTime($"  [{BuildNum++}/{TotalBuilds}] ", BuildConfiguration, TargetDescriptor, Module, Logger, CurrentModuleUnitySize, false, false);
+			if (!CurrentCompileTime.IsValid())
+			{
+				Logger.LogInformation($"Skipping module because it doesn't compile with current settings.");
+				return;
+			}
+
+			TimingData DisableUnityCompileTime = GetBestCompileModuleTime($"  [{BuildNum++}/{TotalBuilds}] ", BuildConfiguration, TargetDescriptor, Module, Logger, TargetUnitySize, false, true);
+
+			int MaxUnitySize = TargetUnitySize * 2;
+			List<TimingData> Timings = new();
+			int UnitySizeInc = MaxUnitySize / UnitySizeDivision;
+			int CurrentUnitySize = UnitySizeInc;
+			for (int UnitySizeIndex = 0; UnitySizeIndex < UnitySizeDivision; UnitySizeIndex++)
+			{
+				TimingData NewTiming = GetBestCompileModuleTime($"  [{BuildNum++}/{TotalBuilds}] ", BuildConfiguration, TargetDescriptor, Module, Logger, CurrentUnitySize, false, false);
+				Timings.Add(NewTiming);
+				CurrentUnitySize += UnitySizeInc;
+
+				if (NewTiming.NumFiles == 1)
+				{
+					break;
+				}
+			}
+
+			Logger.LogInformation($"{Module.Name} Timings CPUTiming(secs) | ExecutorTiming(secs) | NumFiles:");
+			PrintUnityInfo($"Current({CurrentModuleUnitySize})", CurrentCompileTime, Logger);
+			PrintUnityInfo("Disabled", DisableUnityCompileTime, Logger);
+			CurrentUnitySize = UnitySizeInc;
+			TimingData BestTiming = CurrentCompileTime;
+			foreach (TimingData Timing in Timings)
+			{
+				PrintUnityInfo(CurrentUnitySize.ToString(), Timing, Logger);
+				CurrentUnitySize += UnitySizeInc;
+
+				if (Timing.IsValid() &&
+					BestTiming.NumFiles != Timing.NumFiles &&
+					BestTiming.ExecutorTiming > Timing.ExecutorTiming &&
+					BestTiming.CPUTiming > Timing.CPUTiming)
+				{
+					BestTiming = Timing;
+				}
+			}
+
+			if (BestTiming != CurrentCompileTime)
+			{
+				Logger.LogInformation($"Better unity size than current: {BestTiming.UnitySize}");
+			}
+		}
+
+		/// <summary>
+		/// Print the timing data
+		/// </summary>
+		void PrintUnityInfo(string TimingPrefix, TimingData TimingData, ILogger Logger)
+		{
+			const int FirstColWidth = 15;
+			const int ColWidth = 10;
+			if (TimingData.IsValid())
+			{
+				string FormatString = String.Format("  {0,-" + FirstColWidth.ToString() + "}: {1,-" + ColWidth.ToString() + "} | {2,-" + ColWidth.ToString() + "} | {3,-" + ColWidth.ToString() + "}",
+					TimingPrefix,
+					TimingData.CPUTiming,
+					TimingData.ExecutorTiming,
+					TimingData.NumFiles);
+				Logger.LogInformation(FormatString);
+			}
+			else
+			{
+				string FormatString = String.Format("  {0,-" + FirstColWidth.ToString() + "}: Failed", TimingPrefix);
+				Logger.LogInformation(FormatString);
+			}
+		}
+
+		/// <summary>
+		/// Returns best compile timings after building the module times
+		/// </summary>
+		private TimingData GetBestCompileModuleTime(string LogPrefix, BuildConfiguration BuildConfiguration, TargetDescriptor TargetDescriptor, UEBuildModule Module, ILogger Logger, int UnitySize, bool bPriming, bool bDisableUnity)
+		{
+			const int CompileCount = 3;
+			List<TimingData> AllTimingData = new();
+			for (int CompileIndex = 0; CompileIndex < CompileCount; CompileIndex++)
+			{
+				TimingData NewTimingData = CompileModule(LogPrefix, BuildConfiguration, TargetDescriptor, Module, Logger, UnitySize, bPriming, bDisableUnity);
+				if (!NewTimingData.IsValid())
+				{
+					return NewTimingData;
+				}
+				AllTimingData.Add(NewTimingData);
+			}
+
+			TimingData? BestTimingData = AllTimingData.MinBy(TimingData => TimingData.ExecutorTiming);
+			if (BestTimingData == null)
+			{
+				BestTimingData = AllTimingData[0];
+			}
+			return BestTimingData;
+		}
+
+
+		/// <summary>
+		/// Compiles the module and returns the timing information
+		/// </summary>
+		private TimingData CompileModule(string LogPrefix, BuildConfiguration BuildConfiguration, TargetDescriptor TargetDescriptor, UEBuildModule Module, ILogger Logger, int UnitySize, bool bPriming, bool bDisableUnity)
+		{
+			// Store the old arguments
+			string[] OldArgs = TargetDescriptor.AdditionalArguments.GetRawArray();
+			TimingData NewTimingData = new TimingData();
+
+			try
+			{
+				if (!bPriming)
+				{
+					// Clear the output directory
+					Logger.LogInformation($"{LogPrefix}Deleting intermediate directory...");
+					try
+					{
+						var IntermDir = DirectoryItem.GetItemByDirectoryReference(Module.IntermediateDirectory);
+						IntermDir.CacheFiles();
+						IntermDir.ResetCachedInfo();
+						DirectoryReference.Delete(Module.IntermediateDirectory, true);
+					}
+					catch (Exception ex)
+					{
+						Logger.LogError(ex, $"{LogPrefix}Failed to delete {Module.Name}'s intermediate directory.");
+					}
+
+					// Add the module name to the cmdline
+					TargetDescriptor.AdditionalArguments = TargetDescriptor.AdditionalArguments.Append(new string[] { $"-BytesPerUnityCPP={UnitySize}" });
+
+					if (bDisableUnity)
+					{
+						TargetDescriptor.AdditionalArguments = TargetDescriptor.AdditionalArguments.Append(new string[] { $"-DisableUnity" });
+					}
+				}
+
+				using (ISourceFileWorkingSet WorkingSet = new EmptySourceFileWorkingSet())
+				{
+					if (bPriming)
+					{
+						Logger.LogInformation($"{LogPrefix}Priming module...");
+					}
+					else if (bDisableUnity)
+					{
+						Logger.LogInformation($"{LogPrefix}Compiling with no unity files...");
+					}
+					else
+					{
+						Logger.LogInformation($"{LogPrefix}Compiling with unity size '{UnitySize}'...");
+					}
+
+					TimingLogger NewTimingLogger = new(Logger);
+					BuildMode.Build(new List<TargetDescriptor>() { TargetDescriptor }, BuildConfiguration, WorkingSet, BuildOptions.None, null, NewTimingLogger);
+					NewTimingData = NewTimingLogger.TimingData;
+				}
+			}
+			catch
+			{
+				Logger.LogInformation($"{LogPrefix}Compile Failed");
+			}
+
+			NewTimingData.UnitySize = UnitySize;
+			EpicGames.Core.Log.EventParser.Flush(); // we need flush here to get all the logging info for this build
+			if (!bPriming)
+			{
+				if (NewTimingData.IsValid())
+				{
+					Logger.LogInformation($"{LogPrefix}Finished: CPUTime:{NewTimingData.CPUTiming}s | ExecutorTime:{NewTimingData.ExecutorTiming}s | NumFiles:{NewTimingData.NumFiles}");
+				}
+			}
+
+			// Restore the old arguments
+			TargetDescriptor.AdditionalArguments = new CommandLineArguments(OldArgs);
+
+			return NewTimingData;
+		}
+	}
+}
