Reorganized the unity logic to take into account the generated cpp files when determining whether to use unity files. Previously the logic wouldn't take into account the number of generated files which meant that modules that had just generated cpp files wouldn't use unity files.

Before:
Unity Files: 6759
Total CPU Time: 77035.734375 s
Total time in Parallel executor: 2062.95 seconds

After:
Unity Files: 6656
Total CPU Time: 76182.40625 s
Total time in Parallel executor: 2047.47 seconds

diff --git a/Engine/Source/Programs/UnrealBuildTool/Configuration/UEBuildModuleCPP.cs b/Engine/Source/Programs/UnrealBuildTool/Configuration/UEBuildModuleCPP.cs
--- a/Engine/Source/Programs/UnrealBuildTool/Configuration/UEBuildModuleCPP.cs
+++ b/Engine/Source/Programs/UnrealBuildTool/Configuration/UEBuildModuleCPP.cs
@@ -458,56 +458,6 @@
 				MinFilesUsingPrecompiledHeader = 1;
 			}
 
-			// Engine modules will always use unity build mode unless MinSourceFilesForUnityBuildOverride is specified in
-			// the module rules file.  By default, game modules only use unity of they have enough source files for that
-			// to be worthwhile.  If you have a lot of small game modules, consider specifying MinSourceFilesForUnityBuildOverride=0
-			// in the modules that you don't typically iterate on source files in very frequently.
-			int MinSourceFilesForUnityBuild = 2;
-			if (Rules.MinSourceFilesForUnityBuildOverride != 0)
-			{
-				MinSourceFilesForUnityBuild = Rules.MinSourceFilesForUnityBuildOverride;
-			}
-			else if (Target.ProjectFile != null && RulesFile.IsUnderDirectory(DirectoryReference.Combine(Target.ProjectFile.Directory, "Source")))
-			{
-				// Game modules with only a small number of source files are usually better off having faster iteration times
-				// on single source file changes, so we forcibly disable unity build for those modules
-				MinSourceFilesForUnityBuild = Target.MinGameModuleSourceFilesForUnityBuild;
-			}
-
-			// Should we use unity build mode for this module?
-			bool bModuleUsesUnityBuild = false;
-	
-			if (Target.bUseUnityBuild || Target.bForceUnityBuild)
-			{
-				if (Target.bForceUnityBuild)
-				{
-					Logger.LogDebug("Module '{ModuleName}' using unity build mode (bForceUnityBuild enabled for this module)", this.Name);
-					bModuleUsesUnityBuild = true;
-				}
-				else if (!Rules.bUseUnity)
-				{
-					Logger.LogDebug("Module '{ModuleName}' not using unity build mode (bUseUnity disabled for this module)", this.Name);
-					bModuleUsesUnityBuild = false;
-				}
-				else if (InputFiles.CPPFiles.Count < MinSourceFilesForUnityBuild)
-				{
-					Logger.LogDebug("Module '{ModuleName}' not using unity build mode (module with fewer than {NumFiles} source files)", this.Name, MinSourceFilesForUnityBuild);
-					bModuleUsesUnityBuild = false;
-				}
-				else
-				{
-					Logger.LogDebug("Module '{ModuleName}' using unity build mode", this.Name);
-					bModuleUsesUnityBuild = true;
-				}
-			}
-			else
-			{
-				Logger.LogDebug("Module '{ModuleName}' not using unity build mode", this.Name);
-			}
-
-			// Set up the NumIncludedBytesPerUnityCPP for this particular module
-			int NumIncludedBytesPerUnityCPP = (Rules.NumIncludedBytesPerUnityCPPOverride != 0) ? Rules.NumIncludedBytesPerUnityCPPOverride : Target.NumIncludedBytesPerUnityCPP;
-
 			// Set up the environment with which to compile the CPP files
 			CppCompileEnvironment CompileEnvironment = ModuleCompileEnvironment;
 
@@ -539,6 +489,9 @@
 			Dictionary<FileItem, FileItem> SourceFileToUnityFile = new Dictionary<FileItem, FileItem>();
 
 			List<FileItem> CPPFiles = new List<FileItem>(InputFiles.CPPFiles);
+			List<FileItem> GeneratedFileItems = new List<FileItem>();
+			CppCompileEnvironment GeneratedCPPCompileEnvironment = CompileEnvironment;
+			bool bMergeUnityFiles = Target.bMergeModuleAndGeneratedUnityFiles && Rules.bMergeUnityFiles;
 
 			// Compile all the generated CPP files
 			if (GeneratedCppDirectories != null && !CompileEnvironment.bHackHeaderGenerator)
@@ -626,10 +579,7 @@
 						}
 					}
 
-					bool bMergeUnityFiles = Target.bMergeModuleAndGeneratedUnityFiles && Rules.bMergeUnityFiles;
-
 					// Create a compile environment for the generated files. We can disable creating debug info here to improve link times.
-					CppCompileEnvironment GeneratedCPPCompileEnvironment = CompileEnvironment;
 					if (GeneratedCPPCompileEnvironment.bCreateDebugInfo && Target.bDisableDebugInfoForGeneratedCode)
 					{
 						GeneratedCPPCompileEnvironment = new CppCompileEnvironment(GeneratedCPPCompileEnvironment);
@@ -652,7 +602,6 @@
 					}
 
 					// Compile all the generated files
-					List<FileItem> GeneratedFileItems = new List<FileItem>();
 					foreach (string GeneratedFilename in GeneratedFiles.Values)
 					{
 						FileItem GeneratedCppFileItem = FileItem.GetItemByPath(GeneratedFilename);
@@ -661,25 +610,83 @@
 							GeneratedFileItems.Add(GeneratedCppFileItem);
 						}
 					}
-
-					if (bModuleUsesUnityBuild)
-					{
-						if (bMergeUnityFiles)
-						{
-							CPPFiles.AddRange(GeneratedFileItems);
-						}
-						else
-						{
-							Unity.GenerateUnityCPPs(Target, GeneratedFileItems, new List<FileItem>(), CompileEnvironment, WorkingSet, (Rules.ShortName ?? Name) + ".gen", IntermediateDirectory, Graph, SourceFileToUnityFile,
-								out List<FileItem> NormalGeneratedFiles, out List<FileItem> AdaptiveGeneratedFiles, NumIncludedBytesPerUnityCPP);
-							LinkInputFiles.AddRange(CompileFilesWithToolChain(Target, ToolChain, GeneratedCPPCompileEnvironment, ModuleCompileEnvironment, NormalGeneratedFiles, AdaptiveGeneratedFiles, Graph, Logger).ObjectFiles);
-						}
-					}
-					else
-					{
-						LinkInputFiles.AddRange(ToolChain.CompileCPPFiles(GeneratedCPPCompileEnvironment, GeneratedFileItems, IntermediateDirectory, Name, Graph).ObjectFiles);
-					}
 				}
 			}
 
+			// Engine modules will always use unity build mode unless MinSourceFilesForUnityBuildOverride is specified in
+			// the module rules file.  By default, game modules only use unity of they have enough source files for that
+			// to be worthwhile.  If you have a lot of small game modules, consider specifying MinSourceFilesForUnityBuildOverride=0
+			// in the modules that you don't typically iterate on source files in very frequently.
+			int MinSourceFilesForUnityBuild = 2;
+			if (Rules.MinSourceFilesForUnityBuildOverride != 0)
+			{
+				MinSourceFilesForUnityBuild = Rules.MinSourceFilesForUnityBuildOverride;
+			}
+			else if (Target.ProjectFile != null && RulesFile.IsUnderDirectory(DirectoryReference.Combine(Target.ProjectFile.Directory, "Source")))
+			{
+				// Game modules with only a small number of source files are usually better off having faster iteration times
+				// on single source file changes, so we forcibly disable unity build for those modules
+				MinSourceFilesForUnityBuild = Target.MinGameModuleSourceFilesForUnityBuild;
+			}
+
+			// Set up the NumIncludedBytesPerUnityCPP for this particular module
+			int NumIncludedBytesPerUnityCPP = (Rules.NumIncludedBytesPerUnityCPPOverride != 0) ? Rules.NumIncludedBytesPerUnityCPPOverride : Target.NumIncludedBytesPerUnityCPP;
+
+			// Should we use unity build mode for this module?
+			bool bModuleUsesUnityBuild = false;
+			if (Target.bUseUnityBuild || Target.bForceUnityBuild)
+			{
+				int FileCount = CPPFiles.Count;
+
+				// if we are merging the genearted cpp files then that needs to be part of the count
+				if (bMergeUnityFiles)
+				{
+					FileCount += GeneratedFileItems.Count;
+				}
+
+				if (Target.bForceUnityBuild)
+				{
+					Logger.LogDebug("Module '{ModuleName}' using unity build mode (bForceUnityBuild enabled for this module)", this.Name);
+					bModuleUsesUnityBuild = true;
+				}
+				else if (!Rules.bUseUnity)
+				{
+					Logger.LogDebug("Module '{ModuleName}' not using unity build mode (bUseUnity disabled for this module)", this.Name);
+					bModuleUsesUnityBuild = false;
+				}
+				else if (FileCount < MinSourceFilesForUnityBuild)
+				{
+					Logger.LogDebug("Module '{ModuleName}' not using unity build mode (module with fewer than {NumFiles} source files)", this.Name, MinSourceFilesForUnityBuild);
+					bModuleUsesUnityBuild = false;
+				}
+				else
+				{
+					Logger.LogDebug("Module '{ModuleName}' using unity build mode", this.Name);
+					bModuleUsesUnityBuild = true;
+				}
+			}
+			else
+			{
+				Logger.LogDebug("Module '{ModuleName}' not using unity build mode", this.Name);
+			}
+
+			// Compile Generated CPP Files
+			if (bModuleUsesUnityBuild)
+			{
+				if (bMergeUnityFiles)
+				{
+					CPPFiles.AddRange(GeneratedFileItems);
+				}
+				else
+				{
+					Unity.GenerateUnityCPPs(Target, GeneratedFileItems, new List<FileItem>(), CompileEnvironment, WorkingSet, (Rules.ShortName ?? Name) + ".gen", IntermediateDirectory, Graph, SourceFileToUnityFile,
+						out List<FileItem> NormalGeneratedFiles, out List<FileItem> AdaptiveGeneratedFiles, NumIncludedBytesPerUnityCPP);
+					LinkInputFiles.AddRange(CompileFilesWithToolChain(Target, ToolChain, GeneratedCPPCompileEnvironment, ModuleCompileEnvironment, NormalGeneratedFiles, AdaptiveGeneratedFiles, Graph, Logger).ObjectFiles);
+				}
+			}
+			else
+			{
+				LinkInputFiles.AddRange(ToolChain.CompileCPPFiles(GeneratedCPPCompileEnvironment, GeneratedFileItems, IntermediateDirectory, Name, Graph).ObjectFiles);
+			}
+
 			// Compile CPP files
