UBT: Added a log message confirming that a .json file with build actions has been created when using -WriteOutdatedActions command-line argument.

#rnx

diff --git a/Engine/Source/Programs/UnrealBuildTool/Modes/BuildMode.cs b/Engine/Source/Programs/UnrealBuildTool/Modes/BuildMode.cs
--- a/Engine/Source/Programs/UnrealBuildTool/Modes/BuildMode.cs
+++ b/Engine/Source/Programs/UnrealBuildTool/Modes/BuildMode.cs
@@ -593,6 +593,7 @@ namespace UnrealBuildTool
 					using (GlobalTracer.Instance.BuildSpan("ActionGraph.WriteActions").StartActive())
 					{
 						ActionGraph.ExportJson(MergedActionsToExecute, WriteOutdatedActionsFile);
+						Logger.LogInformation("Exported build actions to {WriteOutdatedActionsFile}", WriteOutdatedActionsFile);
 					}
 				}
 				else
