FixIncludePaths: Add -IncludeFilter to limit the #includes that are modified
Add progress printing
Disallow relative parent directory pathing, currently breaks some #includes

#rnx

diff --git a/Engine/Source/Programs/UnrealBuildTool/Modes/FixIncludePathsMode.cs b/Engine/Source/Programs/UnrealBuildTool/Modes/FixIncludePathsMode.cs
--- a/Engine/Source/Programs/UnrealBuildTool/Modes/FixIncludePathsMode.cs
+++ b/Engine/Source/Programs/UnrealBuildTool/Modes/FixIncludePathsMode.cs
@@ -3,12 +3,10 @@
 using System;
 using System.Collections.Generic;
 using System.Linq;
-using System.Text;
 using System.Text.RegularExpressions;
 using EpicGames.Core;
 using UnrealBuildBase;
 using Microsoft.Extensions.Logging;
-using System.Collections;
 
 namespace UnrealBuildTool
 {
@@ -30,6 +28,9 @@ namespace UnrealBuildTool
 		[CommandLine("-Filter=", Description = "Set of filters for files to include in the database. Relative to the root directory, or to the project file.")]
 		List<string> FilterRules = new List<string>();
 
+		[CommandLine("-IncludeFilter=", Description = "Set of filters for #include'd lines to allow updating. Relative to the root directory, or to the project file.")]
+		List<string> IncludeFilterRules = new List<string>();
+
 		[CommandLine("-CheckoutWithP4", Description = "Flags that this task should use p4 to check out the file before updating it.")]
 		public bool bCheckoutWithP4 = false;
 
@@ -65,6 +66,17 @@ namespace UnrealBuildTool
 				}
 			}
 
+			// Parse the include filter argument
+			FileFilter? IncludeFilter = null;
+			if (IncludeFilterRules.Count > 0)
+			{
+				IncludeFilter = new FileFilter(FileFilterType.Exclude);
+				foreach (string FilterRule in IncludeFilterRules)
+				{
+					IncludeFilter.AddRules(FilterRule.Split(';'));
+				}
+			}
+
 			// Force C++ modules to always include their generated code directories
 			UEBuildModuleCPP.bForceAddGeneratedCodeIncludePath = true;
 
@@ -83,18 +95,13 @@ namespace UnrealBuildTool
 					// Create a makefile for the target
 					UEBuildTarget Target = UEBuildTarget.Create(TargetDescriptor, BuildConfiguration.bSkipRulesCompile, BuildConfiguration.bForceRulesCompile, BuildConfiguration.bUsePrecompiled, Logger);
 
-					// Create all the binaries and modules
-					CppCompileEnvironment GlobalCompileEnvironment = Target.CreateCompileEnvironmentForProjectFiles(Logger);
+					// Get InputLists and build Include Filter if necessary
+					HashSet<FileReference> IncludeFileList = new HashSet<FileReference>();
+					Dictionary<UEBuildModuleCPP, List<FileReference>> ModuleToFiles = new Dictionary<UEBuildModuleCPP, List<FileReference>>();
 					foreach (UEBuildBinary Binary in Target.Binaries)
 					{
-						CppCompileEnvironment BinaryCompileEnvironment = Binary.CreateBinaryCompileEnvironment(GlobalCompileEnvironment);
-
 						foreach (UEBuildModuleCPP Module in Binary.Modules.OfType<UEBuildModuleCPP>())
 						{
-							bool IsThirdPartyModule = Module.RulesFile.ContainsName("ThirdParty", Unreal.RootDirectory);
-							if (IsThirdPartyModule)
-								continue;
-
 							UEBuildModuleCPP.InputFileCollection InputFileCollection = Module.FindInputFiles(Target.Platform, new Dictionary<DirectoryItem, FileItem[]>());
 							List<FileItem> InputFiles = new List<FileItem>();
 							InputFiles.AddRange(InputFileCollection.HeaderFiles);
@@ -111,173 +118,218 @@ namespace UnrealBuildTool
 									FileList.Add(fileRef);
 								}
 							}
+							ModuleToFiles[Module] = FileList;
 
-							if (FileList.Any())
+							foreach (FileItem InputFile in InputFiles)
 							{
-								Dictionary<string, string?> PreferredPathCache = new();
-								CppCompileEnvironment env = Module.CreateCompileEnvironmentForIntellisense(Target.Rules, BinaryCompileEnvironment, Logger);
-
-								foreach (var InputFile in FileList)
+								if (IncludeFilter == null || IncludeFilter.Matches(InputFile.Location.MakeRelativeTo(Unreal.RootDirectory)))
 								{
-									List<int> LinesUpdated = new();
-									string[] Text = FileReference.ReadAllLines(InputFile);
-									bool UpdatedText = false;
+									var fileRef = new FileReference(InputFile.AbsolutePath);
+									IncludeFileList.Add(fileRef);
+								}
+							}
+						}
+					}
+
+					// List of modules that are allowed to be processed (not ThirdParty and have unfiltered files)
+					HashSet<UEBuildModuleCPP> AllModules = ModuleToFiles.Where(Item => !Item.Key.RulesFile.ContainsName("ThirdParty", Unreal.RootDirectory) && Item.Value.Any()).Select(Item => Item.Key).ToHashSet();
+
+					// Keep track of progress
+					int Index = 0;
+					int Total = AllModules.Count();
+
+					// Create the global compile environment for this target
+					CppCompileEnvironment GlobalCompileEnvironment = Target.CreateCompileEnvironmentForProjectFiles(Logger);
 
-									for (int i = 0; i < Text.Length; i++)
+					// Create all the binaries and modules
+					foreach (UEBuildBinary Binary in Target.Binaries)
+					{
+						CppCompileEnvironment BinaryCompileEnvironment = Binary.CreateBinaryCompileEnvironment(GlobalCompileEnvironment);
+
+						foreach (UEBuildModuleCPP Module in Binary.Modules.OfType<UEBuildModuleCPP>().Where(Module => AllModules.Contains(Module)))
+						{
+							List<FileReference> FileList = ModuleToFiles[Module];
+
+							Dictionary<string, string?> PreferredPathCache = new();
+							CppCompileEnvironment env = Module.CreateCompileEnvironmentForIntellisense(Target.Rules, BinaryCompileEnvironment, Logger);
+
+							foreach (var InputFile in FileList)
+							{
+								List<int> LinesUpdated = new();
+								string[] Text = FileReference.ReadAllLines(InputFile);
+								bool UpdatedText = false;
+
+								for (int i = 0; i < Text.Length; i++)
+								{
+									var Line = Text[i];
+									int LineNumber = i + 1;
+									Match IncludeMatch = IncludeRegex.Match(Line);
+									if (IncludeMatch.Success)
 									{
-										var Line = Text[i];
-										int LineNumber = i + 1;
-										Match IncludeMatch = IncludeRegex.Match(Line);
-										if (IncludeMatch.Success)
+										string Include = IncludeMatch.Groups[1].Value;
+
+										if (Include.Contains("/Private/") && PublicDirectories.Any(dir => InputFile.FullName.Contains(System.IO.Path.DirectorySeparatorChar + dir + System.IO.Path.DirectorySeparatorChar)))
+										{
+											Logger.LogError("{FileName}({LineNumber}): Can not update #include '{Include}' in the public file because it may break external code that uses it.", InputFile.FullName, LineNumber, Include);
+											continue;
+										}
+
+										if (Include.Contains(".."))
+										{
+											Logger.LogError("{FileName}({LineNumber}): Can not update #include '{Include}', relative pathing is not currently handled.", InputFile.FullName, LineNumber, Include);
+											continue;
+										}
+
+										string? PreferredInclude = null;
+										if (!PreferredPathCache.TryGetValue(Include, out PreferredInclude))
 										{
-											string Include = IncludeMatch.Groups[1].Value;
+											List<DirectoryReference> IncludePaths = new();
+											IncludePaths.Add(new DirectoryReference(System.IO.Directory.GetParent(InputFile.FullName)!));
+											IncludePaths.AddRange(env.UserIncludePaths);
+											IncludePaths.AddRange(env.SystemIncludePaths);
+
+											// search include paths
+											FileReference? FoundIncludeFile = null;
+											DirectoryReference? FoundIncludePath = null;
+											foreach (var IncludePath in IncludePaths)
+											{
+												string Path = System.IO.Path.GetFullPath(System.IO.Path.Combine(IncludePath.FullName, Include));
+												if (System.IO.File.Exists(Path))
+												{
+													FoundIncludeFile = FileReference.FromString(Path);
+													FoundIncludePath = IncludePath;
+													break;
+												}
+											}
 
-											if (Include.Contains("/Private/") && PublicDirectories.Any(dir => InputFile.FullName.Contains(System.IO.Path.DirectorySeparatorChar + dir + System.IO.Path.DirectorySeparatorChar)))
+											if (FoundIncludeFile != null && !IncludeFileList.Contains(FoundIncludeFile))
 											{
-												Logger.LogError("{FileName}({LineNumber}): Can not update #include '{Include}' in the public file because it may break external code that uses it.", InputFile.FullName, LineNumber, Include);
+												Logger.LogDebug("{FileName}({LineNumber}): Skipping '{Include}' because it is filtered out.", InputFile.FullName, LineNumber, Include);
+												PreferredInclude = Include;
+												PreferredPathCache[Include] = PreferredInclude;
 												continue;
 											}
 
-											string? PreferredInclude = null;
-											if (!PreferredPathCache.TryGetValue(Include, out PreferredInclude))
+											if (FoundIncludeFile != null)
 											{
-												List<DirectoryReference> IncludePaths = new();
-												IncludePaths.Add(new DirectoryReference(System.IO.Directory.GetParent(InputFile.FullName)!));
-												IncludePaths.AddRange(env.UserIncludePaths);
-												IncludePaths.AddRange(env.SystemIncludePaths);
-
-												// search include paths
-												string? FullPath = null;
-												DirectoryReference? FoundIncludePath = null;
-												foreach (var IncludePath in IncludePaths)
+												string FullPath = FoundIncludeFile.FullName.Replace('\\', '/');
+												if (FullPath.Contains("ThirdParty"))
 												{
-													string Path = System.IO.Path.GetFullPath(System.IO.Path.Combine(IncludePath.FullName, Include));
-													if (System.IO.File.Exists(Path))
-													{
-														FullPath = Path.Replace('\\', '/');
-														FoundIncludePath = IncludePath;
-														break;
-													}
+													Logger.LogDebug("{FileName}({LineNumber}): Skipping '{Include}' because it is a third party header.", InputFile.FullName, LineNumber, Include);
+													PreferredInclude = Include;
+													PreferredPathCache[Include] = PreferredInclude;
+													continue;
 												}
 
-												if (!string.IsNullOrEmpty(FullPath))
+												// if the include and the source file live in the same directory then it is OK to be relative
+												if (string.Equals(System.IO.Directory.GetParent(FullPath)?.FullName, System.IO.Directory.GetParent(InputFile.FullName)?.FullName, StringComparison.CurrentCultureIgnoreCase) &&
+													string.Equals(Include, System.IO.Path.GetFileName(FullPath), StringComparison.CurrentCultureIgnoreCase))
 												{
-													if (FullPath.Contains("ThirdParty"))
-													{
-														Logger.LogDebug("{FileName}({LineNumber}): Skipping '{Include}' because it is a third party header.", InputFile.FullName, LineNumber, Include);
-														PreferredInclude = Include;
-														PreferredPathCache[Include] = PreferredInclude;
-														continue;
-													}
-
-													// if the include and the source file live in the same directory then it is OK to be relative
-													if (string.Equals(System.IO.Directory.GetParent(FullPath)?.FullName, System.IO.Directory.GetParent(InputFile.FullName)?.FullName, StringComparison.CurrentCultureIgnoreCase) &&
-														string.Equals(Include, System.IO.Path.GetFileName(FullPath), StringComparison.CurrentCultureIgnoreCase))
+													Logger.LogDebug("{FileName}({LineNumber}): Using '{Include}' because it is in the same directory.", InputFile.FullName, LineNumber, Include);
+													PreferredInclude = Include;
+												}
+												else
+												{
+													if (!FullPath.Contains(UnrealRootDirectory))
 													{
-														Logger.LogDebug("{FileName}({LineNumber}): Using '{Include}' because it is in the same directory.", InputFile.FullName, LineNumber, Include);
-														PreferredInclude = Include;
+														Logger.LogDebug("{FileName}({LineNumber}): Skipping '{Include}' because it isn't under the Unreal root directory.", InputFile.FullName, LineNumber, Include);
 													}
 													else
 													{
-														if (!FullPath.Contains(UnrealRootDirectory))
+														string? FoundPreferredPath = PreferredPaths.FirstOrDefault(path => FullPath.Contains(path));
+														if (FoundPreferredPath != null)
 														{
-															Logger.LogDebug("{FileName}({LineNumber}): Skipping '{Include}' because it isn't under the Unreal root directory.", InputFile.FullName, LineNumber, Include);
+															int end = FullPath.LastIndexOf(FoundPreferredPath) + FoundPreferredPath.Length;
+															PreferredInclude = FullPath.Substring(end);
+
+															// Is the current include a shortened version of the preferred include path?
+															if (PreferredInclude != Include && PreferredInclude.Contains(Include))
+															{
+																Logger.LogDebug("{FileName}({LineNumber}): Using '{Include}' because it is shorter than '{PreferredInclude}'.", InputFile.FullName, LineNumber, Include, PreferredInclude);
+																PreferredInclude = Include;
+															}
 														}
 														else
 														{
-															string? FoundPreferredPath = PreferredPaths.FirstOrDefault(path => FullPath.Contains(path));
-															if (FoundPreferredPath != null)
+															PreferredInclude = null;
+
+															string ModulePath = FullPath;
+															FileReference IncludeFileReference = FileReference.FromString(FullPath);
+															DirectoryReference? TempDirectory = IncludeFileReference.Directory;
+															DirectoryReference? FoundDirectory = null;
+															// find the module this include is part of
+															while (TempDirectory != null)
 															{
-																int end = FullPath.LastIndexOf(FoundPreferredPath) + FoundPreferredPath.Length;
-																PreferredInclude = FullPath.Substring(end);
-
-																// Is the current include a shortened version of the preferred include path?
-																if (PreferredInclude != Include && PreferredInclude.Contains(Include))
+																if (DirectoryReference.EnumerateFiles(TempDirectory, $"*.build.cs").Any())
 																{
-																	Logger.LogDebug("{FileName}({LineNumber}): Using '{Include}' because it is shorter than '{PreferredInclude}'.", InputFile.FullName, LineNumber, Include, PreferredInclude);
-																	PreferredInclude = Include;
+																	FoundDirectory = TempDirectory;
+																	break;
 																}
-															}
-															else
-															{
-																PreferredInclude = null;
-
-																string ModulePath = FullPath;
-																FileReference IncludeFileReference = FileReference.FromString(FullPath);
-																DirectoryReference? TempDirectory = IncludeFileReference.Directory;
-																DirectoryReference? FoundDirectory = null;
-																// find the module this include is part of
-																while (TempDirectory != null)
-																{
-																	if (DirectoryReference.EnumerateFiles(TempDirectory, $"*.build.cs").Any())
-																	{
-																		FoundDirectory = TempDirectory;
-																		break;
-																	}
 
-																	TempDirectory = TempDirectory.ParentDirectory;
-																}
+																TempDirectory = TempDirectory.ParentDirectory;
+															}
 
-																if (FoundDirectory != null)
-																{
-																	PreferredInclude = FullPath.Substring(FoundDirectory.FullName.Length + 1);
-																}
+															if (FoundDirectory != null)
+															{
+																PreferredInclude = FullPath.Substring(FoundDirectory.FullName.Length + 1);
 															}
 														}
-
-														PreferredPathCache[Include] = PreferredInclude;
 													}
-												}
-												
-												if (PreferredInclude == null)
-												{
-													Logger.LogDebug("{FileName}({LineNumber}): Could not find path to '{IncludePath}'", InputFile.FullName, LineNumber, Include);
+
+													PreferredPathCache[Include] = PreferredInclude;
 												}
 											}
-
-											if (PreferredInclude != null && Include != PreferredInclude)
+												
+											if (PreferredInclude == null)
 											{
-												Logger.LogInformation("{FileName}({LineNumber}): Updated '{OldInclude}' -> '{NewInclude}'", InputFile.FullName, LineNumber, Include, PreferredInclude);
-												Text[i] = Line.Replace(Include, PreferredInclude);
-												UpdatedText = true;
-												LinesUpdated.Add(i);
+												Logger.LogDebug("{FileName}({LineNumber}): Could not find path to '{IncludePath}'", InputFile.FullName, LineNumber, Include);
 											}
 										}
-									}
 
-									if (UpdatedText)
-									{
-										if (!bNoIncludeSorting)
+										if (PreferredInclude != null && Include != PreferredInclude)
 										{
-											SortIncludes(InputFile, LinesUpdated, Text);
+											Logger.LogInformation("{FileName}({LineNumber}): Updated '{OldInclude}' -> '{NewInclude}'", InputFile.FullName, LineNumber, Include, PreferredInclude);
+											Text[i] = Line.Replace(Include, PreferredInclude);
+											UpdatedText = true;
+											LinesUpdated.Add(i);
 										}
+									}
+								}
 
-										if (!bNoOutput)
+								if (UpdatedText)
+								{
+									if (!bNoIncludeSorting)
+									{
+										SortIncludes(InputFile, LinesUpdated, Text);
+									}
+
+									if (!bNoOutput)
+									{
+										Logger.LogInformation("Updating {IncludePath}", InputFile.FullName);
+										try
 										{
-											Logger.LogInformation("Updating {IncludePath}", InputFile.FullName);
-											try
-											{
-												if (bCheckoutWithP4)
-												{
-													System.Diagnostics.Process Process = new System.Diagnostics.Process();
-													System.Diagnostics.ProcessStartInfo StartInfo = new System.Diagnostics.ProcessStartInfo();
-													Process.StartInfo.WindowStyle = System.Diagnostics.ProcessWindowStyle.Hidden;
-													Process.StartInfo.FileName = "p4.exe";
-													Process.StartInfo.Arguments = $"edit {InputFile.FullName}";
-													Process.Start();
-													Process.WaitForExit();
-												}
-												System.IO.File.WriteAllLines(InputFile.FullName, Text);
-											}
-											catch (Exception ex)
+											if (bCheckoutWithP4)
 											{
-												Logger.LogWarning("Failed to write to file: {Exception}", ex);
+												System.Diagnostics.Process Process = new System.Diagnostics.Process();
+												System.Diagnostics.ProcessStartInfo StartInfo = new System.Diagnostics.ProcessStartInfo();
+												Process.StartInfo.WindowStyle = System.Diagnostics.ProcessWindowStyle.Hidden;
+												Process.StartInfo.FileName = "p4.exe";
+												Process.StartInfo.Arguments = $"edit {InputFile.FullName}";
+												Process.Start();
+												Process.WaitForExit();
 											}
+											System.IO.File.WriteAllLines(InputFile.FullName, Text);
+										}
+										catch (Exception ex)
+										{
+											Logger.LogWarning("Failed to write to file: {Exception}", ex);
 										}
 									}
 								}
 							}
 
+							Logger.LogInformation("[{Index}/{Total}] Processed Module {Name} ({Files} files)", ++Index, Total, Module.Name, FileList.Count);
+
 							ScannedModules.Add(Module);
 						}
 					}
