Add support for XCFramework to UBT on iOS

diff --git a/Engine/Source/Programs/UnrealBuildTool/Configuration/UEBuildFramework.cs b/Engine/Source/Programs/UnrealBuildTool/Configuration/UEBuildFramework.cs
--- a/Engine/Source/Programs/UnrealBuildTool/Configuration/UEBuildFramework.cs
+++ b/Engine/Source/Programs/UnrealBuildTool/Configuration/UEBuildFramework.cs
@@ -4,6 +4,7 @@ using System;
 using System.Collections.Generic;
 using System.Linq;
 using System.Text;
+using System.Xml;
 using System.Threading.Tasks;
 using EpicGames.Core;
 using UnrealBuildBase;
@@ -26,9 +27,14 @@ namespace UnrealBuildTool
 		public readonly FileReference? ZipFile;
 
 		/// <summary>
-		/// Path to the framework on disk.
+		/// Path where the zip should be extracted. May be null.
 		/// </summary>
-		public readonly DirectoryReference? FrameworkDirectory;
+		public readonly DirectoryReference? ZipOutputDirectory;
+
+		/// <summary>
+		/// Base path to the framework on disk.
+		/// </summary>
+		readonly DirectoryReference? FrameworkDirectory;
 
 		/// <summary>
 		/// 
@@ -45,6 +51,11 @@ namespace UnrealBuildTool
 		/// </summary>
 		public FileItem? ExtractedTokenFile;
 
+		/// <summary>
+		/// List of variants contained in this XCFramework. Only non null for XCFrameworks
+		/// </summary>
+		readonly List<XCFrameworkVariantEntry>? XCFrameworkVariants;
+
 		/// <summary>
 		/// Constructor
 		/// </summary>
@@ -68,6 +79,7 @@ namespace UnrealBuildTool
 		{
 			this.Name = Name;
 			this.ZipFile = ZipFile;
+			this.ZipOutputDirectory = OutputDirectory;
 			this.FrameworkDirectory = OutputDirectory;
 			this.CopyBundledAssets = CopyBundledAssets;
 			this.bCopyFramework = bCopyFramework;
@@ -86,6 +98,240 @@ namespace UnrealBuildTool
 			this.FrameworkDirectory = FrameworkDirectory;
 			this.CopyBundledAssets = CopyBundledAssets;
 			this.bCopyFramework = bCopyFramework;
+			if( this.FrameworkDirectory != null && this.FrameworkDirectory.FullName.EndsWith(".xcframework"))
+			{
+				this.XCFrameworkVariants = LoadXCFrameworkVariants();
+			}
+		}
+
+		/// <summary>
+		/// Gets the framework directory for given build settings
+		/// </summary>
+		/// <param name="Platform"></param>
+		/// <param name="Architecture"></param>
+		public DirectoryReference? GetFrameworkDirectory(UnrealTargetPlatform Platform, string Architecture)
+		{
+			if( XCFrameworkVariants != null)
+			{
+				XCFrameworkVariantEntry? FrameworkVariant = XCFrameworkVariants.Find(x => x.Matches(Platform, Architecture));
+				if( FrameworkVariant != null )
+				{
+					return DirectoryReference.Combine(FrameworkDirectory!, FrameworkVariant.LibraryIdentifier);
+				}
+				Log.TraceWarning($"Variant for platform \"{Platform.ToString()}\" with architecture {Architecture} not found in XCFramework {Name}.");
+			}
+			return FrameworkDirectory;
+		}
+
+		/// <summary>
+		/// Loads XCFramework variants description from Info.plist file inside XCFramework structure 
+		/// </summary>
+		List<XCFrameworkVariantEntry>? LoadXCFrameworkVariants()
+		{
+			XmlDocument PlistDoc = new XmlDocument();
+			PlistDoc.Load(FileReference.Combine(FrameworkDirectory!, "Info.plist" ).FullName);
+
+			// Check the plist type
+			XmlNode? CFBundlePackageType = PlistDoc.SelectSingleNode("/plist/dict[key='CFBundlePackageType']/string[1]");
+			if( CFBundlePackageType == null || CFBundlePackageType.NodeType != XmlNodeType.Element || CFBundlePackageType.InnerText != "XFWK" )
+			{
+				Log.TraceWarning($"CFBundlePackageType is not set to XFWK in Info.plist for XCFramework {Name}");
+				return null;
+			}
+
+			// Load Framework variants data from dictionary nodes
+			XmlNodeList? FrameworkVariantDicts = PlistDoc.SelectNodes("/plist/dict[key='AvailableLibraries']/array/dict");			
+			if( FrameworkVariantDicts == null )
+			{
+				Log.TraceWarning($"Invalid Info.plist file for XCFramework {Name}. It will be used as a regular Framework");
+				return null;
+			}
+
+			List<XCFrameworkVariantEntry> Variants = new List<XCFrameworkVariantEntry>();
+			foreach( XmlNode VariantDict in FrameworkVariantDicts )
+			{
+				XCFrameworkVariantEntry? Variant = XCFrameworkVariantEntry.Parse(VariantDict);
+				if( Variant == null )
+				{
+					Log.TraceWarning($"Failed to load variant from Info.plist file for XCFramework {Name}");
+				}
+				else
+				{
+					Log.TraceVerbose($"Found {Variant.LibraryIdentifier} variant in XCFramework {Name}");
+					Variants.Add(Variant);
+				}
+			}
+			return Variants;
+		}
+
+		/// <summary>
+		/// Represents a XCFramework variant description
+		/// </summary>		
+		private class XCFrameworkVariantEntry
+		{
+			/// <summary>
+			/// Identifier of this framework variant. Is in the form platform-architectures-platformvariant.
+			/// Some examples would be ios-arm64, ios-arm64_x86_64-simulator
+			/// It also represents the relative path inside the XCFramefork where the Framework for this variant resides
+			/// </summary>
+			public readonly string LibraryIdentifier;
+
+			/// <summary>
+			/// Relative path where framework lives after applying LibraryIdentifier path
+			/// </summary>
+			public readonly string LibraryPath;
+
+			/// <summary>
+			/// List of supported architectures for this Framework variants
+			/// </summary>
+			public readonly List<string> SupportedArchitectures;
+
+			/// <summary>
+			/// Platform this Framework variant is intended for. Possible values are ios, macos, watchos, tvos
+			/// </summary>
+			public readonly string SupportedPlatform;
+
+			/// <summary>
+			/// Platform variant for this Framework variant. It can be empty or any other value representing a platform variane, like maccatalyst or simulator 
+			/// </summary>
+			public readonly string? SupportedPlatformVariant;
+
+			/// <summary>
+			/// Constructor
+			/// </summary>
+			/// <param name="LibraryIdentifier"></param>
+			/// <param name="LibraryPath"></param>
+			/// <param name="SupportedArchitectures"></param>
+			/// <param name="SupportedPlatform"></param>
+			/// <param name="SupportedPlatformVariant"></param>
+			XCFrameworkVariantEntry(string LibraryIdentifier, string LibraryPath, List<string> SupportedArchitectures, string SupportedPlatform, string? SupportedPlatformVariant)
+			{
+				this.LibraryIdentifier = LibraryIdentifier;
+				this.LibraryPath = LibraryPath;
+				this.SupportedArchitectures = SupportedArchitectures;
+				this.SupportedPlatform = SupportedPlatform;
+				this.SupportedPlatformVariant = SupportedPlatformVariant;
+			}
+
+			/// <summary>
+			/// Returns wether this variant is a match for given platform and architecture 
+			/// </summary>
+			/// <param name="Platform"></param>
+			/// <param name="Architecture"></param>
+			public bool Matches(UnrealTargetPlatform Platform, string Architecture)
+			{
+				if(( Platform == UnrealTargetPlatform.IOS && SupportedPlatform == "ios" ) ||
+				   ( Platform == UnrealTargetPlatform.Mac && SupportedPlatform == "macos" ) ||
+				   ( Platform == UnrealTargetPlatform.TVOS && SupportedPlatform == "tvos" ))
+				{
+					if( Architecture == "-simulator" )
+					{
+						// When using -simulator we don't have the actual architecture. Assume arm64 as other parts of UBT already are doing
+						return ( SupportedPlatformVariant == "simulator" ) && SupportedArchitectures.Contains("arm64");
+					}
+					else
+					{
+						return ( SupportedPlatformVariant == null ) && SupportedArchitectures.Contains(Architecture);
+					}
+				}
+				return false;
+			}
+
+			/// <summary>
+			/// Creates a XCFrameworkVariantEntry by parsing the content from the XCFramework Info.plist file  
+			/// </summary>
+			/// <param name="DictNode">XmlNode loaded form Info.plist containing the representation of a &lt;dict&gt; that contains the variant description</param>
+			public static XCFrameworkVariantEntry? Parse(XmlNode DictNode)
+			{
+				// Entries from a <dict> in a plist file always contains a <key> node and a value node
+				if(( DictNode.ChildNodes.Count % 2 ) != 0)
+				{
+					Log.TraceVerbose($"Invalid dict in Info.plist file for XCFramework");
+					return null;
+				}
+
+				string? LibraryIdentifier = null;
+				string? LibraryPath = null;
+				List<string>? SupportedArchitectures = null;
+				string? SupportedPlatform = null;
+				string? SupportedPlatformVariant = null;
+
+				for( int i = 0; i < DictNode.ChildNodes.Count; i += 2 )
+            	{
+                	XmlNode? Key = DictNode.ChildNodes[i];
+                	XmlNode? Value = DictNode.ChildNodes[i + 1];
+
+					if( Value != null && Key != null && Key.Name == "key" )
+					{
+						switch( Key.InnerText )
+						{
+							case "LibraryIdentifier":
+								LibraryIdentifier = ParseString(Value);
+								break;
+							case "LibraryPath":
+								LibraryPath = ParseString(Value);
+								break;
+							case "SupportedPlatform":
+								SupportedPlatform = ParseString(Value);
+								break;
+							case "SupportedPlatformVariant":
+								SupportedPlatformVariant = ParseString(Value);
+								break;
+							case "SupportedArchitectures":
+								SupportedArchitectures = ParseListOfStrings(Value);
+								break;
+							default:
+								break;
+						}
+					}
+				}
+
+				// All fields but SupportedPlatformVariant are allowed to be present in the framework variant descriprion
+				if( LibraryIdentifier == null || LibraryPath == null || SupportedArchitectures == null || SupportedPlatform == null )
+				{
+					Log.TraceVerbose($"Missing data in Info.plist file for XCFramework");
+					return null;
+				}
+				return new XCFrameworkVariantEntry( LibraryIdentifier, LibraryPath, SupportedArchitectures, SupportedPlatform, SupportedPlatformVariant ); 
+			}
+
+			/// <summary>
+			/// Parses the content from a string value node in a plist file
+			/// </summary>
+			/// <param name="ValueNode">XmlNode we expect to contain a &lt;string&gt; value</param>			
+			static string? ParseString( XmlNode ValueNode )
+			{
+				if( ValueNode.Name != "string" )
+				{
+					Log.TraceVerbose($"Unexpected tag \"{ValueNode.Name}\" while expecting \"string\" in Info.plist file for XCFramework");
+					return null;
+				}
+				return ValueNode.InnerText;
+			}
+
+			/// <summary>
+			/// Parses the content from an array value node in a plist file that has several string entries
+			/// </summary>
+			/// <param name="ValueNode">XmlNode we expect to contain a &lt;array&gt; value with several &lt;string&gt; entries</param>			
+			static List<string>? ParseListOfStrings( XmlNode ValueNode )
+			{
+				if( ValueNode.Name != "array" )
+				{
+					Log.TraceVerbose($"Unexpected tag \"{ValueNode.Name}\" while expecting \"array\" in Info.plist file for XCFramework");
+					return null;
+				}
+
+				List<string> ListOfStrings = new List<string>();
+				foreach( XmlNode ChildNode in ValueNode.ChildNodes )
+				{
+					string? ParsedString = ParseString( ChildNode );
+					if( ParsedString != null )
+					{
+						ListOfStrings.Add( ParsedString );
+					}
+				}
+				return ListOfStrings;
+			}
 		}
 	}
 }
diff --git a/Engine/Source/Programs/UnrealBuildTool/Configuration/UEBuildModule.cs b/Engine/Source/Programs/UnrealBuildTool/Configuration/UEBuildModule.cs
--- a/Engine/Source/Programs/UnrealBuildTool/Configuration/UEBuildModule.cs
+++ b/Engine/Source/Programs/UnrealBuildTool/Configuration/UEBuildModule.cs
@@ -259,7 +259,7 @@ namespace UnrealBuildTool
 					else
 					{
 						// Framework on disk
-						Framework = new UEBuildFramework(FrameworkRules.Name, null, DirectoryReference.Combine(ModuleDirectory, FrameworkRules.Path), FrameworkRules.CopyBundledAssets, FrameworkRules.bCopyFramework);
+						Framework = new UEBuildFramework(FrameworkRules.Name, DirectoryReference.Combine(ModuleDirectory, FrameworkRules.Path), FrameworkRules.CopyBundledAssets, FrameworkRules.bCopyFramework);
 					}
 					PublicAdditionalFrameworks.Add(Framework);
 				}
diff --git a/Engine/Source/Programs/UnrealBuildTool/Platform/IOS/IOSToolChain.cs b/Engine/Source/Programs/UnrealBuildTool/Platform/IOS/IOSToolChain.cs
--- a/Engine/Source/Programs/UnrealBuildTool/Platform/IOS/IOSToolChain.cs
+++ b/Engine/Source/Programs/UnrealBuildTool/Platform/IOS/IOSToolChain.cs
@@ -287,9 +287,10 @@ namespace UnrealBuildTool
 			// Add additional frameworks so that their headers can be found
 			foreach (UEBuildFramework Framework in CompileEnvironment.AdditionalFrameworks)
 			{
-				if (Framework.FrameworkDirectory != null)
+				DirectoryReference? FrameworkDirectory = Framework.GetFrameworkDirectory(CompileEnvironment.Platform, CompileEnvironment.Architecture);
+				if (FrameworkDirectory != null)
 				{
-					string FrameworkDir = Framework.FrameworkDirectory.FullName;
+					string FrameworkDir = FrameworkDirectory.FullName;
 					// embedded frameworks have a framework inside of this directory, so we use this directory. regular frameworks need to go one up to point to the 
 					// directory containing the framework. -F gives a path to look for the -framework
 					if (FrameworkDir.EndsWith(".framework"))
@@ -462,10 +463,11 @@ namespace UnrealBuildTool
 			}
 			foreach (UEBuildFramework Framework in LinkEnvironment.AdditionalFrameworks)
 			{
-				if (Framework.FrameworkDirectory != null)
+				DirectoryReference? FrameworkDirectory = Framework.GetFrameworkDirectory(LinkEnvironment.Platform, LinkEnvironment.Architecture);
+				if (FrameworkDirectory != null)
 				{
 					// If this framework has a directory specified, we'll need to setup the path as well
-					string FrameworkDir = Framework.FrameworkDirectory.FullName;
+					string FrameworkDir = FrameworkDirectory.FullName;
 
 					// embedded frameworks have a framework inside of this directory, so we use this directory. regular frameworks need to go one up to point to the 
 					// directory containing the framework. -F gives a path to look for the -framework
@@ -967,17 +969,17 @@ namespace UnrealBuildTool
 			if(Framework.ExtractedTokenFile == null)
 			{
 				FileItem InputFile = FileItem.GetItemByFileReference(Framework.ZipFile);
-				Framework.ExtractedTokenFile = FileItem.GetItemByFileReference(new FileReference(Framework.FrameworkDirectory!.FullName + ".extracted"));
+				Framework.ExtractedTokenFile = FileItem.GetItemByFileReference(new FileReference(Framework.ZipOutputDirectory!.FullName + ".extracted"));
 
 				StringBuilder ExtractScript = new StringBuilder();
 				ExtractScript.AppendLine("#!/bin/sh");
 				ExtractScript.AppendLine("set -e");
 				// ExtractScript.AppendLine("set -x"); // For debugging
-				ExtractScript.AppendLine(String.Format("[ -d {0} ] && rm -rf {0}", Utils.MakePathSafeToUseWithCommandLine(Framework.FrameworkDirectory.FullName)));
-				ExtractScript.AppendLine(String.Format("unzip -q -o {0} -d {1}", Utils.MakePathSafeToUseWithCommandLine(Framework.ZipFile.FullName), Utils.MakePathSafeToUseWithCommandLine(Framework.FrameworkDirectory.ParentDirectory!.FullName))); // Zip contains folder with the same name, hence ParentDirectory
+				ExtractScript.AppendLine(String.Format("[ -d {0} ] && rm -rf {0}", Utils.MakePathSafeToUseWithCommandLine(Framework.ZipOutputDirectory.FullName)));
+				ExtractScript.AppendLine(String.Format("unzip -q -o {0} -d {1}", Utils.MakePathSafeToUseWithCommandLine(Framework.ZipFile.FullName), Utils.MakePathSafeToUseWithCommandLine(Framework.ZipOutputDirectory.ParentDirectory!.FullName))); // Zip contains folder with the same name, hence ParentDirectory
 				ExtractScript.AppendLine(String.Format("touch {0}", Utils.MakePathSafeToUseWithCommandLine(Framework.ExtractedTokenFile.AbsolutePath)));
 
-				FileItem ExtractScriptFileItem = Graph.CreateIntermediateTextFile(new FileReference(Framework.FrameworkDirectory.FullName + ".sh"), ExtractScript.ToString());
+				FileItem ExtractScriptFileItem = Graph.CreateIntermediateTextFile(new FileReference(Framework.ZipOutputDirectory.FullName + ".sh"), ExtractScript.ToString());
 
 				Action UnzipAction = Graph.CreateAction(ActionType.BuildProject);
 				UnzipAction.CommandPath = new FileReference("/bin/sh");
@@ -987,7 +989,7 @@ namespace UnrealBuildTool
 				UnzipAction.PrerequisiteItems.Add(ExtractScriptFileItem);
 				UnzipAction.ProducedItems.Add(Framework.ExtractedTokenFile);
 				UnzipAction.DeleteItems.Add(Framework.ExtractedTokenFile);
-				UnzipAction.StatusDescription = String.Format("Unzipping : {0} -> {1}", Framework.ZipFile, Framework.FrameworkDirectory);
+				UnzipAction.StatusDescription = String.Format("Unzipping : {0} -> {1}", Framework.ZipFile, Framework.ZipOutputDirectory);
 				UnzipAction.bCanExecuteRemotely = false;
 			}
 			return Framework.ExtractedTokenFile;
@@ -1278,19 +1280,20 @@ namespace UnrealBuildTool
 					DirectoryReference BundleDirectory = Target.bShouldCompileAsDLL ? Executable.Directory.Location : StagedExecutablePath.Directory;
 					foreach (UEBuildFramework Framework in BinaryLinkEnvironment.AdditionalFrameworks)
 					{
-						if (Framework.FrameworkDirectory != null)
+						DirectoryReference? FrameworkDirectory = Framework.GetFrameworkDirectory(BinaryLinkEnvironment.Platform, BinaryLinkEnvironment.Architecture);
+						if (FrameworkDirectory != null)
 						{
 							if (!String.IsNullOrEmpty(Framework.CopyBundledAssets))
 							{
 								// For now, this is hard coded, but we need to loop over all modules, and copy bundled assets that need it
-								DirectoryReference LocalSource = DirectoryReference.Combine(Framework.FrameworkDirectory, Framework.CopyBundledAssets);
+								DirectoryReference LocalSource = DirectoryReference.Combine(FrameworkDirectory, Framework.CopyBundledAssets);
 								string BundleName = Framework.CopyBundledAssets.Substring(Framework.CopyBundledAssets.LastIndexOf('/') + 1);
 								FrameworkNameToSourceDir[BundleName] = LocalSource;
 							}
 
 							if (Framework.bCopyFramework)
 							{
-								string FrameworkDir = Framework.FrameworkDirectory.FullName;
+								string FrameworkDir = FrameworkDirectory.FullName;
 								if (FrameworkDir.EndsWith(".framework"))
 								{
 									FrameworkDir = Path.GetDirectoryName(FrameworkDir)!;
