Support enable ARC in modules

diff --git a/Engine/Source/Programs/UnrealBuildTool/Configuration/ModuleRules.cs b/Engine/Source/Programs/UnrealBuildTool/Configuration/ModuleRules.cs
--- a/Engine/Source/Programs/UnrealBuildTool/Configuration/ModuleRules.cs
+++ b/Engine/Source/Programs/UnrealBuildTool/Configuration/ModuleRules.cs
@@ -720,6 +720,13 @@ namespace UnrealBuildTool
 		/// </summary>
 		public bool bEnableObjCExceptions = false;
 
+		/// <summary>
+		/// Enable objective C automatic reference counting (ARC)
+		/// If you set this to true you should not use shared PCHs for this module. The engine won't be extensively using ARC in the short term  
+		/// Not doing this will result in a compile errors because shared PCHs were compiled with different flags than consumer
+		/// </summary>
+		public bool bEnableObjCAutomaticReferenceCounting = false;
+		
 		/// <summary>
 		/// How to treat shadow variable warnings
 		/// </summary>
diff --git a/Engine/Source/Programs/UnrealBuildTool/Configuration/UEBuildModuleCPP.cs b/Engine/Source/Programs/UnrealBuildTool/Configuration/UEBuildModuleCPP.cs
--- a/Engine/Source/Programs/UnrealBuildTool/Configuration/UEBuildModuleCPP.cs
+++ b/Engine/Source/Programs/UnrealBuildTool/Configuration/UEBuildModuleCPP.cs
@@ -1558,6 +1558,7 @@ namespace UnrealBuildTool
 			Result.bBuildLocallyWithSNDBS = Rules.bBuildLocallyWithSNDBS;
 			Result.bEnableExceptions |= Rules.bEnableExceptions;
 			Result.bEnableObjCExceptions |= Rules.bEnableObjCExceptions;
+			Result.bEnableObjCAutomaticReferenceCounting = Rules.bEnableObjCAutomaticReferenceCounting;
 			Result.ShadowVariableWarningLevel = Rules.ShadowVariableWarningLevel;
 			Result.UnsafeTypeCastWarningLevel = Rules.UnsafeTypeCastWarningLevel;
 			Result.bDisableStaticAnalysis = Rules.bDisableStaticAnalysis;
diff --git a/Engine/Source/Programs/UnrealBuildTool/System/CppCompileEnvironment.cs b/Engine/Source/Programs/UnrealBuildTool/System/CppCompileEnvironment.cs
--- a/Engine/Source/Programs/UnrealBuildTool/System/CppCompileEnvironment.cs
+++ b/Engine/Source/Programs/UnrealBuildTool/System/CppCompileEnvironment.cs
@@ -228,6 +228,13 @@ namespace UnrealBuildTool
 		/// </summary>
 		public bool bEnableObjCExceptions = false;
 
+		/// <summary>
+		/// Enable objective C automatic reference counting (ARC)
+		/// If you set this to true you should not use shared PCHs for this module. The engine won't be extensively using ARC in the short term  
+		/// Not doing this will result in a compile errors because shared PCHs were compiled with different flags than consumer
+		/// </summary>
+		public bool bEnableObjCAutomaticReferenceCounting = false;
+		
 		/// <summary>
 		/// How to treat any warnings in the code
 		/// </summary>
@@ -526,6 +533,7 @@ namespace UnrealBuildTool
 			bRetainFramePointers = Other.bRetainFramePointers;
 			bEnableExceptions = Other.bEnableExceptions;
 			bEnableObjCExceptions = Other.bEnableObjCExceptions;
+			bEnableObjCAutomaticReferenceCounting = Other.bEnableObjCAutomaticReferenceCounting;
 			DefaultWarningLevel = Other.DefaultWarningLevel;
 			DeprecationWarningLevel = Other.DeprecationWarningLevel;
 			ShadowVariableWarningLevel = Other.ShadowVariableWarningLevel;
diff --git a/Engine/Source/Programs/UnrealBuildTool/ToolChain/AppleToolChain.cs b/Engine/Source/Programs/UnrealBuildTool/ToolChain/AppleToolChain.cs
--- a/Engine/Source/Programs/UnrealBuildTool/ToolChain/AppleToolChain.cs
+++ b/Engine/Source/Programs/UnrealBuildTool/ToolChain/AppleToolChain.cs
@@ -254,6 +254,16 @@ namespace UnrealBuildTool
 			return Definition.Contains("\"") ? Definition.Replace("\"", "\\\"") : Definition;
 		}
 
+		protected override void GetCppStandardCompileArgument(CppCompileEnvironment CompileEnvironment, List<string> Arguments)
+		{
+			if (CompileEnvironment.bEnableObjCAutomaticReferenceCounting)
+			{
+				Arguments.Add("-fobjc-arc");
+			}
+
+			base.GetCppStandardCompileArgument(CompileEnvironment, Arguments);
+		}
+
 		protected override void GetCompileArguments_CPP(CppCompileEnvironment CompileEnvironment, List<string> Arguments)
 		{
 			Arguments.Add("-x objective-c++");
diff --git a/Engine/Source/Runtime/Core/Private/Containers/String.cpp b/Engine/Source/Runtime/Core/Private/Containers/String.cpp
--- a/Engine/Source/Runtime/Core/Private/Containers/String.cpp
+++ b/Engine/Source/Runtime/Core/Private/Containers/String.cpp
@@ -324,6 +324,18 @@ void FString::Shrink()
 	Data.Shrink();
 }
 
+#ifdef __OBJC__
+/** Convert FString to Objective-C NSString */
+NSString* FString::GetNSString() const
+{
+#if PLATFORM_TCHAR_IS_4_BYTES
+    return [[[NSString alloc] initWithBytes:Data.GetData() length:Len() * sizeof(TCHAR) encoding:NSUTF32LittleEndianStringEncoding] autorelease];
+#else
+    return [[[NSString alloc] initWithBytes:Data.GetData() length:Len() * sizeof(TCHAR) encoding:NSUTF16LittleEndianStringEncoding] autorelease];
+#endif
+}
+#endif
+
 FString& FString::AppendChar(TCHAR InChar)
 {
 	CheckInvariants();
diff --git a/Engine/Source/Runtime/Core/Public/Apple/ApplePlatformMisc.h b/Engine/Source/Runtime/Core/Public/Apple/ApplePlatformMisc.h
--- a/Engine/Source/Runtime/Core/Public/Apple/ApplePlatformMisc.h
+++ b/Engine/Source/Runtime/Core/Public/Apple/ApplePlatformMisc.h
@@ -23,6 +23,7 @@
 #define UE_DEBUG_BREAK_IMPL() PLATFORM_BREAK()
 
 #ifdef __OBJC__
+#if !__has_feature(objc_arc)
 
 class FScopeAutoreleasePool
 {
@@ -45,6 +46,7 @@ private:
 
 #define SCOPED_AUTORELEASE_POOL const FScopeAutoreleasePool PREPROCESSOR_JOIN(Pool,__LINE__);
 
+#endif // !__has_feature(objc_arc)
 #endif // __OBJC__
 
 /**
diff --git a/Engine/Source/Runtime/Core/Public/Containers/UnrealString.h b/Engine/Source/Runtime/Core/Public/Containers/UnrealString.h
--- a/Engine/Source/Runtime/Core/Public/Containers/UnrealString.h
+++ b/Engine/Source/Runtime/Core/Public/Containers/UnrealString.h
@@ -342,14 +342,7 @@ public:
 
 #ifdef __OBJC__
 	/** Convert FString to Objective-C NSString */
-	FORCEINLINE NSString* GetNSString() const
-	{
-#if PLATFORM_TCHAR_IS_4_BYTES
-		return [[[NSString alloc] initWithBytes:Data.GetData() length:Len() * sizeof(TCHAR) encoding:NSUTF32LittleEndianStringEncoding] autorelease];
-#else
-		return [[[NSString alloc] initWithBytes:Data.GetData() length:Len() * sizeof(TCHAR) encoding:NSUTF16LittleEndianStringEncoding] autorelease];
-#endif
-	}
+    NSString* GetNSString() const;
 #endif
 
 	/** 
