UnrealBuildTool: Only suggest installing visual studio components for Linux or HoloLens if those platforms are valid

#rnx

diff --git a/Engine/Source/Programs/UnrealBuildTool/Platform/Windows/MicrosoftPlatformSDK.Versions.cs b/Engine/Source/Programs/UnrealBuildTool/Platform/Windows/MicrosoftPlatformSDK.Versions.cs
--- a/Engine/Source/Programs/UnrealBuildTool/Platform/Windows/MicrosoftPlatformSDK.Versions.cs
+++ b/Engine/Source/Programs/UnrealBuildTool/Platform/Windows/MicrosoftPlatformSDK.Versions.cs
@@ -63,14 +63,31 @@ namespace UnrealBuildTool
 		static readonly string[] VisualStudioSuggestedComponents = new string[]
 		{
 			"Microsoft.VisualStudio.Workload.CoreEditor",
-			"Microsoft.VisualStudio.Workload.Universal",
 			"Microsoft.VisualStudio.Workload.NativeDesktop",
 			"Microsoft.VisualStudio.Workload.NativeGame",
 			"Microsoft.VisualStudio.Component.VC.Tools.x86.x64",
-			"Microsoft.VisualStudio.Component.VC.Tools.ARM64",
 			"Microsoft.VisualStudio.Component.Windows10SDK",
 		};
 
+		/// <summary>
+		/// Additional set of components that should be suggested to be installed for Visual Studio 2019 or 2022
+		/// to support the HoloLens platform.
+		/// </summary>
+		static readonly string[] VisualStudioSuggestedHololensComponents = new string[]
+		{
+			"Microsoft.VisualStudio.Workload.Universal",
+			"Microsoft.VisualStudio.Component.VC.Tools.ARM64",
+		};
+
+		/// <summary>
+		/// Additional set of components that should be suggested to be installed for Visual Studio 2019 or 2022
+		/// to support the Linux platform.
+		/// </summary>
+		static readonly string[] VisualStudioSuggestedLinuxComponents = new string[]
+		{
+			"Microsoft.VisualStudio.Workload.NativeCrossPlat",
+		};
+
 		/// <summary>
 		/// Additional set of components that should be suggested to be installed for Visual Studio 2019.
 		/// </summary>
@@ -84,18 +101,28 @@ namespace UnrealBuildTool
 		static readonly string[] VisualStudio2022SuggestedComponents = new string[]
 		{
 			"Microsoft.VisualStudio.Workload.ManagedDesktop",
-			"Microsoft.VisualStudio.Workload.NativeCrossPlat",
-			"Microsoft.VisualStudio.Component.VC.14.34.17.4.ARM64",
 			"Microsoft.VisualStudio.Component.VC.14.34.17.4.x86.x64",
 			"Microsoft.Net.Component.4.6.2.TargetingPack",
 		};
 
+		/// <summary>
+		/// Additional set of components that should be suggested to be installed for Visual Studio 2022
+		/// to support the HoloLens platform.
+		/// </summary>
+		static readonly string[] VisualStudio2022SuggestedHololensComponents = new string[]
+		{
+			"Microsoft.VisualStudio.Component.VC.14.34.17.4.ARM64",
+		};
+
 		/// <summary>
 		/// Returns the list of suggested of components that should be suggested to be installed for Visual Studio.
 		/// Used to generate a .vsconfig file which will prompt Visual Studio to ask the user to install these components.
 		/// </summary>
 		public static IEnumerable<string> GetVisualStudioSuggestedComponents(VCProjectFileFormat Format)
 		{
+			bool LinuxValid = InstalledPlatformInfo.IsValidPlatform(UnrealTargetPlatform.Linux);
+			bool HololensValid = InstalledPlatformInfo.IsValidPlatform(UnrealTargetPlatform.HoloLens);
+
 			SortedSet<string> Components = new SortedSet<string>();
 			Components.UnionWith(VisualStudioSuggestedComponents);
 
@@ -106,11 +133,25 @@ namespace UnrealBuildTool
 					break;
 				case VCProjectFileFormat.VisualStudio2022:
 					Components.UnionWith(VisualStudio2022SuggestedComponents);
+					if (HololensValid)
+					{
+						Components.UnionWith(VisualStudio2022SuggestedHololensComponents);
+					}
 					break;
 				default:
 					throw new BuildException("Unsupported Visual Studio version");
 			}
 
+			if (LinuxValid)
+			{
+				Components.UnionWith(VisualStudioSuggestedLinuxComponents);
+			}
+
+			if (HololensValid)
+			{
+				Components.UnionWith(VisualStudioSuggestedHololensComponents);
+			}
+
 			return Components;
 		}
 
