Set arm64 as default architecture for iOS.
There is only support for that one and that fixes Architecture argument being empty when selecting the XCFramework path

https://horde.devtools.epicgames.com/log/63742faa8f4cb2e4dc1242cc?lineindex=59
```
WARNING: Variant for platform "IOS" with architecture  not found in XCFramework FBAEMKit.
WARNING: Variant for platform "IOS" with architecture  not found in XCFramework FBSDKCoreKit_Basics.
WARNING: Variant for platform "IOS" with architecture  not found in XCFramework FBSDKCoreKit.
WARNING: Variant for platform "IOS" with architecture  not found in XCFramework FBSDKLoginKit.
WARNING: Variant for platform "IOS" with architecture  not found in XCFramework FBSDKShareKit.
```

diff --git a/Engine/Source/Programs/UnrealBuildTool/Platform/IOS/UEBuildIOS.cs b/Engine/Source/Programs/UnrealBuildTool/Platform/IOS/UEBuildIOS.cs
--- a/Engine/Source/Programs/UnrealBuildTool/Platform/IOS/UEBuildIOS.cs
+++ b/Engine/Source/Programs/UnrealBuildTool/Platform/IOS/UEBuildIOS.cs
@@ -724,7 +724,7 @@ namespace UnrealBuildTool
 		Dictionary<string, IOSProvisioningData> ProvisionCache = new Dictionary<string, IOSProvisioningData>();
 
 		// by default, use an empty architecture (which is really just a modifer to the platform for some paths/names)
-		public static string IOSArchitecture = "";
+		public static string IOSArchitecture = "arm64";
 
 		public IOSPlatform(UEBuildPlatformSDK InSDK, ILogger Logger)
 			: this(InSDK, UnrealTargetPlatform.IOS, Logger)
