- Fixed the Mac linking of all dependent dylibs when only .cpp changes are made that don't need to relink everything

diff --git a/Engine/Source/Programs/UnrealBuildTool/Platform/Mac/MacToolChain.cs b/Engine/Source/Programs/UnrealBuildTool/Platform/Mac/MacToolChain.cs
--- a/Engine/Source/Programs/UnrealBuildTool/Platform/Mac/MacToolChain.cs
+++ b/Engine/Source/Programs/UnrealBuildTool/Platform/Mac/MacToolChain.cs
@@ -542,7 +542,7 @@ namespace UnrealBuildTool
 
 			// Tell the action that we're building an import library here and it should conditionally be
 			// ignored as a prerequisite for other actions
-			LinkAction.bProducesImportLibrary = RuntimePlatform.IsWindows && (bBuildImportLibraryOnly || LinkEnvironment.bIsBuildingDLL);
+			LinkAction.bProducesImportLibrary = bBuildImportLibraryOnly || LinkEnvironment.bIsBuildingDLL;
 			
 			// Add the output file as a production of the link action.
 			FileItem OutputFile = FileItem.GetItemByFileReference(LinkEnvironment.OutputFilePath);
