[UBT] Add options to enable pie and stack protector for target rules

disabled by default

diff --git a/Engine/Source/Programs/UnrealBuildTool/Configuration/TargetRules.cs b/Engine/Source/Programs/UnrealBuildTool/Configuration/TargetRules.cs
--- a/Engine/Source/Programs/UnrealBuildTool/Configuration/TargetRules.cs
+++ b/Engine/Source/Programs/UnrealBuildTool/Configuration/TargetRules.cs
@@ -994,6 +994,18 @@ namespace UnrealBuildTool
 		[RequiresUniqueBuildEnvironment]
 		public bool bForceEnableRTTI = false;
 
+		/// <summary>
+		/// Enable Position Independent Executable (PIE). Has an overhead cost
+		/// </summary>
+		[CommandLine("-pie")]
+		public bool bEnablePIE = false;
+
+		/// <summary>
+		/// Enable Stack Protection. Has an overhead cost
+		/// </summary>
+		[CommandLine("-stack-protect")]
+		public bool bEnableStackProtection = false;
+
 		/// <summary>
 		/// Compile server-only code.
 		/// </summary>
@@ -2902,6 +2914,16 @@ namespace UnrealBuildTool
 			get { return Inner.bForceEnableRTTI; }
 		}
 
+		public bool bEnablePIE
+		{
+			get { return Inner.bEnablePIE; }
+		}
+
+		public bool bEnableStackProtection
+		{
+			get { return Inner.bEnableStackProtection; }
+		}
+
 		public bool bUseInlining
 		{
 			get { return Inner.bUseInlining; }
diff --git a/Engine/Source/Programs/UnrealBuildTool/Configuration/UEBuildTarget.cs b/Engine/Source/Programs/UnrealBuildTool/Configuration/UEBuildTarget.cs
--- a/Engine/Source/Programs/UnrealBuildTool/Configuration/UEBuildTarget.cs
+++ b/Engine/Source/Programs/UnrealBuildTool/Configuration/UEBuildTarget.cs
@@ -3914,6 +3914,8 @@ namespace UnrealBuildTool
 			GlobalCompileEnvironment.bCheckSystemHeadersForModification = Rules.bCheckSystemHeadersForModification;
 			GlobalCompileEnvironment.bPrintTimingInfo = Rules.bPrintToolChainTimingInfo;
 			GlobalCompileEnvironment.bUseRTTI = Rules.bForceEnableRTTI;
+			GlobalCompileEnvironment.bUsePIE = Rules.bEnablePIE;
+			GlobalCompileEnvironment.bUseStackProtection = Rules.bEnableStackProtection;
 			GlobalCompileEnvironment.bUseInlining = Rules.bUseInlining;
 			GlobalCompileEnvironment.bCompileISPC = Rules.bCompileISPC;
 			GlobalCompileEnvironment.bHideSymbolsByDefault = !Rules.bPublicSymbolsByDefault;
@@ -3944,6 +3946,7 @@ namespace UnrealBuildTool
 			GlobalLinkEnvironment.bUseFastPDBLinking = Rules.bUseFastPDBLinking ?? false;
 			GlobalLinkEnvironment.bDeterministic = Rules.bDeterministic || Rules.bForceDeterministic;
 			GlobalLinkEnvironment.bPrintTimingInfo = Rules.bPrintToolChainTimingInfo;
+			GlobalLinkEnvironment.bUsePIE = Rules.bEnablePIE;
 			GlobalLinkEnvironment.AdditionalArguments = Rules.AdditionalLinkerArguments ?? String.Empty;
 
 			if (Rules.bPGOOptimize && Rules.bPGOProfile)
diff --git a/Engine/Source/Programs/UnrealBuildTool/Platform/Linux/LinuxToolChain.cs b/Engine/Source/Programs/UnrealBuildTool/Platform/Linux/LinuxToolChain.cs
--- a/Engine/Source/Programs/UnrealBuildTool/Platform/Linux/LinuxToolChain.cs
+++ b/Engine/Source/Programs/UnrealBuildTool/Platform/Linux/LinuxToolChain.cs
@@ -546,9 +546,22 @@ namespace UnrealBuildTool
 				Arguments.Add("-fdata-sections");
 			}
 
-			if (bSuppressPIE && !CompileEnvironment.bIsBuildingDLL)
+			// only suppress if happen to be using a system compiler and we have not explicitly requested pie
+			if (!CompileEnvironment.bIsBuildingDLL)
 			{
-				Arguments.Add("-fno-PIE");
+				if (CompileEnvironment.bUsePIE)
+				{
+					Arguments.Add("-fPIE");
+				}
+				else if (bSuppressPIE)
+				{
+					Arguments.Add("-fno-PIE");
+				}
+			}
+
+			if (CompileEnvironment.bUseStackProtection)
+			{
+				Arguments.Add("-fstack-protector");
 			}
 
 			if (PlatformSDK.bVerboseCompiler)
@@ -691,9 +704,12 @@ namespace UnrealBuildTool
 			{
 				Arguments.Add("-Wl,--gc-sections");
 
-				if (bSuppressPIE)
+				if (LinkEnvironment.bUsePIE)
 				{
-					Arguments.Add("-Wl,-no-pie");
+					Arguments.Add("-pie");
+				}
+				else if (bSuppressPIE)
+				{ 					Arguments.Add("-Wl,-no-pie");
 				}
 			}
 
@@ -830,10 +846,20 @@ namespace UnrealBuildTool
 				Logger.LogInformation("Using LTO (link-time optimization).");
 			}
 
-			if (bSuppressPIE)
+			if (CompileEnvironment.bUsePIE)
+			{
+				Logger.LogInformation("Using position independent executables (PIE)");
+			}
+			else if (bSuppressPIE)
 			{
 				Logger.LogInformation("Compiler is set up to generate position independent executables by default, but we're suppressing it.");
 			}
+
+			if (CompileEnvironment.bUseStackProtection)
+			{
+				Logger.LogInformation("Using stack protection");
+			}
+
 			Logger.LogInformation("------------------------------");
 		}
 
diff --git a/Engine/Source/Programs/UnrealBuildTool/System/CppCompileEnvironment.cs b/Engine/Source/Programs/UnrealBuildTool/System/CppCompileEnvironment.cs
--- a/Engine/Source/Programs/UnrealBuildTool/System/CppCompileEnvironment.cs
+++ b/Engine/Source/Programs/UnrealBuildTool/System/CppCompileEnvironment.cs
@@ -170,6 +170,16 @@ namespace UnrealBuildTool
 		/// </summary>
 		public bool bUseRTTI = false;
 
+		/// <summary>
+		/// Use Position Independent Executable (PIE). Has an overhead cost
+		/// </summary>
+		public bool bUsePIE = false;
+
+		/// <summary>
+		/// Use Stack Protection. Has an overhead cost
+		/// </summary>
+		public bool bUseStackProtection = false;
+
 		/// <summary>
 		/// Enable inlining.
 		/// </summary>
@@ -523,6 +533,8 @@ namespace UnrealBuildTool
 			PrecompiledHeaderAction = Other.PrecompiledHeaderAction;
 			bUseSharedBuildEnvironment = Other.bUseSharedBuildEnvironment;
 			bUseRTTI = Other.bUseRTTI;
+			bUsePIE = Other.bUsePIE;
+			bUseStackProtection = Other.bUseStackProtection;
 			bUseInlining = Other.bUseInlining;
 			bCompileISPC = Other.bCompileISPC;
 			bUseAVX = Other.bUseAVX;
diff --git a/Engine/Source/Programs/UnrealBuildTool/System/LinkEnvironment.cs b/Engine/Source/Programs/UnrealBuildTool/System/LinkEnvironment.cs
--- a/Engine/Source/Programs/UnrealBuildTool/System/LinkEnvironment.cs
+++ b/Engine/Source/Programs/UnrealBuildTool/System/LinkEnvironment.cs
@@ -251,6 +251,11 @@ namespace UnrealBuildTool
 		/// </summary>
 		public bool bUseFastPDBLinking;
 
+		/// <summary>
+		/// Use Position Independent Executable (PIE). Has an overhead cost
+		/// </summary>
+		public bool bUsePIE = false;
+
 		/// <summary>
 		/// Whether to ignore dangling (i.e. unresolved external) symbols in modules
 		/// </summary>
@@ -388,6 +393,7 @@ namespace UnrealBuildTool
             bCreateMapFile = Other.bCreateMapFile;
 			bUsePDBFiles = Other.bUsePDBFiles;
 			bUseFastPDBLinking = Other.bUseFastPDBLinking;
+			bUsePIE = Other.bUsePIE;
 			bIgnoreUnresolvedSymbols = Other.bIgnoreUnresolvedSymbols;
 			bDeterministic = Other.bDeterministic;
 			bPrintTimingInfo = Other.bPrintTimingInfo;
