Updated to disable XGE and SNDBS when profiling.
Added a flag that turns off module unity size overrides.

diff --git a/Engine/Source/Programs/UnrealBuildTool/Configuration/ModuleRules.cs b/Engine/Source/Programs/UnrealBuildTool/Configuration/ModuleRules.cs
--- a/Engine/Source/Programs/UnrealBuildTool/Configuration/ModuleRules.cs
+++ b/Engine/Source/Programs/UnrealBuildTool/Configuration/ModuleRules.cs
@@ -830,6 +830,14 @@ namespace UnrealBuildTool
 		/// </summary>
 		public int NumIncludedBytesPerUnityCPPOverride = 0;
 
+		/// <summary>
+		/// Helper function to get the number of byes per unity cpp file
+		/// </summary>
+		public int GetNumIncludedBytesPerUnityCPP()
+		{
+			return (NumIncludedBytesPerUnityCPPOverride != 0 && !Target.bDisableModuleNumIncludedBytesPerUnityCPPOverride) ? NumIncludedBytesPerUnityCPPOverride : Target.NumIncludedBytesPerUnityCPP;
+		}
+
 		/// <summary>
 		/// Module uses a #import so must be built locally when compiling with SN-DBS
 		/// </summary>
diff --git a/Engine/Source/Programs/UnrealBuildTool/Configuration/TargetRules.cs b/Engine/Source/Programs/UnrealBuildTool/Configuration/TargetRules.cs
--- a/Engine/Source/Programs/UnrealBuildTool/Configuration/TargetRules.cs
+++ b/Engine/Source/Programs/UnrealBuildTool/Configuration/TargetRules.cs
@@ -1444,6 +1444,13 @@ namespace UnrealBuildTool
 		[XmlConfigFile(Category = "BuildConfiguration")]
 		public int NumIncludedBytesPerUnityCPP = 384 * 1024;
 
+		/// <summary>
+		/// Disables overrides that are set by the module
+		/// </summary>
+		[CommandLine("-DisableModuleNumIncludedBytesPerUnityCPPOverride", Value = "true")]
+		[XmlConfigFile(Category = "BuildConfiguration")]
+		public bool bDisableModuleNumIncludedBytesPerUnityCPPOverride = false;
+
 		/// <summary>
 		/// Whether to stress test the C++ unity build robustness by including all C++ files files in a project from a single unified file.
 		/// </summary>
@@ -3194,6 +3201,11 @@ namespace UnrealBuildTool
 			get { return Inner.NumIncludedBytesPerUnityCPP; }
 		}
 
+		public bool bDisableModuleNumIncludedBytesPerUnityCPPOverride
+		{
+			get { return Inner.bDisableModuleNumIncludedBytesPerUnityCPPOverride; }
+		}
+
 		public bool bStressTestUnity
 		{
 			get { return Inner.bStressTestUnity; }
diff --git a/Engine/Source/Programs/UnrealBuildTool/Configuration/UEBuildModuleCPP.cs b/Engine/Source/Programs/UnrealBuildTool/Configuration/UEBuildModuleCPP.cs
--- a/Engine/Source/Programs/UnrealBuildTool/Configuration/UEBuildModuleCPP.cs
+++ b/Engine/Source/Programs/UnrealBuildTool/Configuration/UEBuildModuleCPP.cs
@@ -630,7 +630,7 @@ namespace UnrealBuildTool
 			}
 
 			// Set up the NumIncludedBytesPerUnityCPP for this particular module
-			int NumIncludedBytesPerUnityCPP = (Rules.NumIncludedBytesPerUnityCPPOverride != 0) ? Rules.NumIncludedBytesPerUnityCPPOverride : Target.NumIncludedBytesPerUnityCPP;
+			int NumIncludedBytesPerUnityCPP = Rules.GetNumIncludedBytesPerUnityCPP();
 
 			// Should we use unity build mode for this module?
 			bool bModuleUsesUnityBuild = false;
diff --git a/Engine/Source/Programs/UnrealBuildTool/Modes/ProfileUnitySizesMode.cs b/Engine/Source/Programs/UnrealBuildTool/Modes/ProfileUnitySizesMode.cs
--- a/Engine/Source/Programs/UnrealBuildTool/Modes/ProfileUnitySizesMode.cs
+++ b/Engine/Source/Programs/UnrealBuildTool/Modes/ProfileUnitySizesMode.cs
@@ -140,6 +140,8 @@ namespace UnrealBuildTool
 			{
 				List<UEBuildModule> ModuleList = new();
 
+				TargetDescriptor.AdditionalArguments = TargetDescriptor.AdditionalArguments.Append(new string[] { "-NoSNDBS", "-NoXGE" });
+
 				// Create a makefile for the target
 				TimingLogger TimingLogger = new(Logger);
 				UEBuildTarget Target = UEBuildTarget.Create(TargetDescriptor, BuildConfiguration.bSkipRulesCompile, BuildConfiguration.bForceRulesCompile, BuildConfiguration.bUsePrecompiled, TimingLogger);
@@ -186,7 +188,7 @@ namespace UnrealBuildTool
 
 			Logger.LogInformation($"{Module.Name}:");
 
-			int CurrentModuleUnitySize = (Module.Rules.NumIncludedBytesPerUnityCPPOverride != 0) ? Module.Rules.NumIncludedBytesPerUnityCPPOverride : Target.Rules.NumIncludedBytesPerUnityCPP;
+			int CurrentModuleUnitySize = Module.Rules.GetNumIncludedBytesPerUnityCPP();
 			int TargetUnitySize = Target.Rules.NumIncludedBytesPerUnityCPP;
 
 			int BuildNum = 1;
@@ -319,7 +321,7 @@ namespace UnrealBuildTool
 					}
 
 					// Add the module name to the cmdline
-					TargetDescriptor.AdditionalArguments = TargetDescriptor.AdditionalArguments.Append(new string[] { $"-BytesPerUnityCPP={UnitySize}" });
+					TargetDescriptor.AdditionalArguments = TargetDescriptor.AdditionalArguments.Append(new string[] { $"-BytesPerUnityCPP={UnitySize}", "-DisableModuleNumIncludedBytesPerUnityCPPOverride" });
 
 					if (bDisableUnity)
 					{
