UnrealBuildTool: Set VS2022 as the default generated solution and compiler if all available platforms are supported

diff --git a/Engine/Source/Programs/UnrealBuildTool/Platform/Windows/UEBuildWindows.cs b/Engine/Source/Programs/UnrealBuildTool/Platform/Windows/UEBuildWindows.cs
--- a/Engine/Source/Programs/UnrealBuildTool/Platform/Windows/UEBuildWindows.cs
+++ b/Engine/Source/Programs/UnrealBuildTool/Platform/Windows/UEBuildWindows.cs
@@ -995,13 +995,13 @@ namespace UnrealBuildTool
 			{
 				foreach(ProjectFileFormat Format in ProjectFileGeneratorSettings.ParseFormatList(ProjectFileGeneratorSettings.Format, Logger))
 				{
-					if (Format == ProjectFileFormat.VisualStudio2019)
+					if (Format == ProjectFileFormat.VisualStudio2022)
 					{
-						return WindowsCompiler.VisualStudio2019;
+						return WindowsCompiler.VisualStudio2022;
 					}
-					else if (Format == ProjectFileFormat.VisualStudio2022)
+					else if (Format == ProjectFileFormat.VisualStudio2019)
 					{
-						return WindowsCompiler.VisualStudio2022;
+						return WindowsCompiler.VisualStudio2019;
 					}
 				} 
 			}
@@ -1011,13 +1011,13 @@ namespace UnrealBuildTool
 			if (XmlConfig.TryGetValue(typeof(VCProjectFileSettings), "ProjectFileFormat", out ProjectFormatObject))
 			{
 				VCProjectFileFormat ProjectFormat = (VCProjectFileFormat)ProjectFormatObject;
-				if (ProjectFormat == VCProjectFileFormat.VisualStudio2019)
+				if (ProjectFormat == VCProjectFileFormat.VisualStudio2022)
 				{
-					return WindowsCompiler.VisualStudio2019;
+					return WindowsCompiler.VisualStudio2022;
 				}
-				else if (ProjectFormat == VCProjectFileFormat.VisualStudio2022)
+				else if (ProjectFormat == VCProjectFileFormat.VisualStudio2019)
 				{
-					return WindowsCompiler.VisualStudio2022;
+					return WindowsCompiler.VisualStudio2019;
 				}
 			}
 
@@ -1025,48 +1025,48 @@ namespace UnrealBuildTool
 			ProjectFileFormat PreferredAccessor;
 			if(ProjectFileGenerator.GetPreferredSourceCodeAccessor(ProjectFile, out PreferredAccessor))
 			{
-				if(PreferredAccessor == ProjectFileFormat.VisualStudio2019)
+				if(PreferredAccessor == ProjectFileFormat.VisualStudio2022)
 			    {
-				    return WindowsCompiler.VisualStudio2019;
+				    return WindowsCompiler.VisualStudio2022;
 			    }
-				else if (PreferredAccessor == ProjectFileFormat.VisualStudio2022)
+				else if (PreferredAccessor == ProjectFileFormat.VisualStudio2019)
 				{
-					return WindowsCompiler.VisualStudio2022;
+					return WindowsCompiler.VisualStudio2019;
 				}
 			}
 
-			// Second, default based on what's installed, test for 2019 first
-			if (MicrosoftPlatformSDK.HasValidCompiler(WindowsCompiler.VisualStudio2019, Architecture, Logger))
+			// Second, default based on what's installed, test for 2022 first
+			if (MicrosoftPlatformSDK.HasValidCompiler(WindowsCompiler.VisualStudio2022, Architecture, Logger))
 			{
-				return WindowsCompiler.VisualStudio2019;
+				return WindowsCompiler.VisualStudio2022;
 			}
-			else if (MicrosoftPlatformSDK.HasValidCompiler(WindowsCompiler.VisualStudio2022, Architecture, Logger))
+			else if (MicrosoftPlatformSDK.HasValidCompiler(WindowsCompiler.VisualStudio2019, Architecture, Logger))
 			{
-				return WindowsCompiler.VisualStudio2022;
+				return WindowsCompiler.VisualStudio2019;
 			}
 
 			// If we do have a Visual Studio installation, but we're missing just the C++ parts, warn about that.
-			if (TryGetVSInstallDirs(WindowsCompiler.VisualStudio2019, Logger) != null)
-			{
-				string ToolSetWarning = Architecture == WindowsArchitecture.x64 ?
-					"MSVC v142 - VS 2019 C++ x64/x86 build tools (Latest)" :
-					"MSVC v142 - VS 2019 C++ ARM64 build tools (Latest)";
-				Logger.LogWarning("Visual Studio 2019 is installed, but is missing the C++ toolchain. Please verify that the \"{Component}\" component is selected in the Visual Studio 2019 installation options.", ToolSetWarning);
-			}
-			else if (TryGetVSInstallDirs(WindowsCompiler.VisualStudio2022, Logger) != null)
+			if (TryGetVSInstallDirs(WindowsCompiler.VisualStudio2022, Logger) != null)
 			{
 				string ToolSetWarning = Architecture == WindowsArchitecture.x64 ?
 					"MSVC v143 - VS 2022 C++ x64/x86 build tools (Latest)" :
 					"MSVC v143 - VS 2022 C++ ARM64 build tools (Latest)";
 				Logger.LogWarning("Visual Studio 2022 is installed, but is missing the C++ toolchain. Please verify that the \"{Component}\" component is selected in the Visual Studio 2022 installation options.", ToolSetWarning);
 			}
+			else if (TryGetVSInstallDirs(WindowsCompiler.VisualStudio2019, Logger) != null)
+			{
+				string ToolSetWarning = Architecture == WindowsArchitecture.x64 ?
+					"MSVC v142 - VS 2019 C++ x64/x86 build tools (Latest)" :
+					"MSVC v142 - VS 2019 C++ ARM64 build tools (Latest)";
+				Logger.LogWarning("Visual Studio 2019 is installed, but is missing the C++ toolchain. Please verify that the \"{Component}\" component is selected in the Visual Studio 2019 installation options.", ToolSetWarning);
+			}
 			else
 			{
-				Logger.LogWarning("No Visual C++ installation was found. Please download and install Visual Studio 2019 or 2022 with C++ components.");
+				Logger.LogWarning("No Visual C++ installation was found. Please download and install Visual Studio 2022 or 2019 with C++ components.");
 			}
 
-			// Finally, default to VS2019 anyway
-			return WindowsCompiler.VisualStudio2019;
+			// Finally, default to VS2022 anyway
+			return WindowsCompiler.VisualStudio2022;
 		}
 
 		/// <summary>
diff --git a/Engine/Source/Programs/UnrealBuildTool/ProjectFiles/VisualStudio/VCProjectFileGenerator.cs b/Engine/Source/Programs/UnrealBuildTool/ProjectFiles/VisualStudio/VCProjectFileGenerator.cs
--- a/Engine/Source/Programs/UnrealBuildTool/ProjectFiles/VisualStudio/VCProjectFileGenerator.cs
+++ b/Engine/Source/Programs/UnrealBuildTool/ProjectFiles/VisualStudio/VCProjectFileGenerator.cs
@@ -268,19 +268,19 @@ namespace UnrealBuildTool
 			base.SetupSupportedPlatformsAndConfigurations(IncludeAllPlatforms, Logger, out SupportedPlatformNames);
 
 			// If we have a non-default setting for visual studio, check the compiler exists. If not, revert to the default.
-			if (Settings.ProjectFileFormat == VCProjectFileFormat.VisualStudio2019)
+			if (Settings.ProjectFileFormat == VCProjectFileFormat.VisualStudio2022)
 			{
-				if (!WindowsPlatform.HasCompiler(WindowsCompiler.VisualStudio2019, WindowsArchitecture.x64, Logger))
+				if (!WindowsPlatform.HasCompiler(WindowsCompiler.VisualStudio2022, WindowsArchitecture.x64, Logger))
 				{
-					Logger.LogWarning("Visual Studio C++ 2019 installation not found - ignoring preferred project file format.");
+					Logger.LogWarning("Visual Studio C++ 2022 installation not found - ignoring preferred project file format.");
 					Settings.ProjectFileFormat = VCProjectFileFormat.Default;
 				}
 			}
-			else if (Settings.ProjectFileFormat == VCProjectFileFormat.VisualStudio2022)
+			else if (Settings.ProjectFileFormat == VCProjectFileFormat.VisualStudio2019)
 			{
-				if (!WindowsPlatform.HasCompiler(WindowsCompiler.VisualStudio2022, WindowsArchitecture.x64, Logger))
+				if (!WindowsPlatform.HasCompiler(WindowsCompiler.VisualStudio2019, WindowsArchitecture.x64, Logger))
 				{
-					Logger.LogWarning("Visual Studio C++ 2022 installation not found - ignoring preferred project file format.");
+					Logger.LogWarning("Visual Studio C++ 2019 installation not found - ignoring preferred project file format.");
 					Settings.ProjectFileFormat = VCProjectFileFormat.Default;
 				}
 			}
diff --git a/Engine/Source/Programs/Windows/BootstrapPackagedGame/Private/BootstrapPackagedGame.cpp b/Engine/Source/Programs/Windows/BootstrapPackagedGame/Private/BootstrapPackagedGame.cpp
--- a/Engine/Source/Programs/Windows/BootstrapPackagedGame/Private/BootstrapPackagedGame.cpp
+++ b/Engine/Source/Programs/Windows/BootstrapPackagedGame/Private/BootstrapPackagedGame.cpp
@@ -96,8 +96,8 @@ int InstallMissingPrerequisites(const WCHAR* BaseDirectory, const WCHAR* ExecDir
 			// This minimum should match the version installed by
 			// Engine/Source/Programs/PrereqInstaller/Resources/VCRedist/VC_redist.x64.exe
 			const DWORD RequiredMajor = 14;
-			const DWORD RequiredMinor = 31;
-			const DWORD RequiredBld = 31103;
+			const DWORD RequiredMinor = 34;
+			const DWORD RequiredBld = 31931;
 			const DWORD RequiredRbld = 0;
 
 			const DWORD InstalledMajor = RegGetDwordOrZero(Hkey, L"Major");
