UnrealBuildTool: Compiling Microsoft platforms with Clang when using VS2022 runtime requires Clang 14

#rnx

diff --git a/Engine/Source/Programs/UnrealBuildTool/Platform/Windows/MicrosoftPlatformSDK.Versions.cs b/Engine/Source/Programs/UnrealBuildTool/Platform/Windows/MicrosoftPlatformSDK.Versions.cs
--- a/Engine/Source/Programs/UnrealBuildTool/Platform/Windows/MicrosoftPlatformSDK.Versions.cs
+++ b/Engine/Source/Programs/UnrealBuildTool/Platform/Windows/MicrosoftPlatformSDK.Versions.cs
@@ -21,7 +21,8 @@ namespace UnrealBuildTool
 		/// </summary>
 		static readonly VersionNumberRange[] PreferredClangVersions =
 		{
-			VersionNumberRange.Parse("13.0.0", "13.999"),
+			VersionNumberRange.Parse("14.0.0", "14.999"), // VS2022 17.4.x runtime requires Clang 14
+			VersionNumberRange.Parse("13.0.0", "13.999"), // VS2019 16.11 runtime requires Clang 13
 		};
 
 		static readonly VersionNumber MinimumClangVersion = new VersionNumber(13, 0, 0);
