- Moved a utility function from two (copy/pasted) places into ConfigHierarchy - GetStructKeyValuePairs(), which will return a Dictionary that contains the keys and values from a standard UE struct definition (almost always loaded from a .ini file). Similar in nature to GetStructEntry(), but will entries in the struct

diff --git a/Engine/Source/Programs/AutomationTool/CookedEditor/CookedEditor.Automation.cs b/Engine/Source/Programs/AutomationTool/CookedEditor/CookedEditor.Automation.cs
--- a/Engine/Source/Programs/AutomationTool/CookedEditor/CookedEditor.Automation.cs
+++ b/Engine/Source/Programs/AutomationTool/CookedEditor/CookedEditor.Automation.cs
@@ -879,7 +879,7 @@ public class MakeCookedEditor : BuildCommand
 
 		foreach (string Entry in Entries)
 		{
-			Dictionary<string, string> Props = ParseStructProperties(Entry);
+			Dictionary<string, string> Props = ConfigHierarchy.GetStructKeyValuePairs(Entry);
 
 			string SubPath = Props["Path"];
 			string FileWildcard = "*";
@@ -1217,93 +1217,4 @@ public class MakeCookedEditor : BuildCommand
 
 		Context.NonUFSFilesToStage.Add(ReceiptFilename);
 	}
-
-
-
-
-	// @todo: Move this into UBT or something
-	private static Dictionary<string, string> ParseStructProperties(string PropsString)
-	{
-		// we expect parens around a properly encoded struct
-		if (!PropsString.StartsWith("(") || !PropsString.EndsWith(")"))
-		{
-			return null;
-		}
-		// strip ()
-		PropsString = PropsString.Substring(1, PropsString.Length - 2);
-
-		List<string> Props = new List<string>();
-
-		int TokenStart = 0;
-		int StrLen = PropsString.Length;
-		while (TokenStart < StrLen)
-		{
-			// get the next location of each special character
-			int NextComma = PropsString.IndexOf(',', TokenStart);
-			int NextQuote = PropsString.IndexOf('\"', TokenStart);
-			// comma first? easy
-			if (NextComma != -1 && NextComma < NextQuote)
-			{
-				Props.Add(PropsString.Substring(TokenStart, NextComma - TokenStart));
-				TokenStart = NextComma + 1;
-			}
-			// comma but no quotes
-			else if (NextComma != -1 && NextQuote == -1)
-			{
-				Props.Add(PropsString.Substring(TokenStart, NextComma - TokenStart));
-				TokenStart = NextComma + 1;
-			}
-			// neither found, use the rest
-			else if (NextComma == -1 && NextQuote == -1)
-			{
-				Props.Add(PropsString.Substring(TokenStart));
-				break;
-			}
-			// quote first? look for quote after
-			else
-			{
-				NextQuote = PropsString.IndexOf('\"', NextQuote + 1);
-				// are we at the end?
-				if (NextQuote + 1 == StrLen)
-				{
-					// use the rest of the string
-					Props.Add(PropsString.Substring(TokenStart));
-					break;
-				}
-				// it's expected that the following character is a comma, if not, give up
-				if (PropsString[NextQuote + 1] != ',')
-				{
-					break;
-				}
-				// if next is comma, we are done this token
-				Props.Add(PropsString.Substring(TokenStart, (NextQuote - TokenStart) + 1));
-				// skip over the quote and following commma
-				TokenStart = NextQuote + 2;
-			}
-		}
-
-		// now make a dictionary from the properties
-		Dictionary<string, string> KeyValues = new Dictionary<string, string>();
-		foreach (string AProp in Props)
-		{
-			string Prop = AProp.Trim(" \t".ToCharArray());
-			// find the first = (UE properties can't have an equal sign, so it's valid to do)
-			int Equals = Prop.IndexOf('=');
-			// we must have one
-			if (Equals == -1)
-			{
-				continue;
-			}
-
-			string Key = Prop.Substring(0, Equals);
-			string Value = Prop.Substring(Equals + 1);
-			// trim off any quotes around the entire value
-			Value = Value.Trim(" \"".ToCharArray());
-			Key = Key.Trim(" ".ToCharArray());
-			KeyValues.Add(Key, Value);
-		}
-
-		// convert to array type
-		return KeyValues;
-	}
 }
diff --git a/Engine/Source/Programs/UnrealBuildTool/System/ConfigHierarchy.cs b/Engine/Source/Programs/UnrealBuildTool/System/ConfigHierarchy.cs
--- a/Engine/Source/Programs/UnrealBuildTool/System/ConfigHierarchy.cs
+++ b/Engine/Source/Programs/UnrealBuildTool/System/ConfigHierarchy.cs
@@ -1312,6 +1312,114 @@ namespace UnrealBuildTool
 			return null;
 		}
 
+		/// <summary>
+		/// Load a ini value for the given key, then use GetStructKeyValuePairs to return the key/value pairs of the struct
+		/// </summary>
+		/// <param name="Section">Ini section to read from</param>
+		/// <param name="Setting">Name of the ini key to read from</param>
+		/// <returns>Dictionary of struct vars/values</returns>
+		public Dictionary<string, string>? GetStructKeyValuePairsForSetting(string Section, string Setting)
+		{
+			string ConfigEntry;
+			if (GetString(Section, Setting, out ConfigEntry))
+			{
+				return GetStructKeyValuePairs(ConfigEntry);
+			}
+			return null;
+		}
+
+		/// <summary>
+		/// Given a string input (a struct loaded from a .ini file usually), like (Foo=A, Bar="Hello world"), this will return a dictionary of all values, with quotes trimmed off
+		/// In this example, { { Foo, A } , { Bar, Hello world } }
+		/// </summary>
+		/// <param name="Input">String containing a struct representation</param>
+		/// <returns>Dictionary of struct vars/values</returns>
+		public static Dictionary<string, string> GetStructKeyValuePairs(string Input)
+		{
+			// we expect parens around a properly encoded struct
+			if (!Input.StartsWith("(") || !Input.EndsWith(")"))
+			{
+				return new Dictionary<string, string>();
+			}
+			// strip ()
+			Input = Input.Substring(1, Input.Length - 2);
+
+			List<string> Props = new List<string>();
+
+			int TokenStart = 0;
+			int StrLen = Input.Length;
+			while (TokenStart < StrLen)
+			{
+				// get the next location of each special character
+				int NextComma = Input.IndexOf(',', TokenStart);
+				int NextQuote = Input.IndexOf('\"', TokenStart);
+				// comma first? easy
+				if (NextComma != -1 && NextComma < NextQuote)
+				{
+					Props.Add(Input.Substring(TokenStart, NextComma - TokenStart));
+					TokenStart = NextComma + 1;
+				}
+				// comma but no quotes
+				else if (NextComma != -1 && NextQuote == -1)
+				{
+					Props.Add(Input.Substring(TokenStart, NextComma - TokenStart));
+					TokenStart = NextComma + 1;
+				}
+				// neither found, use the rest
+				else if (NextComma == -1 && NextQuote == -1)
+				{
+					Props.Add(Input.Substring(TokenStart));
+					break;
+				}
+				// quote first? look for quote after
+				else
+				{
+					NextQuote = Input.IndexOf('\"', NextQuote + 1);
+					// are we at the end?
+					if (NextQuote + 1 == StrLen)
+					{
+						// use the rest of the string
+						Props.Add(Input.Substring(TokenStart));
+						break;
+					}
+					// it's expected that the following character is a comma, if not, give up
+					if (Input[NextQuote + 1] != ',')
+					{
+						break;
+					}
+					// if next is comma, we are done this token
+					Props.Add(Input.Substring(TokenStart, (NextQuote - TokenStart) + 1));
+					// skip over the quote and following commma
+					TokenStart = NextQuote + 2;
+				}
+			}
+
+			// now make a dictionary from the properties
+			Dictionary<string, string> KeyValues = new Dictionary<string, string>();
+			foreach (string AProp in Props)
+			{
+				string Prop = AProp.Trim(" \t".ToCharArray());
+				// find the first = (UE properties can't have an equal sign, so it's valid to do)
+				int Equals = Prop.IndexOf('=');
+				// we must have one
+				if (Equals == -1)
+				{
+					continue;
+				}
+
+				string Key = Prop.Substring(0, Equals);
+				string Value = Prop.Substring(Equals + 1);
+				// trim off any quotes around the entire value
+				Value = Value.Trim(" \"".ToCharArray());
+				Key = Key.Trim(" ".ToCharArray());
+				KeyValues.Add(Key, Value);
+			}
+
+			// convert to array type
+			return KeyValues;
+		}
+
+
 		#endregion
 	}
 }
