Remove some customizations from iOS and TvOS builds

diff --git a/Engine/Source/Programs/UnrealBuildTool/Platform/IOS/UEBuildIOS.cs b/Engine/Source/Programs/UnrealBuildTool/Platform/IOS/UEBuildIOS.cs
--- a/Engine/Source/Programs/UnrealBuildTool/Platform/IOS/UEBuildIOS.cs
+++ b/Engine/Source/Programs/UnrealBuildTool/Platform/IOS/UEBuildIOS.cs
@@ -723,9 +723,6 @@ namespace UnrealBuildTool
 		List<IOSProjectSettings> CachedProjectSettingsByBundle = new List<IOSProjectSettings>();
 		Dictionary<string, IOSProvisioningData> ProvisionCache = new Dictionary<string, IOSProvisioningData>();
 
-		// by default, use an empty architecture (which is really just a modifer to the platform for some paths/names)
-		public static string IOSArchitecture = "arm64";
-
 		public IOSPlatform(UEBuildPlatformSDK InSDK, ILogger Logger)
 			: this(InSDK, UnrealTargetPlatform.IOS, Logger)
 		{
@@ -736,17 +733,11 @@ namespace UnrealBuildTool
 		{
 		}
 
-		// The current architecture - affects everything about how UBT operates on IOS
 		public override string GetDefaultArchitecture(FileReference? ProjectFile)
 		{
-			return IOSArchitecture;
-		}
-
-		public override string GetFolderNameForArchitecture(string Architecture)
-		{
-			return IOSArchitecture;
+			return "";
 		}
-
+		
 		public override List<FileReference> FinalizeBinaryPaths(FileReference BinaryName, FileReference? ProjectFile, ReadOnlyTargetRules Target)
 		{
 			List<FileReference> BinaryPaths = new List<FileReference>();
diff --git a/Engine/Source/Programs/UnrealBuildTool/Platform/TVOS/UEBuildTVOS.cs b/Engine/Source/Programs/UnrealBuildTool/Platform/TVOS/UEBuildTVOS.cs
--- a/Engine/Source/Programs/UnrealBuildTool/Platform/TVOS/UEBuildTVOS.cs
+++ b/Engine/Source/Programs/UnrealBuildTool/Platform/TVOS/UEBuildTVOS.cs
@@ -38,15 +38,6 @@ namespace UnrealBuildTool
 		{
 		}
 
-		// by default, use an empty architecture (which is really just a modifer to the platform for some paths/names)
-		public static string TVOSArchitecture = "";
-
-		// The current architecture - affects everything about how UBT operates on IOS
-		public override string GetDefaultArchitecture(FileReference? ProjectFile)
-		{
-			return TVOSArchitecture;
-		}
-
 		public override void ValidateTarget(TargetRules Target)
 		{
 			base.ValidateTarget(Target);
