5.1 QFE

This QFE addresses multiple issues with Mac builds in UE 5.1:

* Removed Apple Silicon from default build architectures if installed from Epic Games Launcher, as it is not yet supported in UE 5.1 as a default option.
   * Apple Silicon native build will be available in UE 5.1 through Epic Games Launcher at a later time, but for now please build from GitHub (URL for linking here - https://github.com/EpicGames/UnrealEngine/tree/5.1)
* Fixed a Mobile Preview crash on Mac when Mobile Deferred is enabled and Preview Rendering Level set to iOS Shader 3.1
* Fixed issue that was preventing movable lights from properly casting light on Mac

diff --git a/Engine/Source/Programs/UnrealBuildTool/ProjectFiles/Xcode/XcodeProjectLegacy.cs b/Engine/Source/Programs/UnrealBuildTool/ProjectFiles/Xcode/XcodeProjectLegacy.cs
--- a/Engine/Source/Programs/UnrealBuildTool/ProjectFiles/Xcode/XcodeProjectLegacy.cs
+++ b/Engine/Source/Programs/UnrealBuildTool/ProjectFiles/Xcode/XcodeProjectLegacy.cs
@@ -915,6 +915,10 @@ namespace UnrealBuildTool.XcodeProjectLegacy
 				Content.Append("\t\t\t\tENABLE_TESTABILITY = YES;" + ProjectFileGenerator.NewLine);
 			}
 			Content.Append("\t\t\t\tALWAYS_SEARCH_USER_PATHS = NO;" + ProjectFileGenerator.NewLine);
+			if (Unreal.IsEngineInstalled())
+			{
+				Content.Append("\t\t\t\tARCHS = x86_64;" + ProjectFileGenerator.NewLine);
+			}
 			Content.Append("\t\t\t\tCLANG_CXX_LANGUAGE_STANDARD = \"c++14\";" + ProjectFileGenerator.NewLine);
 			Content.Append("\t\t\t\tGCC_ENABLE_CPP_RTTI = NO;" + ProjectFileGenerator.NewLine);
 			Content.Append("\t\t\t\tGCC_WARN_CHECK_SWITCH_STATEMENTS = NO;" + ProjectFileGenerator.NewLine);
diff --git a/Engine/Source/Runtime/Apple/MetalRHI/Private/MetalTexture.cpp b/Engine/Source/Runtime/Apple/MetalRHI/Private/MetalTexture.cpp
--- a/Engine/Source/Runtime/Apple/MetalRHI/Private/MetalTexture.cpp
+++ b/Engine/Source/Runtime/Apple/MetalRHI/Private/MetalTexture.cpp
@@ -1427,7 +1427,9 @@ uint32 FMetalSurface::GetMemorySize()
 
 uint32 FMetalSurface::GetNumFaces()
 {
-	return GetDesc().Depth * GetDesc().ArraySize;
+	// UE <= 5.0 for Cube Texture FMetalSurface::SizeZ was set to 6
+	// UE >= 5.1 FMetalSurface::SizeZ does not exist and extent.Depth from the create descriptor is set to 1 for cube textures
+	return GetDesc().Depth * GetDesc().ArraySize * (GetDesc().IsTextureCube() ? 6 : 1);
 }
 
 FMetalTexture FMetalSurface::GetDrawableTexture()
diff --git a/Engine/Source/Runtime/Renderer/Private/SceneTextures.cpp b/Engine/Source/Runtime/Renderer/Private/SceneTextures.cpp
--- a/Engine/Source/Runtime/Renderer/Private/SceneTextures.cpp
+++ b/Engine/Source/Runtime/Renderer/Private/SceneTextures.cpp
@@ -582,7 +582,10 @@ void FSceneTextures::InitializeViewFamily(FRDGBuilder& GraphBuilder, FViewFamily
 		ETextureCreateFlags MemorylessFlag = TexCreate_None;
 		if (IsMobileDeferredShadingEnabled(Config.ShaderPlatform) || (Config.NumSamples > 1 && Config.bMemorylessMSAA))
 		{
+		// hotfix for a crash on a Mac mobile preview, proper fix is in 5.2
+		#if !PLATFORM_MAC
 			MemorylessFlag = TexCreate_Memoryless;
+		#endif
 		}
 
 		EPixelFormat DepthAuxFormat = GetMobileSceneDepthAuxPixelFormat(Config.ShaderPlatform, Config.bPreciseDepthAux);
