Fix typo in ld arguments for edge case Android libs that don't start with "lib".

The intention was:
/some/path/libfoo.a -->
AdditionalLibs.Add(libfoo.a) -> -lfoo
AdditionalLibraryPaths.Add(/some/path) -> -L/some/path
/some/path/bar.a -->
AdditionalLibs.Add(/some/path/bar.a) -> /some/path/bar.a

(Copy from 4.27+)

diff --git a/Engine/Source/Programs/UnrealBuildTool/Platform/Android/AndroidToolChain.cs b/Engine/Source/Programs/UnrealBuildTool/Platform/Android/AndroidToolChain.cs
--- a/Engine/Source/Programs/UnrealBuildTool/Platform/Android/AndroidToolChain.cs
+++ b/Engine/Source/Programs/UnrealBuildTool/Platform/Android/AndroidToolChain.cs
@@ -1961,7 +1961,7 @@ namespace UnrealBuildTool
 							}
 							else
 							{
-								AdditionalLibraries.Add(AbsoluteLibraryPath);
+								AdditionalLibraries.Add(Library.FullName);
 							}
 
 							if (!DisableLibCppSharedDependencyValidation() && ReadElfPath != null)
