fix null reference warning

#rnx

diff --git a/Engine/Source/Programs/UnrealBuildTool/System/XmlConfig.cs b/Engine/Source/Programs/UnrealBuildTool/System/XmlConfig.cs
--- a/Engine/Source/Programs/UnrealBuildTool/System/XmlConfig.cs
+++ b/Engine/Source/Programs/UnrealBuildTool/System/XmlConfig.cs
@@ -649,7 +649,7 @@ namespace UnrealBuildTool
 
 			if(CachedSchemaSerializer == null)
 			{
-				CachedSchemaSerializer = XmlSerializer.FromTypes(new Type[] { typeof(XmlSchema) })[0];
+				CachedSchemaSerializer = XmlSerializer.FromTypes(new Type[] { typeof(XmlSchema) })[0]!;
 			}
 
 			StringBuilder Output = new StringBuilder();
