Cleanup Clang check for older XCode versions

diff --git a/Engine/Source/Programs/UnrealBuildTool/Platform/IOS/IOSToolChain.cs b/Engine/Source/Programs/UnrealBuildTool/Platform/IOS/IOSToolChain.cs
--- a/Engine/Source/Programs/UnrealBuildTool/Platform/IOS/IOSToolChain.cs
+++ b/Engine/Source/Programs/UnrealBuildTool/Platform/IOS/IOSToolChain.cs
@@ -217,11 +217,8 @@ namespace UnrealBuildTool
 		{
 			base.GetCompileArguments_Optimizations(CompileEnvironment, Arguments);
 
-			if (CompilerVersionGreaterOrEqual(12, 0, 0))
-			{
-				// We have 'this' vs nullptr comparisons that get optimized away for newer versions of Clang, which is undesirable until we refactor these checks.
-				Arguments.Add("-fno-delete-null-pointer-checks");
-			}
+			// We have 'this' vs nullptr comparisons that get optimized away for newer versions of Clang, which is undesirable until we refactor these checks.
+			Arguments.Add("-fno-delete-null-pointer-checks");
 
 			// use LTO if desired (like VCToolchain does)
 			if (CompileEnvironment.bAllowLTCG)
diff --git a/Engine/Source/Programs/UnrealBuildTool/ToolChain/AppleToolChain.cs b/Engine/Source/Programs/UnrealBuildTool/ToolChain/AppleToolChain.cs
--- a/Engine/Source/Programs/UnrealBuildTool/ToolChain/AppleToolChain.cs
+++ b/Engine/Source/Programs/UnrealBuildTool/ToolChain/AppleToolChain.cs
@@ -68,7 +68,7 @@ namespace UnrealBuildTool
 				Log.TraceInformationOnce("Compiling with non-standard Xcode ({0}): {1}", Reason, DeveloperDir);
 			}
 
-			// Installed engine requires Xcode 11
+			// Installed engine requires Xcode 13
 			if (Unreal.IsEngineInstalled())
 			{
 				string? InstalledSdkVersion = UnrealBuildBase.ApplePlatformSDK.InstalledSDKVersion;
@@ -76,9 +76,9 @@ namespace UnrealBuildTool
 				{
 					throw new BuildException("Unable to get xcode version");
 				}
-				if (int.Parse(InstalledSdkVersion.Substring(0,2)) < 11)
+				if (int.Parse(InstalledSdkVersion.Substring(0,2)) < 13)
 				{
-					throw new BuildException("Building for macOS, iOS and tvOS requires Xcode 11 or newer, Xcode " + InstalledSdkVersion + " detected");
+					throw new BuildException("Building for macOS, iOS and tvOS requires Xcode 13.4.1 or newer, Xcode " + InstalledSdkVersion + " detected");
 				}
 			}
 		}
@@ -290,11 +290,7 @@ namespace UnrealBuildTool
 			Arguments.Add("-x objective-c++-header");
 			GetCppStandardCompileArgument(CompileEnvironment, Arguments);
 			Arguments.Add("-stdlib=libc++");
-
-			if (CompilerVersionGreaterOrEqual(13, 0, 0)) // Note this is supported for >=11 on other clang platforms
-			{
-				Arguments.Add("-fpch-instantiate-templates");
-			}
+			Arguments.Add("-fpch-instantiate-templates");
 		}
 
 		/// <inheritdoc/>
@@ -302,20 +298,7 @@ namespace UnrealBuildTool
 		{
 			base.GetCompileArguments_WarningsAndErrors(CompileEnvironment, Arguments);
 
-			// Flags added in Xcode 13.3 (Apple clang 13.1.6)
-			if (CompilerVersionGreaterOrEqual(13, 0, 0) && CompilerVersionLessThan(13, 1, 6))
-			{
-				Arguments.Remove("-Wno-unused-but-set-variable");
-				Arguments.Remove("-Wno-unused-but-set-parameter");
-				Arguments.Remove("-Wno-ordered-compare-function-pointers");
-			}
-
-			// clang 12.00 has a new warning for copies in ranged loops. Instances have all been fixed up (2020/6/26) but
-			// are likely to be reintroduced due to no equivalent on other platforms at this time so disable the warning
-			if (CompilerVersionGreaterOrEqual(12, 0, 0))
-			{
-				Arguments.Add("-Wno-range-loop-analysis");
-			}
+			Arguments.Add("-Wno-range-loop-analysis");
 		}
 
 		/// <inheritdoc/>
@@ -435,14 +418,6 @@ namespace UnrealBuildTool
 					{
 						Log.TraceInformationOnce("Unable to parse version tokens: {0}", Tokens[3]);
 					}
-					else
-					{
-						if (Major < 12)
-						{
-							Log.TraceInformationOnce("dsymutil version is {0}.{1}.{2}. Using bundled version.", Major, Minor, Patch);
-							bUseInstalledDsymutil = false;
-						}
-					}
 				}
 			}
 
