Fixed an issue where the script would sometimes add a newline at the end of the file.

diff --git a/Engine/Source/Programs/UnrealBuildTool/Modes/InlineGeneratedCpps.cs b/Engine/Source/Programs/UnrealBuildTool/Modes/InlineGeneratedCpps.cs
--- a/Engine/Source/Programs/UnrealBuildTool/Modes/InlineGeneratedCpps.cs
+++ b/Engine/Source/Programs/UnrealBuildTool/Modes/InlineGeneratedCpps.cs
@@ -143,13 +143,12 @@ namespace UnrealBuildTool
 										}
 									}
 
-									var SourceText = FileReference.ReadAllText(InputFile.Location);
-									if (InlineReflectionMarkupRegex.IsMatch(SourceText))
+									var TextLines = System.IO.File.ReadAllLines(InputFile.Location.FullName).ToList();
+									if (TextLines.Any(Text => InlineReflectionMarkupRegex.IsMatch(Text)))
 									{
 										continue;
 									}
 
-									var TextLines = SourceText.Split(Environment.NewLine).ToList();
 									var CleanTextLines = new List<string>(TextLines.Count);
 
 									bool bInComment = false;
