PR #9690: Use vswhere -find to identify if AGDE is installed (Contributed by [REDACTED_PERSON])

#android

diff --git a/Engine/Source/Programs/UnrealBuildTool/Platform/Android/AndroidProjectGenerator.cs b/Engine/Source/Programs/UnrealBuildTool/Platform/Android/AndroidProjectGenerator.cs
--- a/Engine/Source/Programs/UnrealBuildTool/Platform/Android/AndroidProjectGenerator.cs
+++ b/Engine/Source/Programs/UnrealBuildTool/Platform/Android/AndroidProjectGenerator.cs
@@ -2,6 +2,7 @@
 
 using System;
 using System.Collections.Generic;
+using System.Diagnostics;
 using System.Text;
 using System.IO;
 using EpicGames.Core;
@@ -24,12 +25,41 @@ namespace UnrealBuildTool
 		public AndroidProjectGenerator(CommandLineArguments Arguments, ILogger Logger)
 			: base(Arguments, Logger)
 		{
-			if (OperatingSystem.IsWindows())
+			AGDEInstalled = false;
+			if (OperatingSystem.IsWindows() && !Arguments.HasOption("-noagde"))
 			{
 				AGDEInstalled = Microsoft.Win32.Registry.LocalMachine.OpenSubKey(@"SOFTWARE\WOW6432Node\Google\AndroidGameDevelopmentExtension")?.ValueCount > 0;
-				if (Arguments.HasOption("-noagde"))
+
+				if (!AGDEInstalled)
 				{
-					AGDEInstalled = false;
+					try
+					{
+						string? programFiles86 = Environment.GetEnvironmentVariable("ProgramFiles(x86)");
+						if (programFiles86 != null)
+						{
+							string vswhereExe = Path.Join(programFiles86, @"Microsoft Visual Studio\Installer\vswhere.exe");
+							if (File.Exists(vswhereExe))
+							{
+								using (Process p = new Process())
+								{
+									ProcessStartInfo info = new ProcessStartInfo
+									{
+										FileName = vswhereExe,
+										Arguments = @"-find Common7\IDE\Extensions\*\Google.VisualStudio.Android.dll",
+										RedirectStandardOutput = true,
+										UseShellExecute = false
+									};
+									p.StartInfo = info;
+									p.Start();
+									AGDEInstalled = p.StandardOutput.ReadToEnd().Contains("Google.VisualStudio.Android.dll");
+								}
+							}
+						}
+					}
+					catch (Exception ex)
+					{
+						Logger.LogInformation("Failed to identify AGDE installation status: {Message}", ex.Message);
+					}
 				}
 			}
 		}
