Set iOS default architecture that was lost after resolving. This is breaking builds for iOS
check https://p4-swarm.epicgames.net/reviews/23193114/v0,3 vs https://p4-swarm.epicgames.net/reviews/23193114/v0,4

diff --git a/Engine/Source/Programs/UnrealBuildTool/Platform/IOS/UEBuildIOS.cs b/Engine/Source/Programs/UnrealBuildTool/Platform/IOS/UEBuildIOS.cs
--- a/Engine/Source/Programs/UnrealBuildTool/Platform/IOS/UEBuildIOS.cs
+++ b/Engine/Source/Programs/UnrealBuildTool/Platform/IOS/UEBuildIOS.cs
@@ -735,7 +735,7 @@ namespace UnrealBuildTool
 
 		public override string GetDefaultArchitecture(FileReference? ProjectFile)
 		{
-			return "";
+			return "arm64";
 		}
 		
 		public override List<FileReference> FinalizeBinaryPaths(FileReference BinaryName, FileReference? ProjectFile, ReadOnlyTargetRules Target)
