Added ability to force the usage of C++ UHT even if the warning has been enabled.

#rnx

diff --git a/Engine/Source/Programs/UnrealBuildTool/Configuration/BuildConfiguration.cs b/Engine/Source/Programs/UnrealBuildTool/Configuration/BuildConfiguration.cs
--- a/Engine/Source/Programs/UnrealBuildTool/Configuration/BuildConfiguration.cs
+++ b/Engine/Source/Programs/UnrealBuildTool/Configuration/BuildConfiguration.cs
@@ -156,6 +156,12 @@ namespace UnrealBuildTool
 		[XmlConfigFile(Category = "UEBuildConfiguration")]
 		public bool bWarnOnCppUnrealHeaderTool = false;
 
+		/// <summary>
+		/// If true, force C++ version UHT to run as just a warning
+		/// </summary>
+		[CommandLine("-UseCppUHT")]
+		public bool bUseCppUHT = false;
+
 		/// <summary>
 		/// Whether to skip compiling rules assemblies and just assume they are valid
 		/// </summary>
diff --git a/Engine/Source/Programs/UnrealBuildTool/System/ExternalExecution.cs b/Engine/Source/Programs/UnrealBuildTool/System/ExternalExecution.cs
--- a/Engine/Source/Programs/UnrealBuildTool/System/ExternalExecution.cs
+++ b/Engine/Source/Programs/UnrealBuildTool/System/ExternalExecution.cs
@@ -979,33 +979,46 @@ namespace UnrealBuildTool
 		/// </summary>
 		public static void ExecuteHeaderToolIfNecessary(BuildConfiguration BuildConfiguration, FileReference? ProjectFile, TargetMakefile Makefile, string TargetName, ISourceFileWorkingSet WorkingSet, ILogger Logger)
 		{
-			if (!BuildConfiguration.bUseBuiltInUnrealHeaderTool)
+			bool bIsBuildingUHT = TargetName.Equals("UnrealHeaderTool", StringComparison.InvariantCultureIgnoreCase);
+
+			if (bIsBuildingUHT)
 			{
-				if (BuildConfiguration.bWarnOnCppUnrealHeaderTool)
-				{
-					Logger.LogWarning("WARNING: C++ UHT being used because 'bUseBuiltInUnrealHeaderTool' is false.  C++ UHT has been deprecated and will be removed in 5.2");
-				}
-				else
-				{
-					Logger.LogInformation("DEPRECATED: C++ UHT being used because 'bUseBuiltInUnrealHeaderTool' is false.  C++ UHT has been deprecated and will be removed in 5.2");
-				}
+				Logger.LogInformation("DEPRECATED: C++ UHT is being built by request.  C++ UHT has been deprecated and will be removed in 5.2");
+			}
+			else if (BuildConfiguration.bUseCppUHT)
+			{
+				Logger.LogInformation("DEPRECATED: C++ UHT being used because '-UseCppUHT'has been specified.  C++ UHT has been deprecated and will be removed in 5.2");
 			}
-			if (Makefile.RequiredUhtCppPlugins?.Length > 0)
+			else
 			{
-				foreach (string PluginName in Makefile.RequiredUhtCppPlugins)
+				if (!BuildConfiguration.bUseBuiltInUnrealHeaderTool)
 				{
 					if (BuildConfiguration.bWarnOnCppUnrealHeaderTool)
 					{
-						Logger.LogWarning("WARNING: C++ UHT being used because '{PluginName}' does not have a C# version.  C++ UHT has been deprecated and will be removed in 5.2", PluginName);
+						Logger.LogWarning("WARNING: C++ UHT being used because 'bUseBuiltInUnrealHeaderTool' is false.  C++ UHT has been deprecated and will be removed in 5.2");
 					}
 					else
 					{
-						Logger.LogInformation("DEPRECATED: C++ UHT being used because '{PluginName}' does not have a C# version.  C++ UHT has been deprecated and will be removed in 5.2", PluginName);
+						Logger.LogInformation("DEPRECATED: C++ UHT being used because 'bUseBuiltInUnrealHeaderTool' is false.  C++ UHT has been deprecated and will be removed in 5.2");
+					}
+				}
+				if (Makefile.RequiredUhtCppPlugins?.Length > 0)
+				{
+					foreach (string PluginName in Makefile.RequiredUhtCppPlugins)
+					{
+						if (BuildConfiguration.bWarnOnCppUnrealHeaderTool)
+						{
+							Logger.LogWarning("WARNING: C++ UHT being used because '{PluginName}' does not have a C# version.  C++ UHT has been deprecated and will be removed in 5.2", PluginName);
+						}
+						else
+						{
+							Logger.LogInformation("DEPRECATED: C++ UHT being used because '{PluginName}' does not have a C# version.  C++ UHT has been deprecated and will be removed in 5.2", PluginName);
+						}
 					}
 				}
 			}
 
-			if (!BuildConfiguration.bUseBuiltInUnrealHeaderTool || Makefile.RequiredUhtCppPlugins?.Length > 0 || TargetName.Equals("UnrealHeaderTool", StringComparison.InvariantCultureIgnoreCase))
+			if (!BuildConfiguration.bUseBuiltInUnrealHeaderTool || BuildConfiguration.bUseCppUHT || Makefile.RequiredUhtCppPlugins?.Length > 0 || bIsBuildingUHT)
 			{
 				ExecuteExternalHeaderToolIfNecessary(BuildConfiguration, ProjectFile, Makefile, TargetName, WorkingSet, Logger);
 			}
