UBT: When running QueryTargets commandlet, only write out default targets, if one is specified for a given target type.

diff --git a/Engine/Source/Programs/UnrealBuildTool/Modes/QueryTargetsMode.cs b/Engine/Source/Programs/UnrealBuildTool/Modes/QueryTargetsMode.cs
--- a/Engine/Source/Programs/UnrealBuildTool/Modes/QueryTargetsMode.cs
+++ b/Engine/Source/Programs/UnrealBuildTool/Modes/QueryTargetsMode.cs
@@ -29,6 +29,12 @@ namespace UnrealBuildTool
 		[CommandLine("-Output=")]
 		FileReference? OutputFile = null;
 
+		/// <summary>
+		/// Write out all targets, even if a default is specified in the BuildSettings section of the Default*.ini files. 
+		/// </summary>
+		[CommandLine("-IncludeAllTargets")]
+		bool bIncludeAllTargets = false;
+		
 		/// <summary>
 		/// Execute the mode
 		/// </summary>
@@ -62,7 +68,7 @@ namespace UnrealBuildTool
 			}
 
 			// Write information about these targets
-			WriteTargetInfo(ProjectFile, Assembly, OutputFile, Arguments, Logger);
+			WriteTargetInfo(ProjectFile, Assembly, OutputFile, Arguments, Logger, bIncludeAllTargets);
 			Logger.LogInformation("Written {OutputFile}", OutputFile);
 			return 0;
 		}
@@ -92,7 +98,8 @@ namespace UnrealBuildTool
 		/// <param name="OutputFile">Output file to write to</param>
 		/// <param name="Arguments"></param>
 		/// <param name="Logger">Logger for output</param>
-		public static void WriteTargetInfo(FileReference? ProjectFile, RulesAssembly Assembly, FileReference OutputFile, CommandLineArguments Arguments, ILogger Logger)
+		/// <param name="bIncludeAllTargets">Include all targets even if a default target is specified for a given target type.</param>
+		public static void WriteTargetInfo(FileReference? ProjectFile, RulesAssembly Assembly, FileReference OutputFile, CommandLineArguments Arguments, ILogger Logger, bool bIncludeAllTargets = true)
 		{
 			// Construct all the targets in this assembly
 			List<string> TargetNames = new List<string>();
@@ -126,7 +133,17 @@ namespace UnrealBuildTool
 						Logger.LogDebug("{Ex}", ExceptionUtils.FormatException(Ex));
 						continue;
 					}
-
+					
+					// Skip non-default targets if one is specified.
+					if (ProjectFile != null && !bIncludeAllTargets)
+					{
+						string? DefaultTargetName = ProjectFileGenerator.GetProjectDefaultTargetNameForType(ProjectFile.Directory, TargetRules.Type);
+						if (DefaultTargetName != null && DefaultTargetName != TargetName)
+						{
+							continue;
+						}
+					}
+					
 					// Write the target info
 					Writer.WriteObjectStart();
 					Writer.WriteValue("Name", TargetName);
diff --git a/Engine/Source/Programs/UnrealBuildTool/ProjectFiles/ProjectFileGenerator.cs b/Engine/Source/Programs/UnrealBuildTool/ProjectFiles/ProjectFileGenerator.cs
--- a/Engine/Source/Programs/UnrealBuildTool/ProjectFiles/ProjectFileGenerator.cs
+++ b/Engine/Source/Programs/UnrealBuildTool/ProjectFiles/ProjectFileGenerator.cs
@@ -2437,21 +2437,7 @@ namespace UnrealBuildTool
 					{
 						if (GameFolder != null)
 						{
-							FileReference DefaultEngineIniFilename = FileReference.Combine(GameFolder, "Config", "DefaultEngine.ini");
-							ConfigFile? ProjectDefaultEngineIni;
-							ConfigCache.TryReadFile(DefaultEngineIniFilename, out ProjectDefaultEngineIni);
-							if (ProjectDefaultEngineIni != null)
-							{
-								ConfigFileSection? Section;
-								if (ProjectDefaultEngineIni.TryGetSection("/Script/BuildSettings.BuildSettings", out Section))
-								{
-									ConfigLine? Line;
-									if (Section.TryGetLine("DefaultEditorTarget", out Line))
-									{
-										DefaultEditorTarget = Line.Value;
-									}
-								}
-							}
+							DefaultEditorTarget = GetProjectDefaultTargetNameForType(GameFolder, TargetType.Editor);
 						}
 
 						DefaultProjectEditorTargetCache.Add(ProjectFileNameBase, DefaultEditorTarget);
@@ -3063,6 +3049,27 @@ namespace UnrealBuildTool
 			OtherProjectFiles.Add(InProject);
 		}
 
+		/// <summary>
+		/// Retrieves the user-set default target name for a given target type. This is used for cases where
+		/// a target type can have multiple target names but only one is used for default builds.
+		/// </summary>
+		/// <param name="ProjectDirectory">The project directory to find the configuration files in</param>
+		/// <param name="TargetType">The target type to find a default value for.</param>
+		/// <returns>The name of the default target, or null</returns>
+		public static string? GetProjectDefaultTargetNameForType(DirectoryReference ProjectDirectory, TargetType TargetType)
+		{
+			// For now, we only support editor targets.
+			if (TargetType != TargetType.Editor)
+			{
+				return null;
+			}
+			
+			ConfigHierarchy Ini = ConfigCache.ReadHierarchy(ConfigHierarchyType.Engine, ProjectDirectory, BuildHostPlatform.Current.Platform);
+			string? Value;
+			Ini.TryGetValue("/Script/BuildSettings.BuildSettings", "DefaultEditorTarget", out Value);
+			return Value;
+		}
+
 		public virtual string[] GetTargetArguments(string[] Arguments)
 		{
 			// by default we do not forward any arguments to the targets
