Remove static lib version check on Mac

diff --git a/Engine/Source/Programs/UnrealBuildTool/Platform/IOS/IOSPlatformSDK.Versions.cs b/Engine/Source/Programs/UnrealBuildTool/Platform/IOS/IOSPlatformSDK.Versions.cs
--- a/Engine/Source/Programs/UnrealBuildTool/Platform/IOS/IOSPlatformSDK.Versions.cs
+++ b/Engine/Source/Programs/UnrealBuildTool/Platform/IOS/IOSPlatformSDK.Versions.cs
@@ -22,7 +22,7 @@ namespace UnrealBuildTool
 		/// <summary>
 		/// The minimum Clang version that a static library can be built with
 		/// </summary>
-		public override Version MinimumStaticLibClangVersion
+		public Version MinimumStaticLibClangVersion
 		{
 			get
 			{
diff --git a/Engine/Source/Programs/UnrealBuildTool/Platform/Mac/ApplePlatformSDK.Versions.cs b/Engine/Source/Programs/UnrealBuildTool/Platform/Mac/ApplePlatformSDK.Versions.cs
--- a/Engine/Source/Programs/UnrealBuildTool/Platform/Mac/ApplePlatformSDK.Versions.cs
+++ b/Engine/Source/Programs/UnrealBuildTool/Platform/Mac/ApplePlatformSDK.Versions.cs
@@ -37,19 +37,6 @@ namespace UnrealBuildTool
 			MinVersion = MaxVersion = null;
 		}
 
-		/// <summary>
-		/// The minimum Clang version that a macOS static library can be built with
-		/// </summary>
-		public virtual Version MinimumStaticLibClangVersion
-		{
-			get
-			{
-				// return new Version("1200.00.32.29"); // UE 5.0 (Xcode 12.4)
-				// return new Version("1300.00.29.30"); // UE 5.1 (Xcode 13.2.1)
-				return new Version("1400.00.28.1"); // UE 5.1 (Xcode 14.0)
-			}
-		}
-
 		/// <summary>
 		/// The minimum macOS SDK version that a dynamic library can be built with
 		/// </summary>
diff --git a/Engine/Source/Programs/UnrealBuildTool/Platform/Mac/UEBuildMac.cs b/Engine/Source/Programs/UnrealBuildTool/Platform/Mac/UEBuildMac.cs
--- a/Engine/Source/Programs/UnrealBuildTool/Platform/Mac/UEBuildMac.cs
+++ b/Engine/Source/Programs/UnrealBuildTool/Platform/Mac/UEBuildMac.cs
@@ -205,41 +205,6 @@ namespace UnrealBuildTool
 						continue;
 					}
 					ValidatedLibs.Add(LibLoc);
-
-					switch (LibLoc.GetExtension())
-					{
-						case ".a":
-							{
-								// For a static lib, grep it
-								string Args = "-c \"strings ";
-								Args += LibLoc.FullName;
-								Args += " | grep -m1 -i \\(clang\"";
-								string StdOutResult = Utils.RunLocalProcessAndReturnStdOut("bash", Args);
-								if (string.IsNullOrEmpty(StdOutResult))
-								{
-									continue;
-								}
-
-								// This Regex will extract a 2-4 segment version code from string containing a 2-5 segment code 
-								// ie: if given string: "Apple clang version 14.0.0 (clang-1400.0.17.3.1)"
-								//     it'll extract: "1400.0.17.3"  (note the dropped 5th segment)
-								Match M = Regex.Match(StdOutResult, @"(\(clang-(?<ver>\d+.\d+(.(\d+))?(.(\d+))?)(.(\d+))?\))");
-								if (M.Success)
-								{
-									string LibString = M.Groups["ver"].ToString();
-									Version? LibVersion = new Version(LibString);
-									if (LibVersion != null && LibVersion > SDK.MinimumStaticLibClangVersion)
-									{
-										throw new BuildException("macOS Static Library:'{0}' is built with a version of clang newer than UE supports ({1} > {2}). \nPlease rebuild {3} with the minimum supported version of Xcode/clang.", LibLoc, LibString, SDK.MinimumStaticLibClangVersion, LibLoc);
-									}
-								}
-							}
-							break;
-
-						default:
-							// For now, we don't validate any other types of libs (dylib, Framework, etc)
-							break;
-					}
 				}
 			}
 		}
