Enabled the warning on C++ UHT being used

#rnx

diff --git a/Engine/Source/Programs/UnrealBuildTool/Configuration/BuildConfiguration.cs b/Engine/Source/Programs/UnrealBuildTool/Configuration/BuildConfiguration.cs
--- a/Engine/Source/Programs/UnrealBuildTool/Configuration/BuildConfiguration.cs
+++ b/Engine/Source/Programs/UnrealBuildTool/Configuration/BuildConfiguration.cs
@@ -154,7 +154,7 @@ namespace UnrealBuildTool
 		/// If true, generate warnings when C++ UHT is used
 		/// </summary>
 		[XmlConfigFile(Category = "UEBuildConfiguration")]
-		public bool bWarnOnCppUnrealHeaderTool = false;
+		public bool bWarnOnCppUnrealHeaderTool = true;
 
 		/// <summary>
 		/// If true, force C++ version UHT to run as just a warning
