- Fixed PostBuild not being called on Mac compiles

diff --git a/Engine/Source/Programs/UnrealBuildTool/Platform/Mac/MacToolChain.cs b/Engine/Source/Programs/UnrealBuildTool/Platform/Mac/MacToolChain.cs
--- a/Engine/Source/Programs/UnrealBuildTool/Platform/Mac/MacToolChain.cs
+++ b/Engine/Source/Programs/UnrealBuildTool/Platform/Mac/MacToolChain.cs
@@ -1026,9 +1026,8 @@ namespace UnrealBuildTool
 		/// </summary>
 		/// <param name="LinkEnvironment"></param>
 		/// <param name="Executable">FileItem describing the executable to generate app bundle for</param>
-		/// <param name="FixDylibOutputFile"></param>
 		/// <param name="Graph">List of actions to be executed. Additional actions will be added to this list.</param>
-		FileItem FinalizeAppBundle(LinkEnvironment LinkEnvironment, FileItem Executable, FileItem FixDylibOutputFile, IActionGraphBuilder Graph)
+		FileItem FinalizeAppBundle(LinkEnvironment LinkEnvironment, FileItem Executable, IActionGraphBuilder Graph)
 		{
 			// Make a file item for the source and destination files
 			string FullDestPath = Executable.AbsolutePath.Substring(0, Executable.AbsolutePath.IndexOf(".app") + 4);
@@ -1044,7 +1043,11 @@ namespace UnrealBuildTool
 			FileItem BundleScript = FileItem.GetItemByFileReference(FileReference.Combine(LinkEnvironment.IntermediateDirectory!, "FinalizeAppBundle.sh"));
 
 			FinalizeAppBundleAction.CommandArguments = "\"" + BundleScript.AbsolutePath + "\"";
-			FinalizeAppBundleAction.PrerequisiteItems.Add(FixDylibOutputFile);
+			FinalizeAppBundleAction.PrerequisiteItems.Add(Executable);
+			foreach (FileItem Dependency in BundleDependencies)
+			{
+				FinalizeAppBundleAction.PrerequisiteItems.Add(Dependency);
+			}
 			FinalizeAppBundleAction.ProducedItems.Add(DestFile);
 			FinalizeAppBundleAction.StatusDescription = string.Format("Finalizing app bundle: {0}.app", Path.GetFileName(Executable.AbsolutePath));
 			FinalizeAppBundleAction.bCanExecuteRemotely = false;
@@ -1203,9 +1206,9 @@ namespace UnrealBuildTool
 			}
 		}
 
-		public override ICollection<FileItem> PostBuild(FileItem[] Executables, LinkEnvironment BinaryLinkEnvironment, IActionGraphBuilder Graph)
+		public override ICollection<FileItem> PostBuild(FileItem Executable, LinkEnvironment BinaryLinkEnvironment, IActionGraphBuilder Graph)
 		{
-			ICollection<FileItem> OutputFiles = base.PostBuild(Executables, BinaryLinkEnvironment, Graph);
+			ICollection<FileItem> OutputFiles = base.PostBuild(Executable, BinaryLinkEnvironment, Graph);
 
 			if (BinaryLinkEnvironment.bIsBuildingLibrary)
 			{
@@ -1216,23 +1219,17 @@ namespace UnrealBuildTool
 			{
 				foreach (UEBuildBundleResource Resource in BinaryLinkEnvironment.AdditionalBundleResources)
 				{
-					foreach (FileItem Executable in Executables)
-					{
-						OutputFiles.Add(CopyBundleResource(Resource, Executable, BinaryLinkEnvironment.BundleDirectory, Graph));
-					}
+					OutputFiles.Add(CopyBundleResource(Resource, Executable, BinaryLinkEnvironment.BundleDirectory, Graph));
 				}
 			}
 
 			// For Mac, generate the dSYM file if the config file is set to do so
 			if (BinaryLinkEnvironment.bUsePDBFiles == true)
 			{
-				foreach (FileItem Executable in Executables)
-				{
-					DebugInfoFiles.Add(GenerateDebugInfo(Executable, BinaryLinkEnvironment, Graph, Logger));
-				}
+				DebugInfoFiles.Add(GenerateDebugInfo(Executable, BinaryLinkEnvironment, Graph, Logger));
 			}
 
-			if ((BinaryLinkEnvironment.bIsBuildingDLL && (Options & ClangToolChainOptions.OutputDylib) == 0) || (BinaryLinkEnvironment.bIsBuildingConsoleApplication && Executables[0].Name.EndsWith("-Cmd")))
+			if ((BinaryLinkEnvironment.bIsBuildingDLL && (Options & ClangToolChainOptions.OutputDylib) == 0) || (BinaryLinkEnvironment.bIsBuildingConsoleApplication && Executable.Name.EndsWith("-Cmd")))
 			{
 				return OutputFiles;
 			}
@@ -1246,7 +1243,7 @@ namespace UnrealBuildTool
 				}
 				else
 				{
-					OutputFiles.Add(FinalizeAppBundle(BinaryLinkEnvironment, Executables[0], FileItem.GetItemByFileReference(BinaryLinkEnvironment.OutputFilePath), Graph));
+					OutputFiles.Add(FinalizeAppBundle(BinaryLinkEnvironment, Executable, Graph));
 				}
 			}
 
