Fixed HWASan build on Android

diff --git a/Engine/Source/Programs/UnrealBuildTool/Platform/Android/UEDeployAndroid.cs b/Engine/Source/Programs/UnrealBuildTool/Platform/Android/UEDeployAndroid.cs
--- a/Engine/Source/Programs/UnrealBuildTool/Platform/Android/UEDeployAndroid.cs
+++ b/Engine/Source/Programs/UnrealBuildTool/Platform/Android/UEDeployAndroid.cs
@@ -1574,6 +1574,9 @@ namespace UnrealBuildTool
 				string WrapDestFilePath = Path.Combine(WrapDestDir, "wrap.sh");
 				Logger.LogInformation("Copying wrap.sh from {WrapSh} to {WrapDestFilePath}", WrapSh, WrapDestFilePath);
 				File.Copy(WrapSh, WrapDestFilePath, true);
+				FileAttributes Attributes = File.GetAttributes(WrapDestFilePath);
+				Attributes &= ~FileAttributes.ReadOnly;
+				File.SetAttributes(WrapDestFilePath, Attributes);
 			}
 			else
 			{
@@ -4756,7 +4759,7 @@ namespace UnrealBuildTool
 				CopyGfxDebugger(UnrealBuildPath, Arch, NDKArch);
 				CopyVulkanValidationLayers(UnrealBuildPath, Arch, NDKArch, Configuration.ToString());
 				
-				if (Sanitizer != AndroidToolChain.ClangSanitizer.None && Sanitizer != AndroidToolChain.ClangSanitizer.HwAddress)
+				if (Sanitizer != AndroidToolChain.ClangSanitizer.None)
 				{
 					CopyClangSanitizerLib(UnrealBuildPath, Arch, NDKArch, Sanitizer);
 				}
