Gauntlet PGO Profiling support on Windows platforms.

 - also includes -PGOFastGen to enable /FASTGENPROFILE for a better framerate during PGO Profiling.
 - also favoring the PGC file that matches the executable when building PGO Optimize builds rather than failing.

diff --git a/Engine/Source/Programs/AutomationTool/Gauntlet/Platform/Windows/Gauntlet.WindowsPGO.cs b/Engine/Source/Programs/AutomationTool/Gauntlet/Platform/Windows/Gauntlet.WindowsPGO.cs
--- a/Engine/Source/Programs/AutomationTool/Gauntlet/Platform/Windows/Gauntlet.WindowsPGO.cs
+++ b/Engine/Source/Programs/AutomationTool/Gauntlet/Platform/Windows/Gauntlet.WindowsPGO.cs
@@ -1,5 +1,8 @@
 // Copyright Epic Games, Inc. All Rights Reserved.
 
+using AutomationTool;
+using System.IO;
+using System.Text.RegularExpressions;
 using UnrealBuildTool;
 
 namespace Gauntlet
@@ -9,6 +12,8 @@ namespace Gauntlet
 	/// </summary>
 	internal abstract class WinBasePGOPlatform : IPGOPlatform
 	{
+		string LocalOutputDirectory;
+		string PgcFilenamePrefix;
 		UnrealTargetPlatform Platform;
 
 		protected WinBasePGOPlatform(UnrealTargetPlatform InPlatform)
@@ -23,12 +28,66 @@ namespace Gauntlet
 
 		public void GatherResults(string ArtifactPath)
 		{
+			// Find PGC files
+			string PGODir = Path.Combine(ArtifactPath, "PGO");
+			string[] PGCFiles = Directory.GetFiles(PGODir, "*.pgc");
+			if (PGCFiles.Length == 0)
+			{
+				throw new AutomationException("no .pgc files were found in {0}", PGODir );
+			}
+			Log.Info("Found {0} .pgc files in \"{1}\"", PGCFiles.Length, PGODir);
 
+			// Copy the PGC files to the output directory
+			foreach (string SrcFilePath in PGCFiles)
+			{
+				string DstFileName = Path.GetFileName(SrcFilePath);
+
+				// Optionally override the pgc filename...
+				if (!string.IsNullOrWhiteSpace(PgcFilenamePrefix))
+				{
+					int FileIndex;
+					Match ParsedName = Regex.Match(DstFileName, @".*!([0-9]+)\.pgc");
+					if (ParsedName.Success && int.TryParse(ParsedName.Groups[1].Value, out FileIndex))
+					{
+						string NewDstFileName = string.Format("{0}!{1}.pgc", PgcFilenamePrefix, FileIndex);
+						if (NewDstFileName != DstFileName)
+						{
+							Log.Info("Overriding .pgc filename \"{0}\" to \"{1}\"", DstFileName, NewDstFileName);
+							DstFileName = NewDstFileName;
+						}
+					}
+				}
+
+				string DstFilePath = Path.Combine(LocalOutputDirectory, DstFileName);
+				if (File.Exists(DstFilePath))
+				{
+					InternalUtils.SafeDeleteFile(DstFilePath);
+				}
+				File.Copy(SrcFilePath, DstFilePath);
+			}
 		}
 
 		public void ApplyConfiguration(PGOConfig Config)
 		{
+			PgcFilenamePrefix = Config.PgcFilenamePrefix;
+			LocalOutputDirectory = Path.GetFullPath(Config.ProfileOutputDirectory);
+
+			UnrealTestRole ClientRole = Config.RequireRole(UnrealTargetRole.Client);
+			ClientRole.CommandLine += $" -PGOSweepToSaveDir"; // this means the PGC files will be created in the Gauntlet artifact folder
+			ClientRole.ConfigureDevice = ConfigureDevice;			
+		}
+
+		protected abstract string GetLocalCachePath(ITargetDevice Device);
 
+		protected void ConfigureDevice(ITargetDevice Device)
+		{
+			// make sure the PGO directory exists, and is empty
+			string PGODir = Path.Combine(GetLocalCachePath(Device), "UserDir", "Saved", "PGO");
+			if (Directory.Exists(PGODir))
+			{
+				Directory.Delete(PGODir, true);
+			}
+			Directory.CreateDirectory(PGODir);
 		}
 
 		public bool TakeScreenshot(ITargetDevice Device, string ScreenshotDirectory, out string ImageFilename)
@@ -43,6 +102,16 @@ namespace Gauntlet
 		public WindowsPGOPlatform() : base(UnrealTargetPlatform.Win64)
 		{
 		}
+
+		protected override string GetLocalCachePath(ITargetDevice Device)
+		{
+			using (TargetDeviceWindows WinDevice = Device as TargetDeviceWindows)
+			{
+				return WinDevice.LocalCachePath;
+			}
+
+			throw new AutomationException($"{Device} is a {Device.GetType()} - expecting TargetDeviceWindows");
+		}
 	}
 }
 
diff --git a/Engine/Source/Programs/UnrealBuildTool/Platform/Windows/UEBuildWindows.cs b/Engine/Source/Programs/UnrealBuildTool/Platform/Windows/UEBuildWindows.cs
--- a/Engine/Source/Programs/UnrealBuildTool/Platform/Windows/UEBuildWindows.cs
+++ b/Engine/Source/Programs/UnrealBuildTool/Platform/Windows/UEBuildWindows.cs
@@ -399,6 +399,14 @@ namespace UnrealBuildTool
 		/// </summary>
 		public bool bOptimizeGlobalData = true;
 
+		/// <summary>
+		/// If specified along with -PGOProfile, then /FASTGENPROFILE will be used instead of /GENPROFILE. 
+		/// This usually means that the PGO data is generated faster, but the resulting data may not yield as efficient optimizations during -PGOOptimize
+		/// </summary>
+		[XmlConfigFile(Category = "WindowsPlatform")]
+		[CommandLine("-PGOFastGen")]
+		public bool bUseFastGenProfile = false;
+
 		/// <summary>
 		/// (Experimental) Appends the -ftime-trace argument to the command line for Clang to output a JSON file containing a timeline for the compile. 
 		/// See http://aras-p.info/blog/2019/01/16/time-trace-timeline-flame-chart-profiler-for-Clang/ for more info.
@@ -727,6 +735,11 @@ namespace UnrealBuildTool
 			get { return Inner.bOptimizeGlobalData; }
 		}
 
+		public bool bUseFastGenProfile
+		{
+			get { return Inner.bUseFastGenProfile; }
+		}
+
 		public bool bClangTimeTrace
 		{
 			get { return Inner.bClangTimeTrace; }
diff --git a/Engine/Source/Programs/UnrealBuildTool/Platform/Windows/VCToolChain.cs b/Engine/Source/Programs/UnrealBuildTool/Platform/Windows/VCToolChain.cs
--- a/Engine/Source/Programs/UnrealBuildTool/Platform/Windows/VCToolChain.cs
+++ b/Engine/Source/Programs/UnrealBuildTool/Platform/Windows/VCToolChain.cs
@@ -1733,6 +1733,22 @@ namespace UnrealBuildTool
 					TargetFileToSourceFile[FileReference.Combine(ExeDir, ASanDebugRuntimeDLL)] = FileReference.Combine(ASanRuntimeDir, ASanDebugRuntimeDLL);
 				}
 			}
+
+			// copy PGO support binaries for Windows from $(VC_PGO_RunTime_Dir)
+			if (Target.bPGOProfile && Target.Platform.IsInGroup(UnrealPlatformGroup.Windows))
+			{
+				string[] PGOFiles = {
+					"pgort140.dll",
+					"pgosweep.exe",
+				};
+				foreach (string PGOFile in PGOFiles)
+				{			
+					FileReference SrcFile = FileReference.Combine(EnvVars.ToolChainDir, "bin", "Hostx64", "x64", PGOFile);
+					FileReference DstFile = FileReference.Combine(ExeDir, PGOFile);
+					TargetFileToSourceFile[DstFile] = SrcFile;
+					RuntimeDependencies.Add( new RuntimeDependency(DstFile, StagedFileType.NonUFS));
+				}
+			}
 		}
 
 		public virtual FileReference GetApplicationIcon(FileReference? ProjectFile)
@@ -2293,19 +2309,26 @@ namespace UnrealBuildTool
 				// Copy the files there and make them writable...
 				Logger.LogInformation("...copying the profile guided optimization files to output directory...");
 
-				string[] PGDFiles = Directory.GetFiles(LinkEnvironment.PGODirectory!, "*.pgd");
-				string[] PGCFiles = Directory.GetFiles(LinkEnvironment.PGODirectory!, "*.pgc");
-
-				if (PGDFiles.Length > 1)
-				{
-					throw new BuildException("More than one .pgd file found in \"{0}\".", LinkEnvironment.PGODirectory);
-				}
-				else if (PGDFiles.Length == 0)
+				// prefer a PGD file that matches the output file
+				string PGDFile = Path.Combine( LinkEnvironment.PGODirectory!, LinkEnvironment.PGOFilenamePrefix + ".pgd");
+				if (!File.Exists(PGDFile))
 				{
-					Logger.LogWarning("No .pgd files found in \"{PgoDir}\".", LinkEnvironment.PGODirectory);
-					return false;
+					string[] PGDFiles = Directory.GetFiles(LinkEnvironment.PGODirectory!, "*.pgd");
+					if (PGDFiles.Length > 1)
+					{
+						throw new BuildException("More than one .pgd file found in \"{0}\" and \"{1}\" not found ", LinkEnvironment.PGODirectory, PGDFile);
+					}
+					else if (PGDFiles.Length == 0)
+					{
+						Logger.LogWarning("No .pgd files found in \"{PgoDir}\".", LinkEnvironment.PGODirectory);
+						return false;
+					}
+
+					PGDFile = PGDFiles.First();
 				}
 
+
+				string[] PGCFiles = Directory.GetFiles(LinkEnvironment.PGODirectory!, "*.pgc");
 				if (PGCFiles.Length == 0)
 				{
 					Logger.LogWarning("No .pgc files found in \"{PgoDir}\".", LinkEnvironment.PGODirectory);
@@ -2316,7 +2339,6 @@ namespace UnrealBuildTool
 				Directory.CreateDirectory(LinkEnvironment.OutputDirectory!.FullName);
 
 				// Copy the .pgd to the linker output directory, renaming it to match the PGO filename prefix.
-				string PGDFile = PGDFiles.First();
 				string DestPGDFile = Path.Combine(LinkEnvironment.OutputDirectory.FullName, LinkEnvironment.PGOFilenamePrefix + ".pgd");
 				Logger.LogInformation("{Source} -> {Target}", PGDFile, DestPGDFile);
 				File.Copy(PGDFile, DestPGDFile, true);
@@ -2364,7 +2386,14 @@ namespace UnrealBuildTool
 				//Arguments.Add("/GENPROFILE:PGD=" + Path.Combine(LinkEnvironment.PGODirectory, LinkEnvironment.PGOFilenamePrefix + ".pgd"));
 				Arguments.Add("/LTCG");
 				//Arguments.Add("/GENPROFILE:PGD=" + LinkEnvironment.PGOFilenamePrefix + ".pgd");
-				Arguments.Add("/GENPROFILE");
+				if (Target.WindowsPlatform.bUseFastGenProfile)
+				{
+					Arguments.Add("/FASTGENPROFILE");
+				}
+				else
+				{
+					Arguments.Add("/GENPROFILE");
+				}
 				Log.TraceInformationOnce("Enabling Profile Guided Optimization (PGO). Linking will take a while.");
 			}
 		}
diff --git a/Engine/Source/Runtime/Core/Private/Windows/WindowsPlatformMisc.cpp b/Engine/Source/Runtime/Core/Private/Windows/WindowsPlatformMisc.cpp
--- a/Engine/Source/Runtime/Core/Private/Windows/WindowsPlatformMisc.cpp
+++ b/Engine/Source/Runtime/Core/Private/Windows/WindowsPlatformMisc.cpp
@@ -441,6 +441,52 @@ static void PureCallHandler()
 	}
 }
 
+#if ENABLE_PGO_PROFILE
+void PGO_WriteFile()
+{
+	// NB. Using pgosweep.exe means the PGC file will be writable as soon as the title exits & we can control where it is written.
+	// Not using PgoAutoSweep because a) it calls MessageBox() when it encounters an error which would break unattended automation,
+	// and b) it only takes a file name fragment not a full path, so it would be necessary to sweep, find the file and then move it where we want it.
+
+	static uint32 FileCounter = 0;
+
+	// Get the current running process's full path
+	TCHAR ExeFilePath[MAX_PATH + 1];
+	GetModuleFileNameW(NULL, ExeFilePath, MAX_PATH + 1);
+	FString ExeFileName = FPaths::GetCleanFilename(ExeFilePath);
+	FString ExeFolder = FPaths::GetPath(ExeFilePath);
+	FString ExeFileNameWithoutExtension = FPaths::GetBaseFilename(ExeFilePath);
+
+	// Get PGC output directory, defaulting to the exe location but can sweep to the project saved dir so Gauntlet can collect it
+	bool bSweepToSaveDir = FParse::Param(FCommandLine::Get(), TEXT("PGOSweepToSaveDir"));
+	FString OutputDirectory = bSweepToSaveDir  ?  FPaths::Combine(FPaths::ProjectSavedDir(), TEXT("PGO"))  :  ExeFolder;
+
+	// Find next unique PGC file name
+	FString OutputFilePath;
+	do
+	{
+		FString PGCFileName = FString::Printf(TEXT("%s!%d.pgc"), *ExeFileNameWithoutExtension, ++FileCounter);
+		OutputFilePath = FPaths::Combine(OutputDirectory, PGCFileName).Replace(TEXT("/"), TEXT("\\"));
+	} while (GetFileAttributesW(*OutputFilePath) != INVALID_FILE_ATTRIBUTES);
+
+	// Launch PGOSweep & wait for it to finish
+	FString PGOSweepPath = FPaths::Combine(ExeFolder, TEXT("pgosweep.exe"));
+	FString CommandLine = FString::Printf(TEXT("/pid:%d \"%s\" \"%s\""), ::GetCurrentProcessId(), ExeFilePath, *OutputFilePath);
+
+	const bool bLaunchDetached = true;
+	const bool bLaunchHidden = true;
+	const bool bLaunchReallyHidden = bLaunchHidden;
+	FProcHandle ProcHandle = FPlatformProcess::CreateProc(*PGOSweepPath, *CommandLine, bLaunchDetached, bLaunchHidden, bLaunchReallyHidden, nullptr, 0, nullptr, nullptr, nullptr);
+	FPlatformProcess::WaitForProc(ProcHandle);
+	int32 ExitCode = 0;
+	if (FPlatformProcess::GetProcReturnCode(ProcHandle, &ExitCode))
+	{
+		UE_LOG(LogWindows, Log, TEXT("pgosweep.exe exit code %d"), ExitCode);
+	}
+}
+#endif //ENABLE_PGO_PROFILE
+
+
 /*-----------------------------------------------------------------------------
 	SHA-1 functions.
 -----------------------------------------------------------------------------*/
@@ -958,6 +1004,13 @@ void FWindowsPlatformMisc::RequestExitWithStatus(bool Force, uint8 ReturnCode)
 {
 	UE_LOG(LogWindows, Log, TEXT("FPlatformMisc::RequestExitWithStatus(%i, %i)"), Force, ReturnCode);
 
+#if ENABLE_PGO_PROFILE
+	// save current PGO profiling data and terminate immediately
+	PGO_WriteFile();
+	TerminateProcess(GetCurrentProcess(), 0);
+	return;
+#endif
+
 	RequestEngineExit(TEXT("Win RequestExit"));
 	FCoreDelegates::ApplicationWillTerminateDelegate.Broadcast();
 
