- Added GetProjectArchitectures, which UBT will use to determine the architecture to build with when -architecture is not specified (UEBuildMac is now using it to control architectures for programs/projects)
- Changed the very high level loop over targets and architectures to move the Architecture as innermost loop - this allows each target to compile with different architectures in the case of mutli-target compiling
- Added [Editor]DefaultArchitecture to Mac ini's - the [Editor]TargetArchitecture property is now what the target _supports_, and Default is what it will be built with if no architecture is specified (usually will be "Host" to indicate that whatever the current host platform is will be what is compiled. "All" can be used to compile Fat binaries by default.
- Updated the LinkEnvironment constructor that takes another LinkEnvironment and an Architecture, to filter the input files for that architecture, and removed the similar code from MacToolchain
- Updates for Programs for Mac compilation, either:
  -  Adding an .ini and a .uproject files - programs need a .uproject for UBT to be able to find their Config directory in a standard way)
  - Marking the Target as not supporting Mac (some programs call Windows code directly)
  - Updating libs to be fat (intel+arm)

diff --git a/Engine/Config/BaseEngine.ini b/Engine/Config/BaseEngine.ini
--- a/Engine/Config/BaseEngine.ini
+++ b/Engine/Config/BaseEngine.ini
@@ -2852,6 +2852,9 @@ UseFastIntrinsics=False
 EnableMathOptimisations=True
 TargetArchitecture=MacTargetArchitectureUniversal
 EditorTargetArchitecture=MacTargetArchitectureUniversal
+EditorDefaultArchitecture=Host
+DefaultArchitecture=Host
+
 
 [HMDPluginPriority]
 ; Since SteamVR also works with the Oculus Rift, give priority to the native Oculus plugin before trying SteamVR.
diff --git a/Engine/Programs/UnrealLightmass/Config/Mac/MacEngine.ini b/Engine/Programs/UnrealLightmass/Config/Mac/MacEngine.ini
new file mode 100644
--- /dev/null
+++ b/Engine/Programs/UnrealLightmass/Config/Mac/MacEngine.ini
@@ -0,0 +1,3 @@
+[/Script/MacTargetPlatform.MacTargetSettings]
+; Arm64 currently fails to compile due to math intrinsic usage
+TargetArchitecture=MacTargetArchitectureIntel
diff --git a/Engine/Programs/UnrealLightmass/UnrealLightmass.uproject b/Engine/Programs/UnrealLightmass/UnrealLightmass.uproject
new file mode 100644
--- /dev/null
+++ b/Engine/Programs/UnrealLightmass/UnrealLightmass.uproject
@@ -0,0 +1,11 @@
+{
+	"FileVersion": 3,
+	"EngineAssociation": "",
+	"Category": "",
+	"Description": "",
+	"TargetPlatforms": [
+		"Windows",
+		"Mac",
+		"Linux"
+	]
+}
diff --git a/Engine/Programs/Unsync/Config/Mac/MacEngine.ini b/Engine/Programs/Unsync/Config/Mac/MacEngine.ini
new file mode 100644
--- /dev/null
+++ b/Engine/Programs/Unsync/Config/Mac/MacEngine.ini
@@ -0,0 +1,3 @@
+[/Script/MacTargetPlatform.MacTargetSettings]
+; Arm64 currently fails to compile due to math intrinsic usage
+TargetArchitecture=MacTargetArchitectureIntel
diff --git a/Engine/Programs/Unsync/Unsync.uproject b/Engine/Programs/Unsync/Unsync.uproject
new file mode 100644
--- /dev/null
+++ b/Engine/Programs/Unsync/Unsync.uproject
@@ -0,0 +1,11 @@
+{
+	"FileVersion": 3,
+	"EngineAssociation": "",
+	"Category": "",
+	"Description": "",
+	"TargetPlatforms": [
+		"Windows",
+		"Mac",
+		"Linux"
+	]
+}
diff --git a/Engine/Source/Programs/LiveCodingConsole/LiveCodingConsole.Target.cs b/Engine/Source/Programs/LiveCodingConsole/LiveCodingConsole.Target.cs
--- a/Engine/Source/Programs/LiveCodingConsole/LiveCodingConsole.Target.cs
+++ b/Engine/Source/Programs/LiveCodingConsole/LiveCodingConsole.Target.cs
@@ -3,6 +3,7 @@
 using UnrealBuildTool;
 using System.Collections.Generic;
 
+[SupportedPlatforms("Win64")]
 public class LiveCodingConsoleTarget : TargetRules
 {
 	public LiveCodingConsoleTarget(TargetInfo Target) : base(Target)
diff --git a/Engine/Source/Programs/SlateViewer/SlateViewer.Target.cs b/Engine/Source/Programs/SlateViewer/SlateViewer.Target.cs
--- a/Engine/Source/Programs/SlateViewer/SlateViewer.Target.cs
+++ b/Engine/Source/Programs/SlateViewer/SlateViewer.Target.cs
@@ -5,6 +5,7 @@ using System.Collections.Generic;
 
 [SupportedPlatforms(UnrealPlatformClass.Desktop)]
 [SupportedPlatforms("IOS")]
+[SupportedPlatforms("Android")]
 public class SlateViewerTarget : TargetRules
 {
 	public SlateViewerTarget(TargetInfo Target) : base(Target)
diff --git a/Engine/Source/Programs/SymsLibDump/SymsLibDump.Target.cs b/Engine/Source/Programs/SymsLibDump/SymsLibDump.Target.cs
--- a/Engine/Source/Programs/SymsLibDump/SymsLibDump.Target.cs
+++ b/Engine/Source/Programs/SymsLibDump/SymsLibDump.Target.cs
@@ -3,7 +3,7 @@
 using UnrealBuildTool;
 using System.Collections.Generic;
 
-[SupportedPlatforms(UnrealPlatformClass.Desktop)]
+[SupportedPlatforms("Win64", "Linux")]
 public class SymsLibDumpTarget : TargetRules
 {
 	public SymsLibDumpTarget(TargetInfo Target) : base(Target)
diff --git a/Engine/Source/Programs/UnrealBuildTool/Configuration/TargetDescriptor.cs b/Engine/Source/Programs/UnrealBuildTool/Configuration/TargetDescriptor.cs
--- a/Engine/Source/Programs/UnrealBuildTool/Configuration/TargetDescriptor.cs
+++ b/Engine/Source/Programs/UnrealBuildTool/Configuration/TargetDescriptor.cs
@@ -9,6 +9,7 @@ using System.Threading.Tasks;
 using EpicGames.Core;
 using UnrealBuildBase;
 using Microsoft.Extensions.Logging;
+using System.Runtime.InteropServices;
 
 namespace UnrealBuildTool
 {
@@ -362,64 +363,85 @@ namespace UnrealBuildTool
 					}
 				}
 
-				// Parse the architecture parameter, or get the default for the platform
-				List<string> Architectures = new List<string>(Arguments.GetValues("-Architecture=", '+'));
-
 				UEBuildPlatform BuildPlatform = UEBuildPlatform.GetBuildPlatform(Platform);
-				if (Architectures.Count == 0)
-				{
-					Architectures.Add(BuildPlatform.GetDefaultArchitecture(ProjectFile));
-				}
-				else
-				{
-					// If the platform can do these in a single pass then turn them back into a single + separate string.
-					// It is now responsible for splitting them as necessary.
-					if (BuildPlatform.CanCompileArchitecturesInSinglePass(Architectures) || BuildPlatform.CanLinkArchitecturesInSinglePass(Architectures))
-					{ 
-						Architectures = new List<string> { string.Join("+", Architectures.OrderBy(S => S)) };
-					}
-				}				
 
-				foreach(string Architecture in Architectures)
+				// Parse the architecture parameter, or get the default for the platform (and letting the platform clean it up
+				IEnumerable<string> ParamArchitectures = Arguments.GetValues("-Architecture=", '+')
+					.Union(Arguments.GetValues("-Architectures=", '+'))
+					.Select(x => BuildPlatform.FixupArchitecture(x));
+
+				foreach(UnrealTargetConfiguration Configuration in Configurations)
 				{
-					foreach(UnrealTargetConfiguration Configuration in Configurations)
+					// Create all the target descriptors for targets specified by type
+					foreach(string TargetTypeString in Arguments.GetValues("-TargetType="))
 					{
-						// Create all the target descriptors for targets specified by type
-						foreach(string TargetTypeString in Arguments.GetValues("-TargetType="))
+						TargetType TargetType;
+						if(!Enum.TryParse(TargetTypeString, out TargetType))
 						{
-							TargetType TargetType;
-							if(!Enum.TryParse(TargetTypeString, out TargetType))
-							{
-								throw new BuildException("Invalid target type '{0}'", TargetTypeString);
-							}
+							throw new BuildException("Invalid target type '{0}'", TargetTypeString);
+						}
+
+						if (ProjectFile == null)
+						{
+							TargetNames.Add(RulesCompiler.CreateEngineRulesAssembly(bUsePrecompiled, bSkipRulesCompile, bForceRulesCompile, Logger).GetTargetNameByType(TargetType, Platform, Configuration, "", null, Logger));
+						}
+						else
+						{
+							TargetNames.Add(RulesCompiler.CreateProjectRulesAssembly(ProjectFile, bUsePrecompiled, bSkipRulesCompile, bForceRulesCompile, Logger).GetTargetNameByType(TargetType, Platform, Configuration, "", ProjectFile, Logger));
+						}
+					}
+
+					// Make sure we could parse something
+					if (TargetNames.Count == 0)
+					{
+						throw new BuildException("No target name was specified on the command-line.");
+					}
 
-							if (ProjectFile == null)
+					// Create all the target descriptors
+					foreach(string TargetName in TargetNames)
+					{
+						FileReference? TargetProjectFile = ProjectFile;
+						IEnumerable<string> Architectures;
+						if (ParamArchitectures.Count() == 0)
+						{
+							// If a project file was not specified see if we can find one
+							if (TargetProjectFile == null && NativeProjects.TryGetProjectForTarget(TargetName, Logger, out TargetProjectFile))
 							{
-								TargetNames.Add(RulesCompiler.CreateEngineRulesAssembly(bUsePrecompiled, bSkipRulesCompile, bForceRulesCompile, Logger).GetTargetNameByType(TargetType, Platform, Configuration, Architecture, null, Logger));
+								Logger.LogDebug("Found project file for {TargetName} - {ProjectFile}", TargetName, TargetProjectFile);
 							}
-							else
+							// Programs can have a .uproject without finding a matching .Target (since the source and metadata directories are split up)
+							if (TargetProjectFile == null)
 							{
-								TargetNames.Add(RulesCompiler.CreateProjectRulesAssembly(ProjectFile, bUsePrecompiled, bSkipRulesCompile, bForceRulesCompile, Logger).GetTargetNameByType(TargetType, Platform, Configuration, Architecture, ProjectFile, Logger));
+								// find one with a matching name
+								TargetProjectFile = NativeProjects.EnumerateProjectFiles(Log.Logger)
+									.Where(x => x.GetFileNameWithoutAnyExtensions().Equals(TargetName, StringComparison.InvariantCultureIgnoreCase))
+									.FirstOrDefault();
 							}
+							// ask the platform what achitectures it wants for this project
+							Architectures = BuildPlatform.GetProjectArchitectures(TargetProjectFile, TargetName);
+						}
+						else
+						{
+							Architectures = ParamArchitectures;
 						}
 
-						// Make sure we could parse something
-						if (TargetNames.Count == 0)
+						// Make sure we could at least one (even if it's "", the default)
+						if (Architectures.Count() == 0)
 						{
-							throw new BuildException("No target name was specified on the command-line.");
+							throw new BuildException($"Architecture(s) could be determined for platform {Platform}");
 						}
 
-						// Create all the target descriptors
-						foreach(string TargetName in TargetNames)
+						// If the platform can do these in a single pass then turn them back into a single + separate string.
+						// It is now responsible for splitting them as necessary.
+						if (BuildPlatform.NeedsSingleTargetForMultiArchitecture(Architectures))
 						{
-							// If a project file was not specified see if we can find one
-							if (ProjectFile == null && NativeProjects.TryGetProjectForTarget(TargetName, Logger, out ProjectFile))
-							{
-								Logger.LogDebug("Found project file for {TargetName} - {ProjectFile}", TargetName, ProjectFile);
-							}
+							Architectures = new List<string> { string.Join("+", Architectures.OrderBy(S => S)) };
+						}
 
-							// Create the target descriptor
-							TargetDescriptors.Add(new TargetDescriptor(ProjectFile, TargetName, Platform, Configuration, Architecture, Arguments));
+						// Create the target descriptor for each architecture
+						foreach (string Architecture in Architectures)
+						{
+							TargetDescriptors.Add(new TargetDescriptor(TargetProjectFile, TargetName, Platform, Configuration, Architecture, Arguments));
 						}
 					}
 				}
diff --git a/Engine/Source/Programs/UnrealBuildTool/Configuration/UEBuildPlatform.cs b/Engine/Source/Programs/UnrealBuildTool/Configuration/UEBuildPlatform.cs
--- a/Engine/Source/Programs/UnrealBuildTool/Configuration/UEBuildPlatform.cs
+++ b/Engine/Source/Programs/UnrealBuildTool/Configuration/UEBuildPlatform.cs
@@ -229,6 +229,16 @@ namespace UnrealBuildTool
 			return false;
 		}
 
+		/// <summary>
+		/// Determines if a platform needs a single Target object, even if it wants separate compiling and linking
+		/// (one case would be Deploying will combine the Link output into one unit)
+		/// </summary>
+		/// <param name="InArchitectures">Architectures that are being built</param>
+		public virtual bool NeedsSingleTargetForMultiArchitecture(IEnumerable<string> InArchitectures)
+		{
+			return CanCompileArchitecturesInSinglePass(InArchitectures) || CanLinkArchitecturesInSinglePass(InArchitectures);
+		}
+
 		/// <summary>
 		/// Get the default architecture for a project. This may be overriden on the command line to UBT.
 		/// </summary>
@@ -239,6 +249,28 @@ namespace UnrealBuildTool
 			return "";
 		}
 
+		/// <summary>
+		/// Get the default architecture for a project. This may be overriden on the command line to UBT.
+		/// </summary>
+		/// <param name="ProjectFile">Optional project to read settings from </param>
+		/// <param name="TargetName">Optional name of project, can be useful for programs, etc that have no projectfile</param>
+		public virtual IEnumerable<string> GetProjectArchitectures(FileReference? ProjectFile, string? TargetName=null)
+		{
+			// by default, use an empty architecture (which is really just a modifer to the platform for some paths/names)
+			return new string[] { GetDefaultArchitecture(ProjectFile) };
+		}
+
+		/// <summary>
+		/// Convert user specified architecture strings to what the platform wants
+		/// </summary>
+		/// <param name="Architecture"></param>
+		/// <returns></returns>
+		public virtual string FixupArchitecture(string Architecture)
+		{
+			// @todo can we combine with GetFolderNameForArchitecture maybe?
+			return Architecture;
+		}
+
 		/// <summary>
 		/// Get name for architecture-specific directories (can be shorter than architecture name itself)
 		/// </summary>
diff --git a/Engine/Source/Programs/UnrealBuildTool/Platform/Mac/MacToolChain.cs b/Engine/Source/Programs/UnrealBuildTool/Platform/Mac/MacToolChain.cs
--- a/Engine/Source/Programs/UnrealBuildTool/Platform/Mac/MacToolChain.cs
+++ b/Engine/Source/Programs/UnrealBuildTool/Platform/Mac/MacToolChain.cs
@@ -544,7 +544,6 @@ namespace UnrealBuildTool
 				foreach (string Arch in ArchStrings)
 				{
 					LinkEnvironment ArchEnvironment = new LinkEnvironment(LinkEnvironment, Arch);
-					ArchEnvironment.InputFiles = LinkEnvironment.InputFiles.Where(x => x.FullName.Contains($"/{Arch}/")).ToList();
 					if (!bBuildImportLibraryOnly)
 					{
 						ArchEnvironment.OutputFilePaths = LinkEnvironment.OutputFilePaths.Select(x => new FileReference($"{LinkEnvironment.IntermediateDirectory}/{Path.GetFileName(x.FullName)}_{Arch}")).ToList();
diff --git a/Engine/Source/Programs/UnrealBuildTool/Platform/Mac/UEBuildMac.cs b/Engine/Source/Programs/UnrealBuildTool/Platform/Mac/UEBuildMac.cs
--- a/Engine/Source/Programs/UnrealBuildTool/Platform/Mac/UEBuildMac.cs
+++ b/Engine/Source/Programs/UnrealBuildTool/Platform/Mac/UEBuildMac.cs
@@ -3,10 +3,14 @@
 
 using System;
 using System.Collections.Generic;
+using System.Configuration;
+using System.IO;
 using System.Linq;
+using System.Runtime.InteropServices;
 using System.Text.RegularExpressions;
 using EpicGames.Core;
 using Microsoft.Extensions.Logging;
+using UnrealBuildBase;
 
 namespace UnrealBuildTool
 {
@@ -240,7 +244,102 @@ namespace UnrealBuildTool
 		{
 			// by default use Intel.
 			return MacExports.DefaultArchitecture;
-		}		
+		}
+
+		private List<string>? CachedProjectArches = null;
+		private FileReference? CachedArchesProject = null;
+		private string? CachedArchesTarget = null;
+		public override IEnumerable<string> GetProjectArchitectures(FileReference? ProjectFile, string? TargetName)
+		{
+			// use the cached one if it's valid
+			if (CachedProjectArches != null && ProjectFile == CachedArchesProject && TargetName == CachedArchesTarget)
+			{
+				return CachedProjectArches;
+			}
+
+			bool bIsEditor = false;
+			//// get project ini from ProjetFile, or if null, then try to get it from the target rules
+			//DirectoryReference? ProjectDir = null;
+			if (TargetName != null)
+			{
+				RulesAssembly EngineAsm = RulesCompiler.CreateEngineRulesAssembly(Unreal.IsEngineInstalled(), false, false, Log.Logger);
+				TargetRules? Rules = EngineAsm.CreateTargetRules(TargetName, UnrealTargetPlatform.Mac, UnrealTargetConfiguration.Development, "", ProjectFile, null, Log.Logger);
+				bIsEditor = Rules.Type == TargetType.Editor;
+			}
+
+			ConfigHierarchy EngineIni = ConfigCache.ReadHierarchy(ConfigHierarchyType.Engine, ProjectFile?.Directory, UnrealTargetPlatform.Mac);
+
+			// get values from project ini
+			string SupportKey = bIsEditor ? "EditorTargetArchitecture" : "TargetArchitecture";
+			string DefaultKey = bIsEditor ? "EditorDefaultArchitecture" : "DefaultArchitecture";
+			string SupportedArchitecture;
+			string DefaultArchitecture;
+			EngineIni.GetString("/Script/MacTargetPlatform.MacTargetSettings", SupportKey, out SupportedArchitecture);
+			EngineIni.GetString("/Script/MacTargetPlatform.MacTargetSettings", DefaultKey, out DefaultArchitecture);
+			SupportedArchitecture = SupportedArchitecture.ToLower();
+			DefaultArchitecture = DefaultArchitecture.ToLower();
+
+			bool bSupportsArm64 = SupportedArchitecture.Contains("universal") || SupportedArchitecture.Contains("apple");
+			bool bSupportsX86 = SupportedArchitecture.Contains("universal") || SupportedArchitecture.Contains("intel");
+
+			// make sure we found a good value
+			if (!bSupportsArm64 && !bSupportsX86)
+			{
+				throw new BuildException($"Unknown {DefaultKey} value found ('{DefaultArchitecture}') in .ini");
+			}
+
+			// choose a supported architecture(s) based on desired type
+			List<string> Architectures = new();
+			if (DefaultArchitecture.Equals("all", StringComparison.InvariantCultureIgnoreCase))
+			{
+				if (bSupportsArm64)
+				{
+					Architectures.Add("arm64");
+				}
+				if (bSupportsX86)
+				{
+					Architectures.Add("x86_64");
+				}
+			}
+			else if (DefaultArchitecture.Equals("host", StringComparison.InvariantCultureIgnoreCase))
+			{
+				if (MacExports.IsRunningOnAppleArchitecture && bSupportsArm64)
+				{
+					Architectures.Add("arm64");
+				}
+				else
+				{
+					Architectures.Add("x86_64");
+				}
+			}
+			else if (DefaultArchitecture.Contains("apple"))
+			{
+				if (!bSupportsArm64)
+				{
+					throw new BuildException($"{DefaultKey} is set to {DefaultArchitecture}, but AppleSilicon is not a supported architecture");
+				}
+				Architectures.Add("arm64");
+			}
+			else if (DefaultArchitecture.Contains("intel"))
+			{
+				if (!bSupportsArm64)
+				{
+					throw new BuildException($"{DefaultKey} is set to {DefaultArchitecture}, but Intel is not a supported architecture");
+				}
+				Architectures.Add("x86_64");
+			}
+			else
+			{
+				throw new BuildException($"Unknown {DefaultKey} value found ('{DefaultArchitecture}') in .ini");
+			}
+
+
+			CachedProjectArches = Architectures;
+			CachedArchesProject = ProjectFile;
+			CachedArchesTarget = TargetName;
+
+			return CachedProjectArches;
+		}
 
 		/// <summary>
 		/// Determines if the given name is a build product for a target.
diff --git a/Engine/Source/Programs/UnrealBuildTool/System/LinkEnvironment.cs b/Engine/Source/Programs/UnrealBuildTool/System/LinkEnvironment.cs
--- a/Engine/Source/Programs/UnrealBuildTool/System/LinkEnvironment.cs
+++ b/Engine/Source/Programs/UnrealBuildTool/System/LinkEnvironment.cs
@@ -421,11 +421,20 @@ namespace UnrealBuildTool
 			}
 		}
 
+		/// <summary>
+		/// Construct a LinkEnvironment from another, and filter it down to only files/dependencies that apply to the given architecture
+		/// </summary>
+		/// <param name="Other">Parent LinkEnvironment to start with that may have been created from multiple architectures (arch1+arch2)</param>
+		/// <param name="OverrideArchitecture">The single architecture to filter down to</param>
 		public LinkEnvironment(LinkEnvironment Other, string OverrideArchitecture)
 			: this(Other)
 		{
 			Architecture = OverrideArchitecture;
 
+			// filter the input files 
+			string IntermediateDirPart = UEBuildPlatform.GetBuildPlatform(Platform).GetFolderNameForArchitecture(OverrideArchitecture);
+			InputFiles = Other.InputFiles.Where(x => x.Location.ContainsName(IntermediateDirPart, 0)).ToList();
+
 			if (DependenciesToSkipPerArchitecture.Count() > 0)
 			{
 				// add more arrays here?
diff --git a/Engine/Source/Programs/UnrealBuildTool/System/PlatformExports.cs b/Engine/Source/Programs/UnrealBuildTool/System/PlatformExports.cs
--- a/Engine/Source/Programs/UnrealBuildTool/System/PlatformExports.cs
+++ b/Engine/Source/Programs/UnrealBuildTool/System/PlatformExports.cs
@@ -41,11 +41,23 @@ namespace UnrealBuildTool
 		/// <param name="Platform">The platform to get the default architecture for</param>
 		/// <param name="ProjectFile">Project file to read settings from</param>
 		/// <returns>The default architecture</returns>
-		public static string GetDefaultArchitecture(UnrealTargetPlatform Platform, FileReference ProjectFile)
+		public static string GetDefaultArchitecture(UnrealTargetPlatform Platform, FileReference? ProjectFile)
 		{
 			return UEBuildPlatform.GetBuildPlatform(Platform).GetDefaultArchitecture(ProjectFile);
 		}
 
+		/// <summary>
+		/// Gets the default architecture for a given platform
+		/// </summary>
+		/// <param name="Platform">The platform to get the default architecture for</param>
+		/// <param name="ProjectFile">Project file to read settings from</param>
+		/// <param name="TargetName">Target name - most Programs don't have a project file, but they will have a Target name that can be used to determine architecture</param>
+		/// <returns>The default architecture</returns>
+		public static IEnumerable<string> GetProjectArchitectures(UnrealTargetPlatform Platform, FileReference? ProjectFile, string? TargetName=null)
+		{
+			return UEBuildPlatform.GetBuildPlatform(Platform).GetProjectArchitectures(ProjectFile, TargetName);
+		}
+
 		/// <summary>
 		/// Checks whether the given project has a default build configuration
 		/// </summary>
diff --git a/Engine/Source/Programs/UnrealBuildTool/ToolChain/UEToolChain.cs b/Engine/Source/Programs/UnrealBuildTool/ToolChain/UEToolChain.cs
--- a/Engine/Source/Programs/UnrealBuildTool/ToolChain/UEToolChain.cs
+++ b/Engine/Source/Programs/UnrealBuildTool/ToolChain/UEToolChain.cs
@@ -55,15 +55,16 @@ namespace UnrealBuildTool
 			CPPOutput Result;
 			string[] ArchStrings = CompileEnvironment.Architecture.Split('+', StringSplitOptions.RemoveEmptyEntries);
 
+			UEBuildPlatform BuildPlatform = UEBuildPlatform.GetBuildPlatform(CompileEnvironment.Platform);
 			// compile architectures separately if needed
-			if (ArchStrings.Length > 1 && !UEBuildPlatform.GetBuildPlatform(CompileEnvironment.Platform).CanCompileArchitecturesInSinglePass(ArchStrings))
+			if (ArchStrings.Length > 1 && !BuildPlatform.CanCompileArchitecturesInSinglePass(ArchStrings))
 			{
 				Result = new CPPOutput();
 				foreach (string ArchString in ArchStrings)
 				{
 					// determine the output location of intermediates (so, if OutputDir had the arch name in it, like Intermediate/x86+arm64, we would replace it with either emptry string
 					// or a single arch name
-					string ArchReplacement = string.Compare(ArchitectureWithoutMarkup(), ArchString, true) == 0 ? "" : ArchString;
+					string ArchReplacement = string.Compare(ArchitectureWithoutMarkup(), ArchString, true) == 0 ? "" : BuildPlatform.GetFolderNameForArchitecture(ArchString);
 					DirectoryReference ArchOutputDir = new(OutputDir.FullName.Replace(CompileEnvironment.Architecture, ArchReplacement));
 
 					CppCompileEnvironment ArchEnvironment = new(CompileEnvironment, ArchString);
@@ -121,15 +122,16 @@ namespace UnrealBuildTool
 			string[] ArchStrings = LinkEnvironment.Architecture.Split('+', StringSplitOptions.RemoveEmptyEntries);
 
 			// compile architectures separately if needed
-			if (ArchStrings.Length > 1 && !UEBuildPlatform.GetBuildPlatform(LinkEnvironment.Platform).CanLinkArchitecturesInSinglePass(ArchStrings))
+			UEBuildPlatform BuildPlatform = UEBuildPlatform.GetBuildPlatform(LinkEnvironment.Platform);
+			if (ArchStrings.Length > 1 && !BuildPlatform.CanLinkArchitecturesInSinglePass(ArchStrings))
 			{
 				foreach (string ArchString in ArchStrings)
 				{
-					LinkEnvironment ArchEnvironment = new(LinkEnvironment, ArchString);
+					LinkEnvironment ArchEnvironment = new LinkEnvironment(LinkEnvironment, ArchString);
 						
 					// determine the output location of intermediates (so, if OutputDir had the arch name in it, like Intermediate/x86+arm64, we would replace it with either emptry string
 					// or a single arch name
-					string ArchReplacement = string.Compare(ArchitectureWithoutMarkup(), ArchString, true) == 0 ? "" : ArchString;
+					string ArchReplacement = string.Compare(ArchitectureWithoutMarkup(), ArchString, true) == 0 ? "" : BuildPlatform.GetFolderNameForArchitecture(ArchString);
 					ArchEnvironment.OutputFilePaths = LinkEnvironment.OutputFilePaths.Select(x => new FileReference(x.FullName.Replace(LinkEnvironment.Architecture, ArchReplacement))).ToList();
 
 					FileItem? LinkFile = LinkFiles(ArchEnvironment, bBuildImportLibraryOnly, Graph);
diff --git a/Engine/Source/Programs/VirtualProduction/TextureShare/TextureShareSDK.Target.cs b/Engine/Source/Programs/VirtualProduction/TextureShare/TextureShareSDK.Target.cs
--- a/Engine/Source/Programs/VirtualProduction/TextureShare/TextureShareSDK.Target.cs
+++ b/Engine/Source/Programs/VirtualProduction/TextureShare/TextureShareSDK.Target.cs
@@ -5,7 +5,7 @@ using System.IO;
 using EpicGames.Core;
 using System.Runtime.CompilerServices;
 
-[SupportedPlatforms(UnrealPlatformClass.Desktop)]
+[SupportedPlatforms("Win64")]
 public class TextureShareSDKTarget : TargetRules
 {
 	private void ImplPostBuildCopy(string SrcPath, string DestPath)
diff --git a/Engine/Source/Programs/Windows/BootstrapPackagedGame/BootstrapPackagedGame.Target.cs b/Engine/Source/Programs/Windows/BootstrapPackagedGame/BootstrapPackagedGame.Target.cs
--- a/Engine/Source/Programs/Windows/BootstrapPackagedGame/BootstrapPackagedGame.Target.cs
+++ b/Engine/Source/Programs/Windows/BootstrapPackagedGame/BootstrapPackagedGame.Target.cs
@@ -1,9 +1,10 @@
-﻿// Copyright Epic Games, Inc. All Rights Reserved.
+// Copyright Epic Games, Inc. All Rights Reserved.
 
 using UnrealBuildTool;
 using System.Collections.Generic;
 
 [SupportedConfigurations(UnrealTargetConfiguration.Debug, UnrealTargetConfiguration.Development, UnrealTargetConfiguration.Shipping)]
+[SupportedPlatforms("Win64")]
 public class BootstrapPackagedGameTarget : TargetRules
 {
 	public BootstrapPackagedGameTarget(TargetInfo Target) : base(Target)
diff --git a/Engine/Source/Programs/ZenDashboard/ZenDashboard.Target.cs b/Engine/Source/Programs/ZenDashboard/ZenDashboard.Target.cs
--- a/Engine/Source/Programs/ZenDashboard/ZenDashboard.Target.cs
+++ b/Engine/Source/Programs/ZenDashboard/ZenDashboard.Target.cs
@@ -3,6 +3,7 @@
 using UnrealBuildTool;
 using System.Collections.Generic;
 
+[SupportedPlatforms("Win64")]
 public class ZenDashboardTarget : TargetRules
 {
 	public ZenDashboardTarget(TargetInfo Target) : base(Target)
