- Fixed some issues with getting projet architectures

diff --git a/Engine/Source/Programs/UnrealBuildTool/Platform/Mac/UEBuildMac.cs b/Engine/Source/Programs/UnrealBuildTool/Platform/Mac/UEBuildMac.cs
--- a/Engine/Source/Programs/UnrealBuildTool/Platform/Mac/UEBuildMac.cs
+++ b/Engine/Source/Programs/UnrealBuildTool/Platform/Mac/UEBuildMac.cs
@@ -262,8 +262,16 @@ namespace UnrealBuildTool
 			//DirectoryReference? ProjectDir = null;
 			if (TargetName != null)
 			{
-				RulesAssembly EngineAsm = RulesCompiler.CreateEngineRulesAssembly(Unreal.IsEngineInstalled(), false, false, Log.Logger);
-				TargetRules? Rules = EngineAsm.CreateTargetRules(TargetName, UnrealTargetPlatform.Mac, UnrealTargetConfiguration.Development, "", ProjectFile, null, Log.Logger);
+				RulesAssembly RulesAsm;
+				if (ProjectFile == null)
+				{
+					RulesAsm = RulesCompiler.CreateEngineRulesAssembly(Unreal.IsEngineInstalled(), false, false, Log.Logger);
+				}
+				else
+				{
+					RulesAsm = RulesCompiler.CreateProjectRulesAssembly(ProjectFile, Unreal.IsEngineInstalled(), false, false, Log.Logger);
+				}
+				TargetRules? Rules = RulesAsm.CreateTargetRules(TargetName, UnrealTargetPlatform.Mac, UnrealTargetConfiguration.Development, "", ProjectFile, null, Log.Logger);
 				bIsEditor = Rules.Type == TargetType.Editor;
 			}
 
@@ -322,7 +330,7 @@ namespace UnrealBuildTool
 			}
 			else if (DefaultArchitecture.Contains("intel"))
 			{
-				if (!bSupportsArm64)
+				if (!bSupportsX86)
 				{
 					throw new BuildException($"{DefaultKey} is set to {DefaultArchitecture}, but Intel is not a supported architecture");
 				}
