UnrealBuildTool: Suppress msvc /analyze warnings in external headers with /analyze:external-

#rnx

diff --git a/Engine/Source/Programs/UnrealBuildTool/Platform/Windows/VCToolChain.cs b/Engine/Source/Programs/UnrealBuildTool/Platform/Windows/VCToolChain.cs
--- a/Engine/Source/Programs/UnrealBuildTool/Platform/Windows/VCToolChain.cs
+++ b/Engine/Source/Programs/UnrealBuildTool/Platform/Windows/VCToolChain.cs
@@ -330,6 +330,12 @@ namespace UnrealBuildTool
 			{
 				Arguments.Add("/analyze");
 
+				if (Target.WindowsPlatform.Compiler >= WindowsCompiler.VisualStudio2022)
+				{
+					// Ignore warnings in external headers
+					Arguments.Add("/analyze:external-");
+				}
+
 				// Report functions that use a LOT of stack space. You can lower this value if you
 				// want more aggressive checking for functions that use a lot of stack memory.
 				Arguments.Add("/analyze:stacksize" + CompileEnvironment.AnalyzeStackSizeWarning);
