- Added to new architecture APIs to UEBuildPlatform
  - ConvertToReadableArchitecture() - turn an internal platform arch string to a printable version
  - IsX86Architecture() - returns true if the platform arch string represents an x86 architecture
- Changed Linux to use the HumanReadableName as the FolderNameForArchitecture, so it's shorter and it is not obfuscated anymore

diff --git a/Engine/Platforms/Hololens/Source/Programs/UnrealBuildTool/UEBuildHoloLens.cs b/Engine/Platforms/Hololens/Source/Programs/UnrealBuildTool/UEBuildHoloLens.cs
--- a/Engine/Platforms/Hololens/Source/Programs/UnrealBuildTool/UEBuildHoloLens.cs
+++ b/Engine/Platforms/Hololens/Source/Programs/UnrealBuildTool/UEBuildHoloLens.cs
@@ -345,6 +345,11 @@ namespace UnrealBuildTool
 			return WindowsArchitecture.x64.ToString();
 		}
 
+		public override bool IsX86Architecture(string Architecture)
+		{
+			return Architecture == WindowsArchitecture.x64.ToString();
+		}
+
 		/// <summary>
 		/// Get the extension to use for the given binary type
 		/// </summary>
diff --git a/Engine/Source/Programs/UnrealBuildTool/Configuration/TargetDescriptor.cs b/Engine/Source/Programs/UnrealBuildTool/Configuration/TargetDescriptor.cs
--- a/Engine/Source/Programs/UnrealBuildTool/Configuration/TargetDescriptor.cs
+++ b/Engine/Source/Programs/UnrealBuildTool/Configuration/TargetDescriptor.cs
@@ -368,7 +368,7 @@ namespace UnrealBuildTool
 				// Parse the architecture parameter, or get the default for the platform (and letting the platform clean it up
 				IEnumerable<string> ParamArchitectures = Arguments.GetValues("-Architecture=", '+')
 					.Union(Arguments.GetValues("-Architectures=", '+'))
-					.Select(x => BuildPlatform.FixupArchitecture(x));
+					.Select(x => BuildPlatform.ConvertToPlatformArchitecture(x));
 
 				foreach(UnrealTargetConfiguration Configuration in Configurations)
 				{
diff --git a/Engine/Source/Programs/UnrealBuildTool/Configuration/UEBuildPlatform.cs b/Engine/Source/Programs/UnrealBuildTool/Configuration/UEBuildPlatform.cs
--- a/Engine/Source/Programs/UnrealBuildTool/Configuration/UEBuildPlatform.cs
+++ b/Engine/Source/Programs/UnrealBuildTool/Configuration/UEBuildPlatform.cs
@@ -265,12 +265,32 @@ namespace UnrealBuildTool
 		/// </summary>
 		/// <param name="Architecture"></param>
 		/// <returns></returns>
-		public virtual string FixupArchitecture(string Architecture)
+		public virtual string ConvertToPlatformArchitecture(string Architecture)
 		{
-			// @todo can we combine with GetFolderNameForArchitecture maybe?
 			return Architecture;
 		}
 
+		/// <summary>
+		/// Convert user specified architecture strings to what the platform wants
+		/// </summary>
+		/// <param name="Architecture"></param>
+		/// <returns></returns>
+		public virtual string ConvertToReadableArchitecture(string Architecture)
+		{
+			return Architecture;
+		}
+
+		/// <summary>
+		/// Convert user specified architecture strings to what the platform wants
+		/// </summary>
+		/// <param name="Architecture"></param>
+		/// <returns></returns>
+		public virtual bool IsX86Architecture(string Architecture)
+		{
+			// most platforms are x86
+			return true;
+		}
+
 		/// <summary>
 		/// Get name for architecture-specific directories (can be shorter than architecture name itself)
 		/// </summary>
diff --git a/Engine/Source/Programs/UnrealBuildTool/Platform/Android/UEBuildAndroid.cs b/Engine/Source/Programs/UnrealBuildTool/Platform/Android/UEBuildAndroid.cs
--- a/Engine/Source/Programs/UnrealBuildTool/Platform/Android/UEBuildAndroid.cs
+++ b/Engine/Source/Programs/UnrealBuildTool/Platform/Android/UEBuildAndroid.cs
@@ -178,6 +178,11 @@ namespace UnrealBuildTool
 			return BuildHostPlatform.Current.Platform != UnrealTargetPlatform.Linux;
 		}
 
+		public override bool IsX86Architecture(string Architecture)
+		{
+			return Architecture == "-x64";
+		}
+
 		public override bool IsBuildProduct(string FileName, string[] NamePrefixes, string[] NameSuffixes)
 		{
 			return IsBuildProductWithArch(FileName, NamePrefixes, NameSuffixes, ".so")
diff --git a/Engine/Source/Programs/UnrealBuildTool/Platform/IOS/UEBuildIOS.cs b/Engine/Source/Programs/UnrealBuildTool/Platform/IOS/UEBuildIOS.cs
--- a/Engine/Source/Programs/UnrealBuildTool/Platform/IOS/UEBuildIOS.cs
+++ b/Engine/Source/Programs/UnrealBuildTool/Platform/IOS/UEBuildIOS.cs
@@ -737,7 +737,14 @@ namespace UnrealBuildTool
 		{
 			return "arm64";
 		}
-		
+
+		public override bool IsX86Architecture(string Architecture)
+		{
+			return false;
+		}
+
+
+
 		public override List<FileReference> FinalizeBinaryPaths(FileReference BinaryName, FileReference? ProjectFile, ReadOnlyTargetRules Target)
 		{
 			List<FileReference> BinaryPaths = new List<FileReference>();
diff --git a/Engine/Source/Programs/UnrealBuildTool/Platform/Linux/UEBuildLinux.cs b/Engine/Source/Programs/UnrealBuildTool/Platform/Linux/UEBuildLinux.cs
--- a/Engine/Source/Programs/UnrealBuildTool/Platform/Linux/UEBuildLinux.cs
+++ b/Engine/Source/Programs/UnrealBuildTool/Platform/Linux/UEBuildLinux.cs
@@ -189,20 +189,19 @@ namespace UnrealBuildTool
 			}
 		}
 
-		/// <summary>
-		/// Get name for architecture-specific directories (can be shorter than architecture name itself)
-		/// </summary>
+		public override string ConvertToReadableArchitecture(string Architecture)
+		{
+			return IsX86Architecture(Architecture) ? "x64" : "arm64";
+		}
+
+		public override bool IsX86Architecture(string Architecture)
+		{
+			return Architecture == "x86_64-unknown-linux-gnu";
+		}
+
 		public override string GetFolderNameForArchitecture(string Architecture)
 		{
-			// shorten the string (heuristically)
-			uint Sum = 0;
-			int Len = Architecture.Length;
-			for (int Index = 0; Index < Len; ++Index)
-			{
-				Sum += (uint)(Architecture[Index]);
-				Sum <<= 1;	// allowed to overflow
-			}
-			return Sum.ToString("X");
+			return ConvertToReadableArchitecture(Architecture);
 		}
 
 		public override void ResetTarget(TargetRules Target)
diff --git a/Engine/Source/Programs/UnrealBuildTool/Platform/Mac/UEBuildMac.cs b/Engine/Source/Programs/UnrealBuildTool/Platform/Mac/UEBuildMac.cs
--- a/Engine/Source/Programs/UnrealBuildTool/Platform/Mac/UEBuildMac.cs
+++ b/Engine/Source/Programs/UnrealBuildTool/Platform/Mac/UEBuildMac.cs
@@ -208,38 +208,59 @@ namespace UnrealBuildTool
 			}
 		}
 
-
-		/// <summary>
-		/// Returns true since we can do this on Mac (with some caveats, that may necessitate this being an option)
-		/// </summary>
-		/// <param name="InArchitectures">Architectures that are being built</param>
 		public override bool CanCompileArchitecturesInSinglePass(IEnumerable<string> InArchitectures)
 		{
 			return false;
 		}
 
-		/// <summary>
-		/// Returns true since we can do this on Mac (with some caveats, that may necessitate this being an option)
-		/// </summary>
-		/// <param name="InArchitectures">Architectures that are being built</param>
 		public override bool CanLinkArchitecturesInSinglePass(IEnumerable<string> InArchitectures)
 		{
 			return true;
 		}
 
-		/// <summary>
-		/// Allows the platform to override whether the architecture name should be appended to the name of binaries.
-		/// </summary>
-		/// <returns>True if the architecture name should be appended to the binary</returns>
 		public override bool RequiresArchitectureSuffix()
 		{
 			return false;
 		}
 
-		/// <summary>
-		/// Get the default architecture for a project. This may be overriden on the command line to UBT.
-		/// </summary>
-		/// <param name="ProjectFile">Optional project to read settings from </param>
+		public override string ConvertToPlatformArchitecture(string Architecture)
+		{
+			switch (Architecture.ToLower())
+			{
+				case "x86_64":
+				case "x64":
+				case "intel":
+				case "x86":
+					return MacExports.IntelArchitecture;
+
+				case "arm64":
+				case "arm":
+				case "apple":
+				case "applesilicon":
+				case "m1":
+					return MacExports.AppleArchitecture;
+			}
+			throw new BuildException($"Unknown architecture {Architecture}");
+		}
+
+		public override string ConvertToReadableArchitecture(string Architecture)
+		{
+			if (Architecture == MacExports.IntelArchitecture)
+			{
+				return "Intel";
+			}
+			if (Architecture == MacExports.AppleArchitecture)
+			{
+				return "Apple";
+			}
+			return Architecture;
+		}
+
+		public override bool IsX86Architecture(string Architecture)
+		{
+			return Architecture == MacExports.IntelArchitecture;
+		}
+
 		public override string GetDefaultArchitecture(FileReference? ProjectFile)
 		{
 			// by default use Intel.
diff --git a/Engine/Source/Programs/UnrealBuildTool/ToolChain/ClangToolChain.cs b/Engine/Source/Programs/UnrealBuildTool/ToolChain/ClangToolChain.cs
--- a/Engine/Source/Programs/UnrealBuildTool/ToolChain/ClangToolChain.cs
+++ b/Engine/Source/Programs/UnrealBuildTool/ToolChain/ClangToolChain.cs
@@ -1022,7 +1022,8 @@ namespace UnrealBuildTool
 			CompileAction.CommandDescription = IsAnalyzing(CompileEnvironment) ? "Analyze" : "Compile";
 			if (CompileEnvironment.Architecture != "")
 			{
-				CompileAction.CommandDescription += $" [{CompileEnvironment.Architecture}]";
+				string ReadableArch = UEBuildPlatform.GetBuildPlatform(CompileEnvironment.Platform).ConvertToReadableArchitecture(CompileEnvironment.Architecture);
+				CompileAction.CommandDescription += $" [{ReadableArch}]";
 			}
 			CompileAction.StatusDescription = Path.GetFileName(SourceFile.AbsolutePath);
 			CompileAction.bIsGCCCompiler = true;
