- Removed handling of -Architectures= in UBT side until Android code is updated to not _also_ process it (this was a change meant for Android long term, so it will come back)

diff --git a/Engine/Source/Programs/UnrealBuildTool/Configuration/TargetDescriptor.cs b/Engine/Source/Programs/UnrealBuildTool/Configuration/TargetDescriptor.cs
--- a/Engine/Source/Programs/UnrealBuildTool/Configuration/TargetDescriptor.cs
+++ b/Engine/Source/Programs/UnrealBuildTool/Configuration/TargetDescriptor.cs
@@ -367,8 +367,8 @@ namespace UnrealBuildTool
 
 				// Parse the architecture parameter, or get the default for the platform (and letting the platform clean it up
 				IEnumerable<string> ParamArchitectures = Arguments.GetValues("-Architecture=", '+')
-					.Union(Arguments.GetValues("-Architectures=", '+'))
-					.Select(x => BuildPlatform.ConvertToPlatformArchitecture(x));
+					.Select(x => BuildPlatform.ConvertToPlatformArchitecture(x))
+					.Distinct();
 
 				foreach(UnrealTargetConfiguration Configuration in Configurations)
 				{
