When compiling nonunity, add "NonUnity" to the target name to allow side by side unity and nonunity compiles.

diff --git a/Engine/Source/Programs/UnrealBuildTool/Configuration/TargetRules.cs b/Engine/Source/Programs/UnrealBuildTool/Configuration/TargetRules.cs
--- a/Engine/Source/Programs/UnrealBuildTool/Configuration/TargetRules.cs
+++ b/Engine/Source/Programs/UnrealBuildTool/Configuration/TargetRules.cs
@@ -362,6 +362,10 @@ namespace UnrealBuildTool
 				{
 					return NameOverride;
 				}
+				else if (!bUseUnityBuild)
+				{
+					return DefaultName + "NonUnity";
+				}
 
 				return DefaultName;
 			}
