Changing bUseGameplayDebugger to be an enum (Default, Core, Disabled) instead of a bool
Added a new Target's option 'UseGameplayDebuggerCore' to remove GP DBG's menu and default categories

diff --git a/Engine/Source/Editor/GameplayDebugger/GameplayDebuggerEditor.Build.cs b/Engine/Source/Editor/GameplayDebugger/GameplayDebuggerEditor.Build.cs
--- a/Engine/Source/Editor/GameplayDebugger/GameplayDebuggerEditor.Build.cs
+++ b/Engine/Source/Editor/GameplayDebugger/GameplayDebuggerEditor.Build.cs
@@ -45,7 +45,7 @@ namespace UnrealBuildTool.Rules
 					});
 			}
 
-			if (Target.bUseGameplayDebugger)
+			if (Target.bUseGameplayDebuggerCore)
 			{
 				PrecompileForTargets = PrecompileTargetsType.Any;
 			}
diff --git a/Engine/Source/Programs/UnrealBuildTool/Configuration/ModuleRules.cs b/Engine/Source/Programs/UnrealBuildTool/Configuration/ModuleRules.cs
--- a/Engine/Source/Programs/UnrealBuildTool/Configuration/ModuleRules.cs
+++ b/Engine/Source/Programs/UnrealBuildTool/Configuration/ModuleRules.cs
@@ -1279,9 +1279,10 @@ namespace UnrealBuildTool
 		/// </summary>
 		public void SetupGameplayDebuggerSupport(ReadOnlyTargetRules Target, bool bAddAsPublicDependency = false)
 		{
-			if (Target.bUseGameplayDebugger)
+			if (Target.bUseGameplayDebugger || Target.bUseGameplayDebuggerCore)
 			{
-				PublicDefinitions.Add("WITH_GAMEPLAY_DEBUGGER=1");
+				PublicDefinitions.Add("WITH_GAMEPLAY_DEBUGGER_CORE=1");
+				PublicDefinitions.Add("WITH_GAMEPLAY_DEBUGGER=" + (Target.bUseGameplayDebugger ? 1 : 0));
 
 				if (bAddAsPublicDependency)
 				{
@@ -1302,6 +1303,7 @@ namespace UnrealBuildTool
 			}
 			else
 			{
+				PublicDefinitions.Add("WITH_GAMEPLAY_DEBUGGER_CORE=0");
 				PublicDefinitions.Add("WITH_GAMEPLAY_DEBUGGER=0");
 			}
 		}
diff --git a/Engine/Source/Programs/UnrealBuildTool/Configuration/TargetRules.cs b/Engine/Source/Programs/UnrealBuildTool/Configuration/TargetRules.cs
--- a/Engine/Source/Programs/UnrealBuildTool/Configuration/TargetRules.cs
+++ b/Engine/Source/Programs/UnrealBuildTool/Configuration/TargetRules.cs
@@ -266,6 +266,27 @@ namespace UnrealBuildTool
 		SizeAndSpeed
 	}
 
+	/// <summary>
+	/// Determines how the Gameplay Debugger plugin will be activated.
+	/// </summary>
+	public enum GameplayDebuggerOverrideState
+	{
+		/// <summary>
+		/// Default => Not overriden, default usage behavior.
+		/// </summary>
+		Default,
+
+		/// <summary>
+		/// Core => WITH_GAMEPLAY_DEBUGGER = 0 and WITH_GAMEPLAY_DEBUGGER_CORE = 1
+		/// </summary>
+		Core,
+
+		/// <summary>
+		/// Full => WITH_GAMEPLAY_DEBUGGER = 1 and WITH_GAMEPLAY_DEBUGGER_CORE = 1
+		/// </summary>
+		Full
+	}
+
 	/// <summary>
 	/// Utility class for EngineIncludeOrderVersion defines
 	/// </summary>
@@ -770,15 +791,51 @@ namespace UnrealBuildTool
 		public bool bCompilePython = true;
 
 		/// <summary>
-		/// Whether to compile with WITH_GAMEPLAY_DEBUGGER enabled and use GameplayDebugger.
+		/// Whether to compile with WITH_GAMEPLAY_DEBUGGER enabled with all Engine's default gameplay debugger categories.
 		/// </summary>
 		[RequiresUniqueBuildEnvironment]
 		public bool bUseGameplayDebugger
 		{
-			get { return bUseGameplayDebuggerOverride ?? (bBuildDeveloperTools || (Configuration != UnrealTargetConfiguration.Test && Configuration != UnrealTargetConfiguration.Shipping)); }
-			set { bUseGameplayDebuggerOverride = value; }
+			get
+			{
+				if (UseGameplayDebuggerOverride == GameplayDebuggerOverrideState.Default)
+				{
+					return (bBuildDeveloperTools || (Configuration != UnrealTargetConfiguration.Test && Configuration != UnrealTargetConfiguration.Shipping));
+				}
+				else
+				{
+					return UseGameplayDebuggerOverride == GameplayDebuggerOverrideState.Full;
+				}
+			}
+			set
+			{
+				if (value)
+				{
+					UseGameplayDebuggerOverride = GameplayDebuggerOverrideState.Full;
+				}
+				else if (UseGameplayDebuggerOverride != GameplayDebuggerOverrideState.Core)
+				{
+					UseGameplayDebuggerOverride = GameplayDebuggerOverrideState.Default;
+				}
+			}
+		}
+
+		/// <summary>
+		/// Set to true when bUseGameplayDebugger is false but GameplayDebugger's core parts are required.
+		/// </summary>
+		[RequiresUniqueBuildEnvironment]
+		public bool bUseGameplayDebuggerCore
+		{
+			get { return (UseGameplayDebuggerOverride == GameplayDebuggerOverrideState.Core) || bUseGameplayDebugger; }
+			set
+			{
+				if (UseGameplayDebuggerOverride != GameplayDebuggerOverrideState.Full)
+				{
+					UseGameplayDebuggerOverride = value ? GameplayDebuggerOverrideState.Core : GameplayDebuggerOverrideState.Default;
+				}
+			}
 		}
-		bool? bUseGameplayDebuggerOverride;
+		GameplayDebuggerOverrideState UseGameplayDebuggerOverride;
 
 		/// <summary>
 		/// Whether to use Iris.
@@ -2804,7 +2861,12 @@ namespace UnrealBuildTool
 		{
 			get { return Inner.bUseGameplayDebugger; }
 		}
-		
+
+		public bool bUseGameplayDebuggerCore
+		{
+			get { return Inner.bUseGameplayDebuggerCore; }
+		}
+
 		public bool bUseIris
 		{
 			get { return Inner.bUseIris; }
diff --git a/Engine/Source/Runtime/GameplayDebugger/GameplayDebugger.Build.cs b/Engine/Source/Runtime/GameplayDebugger/GameplayDebugger.Build.cs
--- a/Engine/Source/Runtime/GameplayDebugger/GameplayDebugger.Build.cs
+++ b/Engine/Source/Runtime/GameplayDebugger/GameplayDebugger.Build.cs
@@ -46,10 +46,21 @@ namespace UnrealBuildTool.Rules
 
 			SetupIrisSupport(Target);
 
-			if (Target.bUseGameplayDebugger)
+			if (Target.bUseGameplayDebuggerCore)
 			{
 				PrecompileForTargets = PrecompileTargetsType.Any;
 			}
+
+			if (Target.bUseGameplayDebugger || Target.bUseGameplayDebuggerCore)
+			{
+				PublicDefinitions.Add("WITH_GAMEPLAY_DEBUGGER_CORE=1");
+				PublicDefinitions.Add("WITH_GAMEPLAY_DEBUGGER=" + (Target.bUseGameplayDebugger ? 1 : 0));
+			}
+			else
+			{
+				PublicDefinitions.Add("WITH_GAMEPLAY_DEBUGGER_CORE=0");
+				PublicDefinitions.Add("WITH_GAMEPLAY_DEBUGGER=0");
+			}
 		}
 	}
 }
