[UnrealBuildTool]
* ISPCToolChain - changed so NormalizeCommandLinePath is virtual. We want IWYU tool chain to always use absolute paths so these needs to be overridable

diff --git a/Engine/Source/Programs/UnrealBuildTool/ToolChain/ISPCToolChain.cs b/Engine/Source/Programs/UnrealBuildTool/ToolChain/ISPCToolChain.cs
--- a/Engine/Source/Programs/UnrealBuildTool/ToolChain/ISPCToolChain.cs
+++ b/Engine/Source/Programs/UnrealBuildTool/ToolChain/ISPCToolChain.cs
@@ -340,7 +340,7 @@ namespace UnrealBuildTool
 		/// </summary>
 		/// <param name="Reference">The FileSystemReference to normalize</param>
 		/// <returns>Normalized path as a string</returns>
-		protected static string NormalizeCommandLinePath(FileSystemReference Reference)
+		protected virtual string NormalizeCommandLinePath(FileSystemReference Reference)
 		{
 			// Try to use a relative path to shorten command line length.
 			if (Reference.IsUnderDirectory(Unreal.RootDirectory))
@@ -356,7 +356,7 @@ namespace UnrealBuildTool
 		/// </summary>
 		/// <param name="Item">The FileItem to normalize</param>
 		/// <returns>Normalized path as a string</returns>
-		protected static string NormalizeCommandLinePath(FileItem Item)
+		protected virtual string NormalizeCommandLinePath(FileItem Item)
 		{
 			return NormalizeCommandLinePath(Item.Location);
 		}
