Fix iOS compile failure when ISPC is enabled

diff --git a/Engine/Source/Programs/UnrealBuildTool/ToolChain/ISPCToolChain.cs b/Engine/Source/Programs/UnrealBuildTool/ToolChain/ISPCToolChain.cs
--- a/Engine/Source/Programs/UnrealBuildTool/ToolChain/ISPCToolChain.cs
+++ b/Engine/Source/Programs/UnrealBuildTool/ToolChain/ISPCToolChain.cs
@@ -48,6 +48,10 @@ namespace UnrealBuildTool
 					default: Logger.LogWarning("Invalid Android architecture for ISPC. At least one architecture (armv7, x86, etc) needs to be selected in the project settings to build"); break;
 				}
 			}
+			else if (Platform == UnrealTargetPlatform.IOS)
+			{
+				ISPCTargets.Add("neon");
+			}
 			else
 			{
 				Logger.LogWarning("Unsupported ISPC platform target!");
@@ -77,6 +81,10 @@ namespace UnrealBuildTool
 			{
 				ISPCOS += "android";
 			}
+			else if (Platform == UnrealTargetPlatform.IOS)
+			{
+				ISPCOS += "ios";
+			}
 			else if (Platform == UnrealTargetPlatform.Mac)
 			{
 				ISPCOS += "macos";
@@ -120,6 +128,10 @@ namespace UnrealBuildTool
 					default: Logger.LogWarning("Invalid Android architecture for ISPC. At least one architecture (armv7, x86, etc) needs to be selected in the project settings to build"); break;
 				}
 			}
+			else if (Platform == UnrealTargetPlatform.IOS)
+			{
+				ISPCArch += "aarch64"; 
+			}
 			else
 			{
 				Logger.LogWarning("Unsupported ISPC platform target!");
@@ -224,6 +236,7 @@ namespace UnrealBuildTool
 			}
 			else if (UEBuildPlatform.IsPlatformInGroup(Platform, UnrealPlatformGroup.Unix) ||
 					Platform == UnrealTargetPlatform.Mac ||
+					Platform == UnrealTargetPlatform.IOS ||
 					Platform == UnrealTargetPlatform.Android)
 			{
 				Format += "obj";
@@ -251,6 +264,7 @@ namespace UnrealBuildTool
 			}
 			else if (UEBuildPlatform.IsPlatformInGroup(Platform, UnrealPlatformGroup.Unix) ||
 					Platform == UnrealTargetPlatform.Mac ||
+					Platform == UnrealTargetPlatform.IOS ||
 					Platform == UnrealTargetPlatform.Android)
 			{
 				Suffix += ".o";
