- Fixing some issues (mostly with programs) on IOS/Mac
  - Fixed executable path for programs under Restricted
  - Fixed paths when there's an ExeSubDir in the Target
  - Allow for codesigning frameworks under programs without a .uproject

diff --git a/Engine/Source/Programs/AutomationTool/IOS/IOSPlatform.Automation.cs b/Engine/Source/Programs/AutomationTool/IOS/IOSPlatform.Automation.cs
--- a/Engine/Source/Programs/AutomationTool/IOS/IOSPlatform.Automation.cs
+++ b/Engine/Source/Programs/AutomationTool/IOS/IOSPlatform.Automation.cs
@@ -575,11 +575,11 @@ public class IOSPlatform : ApplePlatform
 		return PakParams;
 	}
 
-	public virtual bool PrepForUATPackageOrDeploy(UnrealTargetConfiguration Config, FileReference ProjectFile, string InProjectName, DirectoryReference InProjectDirectory, string InExecutablePath, DirectoryReference InEngineDir, bool bForDistribution, string CookFlavor, bool bIsDataDeploy, bool bCreateStubIPA, bool bIsUEGame)
+	public virtual bool PrepForUATPackageOrDeploy(UnrealTargetConfiguration Config, FileReference ProjectFile, string InProjectName, DirectoryReference InProjectDirectory, FileReference Executable, DirectoryReference InEngineDir, bool bForDistribution, string CookFlavor, bool bIsDataDeploy, bool bCreateStubIPA, bool bIsUEGame)
 	{
-		FileReference TargetReceiptFileName = GetTargetReceiptFileName(Config, InExecutablePath, InEngineDir, InProjectDirectory, bIsUEGame);
+		FileReference TargetReceiptFileName = GetTargetReceiptFileName(Config, Executable.FullName, InEngineDir, InProjectDirectory, bIsUEGame);
 
-		return IOSExports.PrepForUATPackageOrDeploy(Config, ProjectFile, InProjectName, InProjectDirectory, InExecutablePath, InEngineDir, bForDistribution, CookFlavor, bIsDataDeploy, bCreateStubIPA, TargetReceiptFileName, Log.Logger);
+		return IOSExports.PrepForUATPackageOrDeploy(Config, ProjectFile, InProjectName, InProjectDirectory, Executable, InEngineDir, bForDistribution, CookFlavor, bIsDataDeploy, bCreateStubIPA, TargetReceiptFileName, Log.Logger);
 	}
 
 
@@ -795,7 +795,7 @@ public class IOSPlatform : ApplePlatform
 			PrepForUATPackageOrDeploy(TargetConfiguration, Params.RawProjectPath,
 				Params.ShortProjectName,
 				Params.RawProjectPath.Directory,
-				CombinePaths(Path.GetDirectoryName(ProjectGameExeFilename), SC.StageExecutables[0]),
+				FileReference.Combine(new FileReference(ProjectGameExeFilename).Directory!, SC.StageExecutables[0]),
 				DirectoryReference.Combine(SC.LocalRoot, "Engine"),
 				Params.Distribution,
 				"",
diff --git a/Engine/Source/Programs/AutomationTool/TVOS/TVOSPlatform.Automation.cs b/Engine/Source/Programs/AutomationTool/TVOS/TVOSPlatform.Automation.cs
--- a/Engine/Source/Programs/AutomationTool/TVOS/TVOSPlatform.Automation.cs
+++ b/Engine/Source/Programs/AutomationTool/TVOS/TVOSPlatform.Automation.cs
@@ -17,9 +17,9 @@ public class TVOSPlatform : IOSPlatform
 		TargetIniPlatformType = UnrealTargetPlatform.IOS;
 	}
 
-	public override bool PrepForUATPackageOrDeploy(UnrealTargetConfiguration Config, FileReference ProjectFile, string InProjectName, DirectoryReference InProjectDirectory, string InExecutablePath, DirectoryReference InEngineDir, bool bForDistribution, string CookFlavor, bool bIsDataDeploy, bool bCreateStubIPA, bool bIsUEGame)
+	public override bool PrepForUATPackageOrDeploy(UnrealTargetConfiguration Config, FileReference ProjectFile, string InProjectName, DirectoryReference InProjectDirectory, FileReference Executable, DirectoryReference InEngineDir, bool bForDistribution, string CookFlavor, bool bIsDataDeploy, bool bCreateStubIPA, bool bIsUEGame)
 	{
-		string TargetName = Path.GetFileNameWithoutExtension(InExecutablePath).Split("-".ToCharArray())[0];
+		string TargetName = Path.GetFileNameWithoutExtension(Executable.FullName).Split("-".ToCharArray())[0];
 		FileReference TargetReceiptFileName;
 		if (bIsUEGame)
 		{
@@ -29,7 +29,7 @@ public class TVOSPlatform : IOSPlatform
 		{
 			TargetReceiptFileName = TargetReceipt.GetDefaultPath(InProjectDirectory, TargetName, UnrealTargetPlatform.TVOS, Config, "");
 		}
-		return TVOSExports.PrepForUATPackageOrDeploy(Config, ProjectFile, InProjectName, InProjectDirectory, InExecutablePath, InEngineDir, bForDistribution, CookFlavor, bIsDataDeploy, bCreateStubIPA, TargetReceiptFileName, Log.Logger);
+		return TVOSExports.PrepForUATPackageOrDeploy(Config, ProjectFile, InProjectName, InProjectDirectory, Executable, InEngineDir, bForDistribution, CookFlavor, bIsDataDeploy, bCreateStubIPA, TargetReceiptFileName, Log.Logger);
 	}
 
     public override void GetProvisioningData(FileReference InProject, bool bDistribution, out string MobileProvision, out string SigningCertificate, out string Team, out bool bAutomaticSigning)
diff --git a/Engine/Source/Programs/UnrealBuildTool/Platform/IOS/IOSExports.cs b/Engine/Source/Programs/UnrealBuildTool/Platform/IOS/IOSExports.cs
--- a/Engine/Source/Programs/UnrealBuildTool/Platform/IOS/IOSExports.cs
+++ b/Engine/Source/Programs/UnrealBuildTool/Platform/IOS/IOSExports.cs
@@ -73,7 +73,7 @@ namespace UnrealBuildTool
 		/// <param name="ProjectFile"></param>
 		/// <param name="InProjectName"></param>
 		/// <param name="InProjectDirectory"></param>
-		/// <param name="InExecutablePath"></param>
+		/// <param name="Executable"></param>
 		/// <param name="InEngineDir"></param>
 		/// <param name="bForDistribution"></param>
 		/// <param name="CookFlavor"></param>
@@ -82,10 +82,10 @@ namespace UnrealBuildTool
 		/// <param name="BuildReceiptFileName"></param>
 		/// <param name="Logger">Logger for output</param>
 		/// <returns></returns>
-		public static bool PrepForUATPackageOrDeploy(UnrealTargetConfiguration Config, FileReference ProjectFile, string InProjectName, DirectoryReference InProjectDirectory, string InExecutablePath, DirectoryReference InEngineDir, bool bForDistribution, string CookFlavor, bool bIsDataDeploy, bool bCreateStubIPA, FileReference BuildReceiptFileName, ILogger Logger)
+		public static bool PrepForUATPackageOrDeploy(UnrealTargetConfiguration Config, FileReference ProjectFile, string InProjectName, DirectoryReference InProjectDirectory, FileReference Executable, DirectoryReference InEngineDir, bool bForDistribution, string CookFlavor, bool bIsDataDeploy, bool bCreateStubIPA, FileReference BuildReceiptFileName, ILogger Logger)
 		{
 			TargetReceipt Receipt = TargetReceipt.Read(BuildReceiptFileName);
-			return new UEDeployIOS(Logger).PrepForUATPackageOrDeploy(Config, ProjectFile, InProjectName, InProjectDirectory.FullName, InExecutablePath, InEngineDir.FullName, bForDistribution, CookFlavor, bIsDataDeploy, bCreateStubIPA, Receipt);
+			return new UEDeployIOS(Logger).PrepForUATPackageOrDeploy(Config, ProjectFile, InProjectName, InProjectDirectory.FullName, Executable, InEngineDir.FullName, bForDistribution, CookFlavor, bIsDataDeploy, bCreateStubIPA, Receipt);
 		}
 
 		/// <summary>
diff --git a/Engine/Source/Programs/UnrealBuildTool/Platform/IOS/IOSToolChain.cs b/Engine/Source/Programs/UnrealBuildTool/Platform/IOS/IOSToolChain.cs
--- a/Engine/Source/Programs/UnrealBuildTool/Platform/IOS/IOSToolChain.cs
+++ b/Engine/Source/Programs/UnrealBuildTool/Platform/IOS/IOSToolChain.cs
@@ -1859,7 +1859,7 @@ namespace UnrealBuildTool
 				// ensure the plist, entitlements, and provision files are properly copied
 				UEDeployIOS DeployHandler = (Target.Platform == UnrealTargetPlatform.IOS ? new UEDeployIOS(Logger) : new UEDeployTVOS(Logger));
 				DeployHandler.ForDistribution = Target.bForDistribution;
-				DeployHandler.PrepTargetForDeployment(Target.ProjectFile, Target.TargetName, Target.Platform, Target.Configuration, Target.UPLScripts, Target.bCreateStubIPA, BundleID, Target.bBuildAsFramework);
+				DeployHandler.PrepTargetForDeployment(Target.ProjectFile, Target.TargetName, Target.OutputPath, Target.Platform, Target.Configuration, Target.UPLScripts, Target.bCreateStubIPA, BundleID, Target.bBuildAsFramework);
 
 				Logger.LogInformation("Executing {Script}", SignProjectScript);
 
@@ -1938,7 +1938,7 @@ namespace UnrealBuildTool
 				// ensure the plist, entitlements, and provision files are properly copied
 				UEDeployIOS DeployHandler = (Target.Platform == UnrealTargetPlatform.IOS ? new UEDeployIOS(Logger) : new UEDeployTVOS(Logger));
 				DeployHandler.ForDistribution = Target.bForDistribution;
-				DeployHandler.PrepTargetForDeployment(Target.ProjectFile, Target.TargetName, Target.Platform, Target.Configuration, Target.UPLScripts, Target.bCreateStubIPA, BundleID, Target.bBuildAsFramework);
+				DeployHandler.PrepTargetForDeployment(Target.ProjectFile, Target.TargetName, Target.OutputPath, Target.Platform, Target.Configuration, Target.UPLScripts, Target.bCreateStubIPA, BundleID, Target.bBuildAsFramework);
 
 				// write the entitlements file (building on Mac)
 				WriteEntitlements(Target, Logger);
diff --git a/Engine/Source/Programs/UnrealBuildTool/Platform/IOS/UEDeployIOS.cs b/Engine/Source/Programs/UnrealBuildTool/Platform/IOS/UEDeployIOS.cs
--- a/Engine/Source/Programs/UnrealBuildTool/Platform/IOS/UEDeployIOS.cs
+++ b/Engine/Source/Programs/UnrealBuildTool/Platform/IOS/UEDeployIOS.cs
@@ -747,12 +747,12 @@ namespace UnrealBuildTool
 			}
 		}
 
-		public bool PrepForUATPackageOrDeploy(UnrealTargetConfiguration Config, FileReference ProjectFile, string InProjectName, string InProjectDirectory, string InExecutablePath, string InEngineDir, bool bForDistribution, string CookFlavor, bool bIsDataDeploy, bool bCreateStubIPA, TargetReceipt Receipt)
+		public bool PrepForUATPackageOrDeploy(UnrealTargetConfiguration Config, FileReference ProjectFile, string InProjectName, string InProjectDirectory, FileReference Executable, string InEngineDir, bool bForDistribution, string CookFlavor, bool bIsDataDeploy, bool bCreateStubIPA, TargetReceipt Receipt)
 		{
 			List<string> UPLScripts = CollectPluginDataPaths(Receipt.AdditionalProperties, Logger);
 			VersionNumber? SdkVersion = GetSdkVersion(Receipt);
 			bool bBuildAsFramework = GetCompileAsDll(Receipt);
-			return PrepForUATPackageOrDeploy(Config, ProjectFile, InProjectName, InProjectDirectory, InExecutablePath, InEngineDir, bForDistribution, CookFlavor, bIsDataDeploy, bCreateStubIPA, UPLScripts, "", bBuildAsFramework);
+			return PrepForUATPackageOrDeploy(Config, ProjectFile, InProjectName, InProjectDirectory, Executable, InEngineDir, bForDistribution, CookFlavor, bIsDataDeploy, bCreateStubIPA, UPLScripts, "", bBuildAsFramework);
 		}
 
 		void CopyAllProvisions(string ProvisionDir, ILogger Logger)
@@ -804,7 +804,7 @@ namespace UnrealBuildTool
 			}
 		}
 
-		public bool PrepForUATPackageOrDeploy(UnrealTargetConfiguration Config, FileReference? ProjectFile, string InProjectName, string InProjectDirectory, string InExecutablePath, string InEngineDir, bool bForDistribution, string CookFlavor, bool bIsDataDeploy, bool bCreateStubIPA, List<string> UPLScripts, string? BundleID, bool bBuildAsFramework)
+		public bool PrepForUATPackageOrDeploy(UnrealTargetConfiguration Config, FileReference? ProjectFile, string InProjectName, string InProjectDirectory, FileReference Executable, string InEngineDir, bool bForDistribution, string CookFlavor, bool bIsDataDeploy, bool bCreateStubIPA, List<string> UPLScripts, string? BundleID, bool bBuildAsFramework)
 		{
 			if (BuildHostPlatform.Current.Platform != UnrealTargetPlatform.Mac)
 			{
@@ -819,9 +819,9 @@ namespace UnrealBuildTool
 
 			string SubDir = GetTargetPlatformName();
 
-			bool bIsUnrealGame = InExecutablePath.Contains("UnrealGame");
-			string BinaryPath = Path.GetDirectoryName(InExecutablePath)!;
-			string GameExeName = Path.GetFileName(InExecutablePath);
+			bool bIsUnrealGame = Executable.FullName.Contains("UnrealGame");
+			DirectoryReference BinaryPath = Executable.Directory!;
+			string GameExeName = Executable.GetFileName();
 			string GameName = bIsUnrealGame ? "UnrealGame" : GameExeName.Split("-".ToCharArray())[0];
 			string PayloadDirectory = BinaryPath + "/Payload";
 			string AppDirectory = PayloadDirectory + "/" + GameName + ".app";
@@ -830,7 +830,7 @@ namespace UnrealBuildTool
 			string BuildDirectory_NFL = InProjectDirectory + "/Restricted/NotForLicensees/Build/" + SubDir;
 			string IntermediateDirectory = (bIsUnrealGame ? InEngineDir : InProjectDirectory) + "/Intermediate/" + SubDir;
 
-			Directory.CreateDirectory(BinaryPath);
+			DirectoryReference.CreateDirectory(BinaryPath);
 			Directory.CreateDirectory(PayloadDirectory);
 			Directory.CreateDirectory(AppDirectory);
 			Directory.CreateDirectory(BuildDirectory);
@@ -985,33 +985,20 @@ namespace UnrealBuildTool
 		{
 			List<string> UPLScripts = CollectPluginDataPaths(Receipt.AdditionalProperties, Logger);
 			bool bBuildAsFramework = GetCompileAsDll(Receipt);
-			return PrepTargetForDeployment(Receipt.ProjectFile, Receipt.TargetName, Receipt.Platform, Receipt.Configuration, UPLScripts, false, "", bBuildAsFramework);
+			return PrepTargetForDeployment(Receipt.ProjectFile, Receipt.TargetName, Receipt.BuildProducts.First(x => x.Type == BuildProductType.Executable).Path, Receipt.Platform, Receipt.Configuration, UPLScripts, false, "", bBuildAsFramework);
 		}
 
-		public bool PrepTargetForDeployment(FileReference? ProjectFile, string TargetName, UnrealTargetPlatform Platform, UnrealTargetConfiguration Configuration, List<string> UPLScripts, bool bCreateStubIPA, string? BundleID, bool bBuildAsFramework)
+		public bool PrepTargetForDeployment(FileReference? ProjectFile, string TargetName, FileReference Executable, UnrealTargetPlatform Platform, UnrealTargetConfiguration Configuration, List<string> UPLScripts, bool bCreateStubIPA, string? BundleID, bool bBuildAsFramework)
 		{
-			string SubDir = GetTargetPlatformName();
-
 			string GameName = TargetName;
 			string ProjectDirectory = (DirectoryReference.FromFile(ProjectFile) ?? Unreal.EngineDirectory).FullName;
-			string BuildPath = (GameName == "UnrealGame" ? "../../Engine" : ProjectDirectory) + "/Binaries/" + SubDir;
 			bool bIsUnrealGame = GameName.Contains("UnrealGame");
 
-			Console.WriteLine("1 GameName: {0}, ProjectName: {1}", GameName, (ProjectFile == null) ? "" : Path.GetFileNameWithoutExtension(ProjectFile.FullName));
-
-			string DecoratedGameName;
-			if (Configuration == UnrealTargetConfiguration.Development)
-			{
-				DecoratedGameName = GameName;
-			}
-			else
-			{
-				DecoratedGameName = String.Format("{0}-{1}-{2}", GameName, Platform.ToString(), Configuration.ToString());
-			}
+			Console.WriteLine("1 GameName: {0}, ProjectName: {1}, Executable: {2}", GameName, (ProjectFile == null) ? "" : Path.GetFileNameWithoutExtension(ProjectFile.FullName), Executable);
 
 			if (BuildHostPlatform.Current.Platform == UnrealTargetPlatform.Mac && Environment.GetEnvironmentVariable("UBT_NO_POST_DEPLOY") != "true")
 			{
-				return PrepForUATPackageOrDeploy(Configuration, ProjectFile, GameName, ProjectDirectory, BuildPath + "/" + DecoratedGameName, "../../Engine", bForDistribution, "", false, bCreateStubIPA, UPLScripts, BundleID, bBuildAsFramework);
+				return PrepForUATPackageOrDeploy(Configuration, ProjectFile, GameName, ProjectDirectory, Executable, "../../Engine", bForDistribution, "", false, bCreateStubIPA, UPLScripts, BundleID, bBuildAsFramework);
 			}
 			else
 			{
diff --git a/Engine/Source/Programs/UnrealBuildTool/Platform/TVOS/TVOSExports.cs b/Engine/Source/Programs/UnrealBuildTool/Platform/TVOS/TVOSExports.cs
--- a/Engine/Source/Programs/UnrealBuildTool/Platform/TVOS/TVOSExports.cs
+++ b/Engine/Source/Programs/UnrealBuildTool/Platform/TVOS/TVOSExports.cs
@@ -69,7 +69,7 @@ namespace UnrealBuildTool
 		/// <param name="ProjectFile"></param>
 		/// <param name="InProjectName"></param>
 		/// <param name="InProjectDirectory"></param>
-		/// <param name="InExecutablePath"></param>
+		/// <param name="Executable"></param>
 		/// <param name="InEngineDir"></param>
 		/// <param name="bForDistribution"></param>
 		/// <param name="CookFlavor"></param>
@@ -78,10 +78,10 @@ namespace UnrealBuildTool
 		/// <param name="BuildReceiptFileName"></param>
 		/// <param name="Logger"></param>
 		/// <returns></returns>
-		public static bool PrepForUATPackageOrDeploy(UnrealTargetConfiguration Config, FileReference ProjectFile, string InProjectName, DirectoryReference InProjectDirectory, string InExecutablePath, DirectoryReference InEngineDir, bool bForDistribution, string CookFlavor, bool bIsDataDeploy, bool bCreateStubIPA, FileReference BuildReceiptFileName, ILogger Logger)
+		public static bool PrepForUATPackageOrDeploy(UnrealTargetConfiguration Config, FileReference ProjectFile, string InProjectName, DirectoryReference InProjectDirectory, FileReference Executable, DirectoryReference InEngineDir, bool bForDistribution, string CookFlavor, bool bIsDataDeploy, bool bCreateStubIPA, FileReference BuildReceiptFileName, ILogger Logger)
 		{
 			TargetReceipt Receipt = TargetReceipt.Read(BuildReceiptFileName);
-			return new UEDeployTVOS(Logger).PrepForUATPackageOrDeploy(Config, ProjectFile, InProjectName, InProjectDirectory.FullName, InExecutablePath, InEngineDir.FullName, bForDistribution, CookFlavor, bIsDataDeploy, bCreateStubIPA, Receipt);
+			return new UEDeployTVOS(Logger).PrepForUATPackageOrDeploy(Config, ProjectFile, InProjectName, InProjectDirectory.FullName, Executable, InEngineDir.FullName, bForDistribution, CookFlavor, bIsDataDeploy, bCreateStubIPA, Receipt);
 		}
 
 		/// <summary>
diff --git a/Engine/Source/Programs/UnrealBuildTool/ProjectFiles/Xcode/XcodeProject.cs b/Engine/Source/Programs/UnrealBuildTool/ProjectFiles/Xcode/XcodeProject.cs
--- a/Engine/Source/Programs/UnrealBuildTool/ProjectFiles/Xcode/XcodeProject.cs
+++ b/Engine/Source/Programs/UnrealBuildTool/ProjectFiles/Xcode/XcodeProject.cs
@@ -491,26 +491,22 @@ namespace UnrealBuildTool.XcodeProjectXcconfig
 										if (BuildConfigs.Where(Config => Config.DisplayName == ConfigName).ToList().Count == 0)
 										{
 											string TargetName = ProjectTarget.TargetFilePath.GetFileNameWithoutAnyExtensions();
+											// Get the .uproject directory
+											DirectoryReference? UProjectDirectory = DirectoryReference.FromFile(ProjectTarget.UnrealProjectFilePath);
 
 											// Get the output directory
-											DirectoryReference RootDirectory = Unreal.EngineDirectory;
-											// Unique and Monolithic both need to use the target directory not the engine directory
-											if (ProjectTarget.TargetRules.Type != TargetType.Program && (bShouldCompileMonolithic || ProjectTarget.TargetRules.BuildEnvironment == TargetBuildEnvironment.Unique))
+											DirectoryReference RootDirectory;
+											if (UProjectDirectory != null && 
+												(bShouldCompileMonolithic || ProjectTarget.TargetRules.BuildEnvironment == TargetBuildEnvironment.Unique) && 
+												ProjectTarget.TargetRules.File!.IsUnderDirectory(UProjectDirectory))
 											{
-												if (ProjectTarget.UnrealProjectFilePath != null)
-												{
-													RootDirectory = ProjectTarget.UnrealProjectFilePath.Directory;
-												}
+												RootDirectory = UEBuildTarget.GetOutputDirectoryForExecutable(UProjectDirectory, ProjectTarget.TargetRules.File!);
 											}
-
-											if (ProjectTarget.TargetRules.Type == TargetType.Program && ProjectTarget.UnrealProjectFilePath != null)
+											else
 											{
-												RootDirectory = ProjectTarget.UnrealProjectFilePath.Directory;
+												RootDirectory = UEBuildTarget.GetOutputDirectoryForExecutable(Unreal.EngineDirectory, ProjectTarget.TargetRules.File!);
 											}
 
-											// Get the output directory
-											DirectoryReference OutputDirectory = DirectoryReference.Combine(RootDirectory, "Binaries");
-
 											string ExeName = TargetName;
 											if (!bShouldCompileMonolithic && ProjectTarget.TargetRules.Type != TargetType.Program)
 											{
@@ -525,19 +521,31 @@ namespace UnrealBuildTool.XcodeProjectXcconfig
 												}
 											}
 
+											// Get the output directory
+											DirectoryReference OutputDirectory = DirectoryReference.Combine(RootDirectory, "Binaries");
+											DirectoryReference MacBinaryDir = DirectoryReference.Combine(OutputDirectory, "Mac");
+											DirectoryReference IOSBinaryDir = DirectoryReference.Combine(OutputDirectory, "IOS");
+											DirectoryReference TVOSBinaryDir = DirectoryReference.Combine(OutputDirectory, "TVOS");
+											if (!string.IsNullOrEmpty(ProjectTarget.TargetRules.ExeBinariesSubFolder))
+											{
+												MacBinaryDir = DirectoryReference.Combine(MacBinaryDir, ProjectTarget.TargetRules.ExeBinariesSubFolder);
+												IOSBinaryDir = DirectoryReference.Combine(IOSBinaryDir, ProjectTarget.TargetRules.ExeBinariesSubFolder);
+												TVOSBinaryDir = DirectoryReference.Combine(TVOSBinaryDir, ProjectTarget.TargetRules.ExeBinariesSubFolder);
+											}
+
 											if (BuildPlatform.Platform == UnrealTargetPlatform.Mac)
 											{
 												string MacExecutableName = MakeExecutableFileName(ExeName, UnrealTargetPlatform.Mac, Configuration, ProjectTarget.TargetRules.Architecture, ProjectTarget.TargetRules.UndecoratedConfiguration);
 												string IOSExecutableName = MacExecutableName.Replace("-Mac-", "-IOS-");
 												string TVOSExecutableName = MacExecutableName.Replace("-Mac-", "-TVOS-");
-												BuildConfigs.Add(new UnrealBuildConfig(ConfigName, TargetName, FileReference.Combine(OutputDirectory, "Mac", MacExecutableName), FileReference.Combine(OutputDirectory, "IOS", IOSExecutableName), FileReference.Combine(OutputDirectory, "TVOS", TVOSExecutableName), ProjectTarget, Configuration));
+												BuildConfigs.Add(new UnrealBuildConfig(ConfigName, TargetName, FileReference.Combine(MacBinaryDir, MacExecutableName), FileReference.Combine(IOSBinaryDir, IOSExecutableName), FileReference.Combine(TVOSBinaryDir, TVOSExecutableName), ProjectTarget, Configuration));
 											}
 											else if (BuildPlatform.Platform == UnrealTargetPlatform.IOS || BuildPlatform.Platform == UnrealTargetPlatform.TVOS)
 											{
 												string IOSExecutableName = MakeExecutableFileName(ExeName, UnrealTargetPlatform.IOS, Configuration, ProjectTarget.TargetRules.Architecture, ProjectTarget.TargetRules.UndecoratedConfiguration);
 												string TVOSExecutableName = IOSExecutableName.Replace("-IOS-", "-TVOS-");
 												//string MacExecutableName = IOSExecutableName.Replace("-IOS-", "-Mac-");
-												BuildConfigs.Add(new UnrealBuildConfig(ConfigName, TargetName, FileReference.Combine(OutputDirectory, "Mac", IOSExecutableName), FileReference.Combine(OutputDirectory, "IOS", IOSExecutableName), FileReference.Combine(OutputDirectory, "TVOS", TVOSExecutableName), ProjectTarget, Configuration));
+												BuildConfigs.Add(new UnrealBuildConfig(ConfigName, TargetName, FileReference.Combine(MacBinaryDir, IOSExecutableName), FileReference.Combine(IOSBinaryDir, IOSExecutableName), FileReference.Combine(TVOSBinaryDir, TVOSExecutableName), ProjectTarget, Configuration));
 											}
 										}
 									}
@@ -1484,7 +1492,7 @@ namespace UnrealBuildTool.XcodeProjectXcconfig
 				SupportedDevices.Add(UnrealData.IOSProjectSettings!.RuntimeDevices);
 
 				Xcconfig.AppendLine($"SDKROOT{SettingFilter} = iphoneos");
-				Xcconfig.AppendLine($"CONFIGURATION_BUILD_DIR{SettingFilter} = {ProjectRootDir}/Binaries/IOS/Payload");
+				Xcconfig.AppendLine($"CONFIGURATION_BUILD_DIR{SettingFilter} = {UnrealData.AllConfigs[0].IOSExecutablePath!.Directory}/Payload");
 				Xcconfig.AppendLine($"IPHONEOS_DEPLOYMENT_TARGET = {UnrealData.IOSProjectSettings.RuntimeVersion}");
 
 				if (UnrealData.PlistMode == PlistMode.UpdateTemplate)
@@ -1545,7 +1553,7 @@ namespace UnrealBuildTool.XcodeProjectXcconfig
 				SupportedDevices.Add(UnrealData.TVOSProjectSettings!.RuntimeDevices);
 
 				Xcconfig.AppendLine($"SDKROOT{SettingFilter} = appletvos");
-				Xcconfig.AppendLine($"CONFIGURATION_BUILD_DIR{SettingFilter} = {ProjectRootDir}/Binaries/TVOS/Payload");
+				Xcconfig.AppendLine($"CONFIGURATION_BUILD_DIR{SettingFilter} = {UnrealData.AllConfigs[0].TVOSExecutablePath!.Directory}/Payload");
 				Xcconfig.AppendLine($"TVOS_DEPLOYMENT_TARGET = {UnrealData.TVOSProjectSettings.RuntimeVersion}");
 
 				if (UnrealData.PlistMode == PlistMode.UpdateTemplate)
diff --git a/Engine/Source/Programs/UnrealBuildTool/ProjectFiles/Xcode/XcodeProjectLegacy.cs b/Engine/Source/Programs/UnrealBuildTool/ProjectFiles/Xcode/XcodeProjectLegacy.cs
--- a/Engine/Source/Programs/UnrealBuildTool/ProjectFiles/Xcode/XcodeProjectLegacy.cs
+++ b/Engine/Source/Programs/UnrealBuildTool/ProjectFiles/Xcode/XcodeProjectLegacy.cs
@@ -707,45 +707,43 @@ namespace UnrealBuildTool.XcodeProjectLegacy
 		{
 			StringBuilder FrameworkScript = new StringBuilder();
 
-			// nothing to do without a project
-			if (UProjectPath == null)
+			if (UProjectPath != null)
 			{
-				return;
-			}
 
-			// @todo: look also in Project/Build/Frameworks directory!
-			ProjectDescriptor Project = ProjectDescriptor.FromFile(UProjectPath);
-			List<PluginInfo> AvailablePlugins = Plugins.ReadAvailablePlugins(Unreal.EngineDirectory, DirectoryReference.FromFile(UProjectPath), Project.AdditionalPluginDirectories);
+				// @todo: look also in Project/Build/Frameworks directory!
+				ProjectDescriptor Project = ProjectDescriptor.FromFile(UProjectPath);
+				List<PluginInfo> AvailablePlugins = Plugins.ReadAvailablePlugins(Unreal.EngineDirectory, DirectoryReference.FromFile(UProjectPath), Project.AdditionalPluginDirectories);
 
-			// look in each plugin for frameworks
-			// @todo: Cache this kind of things since every target will re-do this work!
-			foreach (PluginInfo PI in AvailablePlugins)
-			{
-				if (!Plugins.IsPluginEnabledForTarget(PI, Project, UnrealTargetPlatform.IOS, UnrealTargetConfiguration.Development, TargetRules.TargetType.Game))
+				// look in each plugin for frameworks
+				// @todo: Cache this kind of things since every target will re-do this work!
+				foreach (PluginInfo PI in AvailablePlugins)
 				{
-					continue;
-				}
+					if (!Plugins.IsPluginEnabledForTarget(PI, Project, UnrealTargetPlatform.IOS, UnrealTargetConfiguration.Development, TargetRules.TargetType.Game))
+					{
+						continue;
+					}
 
-				// for now, we copy and code sign all *.framework.zip, even if the have no code (non-code frameworks are assumed to be *.embeddedframework.zip
-				DirectoryReference FrameworkDir = DirectoryReference.Combine(PI.Directory, "Source/Frameworks");
-				if (!DirectoryReference.Exists(FrameworkDir))
-				{
-					FrameworkDir = DirectoryReference.Combine(PI.Directory, "Frameworks");
-				}
-				if (DirectoryReference.Exists(FrameworkDir))
-				{
-					// look at each zip
-					foreach (FileInfo FI in new System.IO.DirectoryInfo(FrameworkDir.FullName).EnumerateFiles("*.framework.zip"))
+					// for now, we copy and code sign all *.framework.zip, even if the have no code (non-code frameworks are assumed to be *.embeddedframework.zip
+					DirectoryReference FrameworkDir = DirectoryReference.Combine(PI.Directory, "Source/Frameworks");
+					if (!DirectoryReference.Exists(FrameworkDir))
 					{
-						//string Guid = XcodeProjectFileGenerator.MakeXcodeGuid();
-						//string RefGuid = XcodeProjectFileGenerator.MakeXcodeGuid();
+						FrameworkDir = DirectoryReference.Combine(PI.Directory, "Frameworks");
+					}
+					if (DirectoryReference.Exists(FrameworkDir))
+					{
+						// look at each zip
+						foreach (FileInfo FI in new System.IO.DirectoryInfo(FrameworkDir.FullName).EnumerateFiles("*.framework.zip"))
+						{
+							//string Guid = XcodeProjectFileGenerator.MakeXcodeGuid();
+							//string RefGuid = XcodeProjectFileGenerator.MakeXcodeGuid();
 
-						// for FI of foo.framework.zip, this will give us foo.framework
-						//string Framework = Path.GetFileNameWithoutExtension(FI.FullName);
+							// for FI of foo.framework.zip, this will give us foo.framework
+							//string Framework = Path.GetFileNameWithoutExtension(FI.FullName);
 
-						// unzip the framework right into the .app
-						FrameworkScript.AppendFormat("\\techo Unzipping {0}...\\n", FI.FullName);
-						FrameworkScript.AppendFormat("\\tunzip -o -q {0} -d ${{FRAMEWORK_DIR}} -x \\\"__MACOSX/*\\\" \\\"*/.DS_Store\\\"\\n", FI.FullName);
+							// unzip the framework right into the .app
+							FrameworkScript.AppendFormat("\\techo Unzipping {0}...\\n", FI.FullName);
+							FrameworkScript.AppendFormat("\\tunzip -o -q {0} -d ${{FRAMEWORK_DIR}} -x \\\"__MACOSX/*\\\" \\\"*/.DS_Store\\\"\\n", FI.FullName);
+						}
 					}
 				}
 			}
@@ -1099,38 +1097,13 @@ namespace UnrealBuildTool.XcodeProjectLegacy
 					Content.Append("\t\t\t\t\"PRODUCT_BUNDLE_IDENTIFIER[sdk=macosx*]\" = " + MacBundleIdentifier + ";" + ProjectFileGenerator.NewLine);
 				}
 
-				if (bIsUnrealGame || bIsUnrealClient)
+				if (IOSRunTimeVersion != null)
 				{
-					if (IOSRunTimeVersion != null)
-					{
-						Content.Append("\t\t\t\t\"CONFIGURATION_BUILD_DIR[sdk=iphoneos*]\" = \"" + UEDir + "/Engine/Binaries/IOS" + BinariesSubDir + "\";" + ProjectFileGenerator.NewLine);
-					}
-					if (TVOSRunTimeVersion != null)
-					{
-						Content.Append("\t\t\t\t\"CONFIGURATION_BUILD_DIR[sdk=appletvos*]\" = \"" + UEDir + "/Engine/Binaries/TVOS" + BinariesSubDir + "\";" + ProjectFileGenerator.NewLine);
-					}
+					Content.Append("\t\t\t\t\"CONFIGURATION_BUILD_DIR[sdk=iphoneos*]\" = \"" + Config.IOSExecutablePath!.Directory!.FullName + BinariesSubDir + "\";" + ProjectFileGenerator.NewLine);
 				}
-				else if (ProjectFile != null)
-				{
-					if (IOSRunTimeVersion != null)
-					{
-						Content.Append("\t\t\t\t\"CONFIGURATION_BUILD_DIR[sdk=iphoneos*]\" = \"" + GamePath + "/Binaries/IOS" + BinariesSubDir + "\";" + ProjectFileGenerator.NewLine);
-					}
-					if (TVOSRunTimeVersion != null)
-					{
-						Content.Append("\t\t\t\t\"CONFIGURATION_BUILD_DIR[sdk=appletvos*]\" = \"" + GamePath + "/Binaries/TVOS" + BinariesSubDir + "\";" + ProjectFileGenerator.NewLine);
-					}
-				}
-				else
+				if (TVOSRunTimeVersion != null)
 				{
-					if (IOSRunTimeVersion != null)
-					{
-						Content.Append("\t\t\t\t\"CONFIGURATION_BUILD_DIR[sdk=iphoneos*]\" = \"" + UEDir + "/Engine/Binaries/IOS" + BinariesSubDir + "\";" + ProjectFileGenerator.NewLine);
-					}
-					if (TVOSRunTimeVersion != null)
-					{
-						Content.Append("\t\t\t\t\"CONFIGURATION_BUILD_DIR[sdk=appletvos*]\" = \"" + UEDir + "/Engine/Binaries/TVOS" + BinariesSubDir + "\";" + ProjectFileGenerator.NewLine);
-					}
+					Content.Append("\t\t\t\t\"CONFIGURATION_BUILD_DIR[sdk=appletvos*]\" = \"" + Config.TVOSExecutablePath!.Directory!.FullName + BinariesSubDir +	"\";" + ProjectFileGenerator.NewLine);
 				}
 		}
 
@@ -1482,26 +1455,22 @@ namespace UnrealBuildTool.XcodeProjectLegacy
 										if (BuildConfigs.Where(Config => Config.DisplayName == ConfigName).ToList().Count == 0)
 										{
 											string TargetName = ProjectTarget.TargetFilePath.GetFileNameWithoutAnyExtensions();
+											// Get the .uproject directory
+											DirectoryReference? UProjectDirectory = DirectoryReference.FromFile(ProjectTarget.UnrealProjectFilePath);
 
 											// Get the output directory
-											DirectoryReference RootDirectory = Unreal.EngineDirectory;
-											// Unique and Monolithic both need to use the target directory not the engine directory
-											if (ProjectTarget.TargetRules.Type != TargetType.Program && (bShouldCompileMonolithic || ProjectTarget.TargetRules.BuildEnvironment == TargetBuildEnvironment.Unique))
+											DirectoryReference RootDirectory;
+											if (UProjectDirectory != null &&
+												(bShouldCompileMonolithic || ProjectTarget.TargetRules.BuildEnvironment == TargetBuildEnvironment.Unique) &&
+												ProjectTarget.TargetRules.File!.IsUnderDirectory(UProjectDirectory))
 											{
-												if (ProjectTarget.UnrealProjectFilePath != null)
-												{
-													RootDirectory = ProjectTarget.UnrealProjectFilePath.Directory;
-												}
+												RootDirectory = UEBuildTarget.GetOutputDirectoryForExecutable(UProjectDirectory, ProjectTarget.TargetRules.File!);
 											}
-
-											if (ProjectTarget.TargetRules.Type == TargetType.Program && ProjectTarget.UnrealProjectFilePath != null)
+											else
 											{
-												RootDirectory = ProjectTarget.UnrealProjectFilePath.Directory;
+												RootDirectory = UEBuildTarget.GetOutputDirectoryForExecutable(Unreal.EngineDirectory, ProjectTarget.TargetRules.File!);
 											}
 
-											// Get the output directory
-											DirectoryReference OutputDirectory = DirectoryReference.Combine(RootDirectory, "Binaries");
-
 											string ExeName = TargetName;
 											if (!bShouldCompileMonolithic && ProjectTarget.TargetRules.Type != TargetType.Program)
 											{
@@ -1516,19 +1485,31 @@ namespace UnrealBuildTool.XcodeProjectLegacy
 												}
 											}
 
+											// Get the output directory
+											DirectoryReference OutputDirectory = DirectoryReference.Combine(RootDirectory, "Binaries");
+											DirectoryReference MacBinaryDir = DirectoryReference.Combine(OutputDirectory, "Mac");
+											DirectoryReference IOSBinaryDir = DirectoryReference.Combine(OutputDirectory, "IOS");
+											DirectoryReference TVOSBinaryDir = DirectoryReference.Combine(OutputDirectory, "TVOS");
+											if (!string.IsNullOrEmpty(ProjectTarget.TargetRules.ExeBinariesSubFolder))
+											{
+												MacBinaryDir = DirectoryReference.Combine(MacBinaryDir, ProjectTarget.TargetRules.ExeBinariesSubFolder);
+												IOSBinaryDir = DirectoryReference.Combine(IOSBinaryDir, ProjectTarget.TargetRules.ExeBinariesSubFolder);
+												TVOSBinaryDir = DirectoryReference.Combine(TVOSBinaryDir, ProjectTarget.TargetRules.ExeBinariesSubFolder);
+											}
+
 											if (BuildPlatform.Platform == UnrealTargetPlatform.Mac)
 											{
 												string MacExecutableName = MakeExecutableFileName(ExeName, UnrealTargetPlatform.Mac, Configuration, ProjectTarget.TargetRules.Architecture, ProjectTarget.TargetRules.UndecoratedConfiguration);
 												string IOSExecutableName = MacExecutableName.Replace("-Mac-", "-IOS-");
 												string TVOSExecutableName = MacExecutableName.Replace("-Mac-", "-TVOS-");
-												BuildConfigs.Add(new XcodeBuildConfig(ConfigName, TargetName, FileReference.Combine(OutputDirectory, "Mac", MacExecutableName), FileReference.Combine(OutputDirectory, "IOS", IOSExecutableName), FileReference.Combine(OutputDirectory, "TVOS", TVOSExecutableName), ProjectTarget, Configuration));
+												BuildConfigs.Add(new XcodeBuildConfig(ConfigName, TargetName, FileReference.Combine(MacBinaryDir, MacExecutableName), FileReference.Combine(IOSBinaryDir, IOSExecutableName), FileReference.Combine(TVOSBinaryDir, TVOSExecutableName), ProjectTarget, Configuration));
 											}
 											else if (BuildPlatform.Platform == UnrealTargetPlatform.IOS || BuildPlatform.Platform == UnrealTargetPlatform.TVOS)
 											{
 												string IOSExecutableName = MakeExecutableFileName(ExeName, UnrealTargetPlatform.IOS, Configuration, ProjectTarget.TargetRules.Architecture, ProjectTarget.TargetRules.UndecoratedConfiguration);
 												string TVOSExecutableName = IOSExecutableName.Replace("-IOS-", "-TVOS-");
 												//string MacExecutableName = IOSExecutableName.Replace("-IOS-", "-Mac-");
-												BuildConfigs.Add(new XcodeBuildConfig(ConfigName, TargetName, FileReference.Combine(OutputDirectory, "Mac", IOSExecutableName), FileReference.Combine(OutputDirectory, "IOS", IOSExecutableName), FileReference.Combine(OutputDirectory, "TVOS", TVOSExecutableName), ProjectTarget, Configuration));
+												BuildConfigs.Add(new XcodeBuildConfig(ConfigName, TargetName, FileReference.Combine(MacBinaryDir, IOSExecutableName), FileReference.Combine(IOSBinaryDir, IOSExecutableName), FileReference.Combine(TVOSBinaryDir, TVOSExecutableName), ProjectTarget, Configuration));
 											}
 										}
 									}
