UnrealBuildTool: Disallow propagating more environment variables to subprocesses that could affect the msvc compiler

#rnx

diff --git a/Engine/Source/Programs/UnrealBuildTool/Platform/Windows/VCToolChain.cs b/Engine/Source/Programs/UnrealBuildTool/Platform/Windows/VCToolChain.cs
--- a/Engine/Source/Programs/UnrealBuildTool/Platform/Windows/VCToolChain.cs
+++ b/Engine/Source/Programs/UnrealBuildTool/Platform/Windows/VCToolChain.cs
@@ -71,6 +71,13 @@ namespace UnrealBuildTool
 			// add to improve Visual Studio memory usage, but we don't actually need it to set when invoking the compiler. Doing so results in it being converted
 			// into /I arguments by the CL driver, which results in errors due to the command line not fitting into the PDB debug record.
 			Environment.SetEnvironmentVariable("INCLUDE", null);
+
+			// Don't allow the CL or _CL_ environment variable to propagate.
+			Environment.SetEnvironmentVariable("CL", null);
+			Environment.SetEnvironmentVariable("_CL_", null);
+
+			// Don't allow the LIBPATH environment variable to propagate.
+			Environment.SetEnvironmentVariable("LIBPATH", null);
 		}
 
 		/// <summary>
