Add warning log to Xcode clean

diff --git a/Engine/Source/Programs/UnrealBuildTool/Modes/CleanMode.cs b/Engine/Source/Programs/UnrealBuildTool/Modes/CleanMode.cs
--- a/Engine/Source/Programs/UnrealBuildTool/Modes/CleanMode.cs
+++ b/Engine/Source/Programs/UnrealBuildTool/Modes/CleanMode.cs
@@ -67,6 +67,15 @@ namespace UnrealBuildTool
 				throw new BuildException("No targets specified to clean");
 			}
 
+			// Print out warning in Xcode to stop user from trying to make build folder deletable
+			if (Environment.GetEnvironmentVariable("UE_BUILD_FROM_XCODE") == "1" && Unreal.IsEngineInstalled())
+			{
+				Logger.LogInformation("~~~~~~~~~~~~~~~~~~~~~~~~~~~~~~~~~~~~~~~~~~~~~~~~~~~~~~~~~~~~~~~~~~~~~~~~~~~~~~~~~~~~~~~~~~~~");
+				Logger.LogInformation("NOTICE: Please disregard the Xcode prepare clean failed message, we do NOT want Engine/Binaries/Mac to be deletable.");
+				Logger.LogInformation("NOTICE: Cleaning UnrealEditor.app will cause UE fail to launch, you will then need to repair it from Epic Games Launcher.");
+				Logger.LogInformation("~~~~~~~~~~~~~~~~~~~~~~~~~~~~~~~~~~~~~~~~~~~~~~~~~~~~~~~~~~~~~~~~~~~~~~~~~~~~~~~~~~~~~~~~~~~~");
+			}
+
 			// Also add implicit descriptors for cleaning UnrealBuildTool
 			if (!BuildConfiguration.bDoNotBuildUHT)
 			{
