- Added some logging to GenerateProjectFiles, when writing multiple project formats, to indicate what is happening (like why "Generating data for project indexing" is printed twice)
- Added VSMac as a default project format for Mac (so its easier to manage C# code)

diff --git a/Engine/Source/Programs/UnrealBuildTool/Modes/GenerateProjectFilesMode.cs b/Engine/Source/Programs/UnrealBuildTool/Modes/GenerateProjectFilesMode.cs
--- a/Engine/Source/Programs/UnrealBuildTool/Modes/GenerateProjectFilesMode.cs
+++ b/Engine/Source/Programs/UnrealBuildTool/Modes/GenerateProjectFilesMode.cs
@@ -165,7 +165,7 @@ namespace UnrealBuildTool
 			Logger.LogDebug("");
 
 			// Create each project generator and run it
-			List<ProjectFileGenerator> Generators = new List<ProjectFileGenerator>();
+			Dictionary<ProjectFileFormat, ProjectFileGenerator> Generators = new();
 			foreach (ProjectFileFormat ProjectFileFormat in ProjectFileFormats.Distinct())
 			{
 				ProjectFileGenerator Generator;
@@ -221,7 +221,7 @@ namespace UnrealBuildTool
 					default:
 						throw new BuildException("Unhandled project file type '{0}", ProjectFileFormat);
 				}
-				Generators.Add(Generator);
+				Generators[ProjectFileFormat] = Generator;
 			}
 
 			// Check there are no superfluous command line arguments
@@ -230,11 +230,14 @@ namespace UnrealBuildTool
 
 			// Now generate project files
 			ProjectFileGenerator.bGenerateProjectFiles = true;
-			foreach(ProjectFileGenerator Generator in Generators)
+			foreach(KeyValuePair< ProjectFileFormat, ProjectFileGenerator> Pair in Generators)
 			{
-				ProjectFileGenerator.Current = Generator;
-				Arguments.ApplyTo(Generator);
-				bool bGenerateSuccess = Generator.GenerateProjectFiles(PlatformProjectGenerators, Arguments.GetRawArray(), Logger);
+				Logger.LogInformation("");
+				Logger.LogInformation($"Generating {Pair.Key} project files:");
+
+				ProjectFileGenerator.Current = Pair.Value;
+				Arguments.ApplyTo(Pair.Value);
+				bool bGenerateSuccess = Pair.Value.GenerateProjectFiles(PlatformProjectGenerators, Arguments.GetRawArray(), Logger);
 				ProjectFileGenerator.Current = null;
 
 				if (!bGenerateSuccess)
diff --git a/Engine/Source/Programs/UnrealBuildTool/System/BuildHostPlatform.cs b/Engine/Source/Programs/UnrealBuildTool/System/BuildHostPlatform.cs
--- a/Engine/Source/Programs/UnrealBuildTool/System/BuildHostPlatform.cs
+++ b/Engine/Source/Programs/UnrealBuildTool/System/BuildHostPlatform.cs
@@ -363,6 +363,7 @@ namespace UnrealBuildTool
 		internal override IEnumerable<ProjectFileFormat> GetDefaultProjectFileFormats()
 		{
 			yield return ProjectFileFormat.XCode;
+			yield return ProjectFileFormat.VisualStudioMac;
 #if __VPROJECT_AVAILABLE__
 			yield return ProjectFileFormat.VProject;
 #endif
