- Made two modes for modern Xcode projects: OneWorkspacePerPlatform and OneRunTargetPerPlatform. The first makes completely separate .xcworkspaces, projects, and targets for each platform (Mac, IOS, TVOS). The second makes one workspace, and one project for each platform, but the platform will have three run targets (but one build and one index target) - useful for single game workspaces. There are a couple nasty things, like the Project has the Editor configs, but the IOS target doesn't want them, but they show up anyway in dropdowns because of the Project config.
- If there are no modern projects, then the mode is ignored, and it's the same old mode. With a mix of modern and legacy, each platform workspace will point to the same legacy projects (not worth retooling Legacy since it's becoming deprecated)
- Added Launch Storyboard support (it will look for uncompiled, compiled, project, and engine fallback locations). Compiled is useful or old storyboards, and for Windows-based (assuming that will work with Modern)
- Other minor things like metadata discovery, etc

diff --git a/Engine/Build/BatchFiles/Mac/XcodeBuild.sh b/Engine/Build/BatchFiles/Mac/XcodeBuild.sh
--- a/Engine/Build/BatchFiles/Mac/XcodeBuild.sh
+++ b/Engine/Build/BatchFiles/Mac/XcodeBuild.sh
@@ -9,25 +9,15 @@ if [ ["$UE_SKIP_UBT"] == ["1"] ]; then
 	exit 0
 fi
 
+echo Running Xcodebuild $@
+
 # Setup Environment
 source  Engine/Build/BatchFiles/Mac/SetupEnvironment.sh -dotnet Engine/Build/BatchFiles/Mac
 
-# remove environment variable passed from xcode which also has meaning to dotnet, breaking the build
-unset TARGETNAME
-
-# If this is a source drop of the engine make sure that the UnrealBuildTool is up-to-date
-if [ ! -f Engine/Build/InstalledBuild.txt ]; then
-	dotnet build Engine/Source/Programs/UnrealBuildTool/UnrealBuildTool.csproj -c Development -v quiet
-
-	if [ $? -ne 0 ]; then
-		echo "Failed to build the build tool (UnrealBuildTool)"
-		exit 1
-	fi
-fi
-
-
 #echo "Raw Args: $*"
 
+BUILD_UBT=1
+
 case $1 in 
 	"clean")
 		ACTION="clean"
@@ -36,8 +26,30 @@ case $1 in
 	"install")
 		ACTION="install"
 	;;	
+
+	"-nobuildubt")
+		BUILD_UBT=0
+	;;
 esac
 
+
+# If this is a source drop of the engine make sure that the UnrealBuildTool is up-to-date
+if [ $BUILD_UBT == 1 ]; then
+	# remove environment variable passed from xcode which also has meaning to dotnet, breaking the build
+	unset TARGETNAME
+	
+	if [ ! -f Engine/Build/InstalledBuild.txt ]; then
+		dotnet build Engine/Source/Programs/UnrealBuildTool/UnrealBuildTool.csproj -c Development -v quiet
+
+		if [ $? -ne 0 ]; then
+			echo "Failed to build the build tool (UnrealBuildTool)"
+			exit 1
+		fi
+	fi
+fi
+
+
+
 if [ ["$ACTION"] == [""] ]; then
 	ACTION="build"
 	TARGET=$1
@@ -113,6 +125,9 @@ if [ "$ACTION" == "build" ]; then
 
 elif [ $ACTION == "clean" ]; then
 	AdditionalFlags="-clean"
+	
+elif [ $ACTION == "metadata" ]; then
+	AdditionalFlags="-mode=ExportXcodeMetadata"
 fi
 
 echo Running dotnet Engine/Binaries/DotNET/UnrealBuildTool/UnrealBuildTool.dll $TARGET $PLATFORM $CONFIGURATION "$TRAILINGARGS" $UBT_ARCHFLAG $AdditionalFlags
diff --git a/Engine/Source/Programs/UnrealBuildTool/ProjectFiles/Xcode/XcodeProject.cs b/Engine/Source/Programs/UnrealBuildTool/ProjectFiles/Xcode/XcodeProject.cs
--- a/Engine/Source/Programs/UnrealBuildTool/ProjectFiles/Xcode/XcodeProject.cs
+++ b/Engine/Source/Programs/UnrealBuildTool/ProjectFiles/Xcode/XcodeProject.cs
@@ -89,9 +89,13 @@ namespace UnrealBuildTool.XcodeProjectXcconfig
 		public ProjectTarget? ProjectTarget;
 		public UnrealTargetConfiguration BuildConfig;
 
-		public bool bSupportsMac => UnrealData.bSupportsMac && (ProjectTarget == null || ProjectTarget.SupportedPlatforms.Contains(UnrealTargetPlatform.Mac));
-		public bool bSupportsIOS => UnrealData.bSupportsIOS && (ProjectTarget == null || ProjectTarget.SupportedPlatforms.Contains(UnrealTargetPlatform.IOS));
-		public bool bSupportsTVOS => UnrealData.bSupportsTVOS && (ProjectTarget == null || ProjectTarget.SupportedPlatforms.Contains(UnrealTargetPlatform.TVOS));
+		public bool bSupportsMac => Supports(UnrealTargetPlatform.Mac);
+		public bool bSupportsIOS => Supports(UnrealTargetPlatform.IOS);
+		public bool bSupportsTVOS => Supports(UnrealTargetPlatform.TVOS);
+		public bool Supports(UnrealTargetPlatform? Platform)
+		{
+			return UnrealData.Supports(Platform) && (ProjectTarget == null || Platform == null || ProjectTarget.SupportedPlatforms.Contains((UnrealTargetPlatform)Platform));
+		}
 	};
 
 	class UnrealBatchedFiles
@@ -140,12 +144,72 @@ namespace UnrealBuildTool.XcodeProjectXcconfig
 		}
 	}
 
-	enum PlistMode
+	enum MetadataPlatform
+	{
+		MacEditor,
+		Mac,
+		IOS,
+	}
+
+	enum MetadataMode
 	{
 		UsePremade,
 		UpdateTemplate,
 	}
 
+	class Metadata
+	{
+		public MetadataMode Mode;
+		public FileReference File;
+		public string XcodeProjectRelative;
+
+		//public Metadata(MetadataMode Mode, FileReference File, DirectoryReference XcodeProjectFile)
+		//{
+		//	this.Mode = Mode;
+		//	this.File = File;
+		//	XcodeProjectRelative = File.MakeRelativeTo(XcodeProjectFile.ParentDirectory!);
+		//}
+
+		public Metadata(DirectoryReference ProductDirectory, DirectoryReference XcodeProject, ConfigHierarchy Ini, string IniKey, string StandardLocationSubdir, string? ProjectTemplateSubdir = null)
+		{
+			string? PremadeLocation;
+			Ini.TryGetValue("XcodeConfiguration", IniKey, out PremadeLocation);
+
+			if (!string.IsNullOrEmpty(PremadeLocation))
+			{
+				Mode = MetadataMode.UsePremade;
+				File = FileReference.Combine(ProductDirectory, PremadeLocation);
+			}
+			else
+			{
+				if (File == null)
+				{
+					FileReference StandardProjectLocation = FileReference.Combine(ProductDirectory, StandardLocationSubdir);
+					if (FileReference.Exists(StandardProjectLocation))
+					{
+						Mode = MetadataMode.UsePremade;
+						File = StandardProjectLocation;
+					}
+					else
+					{
+						Mode = MetadataMode.UpdateTemplate;
+						if (ProjectTemplateSubdir != null)
+						{
+							File = FileReference.Combine(ProductDirectory, ProjectTemplateSubdir);
+						}
+						else
+						{
+							File = FileReference.Combine(Unreal.EngineDirectory, StandardLocationSubdir);
+						}
+					}
+				}
+			}
+
+			XcodeProjectRelative = File.MakeRelativeTo(XcodeProject.ParentDirectory!);
+		}
+	}
+
+
 	class UnrealData
 	{
 		public bool bIsAppBundle;
@@ -154,23 +218,16 @@ namespace UnrealBuildTool.XcodeProjectXcconfig
 		public bool bIsMergingProjects = false;
 		public bool bWriteCodeSigningSettings = true;
 
-		public PlistMode PlistMode;
-
+		public Dictionary<MetadataPlatform, Metadata> PlistFiles = new();
+		public Dictionary<MetadataPlatform, Metadata> EntitlementsFiles = new();
 
 		public List<UnrealBuildConfig> AllConfigs = new();
 
-		public List<UnrealExtensionInfo> AllExtensions = new();
-
 		public List<UnrealBatchedFiles> BatchedFiles = new();
  
 		public FileReference? UProjectFileLocation = null;
 		public DirectoryReference XcodeProjectFileLocation;
-		public DirectoryReference ProductDir;
-
-		public FileReference? InfoPlistLocation = null;
-		public bool bInfoPlistWasPremade = false;
-		public FileReference? EntitlementsLocation = null;
-		public bool bEntitlementsWasPremade = false;
+		public DirectoryReference ProductDirectory;
 
 		// settings read from project configs
 		public IOSProjectSettings? IOSProjectSettings;
@@ -213,9 +270,13 @@ namespace UnrealBuildTool.XcodeProjectXcconfig
 
 		private XcodeProjectFile? ProjectFile;
 
-		public static bool bSupportsMac => XcodeProjectFileGenerator.XcodePlatforms.Contains(UnrealTargetPlatform.Mac);
-		public static bool bSupportsIOS => XcodeProjectFileGenerator.XcodePlatforms.Contains(UnrealTargetPlatform.IOS);
-		public static bool bSupportsTVOS => XcodeProjectFileGenerator.XcodePlatforms.Contains(UnrealTargetPlatform.TVOS);
+		public static bool bSupportsMac => Supports(UnrealTargetPlatform.Mac);
+		public static bool bSupportsIOS => Supports(UnrealTargetPlatform.IOS);
+		public static bool bSupportsTVOS => Supports(UnrealTargetPlatform.TVOS);
+		public static bool Supports(UnrealTargetPlatform? Platform)
+		{
+			return Platform == null || XcodeProjectFileGenerator.XcodePlatforms.Contains((UnrealTargetPlatform)Platform);
+		}
 
 
 		public UnrealData(FileReference XcodeProjectFileLocation, bool bIsForDistribution, string BundleID, string AppName)
@@ -223,7 +284,7 @@ namespace UnrealBuildTool.XcodeProjectXcconfig
 			// the .xcodeproj is actually a directory
 			this.XcodeProjectFileLocation = new DirectoryReference(XcodeProjectFileLocation.FullName);
 			// default to engine director, will be fixed in Initialize if needed
-			ProductDir = Unreal.EngineDirectory;
+			ProductDirectory = Unreal.EngineDirectory;
 			ProductName = XcodeProjectFileLocation.GetFileNameWithoutAnyExtensions();
 			if (ProductName == "UE5")
 			{
@@ -263,7 +324,7 @@ namespace UnrealBuildTool.XcodeProjectXcconfig
 			// make sure ProjectDir is something good
 			if (UProjectFileLocation != null)
 			{
-				ProductDir = UProjectFileLocation.Directory;
+				ProductDirectory = UProjectFileLocation.Directory;
 			}
 			else if (ProjectFile.ProjectTargets[0].TargetRules!.Type == TargetType.Program)
 			{
@@ -279,11 +340,13 @@ namespace UnrealBuildTool.XcodeProjectXcconfig
 					// if it exists, we have a ProductDir we can use for plists, icons, etc
 					if (DirectoryReference.Exists(ProgramFinder))
 					{
-						ProductDir = ProgramFinder;
+						ProductDirectory = ProgramFinder;
 					}
 				}
 			}
 
+			InitializeMetadata();
+
 			// Figure out all the desired configurations on the unreal side
 			AllConfigs = GetSupportedBuildConfigs(XcodeProjectFileGenerator.XcodePlatforms, Configurations, Logger);
 			// if we can't find any configs, we will fail to create a project
@@ -292,13 +355,6 @@ namespace UnrealBuildTool.XcodeProjectXcconfig
 				return false;
 			}
 
-			// read setings from the configs, now that we have a project
-			ConfigHierarchy SharedPlatformIni = ConfigCache.ReadHierarchy(ConfigHierarchyType.Engine, UProjectFileLocation?.Directory, UnrealTargetPlatform.Mac);
-			SharedPlatformIni.TryGetValue("XcodeConfiguration", "bUseModernCodeSigning", out bUseAutomaticSigning);
-
-			// @todo read from ini
-			PlistMode = PlistMode.UpdateTemplate;
-
 			// this project makes an app bundle (.app directory instead of a raw executable or dylib) if none of the fings make a non-appbundle
 			bIsAppBundle = !AllConfigs.Any(x => x.ProjectTarget!.TargetRules!.bIsBuildingConsoleApplication || x.ProjectTarget.TargetRules.bShouldCompileAsDLL);
 			bHasEditorConfiguration = AllConfigs.Any(x => x.ProjectTarget!.TargetRules!.Type == TargetType.Editor);
@@ -327,12 +383,55 @@ namespace UnrealBuildTool.XcodeProjectXcconfig
 			return true;
 		}
 
+		private void InitializeMetadata()
+		{
+
+			// read setings from the configs, now that we have a project
+			ConfigHierarchy SharedPlatformIni = ConfigCache.ReadHierarchy(ConfigHierarchyType.Engine, UProjectFileLocation?.Directory, UnrealTargetPlatform.Mac);
+			SharedPlatformIni.TryGetValue("XcodeConfiguration", "bUseModernCodeSigning", out bUseAutomaticSigning);
+
+			if (bSupportsMac)
+			{
+				PlistFiles[MetadataPlatform.MacEditor] = new Metadata(ProductDirectory, XcodeProjectFileLocation, SharedPlatformIni, "PremadeMacEditorPlist", "Build/Mac/Resources/Info-Editor.plist");
+				PlistFiles[MetadataPlatform.Mac] = new Metadata(ProductDirectory, XcodeProjectFileLocation, SharedPlatformIni, "PremadeMacPlist", "Build/Mac/Resources/Info.plist");
+				EntitlementsFiles[MetadataPlatform.MacEditor] = new Metadata(ProductDirectory, XcodeProjectFileLocation, SharedPlatformIni, "PremadeMacEditorPlist", "Build/Mac/Resources/Info-Editor.plist");
+				EntitlementsFiles[MetadataPlatform.Mac] = new Metadata(ProductDirectory, XcodeProjectFileLocation, SharedPlatformIni, "PremadeMacPlist", "Build/Mac/Resources/.plist");
+			}
+			if (bSupportsIOS || bSupportsTVOS)
+			{
+				PlistFiles[MetadataPlatform.IOS] = new Metadata(ProductDirectory, XcodeProjectFileLocation, SharedPlatformIni, "PremadeIOSPlist", "Build/IOS/Resources/Info.plist", "Build/IOS/UBTGenerated/Info.plist");
+				EntitlementsFiles[MetadataPlatform.IOS] = new Metadata(ProductDirectory, XcodeProjectFileLocation, SharedPlatformIni, "PremadeIOSPlist", "Build/IOS/Resources/Info.plist");
+			}
+		}
+
+		public string? FindFile(List<string> Paths, UnrealTargetPlatform Platform, bool bMakeRelative)
+		{
+			foreach (string Entry in Paths)
+			{
+				string FinalPath = Entry.Replace("$(Engine)", Unreal.EngineDirectory.FullName);
+				FinalPath = FinalPath.Replace("$(Project)", ProductDirectory.FullName);
+				FinalPath = FinalPath.Replace("$(Platform)", Platform.ToString());
+
+//				Console.WriteLine($"Looking for {FinalPath}");
+				if (File.Exists(FinalPath) || Directory.Exists(FinalPath))
+				{
+//					Console.WriteLine($"  Found it!");
+					if (bMakeRelative)
+					{
+						FinalPath = new FileReference(FinalPath).MakeRelativeTo(XcodeProjectFileLocation.ParentDirectory!);
+					}
+					return FinalPath;
+				}
+			}
+			return null;
+		}
+
 		public string ProjectOrEnginePath(string SubPath, bool bMakeRelative)
 		{
 			string? FinalPath = null;
 			if (UProjectFileLocation != null)
 			{
-				string PathToCheck = Path.Combine(ProductDir.FullName, SubPath);
+				string PathToCheck = Path.Combine(ProductDirectory.FullName, SubPath);
 				if (File.Exists(PathToCheck) || Directory.Exists(PathToCheck))
 				{
 					FinalPath = PathToCheck;
@@ -708,6 +807,19 @@ namespace UnrealBuildTool.XcodeProjectXcconfig
 			bIsReference = InIsReference;
 		}
 
+		public XcodeFileGroup(XcodeFileGroup Other)
+		{
+			GroupName = Other.GroupName;
+			GroupPath = Other.GroupPath;
+			GroupGuid = Other.GroupGuid;
+			bIsReference = Other.bIsReference;
+			Files.AddRange(Other.Files);
+			foreach (string Key in Other.Children.Keys)
+			{
+				Children.Add(Key, new XcodeFileGroup(Other.Children[Key]));
+			}
+		}
+
 		public string GroupGuid;
 		public string GroupName;
 		public string GroupPath;
@@ -716,32 +828,6 @@ namespace UnrealBuildTool.XcodeProjectXcconfig
 		public bool bIsReference;
 	}
 
-	class UnrealExtensionInfo
-	{
-		public UnrealExtensionInfo(string InName)
-		{
-			Name = InName;
-			TargetDependencyGuid = XcodeProjectFileGenerator.MakeXcodeGuid();
-			TargetProxyGuid = XcodeProjectFileGenerator.MakeXcodeGuid();
-			TargetGuid = XcodeProjectFileGenerator.MakeXcodeGuid();
-			ProductGuid = XcodeProjectFileGenerator.MakeXcodeGuid();
-			ResourceBuildPhaseGuid = XcodeProjectFileGenerator.MakeXcodeGuid();
-			ConfigListGuid = XcodeProjectFileGenerator.MakeXcodeGuid();
-			AllConfigs = new Dictionary<string, UnrealBuildConfig>();
-		}
-
-		public string Name;
-		public string TargetDependencyGuid;
-		public string TargetProxyGuid;
-		public string TargetGuid;
-		public string ProductGuid;
-		public string ResourceBuildPhaseGuid;
-		public string ConfigListGuid;
-		public Dictionary<string, UnrealBuildConfig> AllConfigs;
-
-		public string? ConfigurationContents;
-	}
-
 	class XcconfigFile
 	{
 		public string Name;
@@ -749,10 +835,10 @@ namespace UnrealBuildTool.XcodeProjectXcconfig
 		public FileReference FileRef;
 		internal StringBuilder Text;
 
-		public XcconfigFile(DirectoryReference XcodeProjectDirectory, string ConfigName)
+		public XcconfigFile(DirectoryReference XcodeProjectDirectory, UnrealTargetPlatform? Platform, string ConfigName)
 		{
 			Name = ConfigName;
-			FileRef = FileReference.Combine(XcodeProjectDirectory, "Xcconfigs", $"{ConfigName}.xcconfig");
+			FileRef = FileReference.Combine(XcodeProjectDirectory, "Xcconfigs" + (Platform == null ? "" : Platform.ToString()), $"{ConfigName}.xcconfig");
 			Guid = XcodeProjectFileGenerator.MakeXcodeGuid();
 			Text = new StringBuilder();
 		}
@@ -820,8 +906,9 @@ namespace UnrealBuildTool.XcodeProjectXcconfig
 
 		public void CreateXcconfigFile(XcodeProject Project, string Name)
 		{
-			Xcconfig = new XcconfigFile(Project.UnrealData.XcodeProjectFileLocation.ParentDirectory!, Name);
-			Project.FileCollection.AddFileReference(Xcconfig.Guid, $"Xcconfigs/{Xcconfig.FileRef.GetFileName()}", "test.xcconfig", "\"<group>\"", "Xcconfigs");
+			DirectoryReference XcodeProjectDirectory = Project.UnrealData.XcodeProjectFileLocation.ParentDirectory!;
+			Xcconfig = new XcconfigFile(XcodeProjectDirectory, Project.Platform, Name);
+			Project.FileCollection.AddFileReference(Xcconfig.Guid, Xcconfig.FileRef.MakeRelativeTo(XcodeProjectDirectory), "explicitFileType", "test.xcconfig", "\"<group>\"", "Xcconfigs");
 		}
 
 		public virtual void WriteXcconfigFile()
@@ -835,15 +922,25 @@ namespace UnrealBuildTool.XcodeProjectXcconfig
 		/// </summary>
 		/// <param name="Content"></param>
 		/// <param name="Node"></param>
-		public static void WriteNodeAndReferences(StringBuilder Content, XcodeProjectNode Node)
+		/// <param name="WrittenNodes"></param>
+		public static void WriteNodeAndReferences(StringBuilder Content, XcodeProjectNode Node, HashSet<XcodeProjectNode>? WrittenNodes = null)
 		{
+			if (WrittenNodes == null)
+			{
+				WrittenNodes = new();
+			}
+
 			// write the node into the xcode project file
 			Node.Write(Content);
 			Node.WriteXcconfigFile();
 
 			foreach (XcodeProjectNode Reference in Node.References)
 			{
-				WriteNodeAndReferences(Content, Reference);
+				if (!WrittenNodes.Contains(Reference))
+				{
+					WrittenNodes.Add(Reference);
+					WriteNodeAndReferences(Content, Reference, WrittenNodes);
+				}
 			}	
 		}
 	}
@@ -996,18 +1093,25 @@ namespace UnrealBuildTool.XcodeProjectXcconfig
 	}
 	
 	class XcodeBuildConfigList : XcodeProjectNode
-	{
+	{ 
 		public string Guid = XcodeProjectFileGenerator.MakeXcodeGuid();
 		public string TargetName;
 
 		public List<XcodeBuildConfig> BuildConfigs = new();
 
-		public bool bSupportsMac => BuildConfigs.Any(x => x.Info.bSupportsMac);
-		public bool bSupportsIOS => BuildConfigs.Any(x => x.Info.bSupportsIOS);
-		public bool bSupportsTVOS => BuildConfigs.Any(x => x.Info.bSupportsTVOS);
+		public bool bSupportsMac => Supports(UnrealTargetPlatform.Mac);
+		public bool bSupportsIOS => Supports(UnrealTargetPlatform.IOS);
+		public bool bSupportsTVOS => Supports(UnrealTargetPlatform.TVOS);
+		public bool Supports(UnrealTargetPlatform? Platform)
+		{
+			return this.Platform == Platform || (this.Platform == null && BuildConfigs.Any(x => x.Info.Supports(Platform)));
+		}
 
-		public XcodeBuildConfigList(string TargetName, List<UnrealBuildConfig> BuildConfigInfos)
+		private UnrealTargetPlatform? Platform;
+		public XcodeBuildConfigList(UnrealTargetPlatform? Platform, string TargetName, List<UnrealBuildConfig> BuildConfigInfos)
 		{
+			this.Platform = Platform;
+
 			if (BuildConfigInfos.Count == 0)
 			{
 				throw new BuildException("Created a XcodeBuildConfigList with no BuildConfigs. This likely means a target was created too early");
@@ -1016,7 +1120,10 @@ namespace UnrealBuildTool.XcodeProjectXcconfig
 			this.TargetName = TargetName;
 
 			// create build config objects for each info passed in, and them as references
-			BuildConfigs = BuildConfigInfos.Select(x => new XcodeBuildConfig(x)).ToList();
+			IEnumerable<XcodeBuildConfig> Configs = BuildConfigInfos.Select(x => new XcodeBuildConfig(x));
+			// filter out configs that dont match a platform if we are single-platform mode
+			Configs = Configs.Where(x => Platform == null || x.Info.Supports((UnrealTargetPlatform)Platform));
+			BuildConfigs = Configs.ToList();
 			References.AddRange(BuildConfigs);
 		}
 
@@ -1051,7 +1158,7 @@ namespace UnrealBuildTool.XcodeProjectXcconfig
 
 		// Guid for this target
 		public string Guid = XcodeProjectFileGenerator.MakeXcodeGuid();
-		string TargetAppGuid = XcodeProjectFileGenerator.MakeXcodeGuid();
+//		string TargetAppGuid = XcodeProjectFileGenerator.MakeXcodeGuid();
 
 		// com.apple.product-type.application, etc
 		string ProductType;
@@ -1173,16 +1280,18 @@ namespace UnrealBuildTool.XcodeProjectXcconfig
 	{
 		private string ProductGuid = XcodeProjectFileGenerator.MakeXcodeGuid();
 		private UnrealData UnrealData;
+		public UnrealTargetPlatform Platform;
 
-		public XcodeRunTarget(XcodeProject Project)
-			: base(Project.UnrealData.bIsAppBundle ? XcodeTarget.Type.Run_App : XcodeTarget.Type.Run_Tool, Project.UnrealData)
+		public XcodeRunTarget(XcodeProject Project, UnrealTargetPlatform Platform, XcodeBuildTarget? BuildTarget)
+			: base(Project.UnrealData.bIsAppBundle ? XcodeTarget.Type.Run_App : XcodeTarget.Type.Run_Tool, Project.UnrealData, Platform.ToString())
 		{
-			BuildConfigList = new XcodeBuildConfigList(Name, Project.UnrealData.AllConfigs);
+			this.Platform = Platform;
+			BuildConfigList = new XcodeBuildConfigList(Platform, Name, Project.UnrealData.AllConfigs);
 			References.Add(BuildConfigList);
 
 			// add the Product item to the project to be visible in left pane
 			UnrealData = Project.UnrealData;
-			Project.FileCollection.AddFileReference(ProductGuid, UnrealData.ProductName, Project.UnrealData.bIsAppBundle ? "wrapper.application" : "\"compiled.mach-o.executable\"", "BUILT_PRODUCTS_DIR", "Products");
+			Project.FileCollection.AddFileReference(ProductGuid, UnrealData.ProductName, "explicitFileType", Project.UnrealData.bIsAppBundle ? "wrapper.application" : "\"compiled.mach-o.executable\"", "BUILT_PRODUCTS_DIR", "Products");
 
 			if (Project.UnrealData.bIsAppBundle)
 			{
@@ -1190,30 +1299,44 @@ namespace UnrealBuildTool.XcodeProjectXcconfig
 				BuildPhases.Add(ResourcesBuildPhase);
 				References.Add(ResourcesBuildPhase);
 
+				List<string>StoryboardPaths = new List<string>()
+				{
+					"$(Project)/Build/$(Platform)/Resources/Interface/LaunchScreen.storyboardc",
+					"$(Project)/Build/$(Platform)/Resources/Interface/LaunchScreen.storyboard",
+					"$(Project)/Build/Apple/Resources/Interface/LaunchScreen.storyboardc",
+					"$(Project)/Build/Apple/Resources/Interface/LaunchScreen.storyboard",
+					"$(Engine)/Build/$(Platform)/Resources/Interface/LaunchScreen.storyboardc",
+					"$(Engine)/Build/$(Platform)/Resources/Interface/LaunchScreen.storyboard",
+					"$(Engine)/Build/Apple/Resources/Interface/LaunchScreen.storyboardc",
+					"$(Engine)/Build/Apple/Resources/Interface/LaunchScreen.storyboard",
+				};
+
 				// look for Assets
-				if (UnrealData.bSupportsMac)
+				string? StoryboardPath;
+				if (Platform == UnrealTargetPlatform.Mac)
 				{
 					string AssetsPath = UnrealData.ProjectOrEnginePath("Build/Mac/Resources/Assets.xcassets", false);
 					ResourcesBuildPhase.AddResource(new FileReference(AssetsPath));
+					StoryboardPath = UnrealData.FindFile(StoryboardPaths, UnrealTargetPlatform.Mac, false);
 				}
-				if (UnrealData.bSupportsIOS || UnrealData.bSupportsTVOS)
+				else
 				{
 					string AssetsPath = UnrealData.ProjectOrEnginePath("Build/IOS/Resources/Assets.xcassets", false);
 					ResourcesBuildPhase.AddResource(new FileReference(AssetsPath));
+					StoryboardPath = UnrealData.FindFile(StoryboardPaths, UnrealTargetPlatform.IOS, false);
 				}
-			}
 
-			// if we are run-only, then skip some stuff
-			if (!XcodeProjectFileGenerator.bGenerateRunOnlyProject)
-			{
-				// create a biuld target only if we have source files to build
-				if (Project.UnrealData.BatchedFiles.Count != 0)
+				if (StoryboardPath != null)
 				{
-					XcodeBuildTarget BuildTarget = new XcodeBuildTarget(Project.UnrealData);
-					AddDependency(BuildTarget, Project);
+					ResourcesBuildPhase.AddResource(new FileReference(StoryboardPath));
 				}
 			}
 
+			if (BuildTarget != null)
+			{
+				AddDependency(BuildTarget, Project);
+			}
+
 			CreateXcconfigFile(Project, Name);
 
 			// hook up each buildconfig to this Xcconfig 
@@ -1240,23 +1363,33 @@ namespace UnrealBuildTool.XcodeProjectXcconfig
 			Xcconfig.AppendLine("INFOPLIST_OUTPUT_FORMAT = xml");
 			Xcconfig.AppendLine("COMBINE_HIDPI_IMAGES = YES");
 
-			// allow Xcode to generate the final plist file from our input, some INFOPLIST settings and other settings 
-			if (UnrealData.PlistMode == PlistMode.UpdateTemplate)
+			if (Platform == UnrealTargetPlatform.Mac)
+			{
+				Xcconfig.AppendLine($"SUPPORTED_PLATFORMS = macosx");
+				Xcconfig.AppendLine($"SDKROOT = macosx");
+			}
+			else if (Platform == UnrealTargetPlatform.IOS)
+			{
+				Xcconfig.AppendLine($"SUPPORTED_PLATFORMS = iphoneos");// iphonesimulator");
+				Xcconfig.AppendLine($"SDKROOT = iphoneos");
+			}
+			else
 			{
-				Xcconfig.AppendLine("GENERATE_INFOPLIST_FILE = YES");
-				Xcconfig.AppendLine($"ASSETCATALOG_COMPILER_APPICON_NAME = AppIcon");
+				Xcconfig.AppendLine($"SUPPORTED_PLATFORMS = appletvos");
+				Xcconfig.AppendLine($"SDKROOT = appletvos");
 			}
 
+
 			Xcconfig.Write();
 		}
 	}
 
 	class XcodeBuildTarget : XcodeTarget
 	{
-		public XcodeBuildTarget(UnrealData UnrealData)
-			: base(XcodeTarget.Type.Build, UnrealData)
+		public XcodeBuildTarget(XcodeProject Project)
+			: base(XcodeTarget.Type.Build, Project.UnrealData)
 		{
-			BuildConfigList = new XcodeBuildConfigList(Name, UnrealData.AllConfigs);
+			BuildConfigList = new XcodeBuildConfigList(Project.Platform, Name, Project.UnrealData.AllConfigs);
 			References.Add(BuildConfigList);
 		}
 	}
@@ -1271,12 +1404,12 @@ namespace UnrealBuildTool.XcodeProjectXcconfig
 		{
 			UnrealData = Project.UnrealData;
 
-			BuildConfigList = new XcodeBuildConfigList(Name, UnrealData.AllConfigs);
+			BuildConfigList = new XcodeBuildConfigList(Project.Platform, Name, UnrealData.AllConfigs);
 			References.Add(BuildConfigList);
 
 			CreateXcconfigFile(Project, Name);
 			// hook up each buildconfig to this Xcconfig 
-			BuildConfigList!.BuildConfigs.ForEach(x => x.Xcconfig = Xcconfig);
+			BuildConfigList.BuildConfigs.ForEach(x => x.Xcconfig = Xcconfig);
 
 
 			// add all of the files to be natively compiled by this target
@@ -1310,6 +1443,9 @@ namespace UnrealBuildTool.XcodeProjectXcconfig
 
 	class XcodeProject : XcodeProjectNode
 	{
+		// a null platform here means all platforms like the old way
+		public UnrealTargetPlatform? Platform;
+
 		// the blob of data coming from unreal that we can pass around
 		public UnrealData UnrealData;
 
@@ -1320,20 +1456,39 @@ namespace UnrealBuildTool.XcodeProjectXcconfig
 		public string Guid = XcodeProjectFileGenerator.MakeXcodeGuid();
 		private string ProvisioningStyle;
 
-		private XcodeRunTarget RunTarget;
+		public List<XcodeRunTarget> RunTargets = new();
 
 		public XcodeBuildConfigList ProjectBuildConfigs;
 
-		public XcodeProject(UnrealData UnrealData, XcodeFileCollection FileCollection)
+		public XcodeProject(UnrealTargetPlatform? Platform, UnrealData UnrealData, XcodeFileCollection FileCollection)
 		{
+			this.Platform = Platform;
 			this.UnrealData = UnrealData;
 			this.FileCollection = FileCollection;
 
 			ProvisioningStyle = UnrealData.bUseAutomaticSigning ? "Automatic" : "Manual";
-			RunTarget = new XcodeRunTarget(this);
-			References.Add(RunTarget);
 
-			ProjectBuildConfigs = new XcodeBuildConfigList(UnrealData.ProductName, UnrealData.AllConfigs);
+			// if we are run-only, then we don't need a build target (this is shared between platforms if we are doing multi-target)
+			XcodeBuildTarget? BuildTarget = null;
+			if (!XcodeProjectFileGenerator.bGenerateRunOnlyProject)
+			{
+				// create a biuld target only if we have source files to build
+				if (UnrealData.BatchedFiles.Count != 0)
+				{
+					BuildTarget = new XcodeBuildTarget(this);
+				}
+			}
+
+			// create one run target for each platform if our platform is null (ie XcodeProjectGenerator.PerPlatformMode is RunTargetPerPlatform)
+			List<UnrealTargetPlatform> TargetPlatforms = Platform == null ? XcodeProjectFileGenerator.XcodePlatforms : new() { Platform.Value };
+			foreach (UnrealTargetPlatform TargetPlatform in TargetPlatforms)
+			{
+				XcodeRunTarget RunTarget = new XcodeRunTarget(this, TargetPlatform, BuildTarget);
+				RunTargets.Add(RunTarget);
+				References.Add(RunTarget);
+			}
+
+			ProjectBuildConfigs = new XcodeBuildConfigList(Platform, UnrealData.ProductName, UnrealData.AllConfigs);
 			References.Add(ProjectBuildConfigs);
 
 			// create the Projet xcconfig
@@ -1365,9 +1520,12 @@ namespace UnrealBuildTool.XcodeProjectXcconfig
 			Content.WriteLine(4,		"LastUpgradeCheck = 2000;");
 			Content.WriteLine(4,		"ORGANIZATIONNAME = \"Epic Games, Inc.\";");
 			Content.WriteLine(4,		"TargetAttributes = {");
-			Content.WriteLine(5,			$"{RunTarget.Guid} = {{");
-			Content.WriteLine(6,				$"ProvisioningStyle = {ProvisioningStyle};");
-			Content.WriteLine(5,			"};");
+			foreach (XcodeRunTarget RunTarget in RunTargets)
+			{ 
+				Content.WriteLine(5,			$"{RunTarget.Guid} = {{");
+				Content.WriteLine(6,				$"ProvisioningStyle = {ProvisioningStyle};");
+				Content.WriteLine(5,			"};");
+			}
 			Content.WriteLine(4,		"};");
 			Content.WriteLine(3,	"};");
 			Content.WriteLine(3,	$"buildConfigurationList = {ProjectBuildConfigs.Guid} /* Build configuration list for PBXProject \"{ProjectBuildConfigs.TargetName}\" */;");
@@ -1429,11 +1587,9 @@ namespace UnrealBuildTool.XcodeProjectXcconfig
 			return Orientations;
 		}
 
-		string FindBundleID(string RelativePlistPath)
+		string FindBundleID(FileReference PlistFile)
 		{
-			string FullPath = Path.Combine(UnrealData.XcodeProjectFileLocation.ParentDirectory!.FullName, RelativePlistPath);
-
-			string Identifier = XcodeProjectFile.Plist("Print :CFBundleIdentifier", FullPath);
+			string Identifier = XcodeProjectFile.Plist("Print :CFBundleIdentifier", PlistFile.FullName);
 			// handle error
 			if (Identifier.StartsWith("Print:"))
 			{
@@ -1473,6 +1629,7 @@ namespace UnrealBuildTool.XcodeProjectXcconfig
 			Xcconfig.AppendLine($"UE_PRODUCT_NAME_STRIPPED = {UnrealData.ProductName.Replace("_", "").Replace(" ", "")}");
 			Xcconfig.AppendLine($"UE_DISPLAY_NAME = {UnrealData.DisplayName}");
 			Xcconfig.AppendLine($"UE_SIGNING_PREFIX = {SigningPrefix}");
+			Xcconfig.AppendLine($"UE_ENGINE_DIR = {Unreal.EngineDirectory}");
 			Xcconfig.AppendLine("");
 
 
@@ -1484,19 +1641,21 @@ namespace UnrealBuildTool.XcodeProjectXcconfig
 
 			if (ProjectBuildConfigs.bSupportsIOS)
 			{
-				string SettingFilter = XcodeProjectFileGenerator.XcodePlatforms.Count == 1 ? "" : "[sdk=iphoneos*]";
+				string SettingFilter = Platform != null ? "" : "[sdk=iphoneos*]";
 				ConfigHierarchy Ini = ConfigCache.ReadHierarchy(ConfigHierarchyType.Engine, UnrealData.UProjectFileLocation?.Directory, UnrealTargetPlatform.IOS);
 
 				Xcconfig.AppendLine("// IOS Settings");
 
 				SupportedDevices.Add(UnrealData.IOSProjectSettings!.RuntimeDevices);
 
-				Xcconfig.AppendLine($"SDKROOT{SettingFilter} = iphoneos");
 				Xcconfig.AppendLine($"CONFIGURATION_BUILD_DIR{SettingFilter} = {UnrealData.AllConfigs[0].IOSExecutablePath!.Directory}/Payload");
 				Xcconfig.AppendLine($"IPHONEOS_DEPLOYMENT_TARGET = {UnrealData.IOSProjectSettings.RuntimeVersion}");
 
-				if (UnrealData.PlistMode == PlistMode.UpdateTemplate)
+				if (UnrealData.PlistFiles[MetadataPlatform.IOS].Mode == MetadataMode.UpdateTemplate)
 				{
+					// allow Xcode to generate the final plist file from our input, some INFOPLIST settings and other settings 
+					Xcconfig.AppendLine("GENERATE_INFOPLIST_FILE = YES");
+					Xcconfig.AppendLine($"ASSETCATALOG_COMPILER_APPICON_NAME = AppIcon");
 					Xcconfig.AppendLine($"CURRENT_PROJECT_VERSION{SettingFilter} = $(UE_IOS_BUILD_VERSION)");
 
 					// short version string
@@ -1544,7 +1703,7 @@ namespace UnrealBuildTool.XcodeProjectXcconfig
 
 			if (ProjectBuildConfigs.bSupportsTVOS)
 			{
-				string SettingFilter = XcodeProjectFileGenerator.XcodePlatforms.Count == 1 ? "" : "[sdk=appletvos*]";
+				string SettingFilter = Platform != null ? "" : "[sdk=appletvos*]";
 				ConfigHierarchy Ini = ConfigCache.ReadHierarchy(ConfigHierarchyType.Engine, UnrealData.UProjectFileLocation?.Directory, UnrealTargetPlatform.TVOS);
 
 				Xcconfig.AppendLine("");
@@ -1552,12 +1711,14 @@ namespace UnrealBuildTool.XcodeProjectXcconfig
 
 				SupportedDevices.Add(UnrealData.TVOSProjectSettings!.RuntimeDevices);
 
-				Xcconfig.AppendLine($"SDKROOT{SettingFilter} = appletvos");
 				Xcconfig.AppendLine($"CONFIGURATION_BUILD_DIR{SettingFilter} = {UnrealData.AllConfigs[0].TVOSExecutablePath!.Directory}/Payload");
 				Xcconfig.AppendLine($"TVOS_DEPLOYMENT_TARGET = {UnrealData.TVOSProjectSettings.RuntimeVersion}");
 
-				if (UnrealData.PlistMode == PlistMode.UpdateTemplate)
+				if (UnrealData.PlistFiles[MetadataPlatform.IOS].Mode == MetadataMode.UpdateTemplate)
 				{
+					// allow Xcode to generate the final plist file from our input, some INFOPLIST settings and other settings 
+					Xcconfig.AppendLine("GENERATE_INFOPLIST_FILE = YES");
+					Xcconfig.AppendLine($"ASSETCATALOG_COMPILER_APPICON_NAME = AppIcon");
 					Xcconfig.AppendLine($"CURRENT_PROJECT_VERSION{SettingFilter} = $(UE_TVOS_BUILD_VERSION)");
 
 					// short version string
@@ -1588,16 +1749,10 @@ namespace UnrealBuildTool.XcodeProjectXcconfig
 
 			if (ProjectBuildConfigs.bSupportsMac)
 			{
-				string SettingFilter = XcodeProjectFileGenerator.XcodePlatforms.Count == 1 ? "" : "[sdk=macosx*]";
+				string SettingFilter = Platform != null ? "" : "[sdk=macosx*]";
 
 				Xcconfig.AppendLine("");
 				Xcconfig.AppendLine("// Mac Settings");
-				Xcconfig.AppendLine($"SDKROOT{SettingFilter} = macosx");
-
-				if (UnrealData.PlistMode == PlistMode.UpdateTemplate)
-				{
-					Xcconfig.AppendLine($"CURRENT_PROJECT_VERSION{SettingFilter} = $(UE_MAC_BUILD_VERSION)");
-				}
 			}
 
 
@@ -1615,7 +1770,7 @@ namespace UnrealBuildTool.XcodeProjectXcconfig
 			Xcconfig.AppendLine("USE_HEADERMAP = NO");
 			Xcconfig.AppendLine("ONLY_ACTIVE_ARCH = YES");
 			Xcconfig.AppendLine($"TARGETED_DEVICE_FAMILY = {string.Join(",", SupportedDevices)}");
-			
+
 
 			if (UnrealData.bWriteCodeSigningSettings)
 			{
@@ -1626,11 +1781,11 @@ namespace UnrealBuildTool.XcodeProjectXcconfig
 				}
 			}
 
-			if (UnrealData.EntitlementsLocation != null)
-			{
-				string RelativeLoc = UnrealData.EntitlementsLocation.MakeRelativeTo(DirectoryReference.Combine(UnrealData.XcodeProjectFileLocation.ParentDirectory!));
-				Xcconfig.AppendLine($"CODE_SIGN_ENTITLEMENTS = {RelativeLoc}");
-			}
+			//if (UnrealData.EntitlementsLocation != null)
+			//{
+			//	string RelativeLoc = UnrealData.EntitlementsLocation.MakeRelativeTo(DirectoryReference.Combine(UnrealData.XcodeProjectFileLocation.ParentDirectory!));
+			//	Xcconfig.AppendLine($"CODE_SIGN_ENTITLEMENTS = {RelativeLoc}");
+			//}
 
 			Xcconfig.Write();
 
@@ -1639,33 +1794,37 @@ namespace UnrealBuildTool.XcodeProjectXcconfig
 			DirectoryReference? GameDir = UnrealData.UProjectFileLocation?.Directory;
 			string? GamePath = GameDir != null ? XcodeFileCollection.ConvertPath(GameDir.FullName) : null;
 
-			// find engine or premade (from project) plists
-			string MacEditorPlist = UnrealData.ProjectOrEnginePath("Build/Mac/Resources/Info-Editor.plist", true);
-			string MacGamePlist = UnrealData.ProjectOrEnginePath("Build/Mac/Resources/Info.plist", true);
-			string IOSPlist = UnrealData.ProjectOrEnginePath("Build/IOS/Resources/Info.plist", true);
-
 			// now pull the bundle id's out, as xcode will warn if they don't match
 			string MacEditorBundleID = "com.epicgames.UnrealEditor";
 			string MacGameBundleID = UnrealData.BundleIdentifier;
 			string IOSBundleID = UnrealData.BundleIdentifier;
-			if (UnrealData.PlistMode == PlistMode.UsePremade)
+			if (UnrealData.bSupportsMac)
 			{
-				if (UnrealData.bSupportsMac)
+				if (UnrealData.PlistFiles[MetadataPlatform.MacEditor].Mode == MetadataMode.UsePremade)
 				{
-					MacEditorBundleID = FindBundleID(MacEditorPlist);
-					MacGameBundleID = FindBundleID(MacGamePlist);
+					MacEditorBundleID = FindBundleID(UnrealData.PlistFiles[MetadataPlatform.MacEditor].File);
 				}
-				if (UnrealData.bSupportsIOS || UnrealData.bSupportsTVOS)
+				if (UnrealData.PlistFiles[MetadataPlatform.Mac].Mode == MetadataMode.UsePremade)
 				{
-					IOSBundleID = FindBundleID(IOSPlist);
+					MacGameBundleID = FindBundleID(UnrealData.PlistFiles[MetadataPlatform.Mac].File);
 				}
 			}
+			if ((UnrealData.bSupportsIOS || UnrealData.bSupportsTVOS) && UnrealData.PlistFiles[MetadataPlatform.IOS].Mode == MetadataMode.UsePremade)
+			{
+				IOSBundleID = FindBundleID(UnrealData.PlistFiles[MetadataPlatform.IOS].File);
+			}
 
 			// Get Mac architectures supported by this project
 			foreach (UnrealBuildConfig Config in UnrealData.AllConfigs)
 			{
+				XcodeBuildConfig? MatchedConfig = ProjectBuildConfigs.BuildConfigs.FirstOrDefault(x => x.Info == Config);
+				if (MatchedConfig == null || !Config.Supports(Platform))
+				{
+					continue;
+				}
+
 				// hook up the Buildconfig that matches this info to this xcconfig file
-				XcconfigFile ConfigXcconfig = ProjectBuildConfigs!.BuildConfigs.First(x => x.Info == Config).Xcconfig!;
+				XcconfigFile ConfigXcconfig = MatchedConfig.Xcconfig!;
 
 				ConfigXcconfig.AppendLine("// pull in the shared settings for all configs");
 				ConfigXcconfig.AppendLine($"#include \"{UnrealData.ProductName}_Project.xcconfig\"");
@@ -1728,51 +1887,66 @@ namespace UnrealBuildTool.XcodeProjectXcconfig
 				ConfigXcconfig.AppendLine($"SUPPORTED_PLATFORMS = {SupportedPlatforms}");
 
 				// @otodo move VALID_ARCHS up to Project config once we are always universal
-				if (Config.bSupportsMac)
+				if (ProjectBuildConfigs.bSupportsMac)
 				{
-					bool bIsEditor = (Config.ProjectTarget!.TargetRules.Type == TargetType.Editor);
-					string SettingFilter = XcodeProjectFileGenerator.XcodePlatforms.Count == 1 ? "" : "[sdk=macosx*]";
+					bool bIsEditor = (Config.ProjectTarget!.TargetRules!.Type == TargetType.Editor);
+					Metadata PlistData = UnrealData.PlistFiles[bIsEditor ? MetadataPlatform.MacEditor : MetadataPlatform.Mac];
+					Metadata EntitlementsData = UnrealData.EntitlementsFiles[bIsEditor ? MetadataPlatform.MacEditor : MetadataPlatform.Mac];
+
+					string SettingFilter = Platform != null ? "" : "[sdk=macosx*]";
 
 					ConfigXcconfig.AppendLine("");
 					ConfigXcconfig.AppendLine("// Mac setup");
 					ConfigXcconfig.AppendLine($"VALID_ARCHS{SettingFilter} = {SupportedMacArchitectures}");
 					ConfigXcconfig.AppendLine($"CONFIGURATION_BUILD_DIR{SettingFilter} = {MacExecutableDir}");
-					string MacPlistFile = bIsEditor ? MacEditorPlist : MacGamePlist;
-					ConfigXcconfig.AppendLine($"INFOPLIST_FILE{SettingFilter} = {MacPlistFile}");
+					ConfigXcconfig.AppendLine($"INFOPLIST_FILE{SettingFilter} = {PlistData.XcodeProjectRelative}");
 					ConfigXcconfig.AppendLine($"PRODUCT_BUNDLE_IDENTIFIER{SettingFilter} = " + (bIsEditor ? MacEditorBundleID : MacGameBundleID));
 
-					if (UnrealData.PlistMode == PlistMode.UpdateTemplate)
+					if (PlistData.Mode == MetadataMode.UpdateTemplate)
 					{
+						ConfigXcconfig.AppendLine("GENERATE_INFOPLIST_FILE = YES");
+						ConfigXcconfig.AppendLine($"ASSETCATALOG_COMPILER_APPICON_NAME = AppIcon");
+
 						ConfigXcconfig.AppendLine($"PRODUCT_NAME{SettingFilter} = {MacExecutableFileName}");
 						// @todo: get a version for  games, like IOS has
 						ConfigXcconfig.AppendLine($"MARKETING_VERSION{SettingFilter} = {MacToolChain.LoadEngineDisplayVersion()}");
+						ConfigXcconfig.AppendLine($"CURRENT_PROJECT_VERSION{SettingFilter} = $(UE_MAC_BUILD_VERSION)");
+					}
+
+					if (EntitlementsData.Mode == MetadataMode.UsePremade)
+					{
+						ConfigXcconfig.AppendLine($"CODE_SIGN_ENTITLEMENTS{SettingFilter} = {EntitlementsData.XcodeProjectRelative}");
 					}
 				}
 
-				if (Config.bSupportsIOS)
+				if (ProjectBuildConfigs.bSupportsIOS)
 				{
-					string SettingFilter = XcodeProjectFileGenerator.XcodePlatforms.Count == 1 ? "" : "[sdk=iphoneos*]";
+					string SettingFilter = Platform != null ? "" : "[sdk=iphoneos*]";
 
 					ConfigXcconfig.AppendLine("");
 					ConfigXcconfig.AppendLine("// IOS setup");
-					ConfigXcconfig.AppendLine($"INFOPLIST_FILE{SettingFilter} = {IOSPlist}");
+					ConfigXcconfig.AppendLine($"INFOPLIST_FILE{SettingFilter} = {UnrealData.PlistFiles[MetadataPlatform.IOS].XcodeProjectRelative}");
 					ConfigXcconfig.AppendLine($"PRODUCT_BUNDLE_IDENTIFIER{SettingFilter} = {IOSBundleID}");
-					if (UnrealData.PlistMode == PlistMode.UpdateTemplate)
+					if (UnrealData.PlistFiles[MetadataPlatform.IOS].Mode == MetadataMode.UpdateTemplate)
 					{
 						ConfigXcconfig.AppendLine($"PRODUCT_NAME{SettingFilter} = {Config.BuildTarget}"); // @todo: change to Path.GetFileName(Config.IOSExecutablePath) when we stop using payload
 					}
+					if (UnrealData.EntitlementsFiles[MetadataPlatform.IOS].Mode == MetadataMode.UsePremade)
+					{
+						ConfigXcconfig.AppendLine($"CODE_SIGN_ENTITLEMENTS{SettingFilter} = {UnrealData.EntitlementsFiles[MetadataPlatform.IOS].XcodeProjectRelative}");
+					}
 				}
-				if (Config.bSupportsTVOS)
+				if (ProjectBuildConfigs.bSupportsTVOS)
 				{
-					string SettingFilter = XcodeProjectFileGenerator.XcodePlatforms.Count == 1 ? "" : "[sdk=appletvos*]";
+					string SettingFilter = Platform != null ? "" : "[sdk=appletvos*]";
 
 					ConfigXcconfig.AppendLine("");
 					ConfigXcconfig.AppendLine("// TVOS setup");
-					ConfigXcconfig.AppendLine($"INFOPLIST_FILE{SettingFilter} = {IOSPlist}");
+					ConfigXcconfig.AppendLine($"INFOPLIST_FILE{SettingFilter} = {UnrealData.PlistFiles[MetadataPlatform.IOS].XcodeProjectRelative}");
 					ConfigXcconfig.AppendLine($"PRODUCT_BUNDLE_IDENTIFIER{SettingFilter} = {IOSBundleID}");
-					if (UnrealData.PlistMode == PlistMode.UpdateTemplate)
+					if (UnrealData.EntitlementsFiles[MetadataPlatform.IOS].Mode == MetadataMode.UsePremade)
 					{
-						ConfigXcconfig.AppendLine($"PRODUCT_NAME{SettingFilter} = {Config.BuildTarget}"); // @todo: change to Path.GetFileName(Config.IOSExecutablePath) when we stop using payload
+						ConfigXcconfig.AppendLine($"CODE_SIGN_ENTITLEMENTS{SettingFilter} = {UnrealData.EntitlementsFiles[MetadataPlatform.IOS].XcodeProjectRelative}");
 					}
 				}
 
@@ -1846,6 +2020,7 @@ namespace UnrealBuildTool.XcodeProjectXcconfig
 		{
 			public string Guid = XcodeProjectFileGenerator.MakeXcodeGuid();
 			public string RelativePath = "";
+			public string FileKey = "";
 			public string FileType = "";
 			public string SourceTree = "";
 			public string GroupName = "";
@@ -1875,6 +2050,25 @@ namespace UnrealBuildTool.XcodeProjectXcconfig
 			RootDirectory = ProjectFile.UnrealData.UProjectFileLocation == null ? (Unreal.EngineDirectory) : ProjectFile.UnrealData.UProjectFileLocation.Directory;
 		}
 
+		public XcodeFileCollection(XcodeFileCollection Other)
+		{
+			ProjectDirectory = Other.ProjectDirectory;
+			RootDirectory = Other.RootDirectory;
+			MainGroupGuid = Other.MainGroupGuid;
+			AllFiles.AddRange(Other.AllFiles);
+			BuildableResourceFiles.AddRange(Other.AllFiles);
+			ManualFiles.AddRange(Other.ManualFiles);
+
+			BuildableFilesToResponseFile = new(Other.BuildableFilesToResponseFile);
+			GuidsForGroups = new(Other.GuidsForGroups);
+
+			foreach (string Key in Other.Groups.Keys)
+			{
+				Groups.Add(Key, new XcodeFileGroup(Other.Groups[Key]));
+			}
+
+		}
+
 		public void ProcessFile(XcodeSourceFile File, bool bIsForBuild)
 		{
 			// remember all buildable files
@@ -1908,7 +2102,7 @@ namespace UnrealBuildTool.XcodeProjectXcconfig
 			else
 			{
 				string GroupName = File.Reference.IsUnderDirectory(Unreal.EngineDirectory) ? "EngineReferences" : "ExternalReferences";
-				AddFileReference(File.FileRefGuid, File.Reference.FullName, GetFileType(File.Reference.GetExtension()), "<absolute>", "EngineReferences");
+				AddFileReference(File.FileRefGuid, File.Reference.FullName, GetFileKey(File.Reference.GetExtension()), GetFileType(File.Reference.GetExtension()), "<absolute>", "EngineReferences");
 			}
 		}
 
@@ -1921,12 +2115,13 @@ namespace UnrealBuildTool.XcodeProjectXcconfig
 			return GuidsForGroups["Products"];
 		}
 
-		public void AddFileReference(string Guid, string RelativePath, string FileType, string SourceTree/*="SOURCE_ROOT"*/, string GroupName)
+		public void AddFileReference(string Guid, string RelativePath, string FileKey, string FileType, string SourceTree/*="SOURCE_ROOT"*/, string GroupName)
 		{
 			ManualFiles.Add(new ManualFileReference()
 			{
 				Guid = Guid,
 				RelativePath = RelativePath,
+				FileKey = FileKey,
 				FileType = FileType,
 				SourceTree = SourceTree.StartsWith("<") ? $"\"{SourceTree}\"" : SourceTree,
 				GroupName = GroupName,
@@ -1966,15 +2161,16 @@ namespace UnrealBuildTool.XcodeProjectXcconfig
 				string FileName = File.Reference.GetFileName();
 				string FileExtension = Path.GetExtension(FileName);
 				string FileType = GetFileType(FileExtension);
+				string FileKey = GetFileKey(FileExtension);
 				string RelativePath = Utils.CleanDirectorySeparators(File.Reference.MakeRelativeTo(ProjectDirectory), '/');
 				string SourceTree = "SOURCE_ROOT";
 
-				Content.WriteLine($"\t\t{File.FileRefGuid} = {{isa = PBXFileReference; explicitFileType = {FileType}; name = \"{FileName}\"; path = \"{RelativePath}\"; sourceTree = {SourceTree}; }};");
+				Content.WriteLine($"\t\t{File.FileRefGuid} = {{isa = PBXFileReference; {FileKey} = {FileType}; name = \"{FileName}\"; path = \"{RelativePath}\"; sourceTree = {SourceTree}; }};");
 			}
 			foreach (ManualFileReference File in ManualFiles)
 			{
 				string FileName = Path.GetFileName(File.RelativePath);
-				Content.WriteLine($"\t\t{File.Guid} = {{isa = PBXFileReference; explicitFileType = {File.FileType}; name = \"{FileName}\"; path = \"{File.RelativePath}\"; sourceTree = {File.SourceTree}; }};");
+				Content.WriteLine($"\t\t{File.Guid} = {{isa = PBXFileReference; {File.FileKey} = {File.FileType}; name = \"{FileName}\"; path = \"{File.RelativePath}\"; sourceTree = {File.SourceTree}; }};");
 			}
 			Content.WriteLine("/* End PBXFileReference section */");
 			Content.WriteLine();
@@ -1982,7 +2178,7 @@ namespace UnrealBuildTool.XcodeProjectXcconfig
 			WriteGroups(Content);
 		}
 
-		private void WriteGroups(StringBuilder Content /*, List<UnrealExtensionInfo> AllExtensions*/)
+		private void WriteGroups(StringBuilder Content)
 		{
 			XcodeFileGroup? RootGroup = FindRootFileGroup(Groups);
 			if (RootGroup == null)
@@ -2132,11 +2328,26 @@ namespace UnrealBuildTool.XcodeProjectXcconfig
 					return "image.icns";
 				case ".xcassets":
 					return "folder.assetcatalog";
+				case ".storyboardc":
+					return "wrapper.storyboardc";
+				case ".storyboard":
+					return "file.storyboard";
 				default:
 					return "file.text";
 			}
 		}
 
+		// most types Xcode wants as explicitFileType, but some want lastKnownFileType
+		internal static string GetFileKey(string Extension)
+		{
+			switch (Extension)
+			{
+				case ".storyboard":
+					return "lastKnownFileType";
+				default:
+					return "explicitFileType";
+			}
+		}
 		/// <summary>
 		/// Returns true if Extension is a known extension for files containing source code
 		/// </summary>
@@ -2146,11 +2357,11 @@ namespace UnrealBuildTool.XcodeProjectXcconfig
 		}
 
 		/// <summary>
-		/// Returns true if Extension is a known extension for files containing source code
+		/// Returns true if Extension is a known extension for files containing resource data
 		/// </summary>
 		internal static bool IsResourceFile(string Extension)
 		{
-			return Extension == ".xcassets";
+			return Extension == ".xcassets" || Extension == ".storyboardc" || Extension == ".storyboard";
 		}
 
 		/// <summary>
@@ -2242,7 +2453,8 @@ namespace UnrealBuildTool.XcodeProjectXcconfig
 			UnrealData = new UnrealData(InitFilePath, bIsForDistribution, BundleID, AppName);
 
 			// create the container for all the files that will 
-			FileCollection = new XcodeFileCollection(this);
+			SharedFileCollection = new XcodeFileCollection(this);
+			FileCollection = SharedFileCollection;
 		}
 
 		public UnrealData UnrealData;
@@ -2250,15 +2462,19 @@ namespace UnrealBuildTool.XcodeProjectXcconfig
 		/// <summary>
 		///  The PBXPRoject node, root of everything
 		/// </summary>
-		public XcodeProject? RootProject;
+		public Dictionary<XcodeProject, UnrealTargetPlatform?> RootProjects = new();
+
+
+		private XcodeProjectLegacy.XcodeProjectFile? LegacyProjectFile = null;
+		private bool bHasCheckedForLegacy = false;
+		public bool bHasLegacyProject => LegacyProjectFile != null;
 
 		/// <summary>
 		/// Gathers the files and generates project sections
 		/// </summary>
+		private XcodeFileCollection SharedFileCollection;
 		public XcodeFileCollection FileCollection;
 
-		private XcodeProjectLegacy.XcodeProjectFile? LegacyProjectFile = null;
-		private bool bHasCheckedForLegacy = false;
 
 		/// <summary>
 		/// Allocates a generator-specific source file object
@@ -2322,16 +2538,16 @@ namespace UnrealBuildTool.XcodeProjectXcconfig
 
 			foreach (XcodeSourceFile SourceFile in SourceFiles.OfType<XcodeSourceFile>())
 			{
-				FileCollection.ProcessFile(SourceFile, bIsForBuild:IsGeneratedProject);
+				SharedFileCollection.ProcessFile(SourceFile, bIsForBuild:IsGeneratedProject);
 			}
 
 			// cache the main group
-			FileCollection.MainGroupGuid = XcodeFileCollection.GetRootGroupGuid(FileCollection.Groups);
+			SharedFileCollection.MainGroupGuid = XcodeFileCollection.GetRootGroupGuid(SharedFileCollection.Groups);
 
 			// filter each file into the appropriate batch
-			foreach (XcodeSourceFile File in FileCollection.BuildableFilesToResponseFile.Keys)
+			foreach (XcodeSourceFile File in SharedFileCollection.BuildableFilesToResponseFile.Keys)
 			{
-				AddFileToBatch(File, FileCollection);
+				AddFileToBatch(File, SharedFileCollection);
 			}
 
 			// write out the response files for each batch now that everything is done
@@ -2358,6 +2574,27 @@ namespace UnrealBuildTool.XcodeProjectXcconfig
 		}
 
 
+		// null platform means use the old way, with multiple platforms per project
+		private string Suffix(UnrealTargetPlatform? Platform)
+		{
+			return Platform == null ? "" : $" ({Platform})";
+		}
+
+		public DirectoryReference ProjectDirPathForPlatform(UnrealTargetPlatform? Platform)
+		{
+			return DirectoryReference.Combine(ProjectFilePath.Directory, $"{ProjectFilePath.GetFileNameWithoutExtension()}{Suffix(Platform)}{ProjectFilePath.GetExtension()}");
+		}
+
+		public FileReference ProjectFilePathForPlatform(UnrealTargetPlatform? Platform)
+		{
+			return new FileReference(ProjectDirPathForPlatform(Platform).FullName);
+		}
+
+		public FileReference PBXFilePathForPlatform(UnrealTargetPlatform? Platform)
+		{
+			return FileReference.Combine(ProjectDirPathForPlatform(Platform), "project.pbxproj");
+		}
+
 		/// Implements Project interface
 		public override bool WriteProjectFile(List<UnrealTargetPlatform> InPlatforms, List<UnrealTargetConfiguration> InConfigurations, PlatformProjectGeneratorCollection PlatformProjectGenerators, ILogger Logger)
 		{
@@ -2379,81 +2616,109 @@ namespace UnrealBuildTool.XcodeProjectXcconfig
 			string ExistingProjectName = UnrealData.ProductName;
 			FileReference TemplateProject = FileReference.Combine(BuildDirLocation, $"Build/IOS/{UnrealData.ProductName}.xcodeproj/project.pbxproj");
 
+			// @todo this for per-platform!
 			UnrealData.bIsMergingProjects = FileReference.Exists(TemplateProject);
 			UnrealData.bWriteCodeSigningSettings = !UnrealData.bIsMergingProjects;
 
 
 			// turn all UE files into internal representation
 			ProcessSourceFiles();
-			
-			// now create the xcodeproject elements (project -> target -> buildconfigs, etc)
-			RootProject = new XcodeProject(UnrealData, FileCollection);
 
-			bool bSuccess;
-			if (FileReference.Exists(TemplateProject))
+			bool bSuccess = true;
+			foreach (UnrealTargetPlatform? Platform in XcodeProjectFileGenerator.WorkspacePlatforms)
 			{
-				bSuccess = MergeIntoTemplateProject(TemplateProject);
-			}
-			else
-			{
-				// write metadata now so we can add them to the FileCollection
-				ConditionalWriteMetadataFiles(UnrealTargetPlatform.IOS);
+				// skip the platform if the project has no configurations for it
+				if (!UnrealData.AllConfigs.Any(x => x.Supports(Platform)))
+				{
+					continue;
+				}
+				FileReference PBXFilePath = PBXFilePathForPlatform(Platform);
 
-				FileReference PBXProjFilePath = ProjectFilePath + "/project.pbxproj";
+				// now create the xcodeproject elements (project -> target -> buildconfigs, etc)
+				FileCollection = new XcodeFileCollection(SharedFileCollection);
+				XcodeProject RootProject = new XcodeProject(Platform, UnrealData, FileCollection);
+				RootProjects[RootProject] = Platform;
 
-				StringBuilder Content = new StringBuilder();
+				if (FileReference.Exists(TemplateProject))
+				{
+					// @todo hahahaah
+					continue;
+					//bSuccess = MergeIntoTemplateProject(PBXFilePath, RootProject, TemplateProject);
+				}
+				else
+				{
+					// write metadata now so we can add them to the FileCollection
+					ConditionalWriteMetadataFiles(UnrealTargetPlatform.IOS);
 
-				Content.WriteLine(0, "// !$*UTF8*$!");
-				Content.WriteLine(0, "{");
-				Content.WriteLine(1, "archiveVersion = 1;");
-				Content.WriteLine(1, "classes = {");
-				Content.WriteLine(1, "};");
-				Content.WriteLine(1, "objectVersion = 46;");
-				Content.WriteLine(1, "objects = {");
 
-				// write out the list of files and groups
-				FileCollection.Write(Content);
+					StringBuilder Content = new StringBuilder();
 
-				// now write out the project node and its recursive dependent nodes
-				XcodeProjectNode.WriteNodeAndReferences(Content, RootProject);
+					Content.WriteLine(0, "// !$*UTF8*$!");
+					Content.WriteLine(0, "{");
+					Content.WriteLine(1, "archiveVersion = 1;");
+					Content.WriteLine(1, "classes = {");
+					Content.WriteLine(1, "};");
+					Content.WriteLine(1, "objectVersion = 46;");
+					Content.WriteLine(1, "objects = {");
 
-				Content.WriteLine(1, "};");
-				Content.WriteLine(1, $"rootObject = {RootProject.Guid} /* Project object */;");
-				Content.WriteLine(0, "}");
+					// write out the list of files and groups
+					FileCollection.Write(Content);
 
-				// finally write out the pbxproj file!
-				bSuccess = ProjectFileGenerator.WriteFileIfChanged(PBXProjFilePath.FullName, Content.ToString(), Logger, new UTF8Encoding());
-			}
+					// now write out the project node and its recursive dependent nodes
+					XcodeProjectNode.WriteNodeAndReferences(Content, RootProject);
 
+					Content.WriteLine(1, "};");
+					Content.WriteLine(1, $"rootObject = {RootProject.Guid} /* Project object */;");
+					Content.WriteLine(0, "}");
 
-			bool bNeedScheme = UnrealData.CanBuildProjectLocally(this, Logger);
-			if (bNeedScheme)
-			{
-				if (bSuccess)
-				{
-					string TargetName = ProjectFilePath.GetFileNameWithoutAnyExtensions();
-					string RunTargetGuid = XcodeProjectNode.GetNodesOfType<XcodeRunTarget>(RootProject).First()!.Guid;
-					string? BuildTargetGuid = XcodeProjectNode.GetNodesOfType<XcodeBuildTarget>(RootProject).FirstOrDefault()?.Guid;
-					string? IndexTargetGuid = XcodeProjectNode.GetNodesOfType<XcodeIndexTarget>(RootProject).FirstOrDefault()?.Guid;
-					WriteSchemeFile(TargetName, RunTargetGuid, BuildTargetGuid, IndexTargetGuid, UnrealData.bHasEditorConfiguration, UnrealData.UProjectFileLocation != null ? UnrealData.UProjectFileLocation.FullName : "");
+					// finally write out the pbxproj file!
+					bSuccess = ProjectFileGenerator.WriteFileIfChanged(PBXFilePath.FullName, Content.ToString(), Logger, new UTF8Encoding()) && bSuccess;
 				}
-			}
-			else
-			{
-				// clean this up because we don't want it persisting if we narrow our project list
-				DirectoryReference SchemeDir = GetProjectSchemeDirectory();
 
-				if (DirectoryReference.Exists(SchemeDir))
+
+				bool bNeedScheme = UnrealData.CanBuildProjectLocally(this, Logger);
+				if (bNeedScheme)
 				{
-					DirectoryReference.Delete(SchemeDir, true);
+					if (bSuccess)
+					{
+						string ProjectName = ProjectFilePathForPlatform(Platform).GetFileNameWithoutAnyExtensions();
+						string? BuildTargetGuid = XcodeProjectNode.GetNodesOfType<XcodeBuildTarget>(RootProject).FirstOrDefault()?.Guid;
+						string? IndexTargetGuid = XcodeProjectNode.GetNodesOfType<XcodeIndexTarget>(RootProject).FirstOrDefault()?.Guid;
+						WriteSchemeFile(Platform, ProjectName, RootProject.RunTargets, BuildTargetGuid, IndexTargetGuid, UnrealData.bHasEditorConfiguration, UnrealData.UProjectFileLocation != null ? UnrealData.UProjectFileLocation.FullName : "");
+					}
 				}
-			}
+				else
+				{
+					// clean this up because we don't want it persisting if we narrow our project list
+					DirectoryReference SchemeDir = GetProjectSchemeDirectory(Platform);
 
+					if (DirectoryReference.Exists(SchemeDir))
+					{
+						DirectoryReference.Delete(SchemeDir, true);
+					}
+				}
+			}
 			return bSuccess;
 		}
 
+		private UnrealTargetPlatform CurrentPlistPlatform;
 		private void ConditionalWriteMetadataFiles(UnrealTargetPlatform Platform)
 		{
+			CurrentPlistPlatform = Platform;
+
+			// we now use templates or premade, no writing out here
+			foreach (Metadata Data in UnrealData.PlistFiles.Values)
+			{
+				FileCollection.AddFileReference(XcodeProjectFileGenerator.MakeXcodeGuid(), Data.XcodeProjectRelative, "explicitFileType", "text.plist", "\"<group>\"", "Metadata");
+			}
+			foreach (Metadata Data in UnrealData.EntitlementsFiles.Values)
+			{
+				if (Data.Mode == MetadataMode.UsePremade)
+				{
+					FileCollection.AddFileReference(XcodeProjectFileGenerator.MakeXcodeGuid(), Data.XcodeProjectRelative, "explicitFileType", "text.plist", "\"<group>\"", "Metadata");
+				}
+			}
+#if false
 			ConfigHierarchy Ini = ConfigCache.ReadHierarchy(ConfigHierarchyType.Engine, UnrealData.UProjectFileLocation?.Directory, Platform);
 
 			/////////////////////
@@ -2502,6 +2767,11 @@ namespace UnrealBuildTool.XcodeProjectXcconfig
 					Plist("Add com.apple.developer.user-management:0 string runs-as-current-user", Filename);
 				}
 			}
+#endif
+
+
+
+
 
 #if false
 			/////////////////////
@@ -2807,7 +3077,7 @@ namespace UnrealBuildTool.XcodeProjectXcconfig
 
 		private string Plist(string Command)
 		{
-			return Plist(Command, ProjectFilePath.FullName + "/project.pbxproj");
+			return Plist(Command, PBXFilePathForPlatform(CurrentPlistPlatform).FullName);
 		}
 
 		private void PlistSetAdd(string Entry, string Value, string Type="string")
@@ -2885,10 +3155,8 @@ namespace UnrealBuildTool.XcodeProjectXcconfig
 		}
 
 
-		bool MergeIntoTemplateProject(FileReference TemplateProject)
+		bool MergeIntoTemplateProject(FileReference PBXProjFilePath, XcodeProject RootProject, FileReference TemplateProject)
 		{
-			FileReference PBXProjFilePath = ProjectFilePath + "/project.pbxproj";
-
 			// copy existing template project to final location
 			if (FileReference.Exists(PBXProjFilePath))
 			{
@@ -2898,9 +3166,9 @@ namespace UnrealBuildTool.XcodeProjectXcconfig
 			FileReference.Copy(TemplateProject, PBXProjFilePath);
 
 			// write the nodes we need to add (Build/Index targets)
-			XcodeRunTarget RunTarget = XcodeProjectNode.GetNodesOfType<XcodeRunTarget>(RootProject!).First();
+			XcodeRunTarget RunTarget = XcodeProjectNode.GetNodesOfType<XcodeRunTarget>(RootProject).First();
 			XcodeBuildTarget BuildTarget = XcodeProjectNode.GetNodesOfType<XcodeBuildTarget>(RunTarget).First();
-			XcodeIndexTarget IndexTarget = XcodeProjectNode.GetNodesOfType<XcodeIndexTarget>(RootProject!).First();
+			XcodeIndexTarget IndexTarget = XcodeProjectNode.GetNodesOfType<XcodeIndexTarget>(RootProject).First();
 			XcodeDependency BuildDependency = XcodeProjectNode.GetNodesOfType<XcodeDependency>(RunTarget).First();
 
 			// the runtarget and project need to write out so all of their xcconfigs get written as well,
@@ -3007,7 +3275,7 @@ namespace UnrealBuildTool.XcodeProjectXcconfig
 
 					if (bIsProject)
 					{
-						Console.WriteLine("Looking for " + ConfigName);
+						//Console.WriteLine("Looking for " + ConfigName);
 						XcodeBuildConfig Config = RootProject!.ProjectBuildConfigs.BuildConfigs.First(x => x.Info.DisplayName == ConfigName);
 						PlistSetAdd($":objects:{ConfigGuid}:baseConfigurationReference", Config.Xcconfig!.Guid, "string");
 					}
@@ -3051,164 +3319,169 @@ namespace UnrealBuildTool.XcodeProjectXcconfig
 
 #region Schemes
 
-		private FileReference GetUserSchemeManagementFilePath()
+		private FileReference GetUserSchemeManagementFilePath(UnrealTargetPlatform? Platform)
 		{
-			return new FileReference(ProjectFilePath.FullName + "/xcuserdata/" + Environment.UserName + ".xcuserdatad/xcschemes/xcschememanagement.plist");
+			return FileReference.Combine(ProjectDirPathForPlatform(Platform), $"xcuserdata/{Environment.UserName}.xcuserdatad/xcschemes/xcschememanagement.plist");
 		}
 
-		private DirectoryReference GetProjectSchemeDirectory()
+		private DirectoryReference GetProjectSchemeDirectory(UnrealTargetPlatform? Platform)
 		{
-			return new DirectoryReference(ProjectFilePath.FullName + "/xcshareddata/xcschemes");
+			return DirectoryReference.Combine(ProjectDirPathForPlatform(Platform), "xcshareddata/xcschemes");
 		}
 
-		private FileReference GetProjectSchemeFilePathForTarget(string TargetName)
+		private FileReference GetProjectSchemeFilePathForTarget(UnrealTargetPlatform? Platform, string TargetName)
 		{
-			return FileReference.Combine(GetProjectSchemeDirectory(), TargetName + ".xcscheme");
+			return FileReference.Combine(GetProjectSchemeDirectory(Platform), TargetName + ".xcscheme");
 		}
 
-		private void WriteSchemeFile(string TargetName, string TargetGuid, string? BuildTargetGuid, string? IndexTargetGuid, bool bHasEditorConfiguration, string GameProjectPath)
+		private void WriteSchemeFile(UnrealTargetPlatform? Platform, string ProjectName, List<XcodeRunTarget> RunTargets, string? BuildTargetGuid, string? IndexTargetGuid, bool bHasEditorConfiguration, string GameProjectPath)
 		{
+			StringBuilder Content = new StringBuilder();
 
-			FileReference SchemeFilePath = GetProjectSchemeFilePathForTarget(TargetName);
-
-			DirectoryReference.CreateDirectory(SchemeFilePath.Directory);
-
-			string? OldCommandLineArguments = null;
-			if (FileReference.Exists(SchemeFilePath))
+			foreach (XcodeRunTarget RunTarget in RunTargets)
 			{
-				string OldContents = File.ReadAllText(SchemeFilePath.FullName);
-				int OldCommandLineArgumentsStart = OldContents.IndexOf("<CommandLineArguments>") + "<CommandLineArguments>".Length;
-				int OldCommandLineArgumentsEnd = OldContents.IndexOf("</CommandLineArguments>");
-				if (OldCommandLineArgumentsStart != -1 && OldCommandLineArgumentsEnd != -1)
-				{
-					OldCommandLineArguments = OldContents.Substring(OldCommandLineArgumentsStart, OldCommandLineArgumentsEnd - OldCommandLineArgumentsStart);
-				}
-			}
+				string TargetName = RunTarget.Name;
+				string TargetGuid = RunTarget.Guid;
+				bHasEditorConfiguration = RunTarget.BuildConfigList!.BuildConfigs.Any(x => x.Info.ProjectTarget!.TargetRules!.Type == TargetType.Editor);
 
-			string DefaultConfiguration = (bHasEditorConfiguration && !XcodeProjectFileGenerator.bGenerateRunOnlyProject) ? "Development Editor" : "Development";
+				FileReference SchemeFilePath = GetProjectSchemeFilePathForTarget(Platform, TargetName);
 
-			StringBuilder Content = new StringBuilder();
+				DirectoryReference.CreateDirectory(SchemeFilePath.Directory);
 
-			Content.WriteLine("<?xml version=\"1.0\" encoding=\"UTF-8\"?>");
-			Content.WriteLine("<Scheme");
-			Content.WriteLine("   LastUpgradeVersion = \"2000\"");
-			Content.WriteLine("   version = \"1.3\">");
-			Content.WriteLine("   <BuildAction");
-			Content.WriteLine("      parallelizeBuildables = \"YES\"");
-			Content.WriteLine("      buildImplicitDependencies = \"YES\">");
-			Content.WriteLine("      <BuildActionEntries>");
-			Content.WriteLine("         <BuildActionEntry");
-			Content.WriteLine("            buildForTesting = \"YES\"");
-			Content.WriteLine("            buildForRunning = \"YES\"");
-			Content.WriteLine("            buildForProfiling = \"YES\"");
-			Content.WriteLine("            buildForArchiving = \"YES\"");
-			Content.WriteLine("            buildForAnalyzing = \"YES\">");
-			Content.WriteLine("            <BuildableReference");
-			Content.WriteLine("               BuildableIdentifier = \"primary\"");
-			Content.WriteLine("               BlueprintIdentifier = \"" + TargetGuid + "\"");
-			Content.WriteLine("               ReferencedContainer = \"container:" + TargetName + ".xcodeproj\">");
-			Content.WriteLine("            </BuildableReference>");
-			Content.WriteLine("         </BuildActionEntry>");
-			Content.WriteLine("      </BuildActionEntries>");
-			Content.WriteLine("   </BuildAction>");
-			Content.WriteLine("   <TestAction");
-			Content.WriteLine("      buildConfiguration = \"" + DefaultConfiguration + "\"");
-			Content.WriteLine("      selectedDebuggerIdentifier = \"Xcode.DebuggerFoundation.Debugger.LLDB\"");
-			Content.WriteLine("      selectedLauncherIdentifier = \"Xcode.DebuggerFoundation.Launcher.LLDB\"");
-			Content.WriteLine("      shouldUseLaunchSchemeArgsEnv = \"YES\">");
-			Content.WriteLine("      <Testables>");
-			Content.WriteLine("      </Testables>");
-			Content.WriteLine("      <MacroExpansion>");
-			Content.WriteLine("            <BuildableReference");
-			Content.WriteLine("               BuildableIdentifier = \"primary\"");
-			Content.WriteLine("               BlueprintIdentifier = \"" + TargetGuid + "\"");
-			Content.WriteLine("               ReferencedContainer = \"container:" + TargetName + ".xcodeproj\">");
-			Content.WriteLine("            </BuildableReference>");
-			Content.WriteLine("      </MacroExpansion>");
-			Content.WriteLine("      <AdditionalOptions>");
-			Content.WriteLine("      </AdditionalOptions>");
-			Content.WriteLine("   </TestAction>");
-			Content.WriteLine("   <LaunchAction");
-			Content.WriteLine("      buildConfiguration = \"" + DefaultConfiguration + "\"");
-			Content.WriteLine("      selectedDebuggerIdentifier = \"Xcode.DebuggerFoundation.Debugger.LLDB\"");
-			Content.WriteLine("      selectedLauncherIdentifier = \"Xcode.DebuggerFoundation.Launcher.LLDB\"");
-			Content.WriteLine("      launchStyle = \"0\"");
-			Content.WriteLine("      useCustomWorkingDirectory = \"NO\"");
-			Content.WriteLine("      ignoresPersistentStateOnLaunch = \"NO\"");
-			Content.WriteLine("      debugDocumentVersioning = \"NO\"");
-			Content.WriteLine("      debugServiceExtension = \"internal\"");
-			Content.WriteLine("      allowLocationSimulation = \"YES\">");
-			Content.WriteLine("      <BuildableProductRunnable");
-			Content.WriteLine("         runnableDebuggingMode = \"0\">");
-			Content.WriteLine("            <BuildableReference");
-			Content.WriteLine("               BuildableIdentifier = \"primary\"");
-			Content.WriteLine("               BlueprintIdentifier = \"" + TargetGuid + "\"");
-			Content.WriteLine("               ReferencedContainer = \"container:" + TargetName + ".xcodeproj\">");
-			Content.WriteLine("            </BuildableReference>");
-			Content.WriteLine("      </BuildableProductRunnable>");
-			if (string.IsNullOrEmpty(OldCommandLineArguments))
-			{
-				if (bHasEditorConfiguration && TargetName != "UE5")
-				{
-					Content.WriteLine("      <CommandLineArguments>");
-					if (IsForeignProject)
+				string? OldCommandLineArguments = null;
+				if (FileReference.Exists(SchemeFilePath))
+				{
+					string OldContents = File.ReadAllText(SchemeFilePath.FullName);
+					int OldCommandLineArgumentsStart = OldContents.IndexOf("<CommandLineArguments>") + "<CommandLineArguments>".Length;
+					int OldCommandLineArgumentsEnd = OldContents.IndexOf("</CommandLineArguments>");
+					if (OldCommandLineArgumentsStart != -1 && OldCommandLineArgumentsEnd != -1)
 					{
-						Content.WriteLine("         <CommandLineArgument");
-						Content.WriteLine("            argument = \"&quot;" + GameProjectPath + "&quot;\"");
-						Content.WriteLine("            isEnabled = \"YES\">");
-						Content.WriteLine("         </CommandLineArgument>");
+						OldCommandLineArguments = OldContents.Substring(OldCommandLineArgumentsStart, OldCommandLineArgumentsEnd - OldCommandLineArgumentsStart);
 					}
-					else
+				}
+
+				string DefaultConfiguration = (bHasEditorConfiguration && !XcodeProjectFileGenerator.bGenerateRunOnlyProject) ? "Development Editor" : "Development";
+
+				Content.WriteLine("<?xml version=\"1.0\" encoding=\"UTF-8\"?>");
+				Content.WriteLine("<Scheme");
+				Content.WriteLine("   LastUpgradeVersion = \"2000\"");
+				Content.WriteLine("   version = \"1.3\">");
+				Content.WriteLine("   <BuildAction");
+				Content.WriteLine("      parallelizeBuildables = \"YES\"");
+				Content.WriteLine("      buildImplicitDependencies = \"YES\">");
+				Content.WriteLine("      <BuildActionEntries>");
+				Content.WriteLine("         <BuildActionEntry");
+				Content.WriteLine("            buildForTesting = \"YES\"");
+				Content.WriteLine("            buildForRunning = \"YES\"");
+				Content.WriteLine("            buildForProfiling = \"YES\"");
+				Content.WriteLine("            buildForArchiving = \"YES\"");
+				Content.WriteLine("            buildForAnalyzing = \"YES\">");
+				Content.WriteLine("            <BuildableReference");
+				Content.WriteLine("               BuildableIdentifier = \"primary\"");
+				Content.WriteLine("               BlueprintIdentifier = \"" + TargetGuid + "\"");
+				Content.WriteLine("               ReferencedContainer = \"container:" + ProjectName + ".xcodeproj\">");
+				Content.WriteLine("            </BuildableReference>");
+				Content.WriteLine("         </BuildActionEntry>");
+				Content.WriteLine("      </BuildActionEntries>");
+				Content.WriteLine("   </BuildAction>");
+				Content.WriteLine("   <TestAction");
+				Content.WriteLine("      buildConfiguration = \"" + DefaultConfiguration + "\"");
+				Content.WriteLine("      selectedDebuggerIdentifier = \"Xcode.DebuggerFoundation.Debugger.LLDB\"");
+				Content.WriteLine("      selectedLauncherIdentifier = \"Xcode.DebuggerFoundation.Launcher.LLDB\"");
+				Content.WriteLine("      shouldUseLaunchSchemeArgsEnv = \"YES\">");
+				Content.WriteLine("      <Testables>");
+				Content.WriteLine("      </Testables>");
+				Content.WriteLine("      <MacroExpansion>");
+				Content.WriteLine("            <BuildableReference");
+				Content.WriteLine("               BuildableIdentifier = \"primary\"");
+				Content.WriteLine("               BlueprintIdentifier = \"" + TargetGuid + "\"");
+				Content.WriteLine("               ReferencedContainer = \"container:" + ProjectName + ".xcodeproj\">");
+				Content.WriteLine("            </BuildableReference>");
+				Content.WriteLine("      </MacroExpansion>");
+				Content.WriteLine("      <AdditionalOptions>");
+				Content.WriteLine("      </AdditionalOptions>");
+				Content.WriteLine("   </TestAction>");
+				Content.WriteLine("   <LaunchAction");
+				Content.WriteLine("      buildConfiguration = \"" + DefaultConfiguration + "\"");
+				Content.WriteLine("      selectedDebuggerIdentifier = \"Xcode.DebuggerFoundation.Debugger.LLDB\"");
+				Content.WriteLine("      selectedLauncherIdentifier = \"Xcode.DebuggerFoundation.Launcher.LLDB\"");
+				Content.WriteLine("      launchStyle = \"0\"");
+				Content.WriteLine("      useCustomWorkingDirectory = \"NO\"");
+				Content.WriteLine("      ignoresPersistentStateOnLaunch = \"NO\"");
+				Content.WriteLine("      debugDocumentVersioning = \"NO\"");
+				Content.WriteLine("      debugServiceExtension = \"internal\"");
+				Content.WriteLine("      allowLocationSimulation = \"YES\">");
+				Content.WriteLine("      <BuildableProductRunnable");
+				Content.WriteLine("         runnableDebuggingMode = \"0\">");
+				Content.WriteLine("            <BuildableReference");
+				Content.WriteLine("               BuildableIdentifier = \"primary\"");
+				Content.WriteLine("               BlueprintIdentifier = \"" + TargetGuid + "\"");
+				Content.WriteLine("               ReferencedContainer = \"container:" + ProjectName + ".xcodeproj\">");
+				Content.WriteLine("            </BuildableReference>");
+				Content.WriteLine("      </BuildableProductRunnable>");
+				if (string.IsNullOrEmpty(OldCommandLineArguments))
+				{
+					if (bHasEditorConfiguration && TargetName != "UE5")
 					{
+						Content.WriteLine("      <CommandLineArguments>");
+						if (IsForeignProject)
+						{
+							Content.WriteLine("         <CommandLineArgument");
+							Content.WriteLine("            argument = \"&quot;" + GameProjectPath + "&quot;\"");
+							Content.WriteLine("            isEnabled = \"YES\">");
+							Content.WriteLine("         </CommandLineArgument>");
+						}
+						else
+						{
+							Content.WriteLine("         <CommandLineArgument");
+							Content.WriteLine("            argument = \"" + TargetName + "\"");
+							Content.WriteLine("            isEnabled = \"YES\">");
+							Content.WriteLine("         </CommandLineArgument>");
+						}
+						// Always add a configuration argument
 						Content.WriteLine("         <CommandLineArgument");
-						Content.WriteLine("            argument = \"" + TargetName + "\"");
+						Content.WriteLine("            argument = \"-RunConfig=$(Configuration)\"");
 						Content.WriteLine("            isEnabled = \"YES\">");
 						Content.WriteLine("         </CommandLineArgument>");
+						Content.WriteLine("      </CommandLineArguments>");
 					}
-					// Always add a configuration argument
-					Content.WriteLine("         <CommandLineArgument");
-					Content.WriteLine("            argument = \"-RunConfig=$(Configuration)\"");
-					Content.WriteLine("            isEnabled = \"YES\">");
-					Content.WriteLine("         </CommandLineArgument>");
-					Content.WriteLine("      </CommandLineArguments>");
 				}
+				else
+				{
+					Content.WriteLine("      <CommandLineArguments>" + OldCommandLineArguments + "</CommandLineArguments>");
+				}
+				Content.WriteLine("      <AdditionalOptions>");
+				Content.WriteLine("      </AdditionalOptions>");
+				Content.WriteLine("   </LaunchAction>");
+				Content.WriteLine("   <ProfileAction");
+				Content.WriteLine("      buildConfiguration = \"" + DefaultConfiguration + "\"");
+				Content.WriteLine("      shouldUseLaunchSchemeArgsEnv = \"YES\"");
+				Content.WriteLine("      savedToolIdentifier = \"\"");
+				Content.WriteLine("      useCustomWorkingDirectory = \"NO\"");
+				Content.WriteLine("      debugDocumentVersioning = \"NO\">");
+				Content.WriteLine("      <BuildableProductRunnable");
+				Content.WriteLine("         runnableDebuggingMode = \"0\">");
+				Content.WriteLine("            <BuildableReference");
+				Content.WriteLine("               BuildableIdentifier = \"primary\"");
+				Content.WriteLine("               BlueprintIdentifier = \"" + TargetGuid + "\"");
+				Content.WriteLine("               BuildableName = \"" + ProjectName + ".app\"");
+				Content.WriteLine("               BlueprintName = \"" + TargetName + "\"");
+				Content.WriteLine("               ReferencedContainer = \"container:" + ProjectName + ".xcodeproj\">");
+				Content.WriteLine("            </BuildableReference>");
+				Content.WriteLine("      </BuildableProductRunnable>");
+				Content.WriteLine("   </ProfileAction>");
+				Content.WriteLine("   <AnalyzeAction");
+				Content.WriteLine("      buildConfiguration = \"" + DefaultConfiguration + "\">");
+				Content.WriteLine("   </AnalyzeAction>");
+				Content.WriteLine("   <ArchiveAction");
+				Content.WriteLine("      buildConfiguration = \"" + DefaultConfiguration + "\"");
+				Content.WriteLine("      revealArchiveInOrganizer = \"YES\">");
+				Content.WriteLine("   </ArchiveAction>");
+				Content.WriteLine("</Scheme>");
+
+				File.WriteAllText(SchemeFilePath.FullName, Content.ToString(), new UTF8Encoding());
+				Content.Clear();
 			}
-			else
-			{
-				Content.WriteLine("      <CommandLineArguments>" + OldCommandLineArguments + "</CommandLineArguments>");
-			}
-			Content.WriteLine("      <AdditionalOptions>");
-			Content.WriteLine("      </AdditionalOptions>");
-			Content.WriteLine("   </LaunchAction>");
-			Content.WriteLine("   <ProfileAction");
-			Content.WriteLine("      buildConfiguration = \"" + DefaultConfiguration + "\"");
-			Content.WriteLine("      shouldUseLaunchSchemeArgsEnv = \"YES\"");
-			Content.WriteLine("      savedToolIdentifier = \"\"");
-			Content.WriteLine("      useCustomWorkingDirectory = \"NO\"");
-			Content.WriteLine("      debugDocumentVersioning = \"NO\">");
-			Content.WriteLine("      <BuildableProductRunnable");
-			Content.WriteLine("         runnableDebuggingMode = \"0\">");
-			Content.WriteLine("            <BuildableReference");
-			Content.WriteLine("               BuildableIdentifier = \"primary\"");
-			Content.WriteLine("               BlueprintIdentifier = \"" + TargetGuid + "\"");
-			Content.WriteLine("               BuildableName = \"" + TargetName + ".app\"");
-			Content.WriteLine("               BlueprintName = \"" + TargetName + "\"");
-			Content.WriteLine("               ReferencedContainer = \"container:" + TargetName + ".xcodeproj\">");
-			Content.WriteLine("            </BuildableReference>");
-			Content.WriteLine("      </BuildableProductRunnable>");
-			Content.WriteLine("   </ProfileAction>");
-			Content.WriteLine("   <AnalyzeAction");
-			Content.WriteLine("      buildConfiguration = \"" + DefaultConfiguration + "\">");
-			Content.WriteLine("   </AnalyzeAction>");
-			Content.WriteLine("   <ArchiveAction");
-			Content.WriteLine("      buildConfiguration = \"" + DefaultConfiguration + "\"");
-			Content.WriteLine("      revealArchiveInOrganizer = \"YES\">");
-			Content.WriteLine("   </ArchiveAction>");
-			Content.WriteLine("</Scheme>");
-
-			File.WriteAllText(SchemeFilePath.FullName, Content.ToString(), new UTF8Encoding());
-
-			Content.Clear();
 
 			Content.WriteLine("<?xml version=\"1.0\" encoding=\"UTF-8\"?>");
 			Content.WriteLine("<!DOCTYPE plist PUBLIC \"-//Apple//DTD PLIST 1.0//EN\" \"http://www.apple.com/DTDs/PropertyList-1.0.dtd\">");
@@ -3216,19 +3489,25 @@ namespace UnrealBuildTool.XcodeProjectXcconfig
 			Content.WriteLine("<dict>");
 			Content.WriteLine("\t<key>SchemeUserState</key>");
 			Content.WriteLine("\t<dict>");
-			Content.WriteLine("\t\t<key>" + TargetName + ".xcscheme_^#shared#^_</key>");
-			Content.WriteLine("\t\t<dict>");
-			Content.WriteLine("\t\t\t<key>orderHint</key>");
-			Content.WriteLine("\t\t\t<integer>1</integer>");
-			Content.WriteLine("\t\t</dict>");
+			foreach (XcodeRunTarget RunTarget in RunTargets)
+			{
+				Content.WriteLine("\t\t<key>" + RunTarget.Name + ".xcscheme_^#shared#^_</key>");
+				Content.WriteLine("\t\t<dict>");
+				Content.WriteLine("\t\t\t<key>orderHint</key>");
+				Content.WriteLine("\t\t\t<integer>1</integer>");
+				Content.WriteLine("\t\t</dict>");
+			}
 			Content.WriteLine("\t</dict>");
 			Content.WriteLine("\t<key>SuppressBuildableAutocreation</key>");
 			Content.WriteLine("\t<dict>");
-			Content.WriteLine("\t\t<key>" + TargetGuid + "</key>");
-			Content.WriteLine("\t\t<dict>");
-			Content.WriteLine("\t\t\t<key>primary</key>");
-			Content.WriteLine("\t\t\t<true/>");
-			Content.WriteLine("\t\t</dict>");
+			foreach (XcodeRunTarget RunTarget in RunTargets)
+			{
+				Content.WriteLine("\t\t<key>" + RunTarget.Guid + "</key>");
+				Content.WriteLine("\t\t<dict>");
+				Content.WriteLine("\t\t\t<key>primary</key>");
+				Content.WriteLine("\t\t\t<true/>");
+				Content.WriteLine("\t\t</dict>");
+			}
 			if (BuildTargetGuid != null)
 			{
 				Content.WriteLine("\t\t<key>" + BuildTargetGuid + "</key>");
@@ -3249,7 +3528,7 @@ namespace UnrealBuildTool.XcodeProjectXcconfig
 			Content.WriteLine("</dict>");
 			Content.WriteLine("</plist>");
 
-			FileReference ManagementFile = GetUserSchemeManagementFilePath();
+			FileReference ManagementFile = GetUserSchemeManagementFilePath(Platform);
 			if (!DirectoryReference.Exists(ManagementFile.Directory))
 			{
 				DirectoryReference.CreateDirectory(ManagementFile.Directory);
@@ -3265,13 +3544,10 @@ namespace UnrealBuildTool.XcodeProjectXcconfig
 
 		public override string ToString()
 		{
-			return ProjectFilePath.GetFileNameWithoutExtension();
+			throw new NotImplementedException();
+//			return ProjectFilePath.GetFileNameWithoutExtension();
 		}
 
-
-
-
-
 #endregion
 
 	}
diff --git a/Engine/Source/Programs/UnrealBuildTool/ProjectFiles/Xcode/XcodeProjectFileGenerator.cs b/Engine/Source/Programs/UnrealBuildTool/ProjectFiles/Xcode/XcodeProjectFileGenerator.cs
--- a/Engine/Source/Programs/UnrealBuildTool/ProjectFiles/Xcode/XcodeProjectFileGenerator.cs
+++ b/Engine/Source/Programs/UnrealBuildTool/ProjectFiles/Xcode/XcodeProjectFileGenerator.cs
@@ -10,6 +10,12 @@ using Microsoft.Extensions.Logging;
 
 namespace UnrealBuildTool
 {
+	enum XcodePerPlatformMode
+	{
+		OneWorkspacePerPlatform,
+		OneTargetPerPlatform,
+	}
+
 	/// <summary>
 	/// Xcode project file generator implementation
 	/// </summary>
@@ -20,6 +26,9 @@ namespace UnrealBuildTool
 		// always seed the random number the same, so multiple runs of the generator will generate the same project
 		static Random Rand = new Random(0);
 
+		// how to handle multiple platforms
+		public static XcodePerPlatformMode PerPlatformMode = XcodePerPlatformMode.OneWorkspacePerPlatform;
+
 		/// <summary>
 		/// Mark for distribution builds
 		/// </summary>
@@ -83,10 +92,9 @@ namespace UnrealBuildTool
 		/// </summary>
 		public override void CleanProjectFiles(DirectoryReference InPrimaryProjectDirectory, string InPrimaryProjectName, DirectoryReference InIntermediateProjectFilesPath, ILogger Logger)
 		{
-			DirectoryReference PrimaryProjDeleteFilename = DirectoryReference.Combine(InPrimaryProjectDirectory, InPrimaryProjectName + ".xcworkspace");
-			if (DirectoryReference.Exists(PrimaryProjDeleteFilename))
+			foreach (DirectoryReference ProjectFile in DirectoryReference.EnumerateDirectories(InPrimaryProjectDirectory, $"{InPrimaryProjectName}*.xcworkspace"))
 			{
-				DirectoryReference.Delete(PrimaryProjDeleteFilename, true);
+				DirectoryReference.Delete(ProjectFile, true);
 			}
 
 			// Delete the project files folder
@@ -167,100 +175,140 @@ namespace UnrealBuildTool
 			return WriteFileIfChanged(Path, WorkspaceSettingsContent.ToString(), Logger, new UTF8Encoding());
 		}
 
+		private string PrimaryProjectNameForPlatform(UnrealTargetPlatform? Platform)
+		{
+			return Platform == null ? PrimaryProjectName : $"{PrimaryProjectName} ({Platform})";
+		}
+
 		private bool WriteXcodeWorkspace(ILogger Logger)
 		{
 			bool bSuccess = true;
 
-			StringBuilder WorkspaceDataContent = new StringBuilder();
+			// loop opver all projects to see if at least one is modern (if not, we don't bother splitting up by platform)
+			bool bHasModernProjects = false;
+			Action<List<PrimaryProjectFolder>>? FindModern = null;
+			FindModern = (FolderList) =>
+			{
+				foreach (PrimaryProjectFolder CurFolder in FolderList)
+				{
+					var Modern = CurFolder.ChildProjects.FirstOrDefault(P =>
+						P.GetType() == typeof(XcodeProjectXcconfig.XcodeProjectFile) &&
+						!((XcodeProjectXcconfig.XcodeProjectFile)P).bHasLegacyProject);
+					if (Modern != null)
+					{
+						//Logger.LogWarning($"Project {Modern.ProjectFilePath} is modern");
+						bHasModernProjects = true;
+					}
+					if (!bHasModernProjects)
+					{
+						FindModern!(CurFolder.SubFolders);
+					}
+				}
+			};
+			FindModern(RootFolder.SubFolders);
 
-			WorkspaceDataContent.Append("<?xml version=\"1.0\" encoding=\"UTF-8\"?>" + ProjectFileGenerator.NewLine);
-			WorkspaceDataContent.Append("<Workspace" + ProjectFileGenerator.NewLine);
-			WorkspaceDataContent.Append("   version = \"1.0\">" + ProjectFileGenerator.NewLine);
+			// if we want one workspace with multiple platforms, and we have at least one modern project, then process each platform individually
+			// otherwise use null as Platform which means to merge all platforms
+			List<UnrealTargetPlatform?> PlatformsToProcess = bHasModernProjects ? WorkspacePlatforms : NullPlatformList; 
+			foreach (UnrealTargetPlatform? Platform in PlatformsToProcess)
+			{
+				StringBuilder WorkspaceDataContent = new();
+				WorkspaceDataContent.Append("<?xml version=\"1.0\" encoding=\"UTF-8\"?>" + ProjectFileGenerator.NewLine);
+				WorkspaceDataContent.Append("<Workspace" + ProjectFileGenerator.NewLine);
+				WorkspaceDataContent.Append("   version = \"1.0\">" + ProjectFileGenerator.NewLine);
 
-			List<ProjectFile> BuildableProjects = new List<ProjectFile>();
+				List<XcodeProjectXcconfig.XcodeProjectFile> BuildableProjects = new();
 
-			System.Action< List<PrimaryProjectFolder> /* Folders */, string /* Ident */ >? AddProjectsFunction = null;
-			AddProjectsFunction = (FolderList, Ident) =>
-				{
-					foreach (PrimaryProjectFolder CurFolder in FolderList)
+				System.Action<List<PrimaryProjectFolder> /* Folders */, string /* Ident */ >? AddProjectsFunction = null;
+				AddProjectsFunction = (FolderList, Ident) =>
 					{
-						WorkspaceDataContent.Append(Ident + "   <Group" + ProjectFileGenerator.NewLine);
-						WorkspaceDataContent.Append(Ident + "      location = \"container:\"      name = \"" + CurFolder.FolderName + "\">" + ProjectFileGenerator.NewLine);
-
-						AddProjectsFunction!(CurFolder.SubFolders, Ident + "   ");
-				
-						// Filter out anything that isn't an XC project, and that shouldn't be in the workspace
-						IEnumerable<ProjectFile> SupportedProjects =
-								CurFolder.ChildProjects.Where(P => P.GetType() == typeof(XcodeProjectXcconfig.XcodeProjectFile) || P.GetType() == typeof(XcodeProjectLegacy.XcodeProjectFile))
-									.Where(P => XcodeProjectXcconfig.UnrealData.ShouldIncludeProjectInWorkspace(P, Logger))
-									.OrderBy(P => P.ProjectFilePath.GetFileName());
-
-						foreach (ProjectFile XcodeProject in SupportedProjects)
+						foreach (PrimaryProjectFolder CurFolder in FolderList)
 						{
-							WorkspaceDataContent.Append(Ident + "      <FileRef" + ProjectFileGenerator.NewLine);
-							WorkspaceDataContent.Append(Ident + "         location = \"group:" + XcodeProject.ProjectFilePath.MakeRelativeTo(ProjectFileGenerator.PrimaryProjectPath) + "\">" + ProjectFileGenerator.NewLine);
-							WorkspaceDataContent.Append(Ident + "      </FileRef>" + ProjectFileGenerator.NewLine);							
+							WorkspaceDataContent.Append(Ident + "   <Group" + ProjectFileGenerator.NewLine);
+							WorkspaceDataContent.Append(Ident + "      location = \"container:\"      name = \"" + CurFolder.FolderName + "\">" + ProjectFileGenerator.NewLine);
+
+							AddProjectsFunction!(CurFolder.SubFolders, Ident + "   ");
+
+							// Filter out anything that isn't an XC project, and that shouldn't be in the workspace
+							IEnumerable<XcodeProjectXcconfig.XcodeProjectFile> SupportedProjects =
+									CurFolder.ChildProjects.Where(P => P.GetType() == typeof(XcodeProjectXcconfig.XcodeProjectFile))
+										.Select(P => (XcodeProjectXcconfig.XcodeProjectFile)P)
+										.Where(P => XcodeProjectXcconfig.UnrealData.ShouldIncludeProjectInWorkspace(P, Logger))
+										// @todo - still need to handle legacy project getting split up?
+										.Where(P => P.RootProjects.Count == 0 || P.RootProjects.ContainsValue(Platform))
+										.OrderBy(P => P.ProjectFilePath.GetFileName());
+
+							foreach (XcodeProjectXcconfig.XcodeProjectFile XcodeProject in SupportedProjects)
+							{
+								// we have to re-check for each project - if it's a legacy project, even if we wanted it split, it won't be, so always point to 
+								// the shared legacy project
+								FileReference PathToProject = XcodeProject.ProjectFilePath;
+								if (!XcodeProject.bHasLegacyProject && PerPlatformMode == XcodePerPlatformMode.OneWorkspacePerPlatform)
+								{
+									PathToProject = XcodeProject.ProjectFilePathForPlatform(Platform);
+								}
+								WorkspaceDataContent.Append(Ident + "      <FileRef" + ProjectFileGenerator.NewLine);
+								WorkspaceDataContent.Append(Ident + "         location = \"group:" + PathToProject.MakeRelativeTo(ProjectFileGenerator.PrimaryProjectPath) + "\">" + ProjectFileGenerator.NewLine);
+								WorkspaceDataContent.Append(Ident + "      </FileRef>" + ProjectFileGenerator.NewLine);
+							}
+
+							BuildableProjects.AddRange(SupportedProjects);
+
+							WorkspaceDataContent.Append(Ident + "   </Group>" + ProjectFileGenerator.NewLine);
 						}
+					};
+				AddProjectsFunction(RootFolder.SubFolders, "");
 
-						BuildableProjects.AddRange(SupportedProjects);
-
-						WorkspaceDataContent.Append(Ident + "   </Group>" + ProjectFileGenerator.NewLine);
-					}
-				};
-			AddProjectsFunction(RootFolder.SubFolders, "");
-			
-			WorkspaceDataContent.Append("</Workspace>" + ProjectFileGenerator.NewLine);
+				WorkspaceDataContent.Append("</Workspace>" + ProjectFileGenerator.NewLine);
 
-			// Also, update project's schemes index so that the schemes are in a sensible order
-			// (Game, Editor, Client, Server, Programs)
-			int SchemeIndex = 0;
-			BuildableProjects.Sort((ProjA, ProjB) => {
+				// Also, update project's schemes index so that the schemes are in a sensible order
+				// (Game, Editor, Client, Server, Programs)
+				int SchemeIndex = 0;
+				BuildableProjects.Sort((ProjA, ProjB) =>
+				{
 
-				ProjectTarget TargetA = ProjA.ProjectTargets.OfType<ProjectTarget>().OrderBy(T => T.TargetRules!.Type).First();
-				ProjectTarget TargetB = ProjB.ProjectTargets.OfType<ProjectTarget>().OrderBy(T => T.TargetRules!.Type).First();
+					ProjectTarget TargetA = ProjA.ProjectTargets.OfType<ProjectTarget>().OrderBy(T => T.TargetRules!.Type).First();
+					ProjectTarget TargetB = ProjB.ProjectTargets.OfType<ProjectTarget>().OrderBy(T => T.TargetRules!.Type).First();
 
-				TargetType TypeA = TargetA.TargetRules!.Type;
-				TargetType TypeB = TargetB.TargetRules!.Type;
+					TargetType TypeA = TargetA.TargetRules!.Type;
+					TargetType TypeB = TargetB.TargetRules!.Type;
 
-				if (TypeA != TypeB)
-				{
-					return TypeA.CompareTo(TypeB);
-				}
+					if (TypeA != TypeB)
+					{
+						return TypeA.CompareTo(TypeB);
+					}
 
-				return TargetA.Name.CompareTo(TargetB.Name);
-			});
+					return TargetA.Name.CompareTo(TargetB.Name);
+				});
 
-			foreach (ProjectFile XcodeProject in BuildableProjects)
-			{
-				FileReference SchemeManagementFile = XcodeProject.ProjectFilePath + "/xcuserdata/" + Environment.UserName + ".xcuserdatad/xcschemes/xcschememanagement.plist";
-				if (FileReference.Exists(SchemeManagementFile))
+				foreach (XcodeProjectXcconfig.XcodeProjectFile XcodeProject in BuildableProjects)
 				{
-					string SchemeManagementContent = FileReference.ReadAllText(SchemeManagementFile);
-					SchemeManagementContent = SchemeManagementContent.Replace("<key>orderHint</key>\n\t\t\t<integer>1</integer>", "<key>orderHint</key>\n\t\t\t<integer>" + SchemeIndex.ToString() + "</integer>");
-					FileReference.WriteAllText(SchemeManagementFile, SchemeManagementContent);
-					SchemeIndex++;
+					FileReference SchemeManagementFile = XcodeProject.ProjectFilePathForPlatform(Platform) + "/xcuserdata/" + Environment.UserName + ".xcuserdatad/xcschemes/xcschememanagement.plist";
+					if (FileReference.Exists(SchemeManagementFile))
+					{
+						string SchemeManagementContent = FileReference.ReadAllText(SchemeManagementFile);
+						SchemeManagementContent = SchemeManagementContent.Replace("<key>orderHint</key>\n\t\t\t<integer>1</integer>", "<key>orderHint</key>\n\t\t\t<integer>" + SchemeIndex.ToString() + "</integer>");
+						FileReference.WriteAllText(SchemeManagementFile, SchemeManagementContent);
+						SchemeIndex++;
+					}
 				}
-			}
 
-			string ProjectName = PrimaryProjectName;
-			//if (ProjectFilePlatform != XcodeProjectFilePlatform.All)
-			//{
-			//	ProjectName += ProjectFilePlatform == XcodeProjectFilePlatform.Mac ? "_Mac" : (ProjectFilePlatform == XcodeProjectFilePlatform.iOS ? "_IOS" : "_TVOS");
-			//}
-			string WorkspaceDataFilePath = PrimaryProjectPath + "/" + ProjectName + ".xcworkspace/contents.xcworkspacedata";
-			bSuccess = WriteFileIfChanged(WorkspaceDataFilePath, WorkspaceDataContent.ToString(), Logger, new UTF8Encoding());
-			if (bSuccess)
-			{
-				string WorkspaceSettingsFilePath = PrimaryProjectPath + "/" + ProjectName + ".xcworkspace/xcuserdata/" + Environment.UserName + ".xcuserdatad/WorkspaceSettings.xcsettings";
-				bSuccess = WriteWorkspaceSettingsFile(WorkspaceSettingsFilePath, Logger);
-				string WorkspaceSharedSettingsFilePath = PrimaryProjectPath + "/" + ProjectName + ".xcworkspace/xcshareddata/WorkspaceSettings.xcsettings";
-				bSuccess = WriteWorkspaceSharedSettingsFile(WorkspaceSharedSettingsFilePath, Logger);
+				string ProjectName = PrimaryProjectNameForPlatform(Platform);
+				string WorkspaceDataFilePath = PrimaryProjectPath + "/" + ProjectName + ".xcworkspace/contents.xcworkspacedata";
+				Logger.LogInformation($"Writing xcode workspace {Path.GetDirectoryName(WorkspaceDataFilePath)}");
+				bSuccess = WriteFileIfChanged(WorkspaceDataFilePath, WorkspaceDataContent.ToString(), Logger, new UTF8Encoding());
+				if (bSuccess)
+				{
+					string WorkspaceSettingsFilePath = PrimaryProjectPath + "/" + ProjectName + ".xcworkspace/xcuserdata/" + Environment.UserName + ".xcuserdatad/WorkspaceSettings.xcsettings";
+					bSuccess = WriteWorkspaceSettingsFile(WorkspaceSettingsFilePath, Logger);
+					string WorkspaceSharedSettingsFilePath = PrimaryProjectPath + "/" + ProjectName + ".xcworkspace/xcshareddata/WorkspaceSettings.xcsettings";
+					bSuccess = WriteWorkspaceSharedSettingsFile(WorkspaceSharedSettingsFilePath, Logger);
 
-				// cache the location of the workspace, for users of this to know where the final workspace is
-				XCWorkspace = new FileReference(WorkspaceDataFilePath).Directory;
+					// cache the location of the workspace, for users of this to know where the final workspace is
+					XCWorkspace = new FileReference(WorkspaceDataFilePath).Directory;
+				}
 			}
 
-
 			return bSuccess;
 		}
 
@@ -274,6 +322,11 @@ namespace UnrealBuildTool
 		/// </summary>
 		static public List<UnrealTargetPlatform> XcodePlatforms = new();
 
+		static public List<UnrealTargetPlatform?> WorkspacePlatforms = new ();
+		static public List<UnrealTargetPlatform?> RunTargetPlatforms = new();
+		static public List<UnrealTargetPlatform?> NullPlatformList = new() { null };
+
+
 		/// <summary>
 		/// Should we generate only a run project (no build/index targets)
 		/// </summary>
@@ -306,6 +359,15 @@ namespace UnrealBuildTool
 				}
 			}
 
+			if (PerPlatformMode == XcodePerPlatformMode.OneWorkspacePerPlatform)
+			{
+				WorkspacePlatforms = XcodePlatforms.Select<UnrealTargetPlatform, UnrealTargetPlatform?>(x => x).ToList();
+			}
+			else
+			{
+				WorkspacePlatforms = NullPlatformList;
+			}
+
 			foreach (string CurArgument in Arguments)
 			{
 				if (CurArgument.Contains("-iOSDeployOnly", StringComparison.InvariantCultureIgnoreCase) ||
