- Fixed the location of the .target files for Programs that have .uprojects (to use Engine/Binaries instead of Engine/Programs/Foo/Binaries), and also Intermediates to go back under Engine/Intermediate

diff --git a/Engine/Source/Programs/UnrealBuildTool/Configuration/UEBuildTarget.cs b/Engine/Source/Programs/UnrealBuildTool/Configuration/UEBuildTarget.cs
--- a/Engine/Source/Programs/UnrealBuildTool/Configuration/UEBuildTarget.cs
+++ b/Engine/Source/Programs/UnrealBuildTool/Configuration/UEBuildTarget.cs
@@ -1180,10 +1180,18 @@ namespace UnrealBuildTool
 				}
 			}
 
+			// track where intermediates and receipt files will be saved under (Programs with .uproject files still want to put the receipt
+			// under Engine/Binaries, not Engine/Programs/Foo/Binaries
+			DirectoryReference OutputRootDirectory = Unreal.EngineDirectory;
+
 			// Figure out what the project directory is. If we have a uproject file, use that. Otherwise use the engine directory.
 			if (ProjectFile != null)
 			{
 				ProjectDirectory = ProjectFile.Directory;
+				if (Rules.Type != TargetType.Program)
+				{
+					OutputRootDirectory = ProjectDirectory;
+				}
 			}
 			else
 			{
diff --git a/Engine/Source/Programs/UnrealBuildTool/System/TargetMakefile.cs b/Engine/Source/Programs/UnrealBuildTool/System/TargetMakefile.cs
--- a/Engine/Source/Programs/UnrealBuildTool/System/TargetMakefile.cs
+++ b/Engine/Source/Programs/UnrealBuildTool/System/TargetMakefile.cs
@@ -921,7 +921,12 @@ namespace UnrealBuildTool
 		/// <returns>Path to the makefile</returns>
 		public static FileReference GetLocation(FileReference? ProjectFile, string TargetName, UnrealTargetPlatform Platform, string Architecture, UnrealTargetConfiguration Configuration)
 		{
-			DirectoryReference BaseDirectory = DirectoryReference.FromFile(ProjectFile) ?? Unreal.EngineDirectory;
+			DirectoryReference BaseDirectory = Unreal.EngineDirectory;
+			// Programs with .uprojects still want the Intermediate under Engine (see UEBuildTarget constructor for similar code) 
+			if (ProjectFile != null && !ProjectFile.ContainsName("Programs", 0))
+			{
+				BaseDirectory = ProjectFile.Directory;
+			}
 			return FileReference.Combine(BaseDirectory, UEBuildTarget.GetPlatformIntermediateFolder(Platform, Architecture, false), TargetName, Configuration.ToString(), "Makefile.bin");
 		}
 	}
