UnrealBuildTool: Suppress -Wdeprecated-anon-enum-enum-conversion clang warnings when compiling with c++20

#rnx

diff --git a/Engine/Source/Programs/UnrealBuildTool/ToolChain/ClangToolChain.cs b/Engine/Source/Programs/UnrealBuildTool/ToolChain/ClangToolChain.cs
--- a/Engine/Source/Programs/UnrealBuildTool/ToolChain/ClangToolChain.cs
+++ b/Engine/Source/Programs/UnrealBuildTool/ToolChain/ClangToolChain.cs
@@ -532,6 +532,11 @@ namespace UnrealBuildTool
 			Arguments.Add("-Wno-enum-enum-conversion");                 // https://clang.llvm.org/docs/DiagnosticsReference.html#wenum-enum-conversion					// ?? no reason given
 			Arguments.Add("-Wno-enum-float-conversion");                // https://clang.llvm.org/docs/DiagnosticsReference.html#wenum-float-conversion					// ?? no reason given
 
+			if (CompileEnvironment.CppStandard >= CppStandardVersion.Cpp20)
+			{
+				Arguments.Add("-Wno-deprecated-anon-enum-enum-conversion"); // https://clang.llvm.org/docs/DiagnosticsReference.html#wdeprecated-anon-enum-enum-conversion // new warning for C++20
+			}
+
 			if (CompilerVersionGreaterOrEqual(13, 0, 0))
 			{
 				Arguments.Add("-Wno-unused-but-set-variable");           // https://clang.llvm.org/docs/DiagnosticsReference.html#wunused-but-set-variable				// new warning for clang 13
