Fix remote mac builds not copying over the MarketplaceRules.dll and support files when engine is installed, due to the files being saved off to a dynamic directory generally under $HOME/AppData/Local/UnrealEngine/Intermediate/* on the windows side.

diff --git a/Engine/Source/Programs/UnrealBuildTool/ToolChain/RemoteMac.cs b/Engine/Source/Programs/UnrealBuildTool/ToolChain/RemoteMac.cs
--- a/Engine/Source/Programs/UnrealBuildTool/ToolChain/RemoteMac.cs
+++ b/Engine/Source/Programs/UnrealBuildTool/ToolChain/RemoteMac.cs
@@ -929,6 +929,14 @@ namespace UnrealBuildTool
 			if(Unreal.IsEngineInstalled())
 			{
 				EngineFilters.Add(FileReference.Combine(Unreal.EngineDirectory, "Build", "Rsync", "RsyncEngineInstalled.txt"));
+
+				// Upload MarketplaceRules.dll (if it exists) when in InstallBuild. The path is specified in RulesCompiler::CreateMarketplaceRulesAssembly()
+				DirectoryReference MarketplaceEngineDir = DirectoryReference.Combine(UnrealBuildTool.WritableEngineDirectory, "Intermediate", "Build", "BuildRules");
+				if(DirectoryReference.Exists(MarketplaceEngineDir))
+				{
+					Logger.LogInformation("[Remote] Uploading MarketplaceRules Engine files located in {MarketplaceEngineDir}", MarketplaceEngineDir);
+					UploadDirectory(MarketplaceEngineDir, Logger);
+				}
 			}
 			EngineFilters.Add(FileReference.Combine(Unreal.EngineDirectory, "Build", "Rsync", "RsyncEngine.txt"));
 
