Fix UBT compile warning due to name change from UnrealBuildTool->Unreal caused by merging in CL23535789

diff --git a/Engine/Source/Programs/UnrealBuildTool/ToolChain/RemoteMac.cs b/Engine/Source/Programs/UnrealBuildTool/ToolChain/RemoteMac.cs
--- a/Engine/Source/Programs/UnrealBuildTool/ToolChain/RemoteMac.cs
+++ b/Engine/Source/Programs/UnrealBuildTool/ToolChain/RemoteMac.cs
@@ -931,7 +931,7 @@ namespace UnrealBuildTool
 				EngineFilters.Add(FileReference.Combine(Unreal.EngineDirectory, "Build", "Rsync", "RsyncEngineInstalled.txt"));
 
 				// Upload MarketplaceRules.dll (if it exists) when in InstallBuild. The path is specified in RulesCompiler::CreateMarketplaceRulesAssembly()
-				DirectoryReference MarketplaceEngineDir = DirectoryReference.Combine(UnrealBuildTool.WritableEngineDirectory, "Intermediate", "Build", "BuildRules");
+				DirectoryReference MarketplaceEngineDir = DirectoryReference.Combine(Unreal.WritableEngineDirectory, "Intermediate", "Build", "BuildRules");
 				if(DirectoryReference.Exists(MarketplaceEngineDir))
 				{
 					Logger.LogInformation("[Remote] Uploading MarketplaceRules Engine files located in {MarketplaceEngineDir}", MarketplaceEngineDir);
