- Fixed issue with new variable not being used (had a file across CLs and one change didn't make it into the checked in file)

diff --git a/Engine/Source/Programs/UnrealBuildTool/Configuration/UEBuildTarget.cs b/Engine/Source/Programs/UnrealBuildTool/Configuration/UEBuildTarget.cs
--- a/Engine/Source/Programs/UnrealBuildTool/Configuration/UEBuildTarget.cs
+++ b/Engine/Source/Programs/UnrealBuildTool/Configuration/UEBuildTarget.cs
@@ -1199,7 +1199,7 @@ namespace UnrealBuildTool
 			}
 
 			// Build the project intermediate directory
-			ProjectIntermediateDirectory = DirectoryReference.Combine(ProjectDirectory, PlatformIntermediateFolder, TargetName, Configuration.ToString());
+			ProjectIntermediateDirectory = DirectoryReference.Combine(OutputRootDirectory, PlatformIntermediateFolder, TargetName, Configuration.ToString());
 
 			// Build the engine intermediate directory. If we're building agnostic engine binaries, we can use the engine intermediates folder. Otherwise we need to use the project intermediates directory.
 			if (!bUseSharedBuildEnvironment)
@@ -1216,7 +1216,7 @@ namespace UnrealBuildTool
 			}
 
 			// Get the receipt path for this target
-			ReceiptFileName = TargetReceipt.GetDefaultPath(ProjectDirectory, TargetName, Platform, Configuration, Architecture);
+			ReceiptFileName = TargetReceipt.GetDefaultPath(OutputRootDirectory, TargetName, Platform, Configuration, Architecture);
 
 			// Read the project descriptor
 			if (ProjectFile != null)
