UnrealBuildTool: Add missing FileInlineGenCPPMap to CppEnvironment copy constructor. Via UDN.

#rnx

diff --git a/Engine/Source/Programs/UnrealBuildTool/System/CppCompileEnvironment.cs b/Engine/Source/Programs/UnrealBuildTool/System/CppCompileEnvironment.cs
--- a/Engine/Source/Programs/UnrealBuildTool/System/CppCompileEnvironment.cs
+++ b/Engine/Source/Programs/UnrealBuildTool/System/CppCompileEnvironment.cs
@@ -601,6 +601,7 @@ namespace UnrealBuildTool
 			bCheckSystemHeadersForModification = Other.bCheckSystemHeadersForModification;
 			ForceIncludeFiles.AddRange(Other.ForceIncludeFiles);
 			AdditionalPrerequisites.AddRange(Other.AdditionalPrerequisites);
+			FileInlineGenCPPMap = new Dictionary<FileItem, List<FileItem>>(Other.FileInlineGenCPPMap);
 			Definitions.AddRange(Other.Definitions);
 			AdditionalArguments = Other.AdditionalArguments;
 			AdditionalFrameworks.AddRange(Other.AdditionalFrameworks);
