PR #9942: Added support for coroutines (enabled via bEnableCppCoroutinesForEvaluation) in Rider (Contributed by [REDACTED_PERSON])

diff --git a/Engine/Source/Programs/UnrealBuildTool/ProjectFiles/Project.cs b/Engine/Source/Programs/UnrealBuildTool/ProjectFiles/Project.cs
--- a/Engine/Source/Programs/UnrealBuildTool/ProjectFiles/Project.cs
+++ b/Engine/Source/Programs/UnrealBuildTool/ProjectFiles/Project.cs
@@ -248,6 +248,15 @@ namespace UnrealBuildTool
 			protected set;
 		}
 
+		/// <summary>
+		/// Information for IntelliSense: whether coroutines are enabled
+		/// </summary>
+		public bool IntelliSenseEnableCoroutines
+		{
+			get;
+			protected set;
+		}
+
 		/// All of the targets in this project.  All non-stub projects must have at least one target.
 		public readonly List<Project> ProjectTargets = new List<Project>();
 
@@ -364,6 +373,7 @@ namespace UnrealBuildTool
 				AddIntelliSenseIncludePaths(BuildEnvironment.UserIncludePaths, CompileEnvironment.UserIncludePaths);
 			}
 
+			SetIntelliSenseEnableCoroutines(CompileEnvironment.bEnableCoroutines);
 			SetIntelliSenseCppVersion(Module.Rules.CppStandard);
 		}
 
@@ -445,6 +455,15 @@ namespace UnrealBuildTool
 				}
 			}
 		}
+		
+		/// <summary>
+		/// Sets whether coroutines should be enabled or not
+		/// </summary>
+		/// <param name="Value">Should enable them or not</param>
+		public void SetIntelliSenseEnableCoroutines(bool Value)
+		{
+			IntelliSenseEnableCoroutines |= Value;
+		}
 
 		/// <summary>
 		/// Add the given project to the DepondsOn project list.
diff --git a/Engine/Source/Programs/UnrealBuildTool/ProjectFiles/Rider/RiderProjectFile.cs b/Engine/Source/Programs/UnrealBuildTool/ProjectFiles/Rider/RiderProjectFile.cs
--- a/Engine/Source/Programs/UnrealBuildTool/ProjectFiles/Rider/RiderProjectFile.cs
+++ b/Engine/Source/Programs/UnrealBuildTool/ProjectFiles/Rider/RiderProjectFile.cs
@@ -764,7 +764,8 @@ namespace UnrealBuildTool
 				bUseStaticCRT = CompileEnvironment.bUseStaticCRT,
 				PrecompiledHeaderAction = CompileEnvironment.PrecompiledHeaderAction.ToString(),
 				PrecompiledHeaderFile = CompileEnvironment.PrecompiledHeaderFile?.ToString(),
-				ForceIncludeFiles = CompileEnvironment.ForceIncludeFiles.Select(Item => Item.ToString()).ToList()
+				ForceIncludeFiles = CompileEnvironment.ForceIncludeFiles.Select(Item => Item.ToString()).ToList(),
+				bEnableCoroutines = CompileEnvironment.bEnableCoroutines
 			};
 
 			if (CurrentTarget!.Platform.IsInGroup(UnrealPlatformGroup.Windows))
diff --git a/Engine/Source/Programs/UnrealBuildTool/ProjectFiles/Rider/ToolchainInfo.cs b/Engine/Source/Programs/UnrealBuildTool/ProjectFiles/Rider/ToolchainInfo.cs
--- a/Engine/Source/Programs/UnrealBuildTool/ProjectFiles/Rider/ToolchainInfo.cs
+++ b/Engine/Source/Programs/UnrealBuildTool/ProjectFiles/Rider/ToolchainInfo.cs
@@ -27,7 +27,8 @@ namespace UnrealBuildTool
 		public string? PrecompiledHeaderFile;
 		public List<string>? ForceIncludeFiles;
 		public string? Compiler;
-		public bool bStrictConformanceMode = false;
+		public bool bStrictConformanceMode;
+		public bool bEnableCoroutines;
 
 		public IEnumerable<Tuple<string, object?>> GetDiff(ToolchainInfo Other)
 		{
@@ -58,7 +59,16 @@ namespace UnrealBuildTool
 		{
 			if (ReferenceEquals(null, Other)) return false;
 			if (ReferenceEquals(this, Other)) return true;
-			return PrecompiledHeaderAction == Other.PrecompiledHeaderAction && CppStandard == Other.CppStandard && bUseRTTI == Other.bUseRTTI && bEnableExceptions == Other.bEnableExceptions && bIsBuildingLibrary == Other.bIsBuildingLibrary && bIsBuildingDLL == Other.bIsBuildingDLL && Architecture == Other.Architecture && Configuration == Other.Configuration && bOptimizeCode == Other.bOptimizeCode && bUseInlining == Other.bUseInlining && bUseUnity == Other.bUseUnity && bCreateDebugInfo == Other.bCreateDebugInfo && bUseAVX == Other.bUseAVX && bUseDebugCRT == Other.bUseDebugCRT && bUseStaticCRT == Other.bUseStaticCRT && PrecompiledHeaderFile == Other.PrecompiledHeaderFile && Equals(ForceIncludeFiles, Other.ForceIncludeFiles) && Compiler == Other.Compiler && bStrictConformanceMode == Other.bStrictConformanceMode;
+			return CppStandard == Other.CppStandard && bUseRTTI == Other.bUseRTTI &&
+				bEnableExceptions == Other.bEnableExceptions && bIsBuildingLibrary == Other.bIsBuildingLibrary &&
+				bIsBuildingDLL == Other.bIsBuildingDLL && Architecture == Other.Architecture &&
+				Configuration == Other.Configuration && bOptimizeCode == Other.bOptimizeCode &&
+				bUseInlining == Other.bUseInlining && bUseUnity == Other.bUseUnity &&
+				bCreateDebugInfo == Other.bCreateDebugInfo && bUseAVX == Other.bUseAVX &&
+				bUseDebugCRT == Other.bUseDebugCRT && bUseStaticCRT == Other.bUseStaticCRT &&
+				PrecompiledHeaderAction == Other.PrecompiledHeaderAction && PrecompiledHeaderFile == Other.PrecompiledHeaderFile &&
+				Equals(ForceIncludeFiles, Other.ForceIncludeFiles) && Compiler == Other.Compiler &&
+				bStrictConformanceMode == Other.bStrictConformanceMode && bEnableCoroutines == Other.bEnableCoroutines;
 		}
 
 		public override bool Equals(object? Obj)
@@ -91,6 +101,7 @@ namespace UnrealBuildTool
 				HashCode = (HashCode * 397) ^ (ForceIncludeFiles != null ? ForceIncludeFiles.GetHashCode() : 0);
 				HashCode = (HashCode * 397) ^ (Compiler != null ? Compiler.GetHashCode() : 0);
 				HashCode = (HashCode * 397) ^ bStrictConformanceMode.GetHashCode();
+				HashCode = (HashCode * 397) ^ bEnableCoroutines.GetHashCode();
 				return HashCode;
 			}
 		}
diff --git a/Engine/Source/Programs/UnrealBuildTool/ProjectFiles/VisualStudio/VCProject.cs b/Engine/Source/Programs/UnrealBuildTool/ProjectFiles/VisualStudio/VCProject.cs
--- a/Engine/Source/Programs/UnrealBuildTool/ProjectFiles/VisualStudio/VCProject.cs
+++ b/Engine/Source/Programs/UnrealBuildTool/ProjectFiles/VisualStudio/VCProject.cs
@@ -616,17 +616,32 @@ namespace UnrealBuildTool
 				case CppStandardVersion.Cpp17:
 					return "/std:c++17";
 				case CppStandardVersion.Cpp20:
+					return "/std:c++20";
 				case CppStandardVersion.Latest:
 					return "/std:c++latest";
-				// Will be added when MSVC is feature-complete.
-				// https://docs.microsoft.com/en-us/cpp/build/reference/std-specify-language-standard-version?view=msvc-160
-				// case CppStandardVersion.Cpp20:
-				//	return "/std:c++20";
 				default:
 					throw new BuildException($"Unsupported C++ standard type set: {Version}");
 			}
 		}
 
+		/// <summary>
+		/// Gets compiler switch for specifying in AdditionalOptions in .vcxproj file for coroutines support
+		/// </summary>
+		private string GetEnableCoroutinesArgument()
+		{
+			if (IntelliSenseEnableCoroutines)
+			{
+				if (GetCompilerForIntellisense().IsMSVC())
+				{
+					return "/await:strict";
+				}
+
+				return "-fcoroutines-ts";
+			}
+
+			return string.Empty;
+		}
+
 		HashSet<string>? InvalidConfigPlatformNames;
 		List<ProjectConfigAndTargetCombination>? ProjectConfigAndTargetCombinations;
 
@@ -1116,8 +1131,8 @@ namespace UnrealBuildTool
 				VCProjectFileContent.AppendLine("    <IncludePath>$(IncludePath){0}</IncludePath>", (SharedIncludeSearchPaths.Length > 0 ? (";" + SharedIncludeSearchPaths) : ""));
 				VCProjectFileContent.AppendLine("    <NMakeForcedIncludes>$(NMakeForcedIncludes)</NMakeForcedIncludes>");
 				VCProjectFileContent.AppendLine("    <NMakeAssemblySearchPath>$(NMakeAssemblySearchPath)</NMakeAssemblySearchPath>");
-				VCProjectFileContent.AppendLine("    <AdditionalOptions>{0}</AdditionalOptions>",
-					GetCppStandardCompileArgument(GetIntelliSenseCppVersion()));
+				VCProjectFileContent.AppendLine("    <AdditionalOptions>{0} {1}</AdditionalOptions>",
+					GetCppStandardCompileArgument(GetIntelliSenseCppVersion()), GetEnableCoroutinesArgument());
 				VCProjectFileContent.AppendLine("  </PropertyGroup>");
 			}
 
