PR #9958: Fixing generating Rider project model (Contributed by [REDACTED_PERSON])

diff --git a/Engine/Source/Programs/UnrealBuildTool/ProjectFiles/Rider/RiderProjectFile.cs b/Engine/Source/Programs/UnrealBuildTool/ProjectFiles/Rider/RiderProjectFile.cs
--- a/Engine/Source/Programs/UnrealBuildTool/ProjectFiles/Rider/RiderProjectFile.cs
+++ b/Engine/Source/Programs/UnrealBuildTool/ProjectFiles/Rider/RiderProjectFile.cs
@@ -73,12 +73,7 @@ namespace UnrealBuildTool
 			DirectoryReference ProjectRootFolder = RootPath;
 			List<TargetEntry> FileToTarget = new List<TargetEntry>();
 			foreach (UnrealTargetPlatform Platform in InPlatforms)
-			{
-				if (!IsPlatformInHostGroup(Platform))
-				{
-					continue;
-				}
-				
+			{				
 				foreach (UnrealTargetConfiguration Configuration in InConfigurations)
 				{
 					foreach (ProjectTarget ProjectTarget in ProjectTargets.OfType<ProjectTarget>())
