- Fixed issue with Programs writing to Engine/Binaries when another project's .uproject is passed in as context for compiling the program.

diff --git a/Engine/Source/Programs/UnrealBuildTool/Configuration/UEBuildTarget.cs b/Engine/Source/Programs/UnrealBuildTool/Configuration/UEBuildTarget.cs
--- a/Engine/Source/Programs/UnrealBuildTool/Configuration/UEBuildTarget.cs
+++ b/Engine/Source/Programs/UnrealBuildTool/Configuration/UEBuildTarget.cs
@@ -1188,7 +1188,9 @@ namespace UnrealBuildTool
 			if (ProjectFile != null)
 			{
 				ProjectDirectory = ProjectFile.Directory;
-				if (Rules.Type != TargetType.Program)
+				// programs are often compiled with another project passed in for context, so if it's using a _different_ uproject, use that uproject
+				// for the output location
+				if (Rules.Type != TargetType.Program || string.Compare(ProjectFile.GetFileNameWithoutAnyExtensions(), Rules.Name, true) != 0)
 				{
 					OutputRootDirectory = ProjectDirectory;
 				}
