PR #8562: Add further debug output for makefile argument changes. (Contributed by [REDACTED_PERSON])

diff --git a/Engine/Source/Programs/UnrealBuildTool/System/TargetMakefile.cs b/Engine/Source/Programs/UnrealBuildTool/System/TargetMakefile.cs
--- a/Engine/Source/Programs/UnrealBuildTool/System/TargetMakefile.cs
+++ b/Engine/Source/Programs/UnrealBuildTool/System/TargetMakefile.cs
@@ -629,6 +629,8 @@ namespace UnrealBuildTool
 				// Check if the arguments are different
 				if(!Enumerable.SequenceEqual(Makefile.AdditionalArguments!, Arguments))
 				{
+					Logger.LogDebug("Old command line arguments:\n", string.Join(' ', Makefile.AdditionalArguments!));
+					Logger.LogDebug("New command line arguments:\n", string.Join(' ', Arguments));
 					ReasonNotLoaded = "command line arguments changed";
 					return null;
 				}
