PR #8657: Put tool projects in proper subfolders (Contributed by [REDACTED_PERSON])

diff --git a/Engine/Source/Programs/UnrealBuildTool/ProjectFiles/ProjectFileGenerator.cs b/Engine/Source/Programs/UnrealBuildTool/ProjectFiles/ProjectFileGenerator.cs
--- a/Engine/Source/Programs/UnrealBuildTool/ProjectFiles/ProjectFileGenerator.cs
+++ b/Engine/Source/Programs/UnrealBuildTool/ProjectFiles/ProjectFileGenerator.cs
@@ -548,6 +548,7 @@ namespace UnrealBuildTool
 		/// </summary>
 		void DiscoverCSharpProgramProjects(List<DirectoryReference> AllEngineDirectories, List<FileReference> AllGameProjects, PrimaryProjectFolder ProgramsFolder, ILogger Logger)
 		{
+			Debugger.Launch();
 			List<FileReference> FoundProjects = new List<FileReference>();
 			List<DirectoryReference> ProjectDirs = new List<DirectoryReference>();
 
@@ -587,7 +588,22 @@ namespace UnrealBuildTool
 							Project.ShouldBuildByDefaultForSolutionTargets = true;
 
 							AddExistingProjectFile(Project, bForceDevelopmentConfiguration: false);
-							ProgramsFolder.ChildProjects.Add(Project);
+
+							PrimaryProjectFolder ProjectFolderToAddTo = ProgramsFolder;
+							if (FoundProject.Directory != null && FoundProject.Directory.ParentDirectory != null)
+							{
+								string ProjectRelativePath = FoundProject.Directory.ParentDirectory.MakeRelativeTo(ProjectDir);
+								if (ProjectRelativePath != ".")
+								{
+									string[] Directories = ProjectRelativePath.Split(Path.DirectorySeparatorChar);
+									foreach (string IntermediateDir in Directories)
+									{
+										ProjectFolderToAddTo = ProjectFolderToAddTo.AddSubFolder(IntermediateDir);
+									}
+								}
+							}
+
+							ProjectFolderToAddTo.ChildProjects.Add(Project);
 						}
 						else
 						{
