PR #8992: Add a GetPlugin method (Contributed by [REDACTED_PERSON])

diff --git a/Engine/Source/Programs/UnrealBuildTool/System/Plugins.cs b/Engine/Source/Programs/UnrealBuildTool/System/Plugins.cs
--- a/Engine/Source/Programs/UnrealBuildTool/System/Plugins.cs
+++ b/Engine/Source/Programs/UnrealBuildTool/System/Plugins.cs
@@ -248,6 +248,26 @@ namespace UnrealBuildTool
 			DirectoryItem.ResetAllCachedInfo_SLOW();
 		}
 
+		/// <summary>
+		/// Add a means to retrieve a plugin to determine if it is present.
+		/// </summary>
+		/// <param name="PluginName">The name of the plugin.</param>
+		/// <returns>The cached plugin info, or null if not found.</returns>
+		public static PluginInfo? GetPlugin(string PluginName)
+		{
+			foreach (var Pair in PluginInfoCache)
+			{
+				foreach (var Plugin in Pair.Value)
+				{
+					if (string.Equals(Plugin.Name, PluginName, StringComparison.InvariantCultureIgnoreCase))
+					{
+						return Plugin;
+					}
+				}
+			}
+			return null;
+		}
+
 		/// <summary>
 		/// Returns a filtered list of plugins as a name:plugin dictionary to ensure that any game plugins override engine plugins with the same
 		/// name, and otherwise that no two plugins with the same name exist. 
