Added the ability to throw errors or warnings if the generated header is not properly included or if a regular or namespaced enum does not have an underlying type specified. These are both silent at the moment.

diff --git a/Engine/Source/Programs/Shared/EpicGames.UHT/Parsers/UhtEnumParser.cs b/Engine/Source/Programs/Shared/EpicGames.UHT/Parsers/UhtEnumParser.cs
--- a/Engine/Source/Programs/Shared/EpicGames.UHT/Parsers/UhtEnumParser.cs
+++ b/Engine/Source/Programs/Shared/EpicGames.UHT/Parsers/UhtEnumParser.cs
@@ -120,6 +120,27 @@ namespace EpicGames.UHT.Parsers
 							break;
 					}
 
+					if (enumObject.CppForm != UhtEnumCppForm.EnumClass && enumObject.UnderlyingType == UhtEnumUnderlyingType.Unspecified)
+					{
+						UhtIssueBehavior enumUnderlyingTypeBehavior = enumObject.Package.IsPartOfEngine ? topScope.Session.Config!.EngineEnumUnderlyingTypeNotSet
+							: topScope.Session.Config!.NonEngineEnumUnderlyingTypeNotSet;
+
+						string logMessage = $"Underlying type must be specified.";
+						switch (enumUnderlyingTypeBehavior)
+						{
+							case UhtIssueBehavior.AllowSilently:
+								break;
+
+							case UhtIssueBehavior.AllowAndLog:
+								topScope.TokenReader.LogTrace(enumToken.InputLine, logMessage);
+								break;
+
+							default:
+								topScope.TokenReader.LogError(enumToken.InputLine, logMessage);
+								break;
+						}
+					}
+
 					tokenContext.Reset($"UENUM {enumObject.SourceName}");
 
 					topScope.AddModuleRelativePathToMetaData();
diff --git a/Engine/Source/Programs/Shared/EpicGames.UHT/Types/Properties/UhtObjectPropertyBase.cs b/Engine/Source/Programs/Shared/EpicGames.UHT/Types/Properties/UhtObjectPropertyBase.cs
--- a/Engine/Source/Programs/Shared/EpicGames.UHT/Types/Properties/UhtObjectPropertyBase.cs
+++ b/Engine/Source/Programs/Shared/EpicGames.UHT/Types/Properties/UhtObjectPropertyBase.cs
@@ -233,8 +233,8 @@ namespace EpicGames.UHT.Types
 		/// <param name="typeStartPos">Starting character position of the type</param>
 		/// <param name="alternativeTypeDesc">Suggested alternate declaration</param>
 		/// <exception cref="UhtIceException">Thrown if the behavior type is unexpected</exception>
-		public static void ConditionalLogPointerUsage(UhtPropertySettings propertySettings, UhtPointerMemberBehavior engineBehavior, UhtPointerMemberBehavior enginePluginBehavior,
-			UhtPointerMemberBehavior nonEngineBehavior, string pointerTypeDesc, IUhtTokenReader tokenReader, int typeStartPos, string? alternativeTypeDesc)
+		public static void ConditionalLogPointerUsage(UhtPropertySettings propertySettings, UhtIssueBehavior engineBehavior, UhtIssueBehavior enginePluginBehavior,
+			UhtIssueBehavior nonEngineBehavior, string pointerTypeDesc, IUhtTokenReader tokenReader, int typeStartPos, string? alternativeTypeDesc)
 		{
 			if (propertySettings.PropertyCategory != UhtPropertyCategory.Member)
 			{
@@ -242,7 +242,7 @@ namespace EpicGames.UHT.Types
 			}
 
 			UhtPackage package = propertySettings.Outer.Package;
-			UhtPointerMemberBehavior behavior = nonEngineBehavior;
+			UhtIssueBehavior behavior = nonEngineBehavior;
 			if (package.IsPartOfEngine)
 			{
 				if (package.IsPlugin)
@@ -255,7 +255,7 @@ namespace EpicGames.UHT.Types
 				}
 			}
 
-			if (behavior == UhtPointerMemberBehavior.AllowSilently)
+			if (behavior == UhtIssueBehavior.AllowSilently)
 			{
 				return;
 			}
@@ -267,7 +267,7 @@ namespace EpicGames.UHT.Types
 
 			switch (behavior)
 			{
-				case UhtPointerMemberBehavior.Disallow:
+				case UhtIssueBehavior.Disallow:
 					if (!String.IsNullOrEmpty(alternativeTypeDesc))
 					{
 						tokenReader.LogError($"{pointerTypeDesc} usage in member declaration detected [[[{type}]]].  This is disallowed for the target/module, consider {alternativeTypeDesc} as an alternative.");
@@ -278,7 +278,7 @@ namespace EpicGames.UHT.Types
 					}
 					break;
 
-				case UhtPointerMemberBehavior.AllowAndLog:
+				case UhtIssueBehavior.AllowAndLog:
 					if (!String.IsNullOrEmpty(alternativeTypeDesc))
 					{
 						tokenReader.LogTrace($"{pointerTypeDesc} usage in member declaration detected [[[{type}]]].  Consider {alternativeTypeDesc} as an alternative.");
diff --git a/Engine/Source/Programs/Shared/EpicGames.UHT/Utils/UhtConfig.cs b/Engine/Source/Programs/Shared/EpicGames.UHT/Utils/UhtConfig.cs
--- a/Engine/Source/Programs/Shared/EpicGames.UHT/Utils/UhtConfig.cs
+++ b/Engine/Source/Programs/Shared/EpicGames.UHT/Utils/UhtConfig.cs
@@ -7,23 +7,22 @@ namespace EpicGames.UHT.Utils
 {
 
 	/// <summary>
-	/// Describes how pointers will generate warnings or errors
+	/// Describes whether issues will generate warnings or errors
 	/// </summary>
-	public enum UhtPointerMemberBehavior
+	public enum UhtIssueBehavior
 	{
-
 		/// <summary>
 		/// An error will be generated.
 		/// </summary>
 		Disallow,
 
 		/// <summary>
-		/// Ignore the pointer
+		/// Ignore the issue.
 		/// </summary>
 		AllowSilently,
 
 		/// <summary>
-		/// Log a warning about the pointer
+		/// Log a warning.
 		/// </summary>
 		AllowAndLog,
 	};
@@ -42,32 +41,32 @@ namespace EpicGames.UHT.Utils
 		/// <summary>
 		/// Pointer warning for native pointers in the engine
 		/// </summary>
-		public UhtPointerMemberBehavior EngineNativePointerMemberBehavior { get; }
+		public UhtIssueBehavior EngineNativePointerMemberBehavior { get; }
 
 		/// <summary>
 		/// Pointer warning for object pointers in the engine
 		/// </summary>
-		public UhtPointerMemberBehavior EngineObjectPtrMemberBehavior { get; }
+		public UhtIssueBehavior EngineObjectPtrMemberBehavior { get; }
 
 		/// <summary>
 		/// Pointer warning for native pointers in engine plugins
 		/// </summary>
-		public UhtPointerMemberBehavior EnginePluginNativePointerMemberBehavior { get; }
+		public UhtIssueBehavior EnginePluginNativePointerMemberBehavior { get; }
 
 		/// <summary>
 		/// Pointer warning for object pointers in engine plugins
 		/// </summary>
-		public UhtPointerMemberBehavior EnginePluginObjectPtrMemberBehavior { get; }
+		public UhtIssueBehavior EnginePluginObjectPtrMemberBehavior { get; }
 
 		/// <summary>
 		/// Pointer warning for native pointers outside the engine
 		/// </summary>
-		public UhtPointerMemberBehavior NonEngineNativePointerMemberBehavior { get; }
+		public UhtIssueBehavior NonEngineNativePointerMemberBehavior { get; }
 
 		/// <summary>
 		/// Pointer warning for object pointers outside the engine
 		/// </summary>
-		public UhtPointerMemberBehavior NonEngineObjectPtrMemberBehavior { get; }
+		public UhtIssueBehavior NonEngineObjectPtrMemberBehavior { get; }
 
 		/// <summary>
 		/// If true, deprecation warnings should be shown
@@ -84,6 +83,26 @@ namespace EpicGames.UHT.Utils
 		/// </summary>
 		public bool AreRigVMUInterfaceProeprtiesEnabled { get; }
 
+		/// <summary>
+		/// Behavior for when generated headers aren't properly included in engine code.
+		/// </summary>
+		public UhtIssueBehavior EngineMissingGeneratedHeaderIncludeBehavior { get; }
+
+		/// <summary>
+		/// Behavior for when generated headers aren't properly included in non-engine code.
+		/// </summary>
+		public UhtIssueBehavior NonEngineMissingGeneratedHeaderIncludeBehavior { get; }
+
+		/// <summary>
+		/// Behavior for when enum underlying types aren't set for engine code.
+		/// </summary>
+		public UhtIssueBehavior EngineEnumUnderlyingTypeNotSet { get; }
+
+		/// <summary>
+		/// Behavior for when enum underlying types aren't set for non engine code.
+		/// </summary>
+		public UhtIssueBehavior NonEngineEnumUnderlyingTypeNotSet { get; }
+
 		/// <summary>
 		/// If the token references a remapped identifier, update the value in the token 
 		/// </summary>
diff --git a/Engine/Source/Programs/UnrealBuildTool/Modes/UnrealHeaderToolMode.cs b/Engine/Source/Programs/UnrealBuildTool/Modes/UnrealHeaderToolMode.cs
--- a/Engine/Source/Programs/UnrealBuildTool/Modes/UnrealHeaderToolMode.cs
+++ b/Engine/Source/Programs/UnrealBuildTool/Modes/UnrealHeaderToolMode.cs
@@ -64,32 +64,32 @@ namespace UnrealBuildTool.Modes
 		/// <summary>
 		/// Internal version of pointer warning for native pointers in the engine
 		/// </summary>
-		private readonly UhtPointerMemberBehavior _engineNativePointerMemberBehavior = UhtPointerMemberBehavior.AllowSilently;
+		private readonly UhtIssueBehavior _engineNativePointerMemberBehavior = UhtIssueBehavior.AllowSilently;
 
 		/// <summary>
 		/// Internal version of pointer warning for object pointers in the engine
 		/// </summary>
-		private readonly UhtPointerMemberBehavior _engineObjectPtrMemberBehavior = UhtPointerMemberBehavior.AllowSilently;
+		private readonly UhtIssueBehavior _engineObjectPtrMemberBehavior = UhtIssueBehavior.AllowSilently;
 
 		/// <summary>
 		/// Internal version of pointer warning for native pointers in engine plugins
 		/// </summary>
-		private readonly UhtPointerMemberBehavior _enginePluginNativePointerMemberBehavior = UhtPointerMemberBehavior.AllowSilently;
+		private readonly UhtIssueBehavior _enginePluginNativePointerMemberBehavior = UhtIssueBehavior.AllowSilently;
 
 		/// <summary>
 		/// Internal version of pointer warning for object pointers in engine plugins
 		/// </summary>
-		private readonly UhtPointerMemberBehavior _enginePluginObjectPtrMemberBehavior = UhtPointerMemberBehavior.AllowSilently;
+		private readonly UhtIssueBehavior _enginePluginObjectPtrMemberBehavior = UhtIssueBehavior.AllowSilently;
 
 		/// <summary>
 		/// Internal version of pointer warning for native pointers outside the engine
 		/// </summary>
-		private readonly UhtPointerMemberBehavior _nonEngineNativePointerMemberBehavior = UhtPointerMemberBehavior.AllowSilently;
+		private readonly UhtIssueBehavior _nonEngineNativePointerMemberBehavior = UhtIssueBehavior.AllowSilently;
 
 		/// <summary>
 		/// Internal version of pointer warning for object pointers outside the engine
 		/// </summary>
-		private readonly UhtPointerMemberBehavior _nonEngineObjectPtrMemberBehavior = UhtPointerMemberBehavior.AllowSilently;
+		private readonly UhtIssueBehavior _nonEngineObjectPtrMemberBehavior = UhtIssueBehavior.AllowSilently;
 
 		/// <summary>
 		/// If true, deprecation warnings should be shown
@@ -106,27 +106,47 @@ namespace UnrealBuildTool.Modes
 		/// </summary>
 		private readonly bool _areRigVMUInterfaceProeprtiesEnabled = false;
 
+		/// <summary>
+		/// Internal version of behavior when there is a missing generated header include in the engine code.
+		/// </summary>
+		private readonly UhtIssueBehavior _engineMissingGeneratedHeaderIncludeBehavior = UhtIssueBehavior.AllowSilently;
+
+		/// <summary>
+		/// Internal version of behavior when there is a missing generated header include in non engine code.
+		/// </summary>
+		private readonly UhtIssueBehavior _nonEngineMissingGeneratedHeaderIncludeBehavior = UhtIssueBehavior.AllowSilently;
+
+		/// <summary>
+		/// Internal version of the behavior set when the underlying type of a regular and namespaced enum isn't set in the engine code.
+		/// </summary>
+		private readonly UhtIssueBehavior _engineEnumUnderlyingTypeNotSet = UhtIssueBehavior.AllowSilently;
+
+		/// <summary>
+		/// Internal version of the behavior set when the underlying type of a regular and namespaced enum isn't set in the non engine code.
+		/// </summary>
+		private readonly UhtIssueBehavior _nonEngineEnumUnderlyingTypeNotSet = UhtIssueBehavior.AllowSilently;
+
 		#region IUhtConfig Implementation
 		/// <inheritdoc/>
 		public EGeneratedCodeVersion DefaultGeneratedCodeVersion => this._defaultGeneratedCodeVersion;
 
 		/// <inheritdoc/>
-		public UhtPointerMemberBehavior EngineNativePointerMemberBehavior => this._engineNativePointerMemberBehavior;
+		public UhtIssueBehavior EngineNativePointerMemberBehavior => this._engineNativePointerMemberBehavior;
 
 		/// <inheritdoc/>
-		public UhtPointerMemberBehavior EngineObjectPtrMemberBehavior => this._engineObjectPtrMemberBehavior;
+		public UhtIssueBehavior EngineObjectPtrMemberBehavior => this._engineObjectPtrMemberBehavior;
 
 		/// <inheritdoc/>
-		public UhtPointerMemberBehavior EnginePluginNativePointerMemberBehavior => this._enginePluginNativePointerMemberBehavior;
+		public UhtIssueBehavior EnginePluginNativePointerMemberBehavior => this._enginePluginNativePointerMemberBehavior;
 
 		/// <inheritdoc/>
-		public UhtPointerMemberBehavior EnginePluginObjectPtrMemberBehavior => this._enginePluginObjectPtrMemberBehavior;
+		public UhtIssueBehavior EnginePluginObjectPtrMemberBehavior => this._enginePluginObjectPtrMemberBehavior;
 
 		/// <inheritdoc/>
-		public UhtPointerMemberBehavior NonEngineNativePointerMemberBehavior => this._nonEngineNativePointerMemberBehavior;
+		public UhtIssueBehavior NonEngineNativePointerMemberBehavior => this._nonEngineNativePointerMemberBehavior;
 
 		/// <inheritdoc/>
-		public UhtPointerMemberBehavior NonEngineObjectPtrMemberBehavior => this._nonEngineObjectPtrMemberBehavior;
+		public UhtIssueBehavior NonEngineObjectPtrMemberBehavior => this._nonEngineObjectPtrMemberBehavior;
 
 		/// <summary>
 		/// If true, UObject properties are enabled in RigVM
@@ -143,6 +163,18 @@ namespace UnrealBuildTool.Modes
 		/// </summary>
 		public bool ShowDeprecations => this._showDeprecations;
 
+		/// <inheritdoc/>
+		public UhtIssueBehavior EngineMissingGeneratedHeaderIncludeBehavior => this._engineMissingGeneratedHeaderIncludeBehavior;
+
+		/// <inheritdoc/>
+		public UhtIssueBehavior NonEngineMissingGeneratedHeaderIncludeBehavior => this._nonEngineMissingGeneratedHeaderIncludeBehavior;
+
+		/// <inheritdoc/>
+		public UhtIssueBehavior EngineEnumUnderlyingTypeNotSet => this._engineEnumUnderlyingTypeNotSet;
+
+		/// <inheritdoc/>
+		public UhtIssueBehavior NonEngineEnumUnderlyingTypeNotSet => this._nonEngineEnumUnderlyingTypeNotSet;
+
 		/// <inheritdoc/>
 		public void RedirectTypeIdentifier(ref UhtToken Token)
 		{
@@ -227,15 +259,19 @@ namespace UnrealBuildTool.Modes
 			this._units = GetHashSet("UnrealHeaderTool", "Units", StringViewComparer.OrdinalIgnoreCase);
 			this._delegateParameterCountStrings = GetList("UnrealHeaderTool", "DelegateParameterCountStrings");
 			this._defaultGeneratedCodeVersion = GetGeneratedCodeVersion("UnrealHeaderTool", "DefaultGeneratedCodeVersion", EGeneratedCodeVersion.V1);
-			this._engineNativePointerMemberBehavior = GetPointerMemberBehavior("UnrealHeaderTool", "EngineNativePointerMemberBehavior", UhtPointerMemberBehavior.AllowSilently);
-			this._engineObjectPtrMemberBehavior = GetPointerMemberBehavior("UnrealHeaderTool", "EngineObjectPtrMemberBehavior", UhtPointerMemberBehavior.AllowSilently);
-			this._enginePluginNativePointerMemberBehavior = GetPointerMemberBehavior("UnrealHeaderTool", "EnginePluginNativePointerMemberBehavior", UhtPointerMemberBehavior.AllowSilently);
-			this._enginePluginObjectPtrMemberBehavior = GetPointerMemberBehavior("UnrealHeaderTool", "EnginePluginObjectPtrMemberBehavior", UhtPointerMemberBehavior.AllowSilently);
-			this._nonEngineNativePointerMemberBehavior = GetPointerMemberBehavior("UnrealHeaderTool", "NonEngineNativePointerMemberBehavior", UhtPointerMemberBehavior.AllowSilently);
-			this._nonEngineObjectPtrMemberBehavior = GetPointerMemberBehavior("UnrealHeaderTool", "NonEngineObjectPtrMemberBehavior", UhtPointerMemberBehavior.AllowSilently);
+			this._engineNativePointerMemberBehavior = GetPointerMemberBehavior("UnrealHeaderTool", "EngineNativePointerMemberBehavior", UhtIssueBehavior.AllowSilently);
+			this._engineObjectPtrMemberBehavior = GetPointerMemberBehavior("UnrealHeaderTool", "EngineObjectPtrMemberBehavior", UhtIssueBehavior.AllowSilently);
+			this._enginePluginNativePointerMemberBehavior = GetPointerMemberBehavior("UnrealHeaderTool", "EnginePluginNativePointerMemberBehavior", UhtIssueBehavior.AllowSilently);
+			this._enginePluginObjectPtrMemberBehavior = GetPointerMemberBehavior("UnrealHeaderTool", "EnginePluginObjectPtrMemberBehavior", UhtIssueBehavior.AllowSilently);
+			this._nonEngineNativePointerMemberBehavior = GetPointerMemberBehavior("UnrealHeaderTool", "NonEngineNativePointerMemberBehavior", UhtIssueBehavior.AllowSilently);
+			this._nonEngineObjectPtrMemberBehavior = GetPointerMemberBehavior("UnrealHeaderTool", "NonEngineObjectPtrMemberBehavior", UhtIssueBehavior.AllowSilently);
 			this._areRigVMUObjectProeprtiesEnabled = GetBoolean("UnrealHeaderTool", "AreRigVMUObjectProeprtiesEnabled", false);
 			this._areRigVMUInterfaceProeprtiesEnabled = GetBoolean("UnrealHeaderTool", "AreRigVMUInterfaceProeprtiesEnabled", false);
 			this._showDeprecations = GetBoolean("UnrealHeaderTool", "ShowDeprecations", true);
+			this._engineMissingGeneratedHeaderIncludeBehavior = GetPointerMemberBehavior("UnrealHeaderTool", "EngineMissingGeneratedHeaderIncludeBehavior", UhtIssueBehavior.AllowSilently);
+			this._nonEngineMissingGeneratedHeaderIncludeBehavior = GetPointerMemberBehavior("UnrealHeaderTool", "NonEngineMissingGeneratedHeaderIncludeBehavior", UhtIssueBehavior.AllowSilently);
+			this._engineEnumUnderlyingTypeNotSet = GetPointerMemberBehavior("UnrealHeaderTool", "EngineEnumUnderlyingTypeNotSet", UhtIssueBehavior.AllowSilently);
+			this._nonEngineEnumUnderlyingTypeNotSet = GetPointerMemberBehavior("UnrealHeaderTool", "NonEngineEnumUnderlyingTypeNotSet", UhtIssueBehavior.AllowSilently);
 		}
 
 		private bool GetBoolean(string SectionName, string KeyName, bool bDefault)
@@ -247,11 +283,11 @@ namespace UnrealBuildTool.Modes
 			return bDefault;
 		}
 
-		private UhtPointerMemberBehavior GetPointerMemberBehavior(string SectionName, string KeyName, UhtPointerMemberBehavior Default)
+		private UhtIssueBehavior GetPointerMemberBehavior(string SectionName, string KeyName, UhtIssueBehavior Default)
 		{
 			if (this._ini.TryGetValue(SectionName, KeyName, out string? BehaviorStr))
 			{
-				if (!Enum.TryParse(BehaviorStr, out UhtPointerMemberBehavior Value))
+				if (!Enum.TryParse(BehaviorStr, out UhtIssueBehavior Value))
 				{
 					throw new Exception(string.Format("Unrecognized native pointer member behavior '{0}'", BehaviorStr));
 				}
