PR #9077: Make VSCode .ignore entries absolute (Contributed by [REDACTED_PERSON])

diff --git a/Engine/Source/Programs/UnrealBuildTool/ProjectFiles/VisualStudioCode/VSCodeProjectFileGenerator.cs b/Engine/Source/Programs/UnrealBuildTool/ProjectFiles/VisualStudioCode/VSCodeProjectFileGenerator.cs
--- a/Engine/Source/Programs/UnrealBuildTool/ProjectFiles/VisualStudioCode/VSCodeProjectFileGenerator.cs
+++ b/Engine/Source/Programs/UnrealBuildTool/ProjectFiles/VisualStudioCode/VSCodeProjectFileGenerator.cs
@@ -1561,7 +1561,7 @@ namespace UnrealBuildTool
 
 				foreach (string PathToExclude in PathsToExclude)
 				{
-					OutFile.AppendLine(PathToExclude.Replace('\\', '/').Replace(WorkspaceRootPath, ""));
+					OutFile.AppendLine(PathToExclude.Replace('\\', '/').Replace(WorkspaceRootPath, "/"));
 				}
 			}
 			FileReference.WriteAllText(FileReference.Combine(PrimaryProjectPath, ".ignore"), OutFile.ToString());
