PR #9842: [BugFix][UBT][AndroidToolChain] Fix Regex failed to parse AndroidClangBuild (Contributed by [REDACTED_PERSON])

diff --git a/Engine/Source/Programs/UnrealBuildTool/Platform/Android/AndroidToolChain.cs b/Engine/Source/Programs/UnrealBuildTool/Platform/Android/AndroidToolChain.cs
--- a/Engine/Source/Programs/UnrealBuildTool/Platform/Android/AndroidToolChain.cs
+++ b/Engine/Source/Programs/UnrealBuildTool/Platform/Android/AndroidToolChain.cs
@@ -236,13 +236,13 @@ namespace UnrealBuildTool
 			ClangPath = Utils.CollapseRelativeDirectories(Path.Combine(NDKPath, @"toolchains/llvm", ArchitecturePath, @"bin/clang++" + ExeExtension));
 
 			// Android (6317467 based on r365631c1) clang version 9.0.8 
-			AndroidClangBuild = Utils.RunLocalProcessAndReturnStdOut(ClangPath, "--version", Logger);
+			string AndroidClangBuildTmp = Utils.RunLocalProcessAndReturnStdOut(ClangPath, "--version", Logger);
 			try
 			{
-				AndroidClangBuild = Regex.Match(AndroidClangBuild, @"(\w+) based on").Groups[1].ToString();
+				AndroidClangBuild = Regex.Match(AndroidClangBuildTmp, @"(\w+) based on").Groups[1].ToString();
 				if (String.IsNullOrEmpty(AndroidClangBuild))
 				{
-					AndroidClangBuild = Regex.Match(AndroidClangBuild, @"(\w+), based on").Groups[1].ToString();
+					AndroidClangBuild = Regex.Match(AndroidClangBuildTmp, @"(\w+), based on").Groups[1].ToString();
 				}
 			}
 			catch
