PR #9272: feat: add support for early init scripts (Contributed by [REDACTED_PERSON])

diff --git a/Engine/Source/Programs/UnrealBuildTool/Modes/BuildMode.cs b/Engine/Source/Programs/UnrealBuildTool/Modes/BuildMode.cs
--- a/Engine/Source/Programs/UnrealBuildTool/Modes/BuildMode.cs
+++ b/Engine/Source/Programs/UnrealBuildTool/Modes/BuildMode.cs
@@ -255,6 +255,110 @@ namespace UnrealBuildTool
 			return 0;
 		}
 
+		/// <summary>
+		/// Creates scripts for executing the init scripts (only supported on project descriptors)
+		/// </summary>
+		/// <param name="TargetDescriptor">The current target</param>
+		/// <returns>List of created script files</returns>
+		public static FileReference[] CreateInitScripts(TargetDescriptor TargetDescriptor)
+		{
+			if (TargetDescriptor.ProjectFile == null)
+			{
+				return new FileReference[] { };
+			}
+
+			ProjectDescriptor ProjectDescriptor = ProjectDescriptor.FromFile(TargetDescriptor.ProjectFile);
+			List<string[]> InitCommandBatches = new List<string[]>();
+
+			if(ProjectDescriptor != null && ProjectDescriptor.InitSteps != null)
+			{
+				if(ProjectDescriptor.InitSteps.TryGetCommands(BuildHostPlatform.Current.Platform, out string[]? Commands))
+				{
+					InitCommandBatches.Add(Commands);
+				}
+			}
+
+			if (InitCommandBatches.Count == 0)
+			{
+				return new FileReference[] { };
+			}
+
+			DirectoryReference ProjectDirectory = DirectoryReference.FromFile(TargetDescriptor.ProjectFile);
+			string PlatformIntermediateFolder = UEBuildTarget.GetPlatformIntermediateFolder(TargetDescriptor.Platform, TargetDescriptor.Architecture, false);
+
+			DirectoryReference ProjectIntermediateDirectory = DirectoryReference.Combine(ProjectDirectory, PlatformIntermediateFolder, TargetDescriptor.Name, TargetDescriptor.Configuration.ToString());
+
+			return WriteInitScripts(TargetDescriptor, BuildHostPlatform.Current.Platform, ProjectIntermediateDirectory, "Init", InitCommandBatches);
+		}
+
+		/// <summary>
+		/// Write scripts containing the custom build steps for the given host platform
+		/// </summary>
+		/// <param name="TargetDescriptor">The current target</param>
+		/// <param name="HostPlatform">The current host platform</param>
+		/// <param name="Directory">The output directory for the scripts</param>
+		/// <param name="FilePrefix">Bare prefix for all the created script files</param>
+		/// <param name="CommandBatches">List of custom build steps</param>
+		/// <returns>List of created script files</returns>
+		private static FileReference[] WriteInitScripts(TargetDescriptor TargetDescriptor, UnrealTargetPlatform HostPlatform, DirectoryReference Directory, string FilePrefix, List<string[]> CommandBatches)
+		{
+			List<FileReference> ScriptFiles = new List<FileReference>();
+			foreach(string[] CommandBatch in CommandBatches)
+			{
+				// Find all the standard variables
+				Dictionary<string, string> Variables = GetTargetVariables(TargetDescriptor);
+
+				// Get the output path to the script
+				string ScriptExtension = RuntimePlatform.IsWindows ? ".bat" : ".sh";
+				FileReference ScriptFile = FileReference.Combine(Directory, String.Format("{0}-{1}{2}", FilePrefix, ScriptFiles.Count + 1, ScriptExtension));
+
+				// Write it to disk
+				List<string> Contents = new List<string>();
+				if(RuntimePlatform.IsWindows)
+				{
+					Contents.Insert(0, "@echo off");
+				}
+				foreach(string Command in CommandBatch)
+				{
+					Contents.Add(Utils.ExpandVariables(Command, Variables));
+				}
+				if(!DirectoryReference.Exists(ScriptFile.Directory))
+				{
+					DirectoryReference.CreateDirectory(ScriptFile.Directory);
+				}
+				FileReference.WriteAllLines(ScriptFile, Contents);
+
+				// Add the output file to the list of generated scripts
+				ScriptFiles.Add(ScriptFile);
+			}
+			return ScriptFiles.ToArray();
+		}
+
+		/// <summary>
+		/// Gets a list of variables that can be expanded in paths referenced by this target
+		/// </summary>
+		/// <param name="TargetDescriptor">The current target</param>
+		/// <returns>Map of variable names to values</returns>
+		private static Dictionary<string, string> GetTargetVariables(TargetDescriptor TargetDescriptor)
+		{
+			Dictionary<string, string> Variables = new Dictionary<string,string>();
+			Variables.Add("RootDir", Unreal.RootDirectory.FullName);
+			Variables.Add("EngineDir", Unreal.EngineDirectory.FullName);
+			Variables.Add("TargetName", TargetDescriptor.Name);
+			Variables.Add("TargetPlatform", TargetDescriptor.Platform.ToString());
+			Variables.Add("TargetConfiguration", TargetDescriptor.Configuration.ToString());
+			if(TargetDescriptor.ProjectFile != null)
+			{
+				Variables.Add("ProjectDir", TargetDescriptor.ProjectFile.Directory.FullName);
+				Variables.Add("ProjectFile", TargetDescriptor.ProjectFile.FullName);
+			}
+			if (BuildVersion.TryRead(BuildVersion.GetDefaultFileName(), out BuildVersion? Version))
+			{
+				Variables.Add("EngineVersion", String.Format("{0}.{1}.{2}", Version.MajorVersion, Version.MinorVersion, Version.PatchVersion));
+			}
+			return Variables;
+		}
+
 		/// <summary>
 		/// Build a list of targets
 		/// </summary>
@@ -272,6 +376,13 @@ namespace UnrealBuildTool
 
 			for (int Idx = 0; Idx < TargetDescriptors.Count; ++Idx)
 			{
+				// Create and execute the init scripts
+				FileReference[] InitScripts = CreateInitScripts(TargetDescriptors[Idx]);
+				if (InitScripts.Length > 0)
+				{
+					Utils.ExecuteCustomBuildSteps(InitScripts, Logger);
+				}
+
 				TargetMakefile NewMakefile = CreateMakefile(BuildConfiguration, TargetDescriptors[Idx], WorkingSet, Logger);
 				TargetMakefiles.Add(NewMakefile);
 				if (!bSkipPreBuildTargets)
diff --git a/Engine/Source/Programs/UnrealBuildTool/System/ProjectDescriptor.cs b/Engine/Source/Programs/UnrealBuildTool/System/ProjectDescriptor.cs
--- a/Engine/Source/Programs/UnrealBuildTool/System/ProjectDescriptor.cs
+++ b/Engine/Source/Programs/UnrealBuildTool/System/ProjectDescriptor.cs
@@ -100,6 +100,11 @@ namespace UnrealBuildTool
 		/// </summary>
 		public uint EpicSampleNameHash;
 
+		/// <summary>
+		/// Steps to execute before creating rules assemblies in this project
+		/// </summary>
+		public CustomBuildSteps? InitSteps;
+
 		/// <summary>
 		/// Steps to execute before building targets in this project
 		/// </summary>
@@ -194,7 +199,8 @@ namespace UnrealBuildTool
 			// Get the sample name hash
 			RawObject.TryGetUnsignedIntegerField("EpicSampleNameHash", out EpicSampleNameHash);
 
-			// Read the pre and post-build steps
+			// Read the init, pre and post-build steps
+			CustomBuildSteps.TryRead(RawObject, "InitSteps", out InitSteps);
 			CustomBuildSteps.TryRead(RawObject, "PreBuildSteps", out PreBuildSteps);
 			CustomBuildSteps.TryRead(RawObject, "PostBuildSteps", out PostBuildSteps);
 		}
@@ -306,6 +312,10 @@ namespace UnrealBuildTool
 			}
 
 			// Write the custom build steps
+			if(InitSteps != null)
+			{
+				InitSteps.Write(Writer, "InitSteps");
+			}
 			if(PreBuildSteps != null)
 			{
 				PreBuildSteps.Write(Writer, "PreBuildSteps");
