UnrealBuildTool: Remove debug launch accidently checked in with GitHub PR

#rnx

diff --git a/Engine/Source/Programs/UnrealBuildTool/ProjectFiles/ProjectFileGenerator.cs b/Engine/Source/Programs/UnrealBuildTool/ProjectFiles/ProjectFileGenerator.cs
--- a/Engine/Source/Programs/UnrealBuildTool/ProjectFiles/ProjectFileGenerator.cs
+++ b/Engine/Source/Programs/UnrealBuildTool/ProjectFiles/ProjectFileGenerator.cs
@@ -548,7 +548,6 @@ namespace UnrealBuildTool
 		/// </summary>
 		void DiscoverCSharpProgramProjects(List<DirectoryReference> AllEngineDirectories, List<FileReference> AllGameProjects, PrimaryProjectFolder ProgramsFolder, ILogger Logger)
 		{
-			Debugger.Launch();
 			List<FileReference> FoundProjects = new List<FileReference>();
 			List<DirectoryReference> ProjectDirs = new List<DirectoryReference>();
 
