PR #9023: fix issue of writing make files with UTF8-BOM encoding (Contributed by [REDACTED_PERSON])

diff --git a/Engine/Source/Programs/UnrealBuildTool/ProjectFiles/ProjectFileGenerator.cs b/Engine/Source/Programs/UnrealBuildTool/ProjectFiles/ProjectFileGenerator.cs
--- a/Engine/Source/Programs/UnrealBuildTool/ProjectFiles/ProjectFileGenerator.cs
+++ b/Engine/Source/Programs/UnrealBuildTool/ProjectFiles/ProjectFileGenerator.cs
@@ -2907,13 +2907,9 @@ namespace UnrealBuildTool
 					// at the start of the file (adding two bytes to the file length).  For most
 					// files this is only mildly annoying but for Makefiles it can actually make
 					// them un-useable.
-					// TODO(sbc): See if we can just drop the Encoding.UTF8 argument on all
-					// platforms.  In this case UTF8 encoding will still be used but without the
-					// BOM, which is, AFAICT, desirable in almost all cases.
-					if (BuildHostPlatform.Current.Platform == UnrealTargetPlatform.Linux || BuildHostPlatform.Current.Platform == UnrealTargetPlatform.Mac)
-						File.WriteAllText(FileName, NewFileContents, new UTF8Encoding());
-					else
-						File.WriteAllText(FileName, NewFileContents, InEncoding != null ? InEncoding : Encoding.UTF8);
+					// To fix this we explicitly define UTF8 Encoding without BOM for all platform
+					// if another encoding is not specified in the call
+					File.WriteAllText(FileName, NewFileContents, InEncoding != null ? InEncoding : new UTF8Encoding(false));
 					Logger.LogDebug("Saved {Path}", Path.GetFileName(FileName));
 				}
 				catch (Exception ex)
