UBT: Expose OptimizationLevel to cmd line

#rnx

diff --git a/Engine/Source/Programs/UnrealBuildTool/Configuration/TargetRules.cs b/Engine/Source/Programs/UnrealBuildTool/Configuration/TargetRules.cs
--- a/Engine/Source/Programs/UnrealBuildTool/Configuration/TargetRules.cs
+++ b/Engine/Source/Programs/UnrealBuildTool/Configuration/TargetRules.cs
@@ -1233,6 +1233,7 @@ namespace UnrealBuildTool
 		/// Allows to fine tune optimizations level for speed and\or code size
 		/// </summary>
 		[RequiresUniqueBuildEnvironment]
+		[CommandLine("-OptimizationLevel=")]
 		public OptimizationMode OptimizationLevel = OptimizationMode.Speed;
 
 		/// <summary>
