Turned off VTune and ISPC when compiling windows arm64

diff --git a/Engine/Source/Programs/UnrealBuildTool/Platform/Windows/UEBuildWindows.cs b/Engine/Source/Programs/UnrealBuildTool/Platform/Windows/UEBuildWindows.cs
--- a/Engine/Source/Programs/UnrealBuildTool/Platform/Windows/UEBuildWindows.cs
+++ b/Engine/Source/Programs/UnrealBuildTool/Platform/Windows/UEBuildWindows.cs
@@ -932,6 +932,9 @@ namespace UnrealBuildTool
 					{
 						"OpenImageDenoise"
 					});
+					
+					// VTune does not support ARM
+					Target.GlobalDefinitions.Add("UE_EXTERNAL_PROFILING_ENABLED=0");
 				}
 			}
 
@@ -994,6 +997,11 @@ namespace UnrealBuildTool
 
 			Target.bCompileISPC = true;
 
+			if (Platform == UnrealTargetPlatform.Win64 && Target.Architecture == "arm64")
+			{
+				Target.bCompileISPC = false; // The version of ISPC we currently use does not support Windows Aarch64
+			}
+
 			// Initialize the VC environment for the target, and set all the version numbers to the concrete values we chose
 			Target.WindowsPlatform.Environment = CreateVCEnvironment(Target);
 
diff --git a/Engine/Source/Runtime/Core/Core.Build.cs b/Engine/Source/Runtime/Core/Core.Build.cs
--- a/Engine/Source/Runtime/Core/Core.Build.cs
+++ b/Engine/Source/Runtime/Core/Core.Build.cs
@@ -55,13 +55,17 @@ public class Core : ModuleRules
 		if (Target.Platform.IsInGroup(UnrealPlatformGroup.Windows))
 		{
 			AddEngineThirdPartyPrivateStaticDependencies(Target,
-				"zlib");
-
-			AddEngineThirdPartyPrivateStaticDependencies(Target, 
 				"IntelTBB",
-				"IntelVTune"
+				"zlib"
 				);
 
+			if (Target.Architecture.IndexOf("arm", StringComparison.OrdinalIgnoreCase) == -1)
+			{
+				AddEngineThirdPartyPrivateStaticDependencies(Target,
+					"IntelVTune"
+					);
+			}
+
 			// We do not want the static analyzer to run on thirdparty code
 			if (Target.StaticAnalyzer == StaticAnalyzer.None) 
 			{
