Allow UPL to use tools schema
Allow "build" as part of project names

#android

diff --git a/Engine/Source/Programs/UnrealBuildTool/Platform/Android/UEDeployAndroid.cs b/Engine/Source/Programs/UnrealBuildTool/Platform/Android/UEDeployAndroid.cs
--- a/Engine/Source/Programs/UnrealBuildTool/Platform/Android/UEDeployAndroid.cs
+++ b/Engine/Source/Programs/UnrealBuildTool/Platform/Android/UEDeployAndroid.cs
@@ -152,7 +152,7 @@ namespace UnrealBuildTool
 				}
 			}
 
-			UPL = new UnrealPluginLanguage(ProjectFile, inPluginExtraData, NDKArches, "http://schemas.android.com/apk/res/android", "xmlns:android=\"http://schemas.android.com/apk/res/android\"", UnrealTargetPlatform.Android, Logger);
+			UPL = new UnrealPluginLanguage(ProjectFile, inPluginExtraData, NDKArches, "http://schemas.android.com/apk/res/android", "xmlns:android=\"http://schemas.android.com/apk/res/android\" xmlns:tools=\"http://schemas.android.com/tools\"", UnrealTargetPlatform.Android, Logger);
 			UPLHashCode = UPL.GetUPLHash();
 //			APL.SetTrace();
 		}
@@ -2708,7 +2708,7 @@ namespace UnrealBuildTool
 
 			StringBuilder Text = new StringBuilder();
 			Text.AppendLine(XML_HEADER);
-			Text.AppendLine("<manifest xmlns:android=\"http://schemas.android.com/apk/res/android\"");
+			Text.AppendLine("<manifest xmlns:android=\"http://schemas.android.com/apk/res/android\" xmlns:tools=\"http://schemas.android.com/tools\"");
 			Text.AppendLine(string.Format("          package=\"{0}\"", PackageName));
 			if (ExtraManifestNodeTags != null)
 			{
@@ -4976,6 +4976,7 @@ namespace UnrealBuildTool
 
 					ABIFilter += ", \"" + NDKArch + "\"";
 
+					int UnrealBuildGradlePathLength = UnrealBuildGradlePath.Length;
 					string[] SourceFiles = Directory.GetFiles(UnrealBuildGradlePath, "*.*", SearchOption.AllDirectories);
 					foreach (string Filename in SourceFiles)
 					{
@@ -4984,7 +4985,7 @@ namespace UnrealBuildTool
 
 						// skip the build directories
 						string Workname = Filename.Replace("\\", "/");
-						string DirectoryName = Path.GetDirectoryName(Filename)!;
+						string DirectoryName = Path.GetDirectoryName(Filename)!.Substring(UnrealBuildGradlePathLength);
 						if (DirectoryName.Contains("build") || Workname.Contains("/."))
 						{
 							continue;
