Emit a warning when a plugin extension specifies an explictly empty platform allow list, suggesting the author expects some sort of inheritence.

#rnx

diff --git a/Engine/Source/Programs/UnrealBuildTool/System/PluginDescriptor.cs b/Engine/Source/Programs/UnrealBuildTool/System/PluginDescriptor.cs
--- a/Engine/Source/Programs/UnrealBuildTool/System/PluginDescriptor.cs
+++ b/Engine/Source/Programs/UnrealBuildTool/System/PluginDescriptor.cs
@@ -350,13 +350,7 @@ namespace UnrealBuildTool
 			try
 			{
 				PluginDescriptor Descriptor = new PluginDescriptor(RawObject, FileName);
-				if (Descriptor.Modules != null)
-				{
-					foreach (ModuleDescriptor Module in Descriptor.Modules)
-					{
-						Module.Validate(FileName);
-					}
-				}
+				Descriptor.Validate(FileName);
 				return Descriptor;
 			}
 			catch (JsonParseException ParseException)
@@ -482,6 +476,34 @@ namespace UnrealBuildTool
 			}
 		}
 
+		/// <summary>
+		/// Produces any warnings and errors for the plugin descriptor
+		/// </summary>
+		/// <param name="FileName">File containing the plugin</param>
+		public void Validate(FileReference FileName)
+		{
+			if (Modules != null)
+			{
+				foreach (ModuleDescriptor Module in Modules)
+				{
+					Module.Validate(FileName);
+				}
+
+				if (bIsPluginExtension)
+				{
+					foreach (ModuleDescriptor ChildModule in Modules)
+					{
+						if (ChildModule.bHasExplicitPlatforms && ChildModule.PlatformAllowList != null && ChildModule.PlatformAllowList.Count == 0)
+						{
+							// The order that child plugins are merged into the parent is undefined - there is no heirarchy. Only the PlatformAllowList and PlatformDenyList are currently merged into the parent
+							// Having an explicity-empty list here suggests someone is trying to create a heirarchy and may get caught out if their module declaration also includes TargetAllowList, ProgramAllowList etc. as these properties will be ignored
+							Log.TraceWarningOnce(FileName, $"Plugin extensions should not declare HasExplicitPlatforms with an empty PlatformAllowList. (module {ChildModule.Name})");
+						}
+					}
+				}
+			}
+		}
+
 		/// <summary>
 		/// Determines if this reference enables the plugin for a given platform
 		/// </summary>
