[UnrealBuildTool - SourceFileMetadataCache]
* Added tracking of markup containing information saying if a header can be used as a standalone header unit or not. This is also used by IWYU path to figure out if a header can be built stand alone
* Added tracking of includes in headers. This is used in header units path to figure out dependencies between headers (since header units will want to import produced header units from other headers

This increases the file metadata cache from 3mb to 4mb for a large game but  should not affect ubt startup times noticable

diff --git a/Engine/Source/Programs/UnrealBuildTool/System/SourceFileMetadataCache.cs b/Engine/Source/Programs/UnrealBuildTool/System/SourceFileMetadataCache.cs
--- a/Engine/Source/Programs/UnrealBuildTool/System/SourceFileMetadataCache.cs
+++ b/Engine/Source/Programs/UnrealBuildTool/System/SourceFileMetadataCache.cs
@@ -6,7 +6,6 @@ using System.Collections.Generic;
 using System.Diagnostics;
 using System.IO;
 using System.Linq;
-using System.Text;
 using System.Text.RegularExpressions;
 using System.Threading.Tasks;
 using EpicGames.Core;
@@ -16,6 +15,22 @@ using UnrealBuildBase;
 
 namespace UnrealBuildTool
 {
+	/// <summary>
+	/// The header unit type markup is used for multiple things:
+	/// 1. To figure out if a header can be compiled by itself or not. Many includes are included in the middle of other includes and those can never be compiled by themselves
+	///    This markup is used by both msvc header unit feature and IWYU toolchain
+	/// 2. Tell if the header even supports being part of a header unit. If it does not support it will also prevent all includes including it from producing a header unit
+	/// This is how to markup in headers to provide above info:
+	///    // HEADER_UNIT_SKIP - Here you can write why this file can't be compiled standalone
+	///    // HEADER_UNIT_UNSUPPORTED - Here you can write why this file can't be part of header units.
+	/// </summary>
+	enum HeaderUnitType
+	{
+		Valid = 0,
+		Unsupported = 1,
+		Skip = 2
+	}
+
 	/// <summary>
 	/// Caches information about C++ source files; whether they contain reflection markup, what the first included header is, and so on.
 	/// </summary>
@@ -38,9 +53,9 @@ namespace UnrealBuildTool
 		}
 
 		/// <summary>
-		/// Information about whether a file contains reflection markup
+		/// Information about whether a header file contains reflection markup and what includes it has (the superset, ignoring #if/#endif)
 		/// </summary>
-		class ReflectionInfo
+		class HeaderInfo
 		{
 			/// <summary>
 			/// Last write time of the file when the data was cached
@@ -51,6 +66,13 @@ namespace UnrealBuildTool
 			/// Whether or not the file contains reflection markup
 			/// </summary>
 			public bool bContainsMarkup;
+
+			/// <summary>
+			/// List of includes that header contains
+			/// </summary>
+			public List<string>? Includes;
+
+			public HeaderUnitType UnitType;
 		}
 
 		/// <summary>
@@ -72,7 +94,7 @@ namespace UnrealBuildTool
 		/// <summary>
 		/// The current file version
 		/// </summary>
-		public const int CurrentVersion = 4;
+		public const int CurrentVersion = 5;
 
 		/// <summary>
 		/// Location of this dependency cache
@@ -97,7 +119,7 @@ namespace UnrealBuildTool
 		/// <summary>
 		/// Map from file item to header file info
 		/// </summary>
-		ConcurrentDictionary<FileItem, ReflectionInfo> FileToReflectionInfo = new ConcurrentDictionary<FileItem, ReflectionInfo>();
+		ConcurrentDictionary<FileItem, HeaderInfo> FileToHeaderInfo = new ConcurrentDictionary<FileItem, HeaderInfo>();
 
 		/// <summary>
 		/// Map from file item to inline info
@@ -122,7 +144,7 @@ namespace UnrealBuildTool
 		/// <summary>
 		/// Regex that matches C++ code with UObject declarations which we will need to generated code for.
 		/// </summary>
-		static readonly Regex ReflectionMarkupRegex = new Regex("^\\s*U(CLASS|STRUCT|ENUM|INTERFACE|DELEGATE)\\b", RegexOptions.Compiled | RegexOptions.Multiline);
+		static readonly Regex ReflectionMarkupRegex = new Regex("^\\s*U(CLASS|STRUCT|ENUM|INTERFACE|DELEGATE)\\b", RegexOptions.Compiled);
 
 		/// <summary>
 		/// Regex that matches #include UE_INLINE_GENERATED_CPP_BY_NAME(****) statements.
@@ -240,30 +262,75 @@ namespace UnrealBuildTool
 		/// <summary>
 		/// Determines whether the given file contains reflection markup
 		/// </summary>
-		/// <param name="SourceFile">The source file to parse</param>
+		/// <param name="HeaderFile">The source file to parse</param>
 		/// <returns>True if the file contains reflection markup</returns>
-		public bool ContainsReflectionMarkup(FileItem SourceFile)
+		public bool ContainsReflectionMarkup(FileItem HeaderFile)
 		{
-			if(Parent != null && !SourceFile.Location.IsUnderDirectory(BaseDirectory))
+			if(Parent != null && !HeaderFile.Location.IsUnderDirectory(BaseDirectory))
 			{
-				return Parent.ContainsReflectionMarkup(SourceFile);
+				return Parent.ContainsReflectionMarkup(HeaderFile);
 			}
 			else
 			{
-				ReflectionInfo? ReflectionInfo;
-				if(!FileToReflectionInfo.TryGetValue(SourceFile, out ReflectionInfo) || SourceFile.LastWriteTimeUtc.Ticks > ReflectionInfo.LastWriteTimeUtc)
+				HeaderInfo? HeaderInfo;
+				if(!FileToHeaderInfo.TryGetValue(HeaderFile, out HeaderInfo) || HeaderFile.LastWriteTimeUtc.Ticks > HeaderInfo.LastWriteTimeUtc)
 				{
-					ReflectionInfo = new ReflectionInfo();
-					ReflectionInfo.LastWriteTimeUtc = SourceFile.LastWriteTimeUtc.Ticks;
-					ReflectionInfo.bContainsMarkup = ReflectionMarkupRegex.IsMatch(FileReference.ReadAllText(SourceFile.Location));
-					FileToReflectionInfo[SourceFile] = ReflectionInfo;
+					HeaderInfo = ParseHeader(HeaderFile);
+					FileToHeaderInfo[HeaderFile] = HeaderInfo;
 					bModified = true;
 				}
-				return ReflectionInfo.bContainsMarkup;
+				return HeaderInfo.bContainsMarkup;
+			}
+		}
+
+		/// <summary>
+		/// Read entire header file to find markup and includes
+		/// </summary>
+		/// <param name="HeaderFile">The header file to parse</param>
+		/// <returns>A HeaderInfo struct containing information about header</returns>
+		private HeaderInfo ParseHeader(FileItem HeaderFile)
+		{
+			HeaderInfo HeaderInfo = new();
+			HeaderInfo.LastWriteTimeUtc = HeaderFile.LastWriteTimeUtc.Ticks;
+
+			bool bContainsMarkup = false;
+			SortedSet<string> Includes = new();
+			foreach (string Line in FileReference.ReadAllLines(HeaderFile.Location))
+			{
+				if (!bContainsMarkup)
+				{
+					bContainsMarkup = ReflectionMarkupRegex.IsMatch(Line);
+				}
+
+				if (Line.AsSpan().TrimStart().StartsWith("#include"))
+				{
+					int FirstQuotation = Line.IndexOf('"');
+					int FirstAngleBracket = Line.IndexOf('<');
+					if (FirstQuotation != -1 && (FirstAngleBracket == -1 || FirstQuotation < FirstAngleBracket)) // Handle #include <foo.h> // Some text with "
+					{
+						int SecondQuotation = Line.IndexOf('"', FirstQuotation + 1);
+						Includes.Add(Line.Substring(FirstQuotation + 1, SecondQuotation - FirstQuotation - 1));
+					}
+				}
+
+				int HeaderUnitIndex = Line.IndexOf("HEADER_UNIT_");
+				if (HeaderUnitIndex != -1)
+				{
+					ReadOnlySpan<char> Span = Line.AsSpan(HeaderUnitIndex + "HEADER_UNIT_".Length);
+					if (Span.StartsWith("UNSUPPORTED"))
+						HeaderInfo.UnitType = HeaderUnitType.Unsupported;
+					else if (Span.StartsWith("SKIP"))
+						HeaderInfo.UnitType = HeaderUnitType.Skip;
+				}
 			}
+
+			HeaderInfo.bContainsMarkup = bContainsMarkup;
+			HeaderInfo.Includes = Includes.ToList();
+			return HeaderInfo;
 		}
 
 
+
 		/// <summary>
 		/// Returns a list of inlined generated cpps that this source file contains.
 		/// </summary>
@@ -294,6 +361,50 @@ namespace UnrealBuildTool
 			}
 		}
 
+		/// <summary>
+		/// Returns a HeaderInfo struct for a header file (and parse the file if not already cached)
+		/// </summary>
+		/// <param name="HeaderFile">The header file to parse</param>
+		/// <returns>HeaderInfo for header file</returns>
+		HeaderInfo GetHeaderInfo(FileItem HeaderFile)
+		{
+			if (Parent != null && !HeaderFile.Location.IsUnderDirectory(BaseDirectory))
+			{
+				return Parent.GetHeaderInfo(HeaderFile);
+			}
+			else
+			{
+				HeaderInfo? HeaderInfo;
+				if (!FileToHeaderInfo.TryGetValue(HeaderFile, out HeaderInfo) || HeaderFile.LastWriteTimeUtc.Ticks > HeaderInfo.LastWriteTimeUtc)
+				{
+					HeaderInfo = ParseHeader(HeaderFile);
+					FileToHeaderInfo[HeaderFile] = HeaderInfo;
+					bModified = true;
+				}
+				return HeaderInfo;
+			}
+		}
+
+		/// <summary>
+		/// Returns header unit type for a header file (and parse the file if not already cached)
+		/// </summary>
+		/// <param name="HeaderFile">The header file to parse</param>
+		/// <returns>Header unit type</returns>
+		public HeaderUnitType GetHeaderUnitType(FileItem HeaderFile)
+		{
+			return GetHeaderInfo(HeaderFile).UnitType;
+		}
+
+		/// <summary>
+		/// Returns all #includes existing inside a header file (and parse the file if not already cached)
+		/// </summary>
+		/// <param name="HeaderFile">The header file to parse</param>
+		/// <returns>List of includes</returns>
+		public List<string> GetHeaderIncludes(FileItem HeaderFile)
+		{
+			return GetHeaderInfo(HeaderFile).Includes!;
+		}
+
 		/// <summary>
 		/// Parse the first include directive from a source file
 		/// </summary>
@@ -440,11 +551,13 @@ namespace UnrealBuildTool
 					{
 						FileItem File = Reader.ReadCompactFileItem();
 
-						ReflectionInfo ReflectionInfo = new ReflectionInfo();
-						ReflectionInfo.LastWriteTimeUtc = Reader.ReadLong();
-						ReflectionInfo.bContainsMarkup = Reader.ReadBool();
+						HeaderInfo HeaderInfo = new HeaderInfo();
+						HeaderInfo.LastWriteTimeUtc = Reader.ReadLong();
+						HeaderInfo.bContainsMarkup = Reader.ReadBool();
+						HeaderInfo.UnitType = (HeaderUnitType)Reader.ReadByte();
+						HeaderInfo.Includes = Reader.ReadList(() => Reader.ReadString())!;
 
-						FileToReflectionInfo[File] = ReflectionInfo;
+						FileToHeaderInfo[File] = HeaderInfo;
 					}
 
 					int FileToInlineMarkupFlagCount = Reader.ReadInt();
@@ -487,12 +600,14 @@ namespace UnrealBuildTool
 						Writer.WriteString(Pair.Value.IncludeText);
 					}
 
-					Writer.WriteInt(FileToReflectionInfo.Count);
-					foreach(KeyValuePair<FileItem, ReflectionInfo> Pair in FileToReflectionInfo)
+					Writer.WriteInt(FileToHeaderInfo.Count);
+					foreach(KeyValuePair<FileItem, HeaderInfo> Pair in FileToHeaderInfo)
 					{
 						Writer.WriteCompactFileItem(Pair.Key);
 						Writer.WriteLong(Pair.Value.LastWriteTimeUtc);
 						Writer.WriteBool(Pair.Value.bContainsMarkup);
+						Writer.WriteByte((byte)Pair.Value.UnitType);
+						Writer.WriteList(Pair.Value.Includes, Item => Writer.WriteString(Item));
 					}
 
 					Writer.WriteInt(FileToInlineReflectionInfo.Count);
