[UnrealBuildTool - ClangToolChain]
* Fixed so the code adding files to the response file is going through a virtual method which produces the path. This makes it possible for derived classes to control if they want to produce files with other extension names for the various files (used by iwyu toolchain which is not outputting ".o" files etc

diff --git a/Engine/Source/Programs/UnrealBuildTool/ToolChain/ClangToolChain.cs b/Engine/Source/Programs/UnrealBuildTool/ToolChain/ClangToolChain.cs
--- a/Engine/Source/Programs/UnrealBuildTool/ToolChain/ClangToolChain.cs
+++ b/Engine/Source/Programs/UnrealBuildTool/ToolChain/ClangToolChain.cs
@@ -862,6 +862,11 @@ namespace UnrealBuildTool
 			GetCompileArguments_AdditionalArgs(CompileEnvironment, Arguments);
 		}
 
+		protected virtual string GetFileNameFromExtension(string AbsolutePath, string Extension)
+		{
+			return Path.GetFileName(AbsolutePath) + Extension;
+		}
+
 		/// <summary>
 		/// Compile arguments for specific files in a module. Also updates Action and CPPOutput results.
 		/// </summary>
@@ -919,12 +924,12 @@ namespace UnrealBuildTool
 			FileItem OutputFile;
 			if (CompileEnvironment.PrecompiledHeaderAction == PrecompiledHeaderAction.Create)
 			{
-				OutputFile = FileItem.GetItemByFileReference(FileReference.Combine(OutputDir, Path.GetFileName(SourceFile.AbsolutePath) + ".gch"));
+				OutputFile = FileItem.GetItemByFileReference(FileReference.Combine(OutputDir, GetFileNameFromExtension(SourceFile.AbsolutePath, ".gch")));
 				CompileResult.PrecompiledHeaderFile = OutputFile;
 			}
 			else if (CompileEnvironment.bPreprocessOnly)
 			{
-				OutputFile = FileItem.GetItemByFileReference(FileReference.Combine(OutputDir, Path.GetFileName(SourceFile.AbsolutePath) + ".i"));
+				OutputFile = FileItem.GetItemByFileReference(FileReference.Combine(OutputDir, GetFileNameFromExtension(SourceFile.AbsolutePath, ".i")));
 				CompileResult.ObjectFiles.Add(OutputFile);
 
 				// Clang does EITHER pre-process or object file.
@@ -937,7 +942,7 @@ namespace UnrealBuildTool
 			else
 			{
 				string ObjectFileExtension = (CompileEnvironment.AdditionalArguments != null && CompileEnvironment.AdditionalArguments.Contains("-emit-llvm")) ? ".bc" : ".o";
-				OutputFile = FileItem.GetItemByFileReference(FileReference.Combine(OutputDir, $"{Path.GetFileName(SourceFile.AbsolutePath)}{ObjectFileExtension}"));
+				OutputFile = FileItem.GetItemByFileReference(FileReference.Combine(OutputDir, GetFileNameFromExtension(SourceFile.AbsolutePath, ObjectFileExtension)));
 				CompileResult.ObjectFiles.Add(OutputFile);
 			}
 
@@ -952,7 +957,7 @@ namespace UnrealBuildTool
 				// Generate the timing info
 				if (CompileEnvironment.bPrintTimingInfo)
 				{
-					FileItem TraceFile = FileItem.GetItemByFileReference(FileReference.Combine(OutputDir, Path.GetFileName(SourceFile.AbsolutePath) + ".json"));
+					FileItem TraceFile = FileItem.GetItemByFileReference(FileReference.Combine(OutputDir, GetFileNameFromExtension(SourceFile.AbsolutePath, ".json")));
 					Arguments.Add("-ftime-trace");
 					CompileAction.ProducedItems.Add(TraceFile);
 				}
@@ -960,7 +965,7 @@ namespace UnrealBuildTool
 				// Generate the included header dependency list
 				if (!PreprocessDepends)
 				{
-					FileItem DependencyListFile = FileItem.GetItemByFileReference(FileReference.Combine(OutputDir, Path.GetFileName(SourceFile.AbsolutePath) + ".d"));
+					FileItem DependencyListFile = FileItem.GetItemByFileReference(FileReference.Combine(OutputDir, GetFileNameFromExtension(SourceFile.AbsolutePath, ".d")));
 					Arguments.Add(GetDepencenciesListFileArgument(DependencyListFile));
 					CompileAction.DependencyListFile = DependencyListFile;
 					CompileAction.ProducedItems.Add(DependencyListFile);
@@ -968,7 +973,7 @@ namespace UnrealBuildTool
 			}
 			else
 			{
-				OutputFile = FileItem.GetItemByFileReference(FileReference.Combine(OutputDir, Path.GetFileName(SourceFile.AbsolutePath) + ".analysis"));
+				OutputFile = FileItem.GetItemByFileReference(FileReference.Combine(OutputDir, GetFileNameFromExtension(SourceFile.AbsolutePath, ".analysis")));
 			}
 
 			// Add the parameters needed to compile the output file to the command-line.
