[UnrealBuildTool]
* BuildMode - Added so CreateMakeFile and Build (using makefile) are "internal" instead of private. This means that other modes (such as IWYU) can re-use the code

diff --git a/Engine/Source/Programs/UnrealBuildTool/Modes/BuildMode.cs b/Engine/Source/Programs/UnrealBuildTool/Modes/BuildMode.cs
--- a/Engine/Source/Programs/UnrealBuildTool/Modes/BuildMode.cs
+++ b/Engine/Source/Programs/UnrealBuildTool/Modes/BuildMode.cs
@@ -411,7 +411,7 @@ namespace UnrealBuildTool
 		/// <param name="WriteOutdatedActionsFile">Files to write the list of outdated actions to (rather than building them)</param>
 		/// <param name="Logger">Logger for output</param>
 		/// <returns>Result from the compilation</returns>
-		static void Build(TargetMakefile[] Makefiles, List<TargetDescriptor> TargetDescriptors, BuildConfiguration BuildConfiguration, BuildOptions Options, FileReference? WriteOutdatedActionsFile, ILogger Logger)
+		internal static void Build(TargetMakefile[] Makefiles, List<TargetDescriptor> TargetDescriptors, BuildConfiguration BuildConfiguration, BuildOptions Options, FileReference? WriteOutdatedActionsFile, ILogger Logger)
 		{
 			// Execute the build
 			if ((Options & BuildOptions.SkipBuild) == 0)
@@ -800,7 +800,7 @@ namespace UnrealBuildTool
 		/// <param name="WorkingSet">Set of source files which are part of the working set</param>
 		/// <param name="Logger">Logger for output</param>
 		/// <returns>Makefile for the given target</returns>
-		static TargetMakefile CreateMakefile(BuildConfiguration BuildConfiguration, TargetDescriptor TargetDescriptor, ISourceFileWorkingSet WorkingSet, ILogger Logger)
+		internal static TargetMakefile CreateMakefile(BuildConfiguration BuildConfiguration, TargetDescriptor TargetDescriptor, ISourceFileWorkingSet WorkingSet, ILogger Logger)
 		{
 			// Get the path to the makefile for this target
 			FileReference? MakefileLocation = null;
