[UnrealBuildTool]
* Added code to fixup definitions provided through .build.cs files. There are actually strict rules around the syntax of defines (reading msvc documentation). There should be no space before or after the '=' character. so "MY_DEFINE = 1" is actually invalid. When building with header units this actually produce invalid data but in normal builds it seems to swallow it (even though the documentation says it shouldn't)

I would have wanted to fix the code producing bad definition format but it is in a whole bunch of modules so I suspect external teams might have this too, so instead we fix it up by trimming the strings before and after '='

diff --git a/Engine/Source/Programs/UnrealBuildTool/Configuration/UEBuildModule.cs b/Engine/Source/Programs/UnrealBuildTool/Configuration/UEBuildModule.cs
--- a/Engine/Source/Programs/UnrealBuildTool/Configuration/UEBuildModule.cs
+++ b/Engine/Source/Programs/UnrealBuildTool/Configuration/UEBuildModule.cs
@@ -561,6 +561,36 @@ namespace UnrealBuildTool
 			}
 		}
 
+		/// <summary>
+		/// Add definitions from source to target. This code also reformats defines to match strict rules which does not allow space before and after '='
+		/// </summary>
+		protected void AddDefinitions(List<string> Target, HashSet<string> Source)
+		{
+			StringBuilder Builder = new();
+			foreach (string Def in Source)
+			{
+				string FixedDef = Def;
+				int IndexOfAssign = Def.IndexOf("=");
+				if (IndexOfAssign != -1)
+				{
+					ReadOnlySpan<char> DefSpan = Def.AsSpan();
+					ReadOnlySpan<char> Name = DefSpan.Slice(0, IndexOfAssign);
+					ReadOnlySpan<char> NameTrim = Name.Trim();
+					ReadOnlySpan<char> Value = DefSpan.Slice(IndexOfAssign + 1);
+					ReadOnlySpan<char> ValueTrim = Value.Trim();
+					if (Name.Length != NameTrim.Length || Value.Length != ValueTrim.Length)
+					{
+						Builder.Clear();
+						Builder.Append(NameTrim);
+						Builder.Append('=');
+						Builder.Append(ValueTrim);
+						FixedDef = Builder.ToString();
+					}
+				}
+				Target.Add(FixedDef);
+			}
+		}
+
 		/// <summary>
 		/// Sets up the environment for compiling any module that includes the public interface of this module.
 		/// </summary>
@@ -603,7 +633,7 @@ namespace UnrealBuildTool
 			SystemIncludePaths.UnionWith(PublicSystemIncludePaths);
 
 			// Add this module's public definitions
-			Definitions.AddRange(PublicDefinitions);
+			AddDefinitions(Definitions, PublicDefinitions);
 
 			// Add the import or export declaration for the module
 			if (Rules.Type == ModuleRules.ModuleType.CPlusPlus)
diff --git a/Engine/Source/Programs/UnrealBuildTool/Configuration/UEBuildModuleCPP.cs b/Engine/Source/Programs/UnrealBuildTool/Configuration/UEBuildModuleCPP.cs
--- a/Engine/Source/Programs/UnrealBuildTool/Configuration/UEBuildModuleCPP.cs
+++ b/Engine/Source/Programs/UnrealBuildTool/Configuration/UEBuildModuleCPP.cs
@@ -1618,7 +1618,8 @@ namespace UnrealBuildTool
 			}
 
 			// Add the module's public and private definitions.
-			Result.Definitions.AddRange(PublicDefinitions);
+			AddDefinitions(Result.Definitions, PublicDefinitions);
+
 			Result.Definitions.AddRange(Rules.PrivateDefinitions);
 
 			// Add the project definitions
