[UnrealBuildTool]
* Fixed log warning to show correct data (and is now copy-pastable)

diff --git a/Engine/Source/Programs/UnrealBuildTool/Configuration/UEBuildModuleCPP.cs b/Engine/Source/Programs/UnrealBuildTool/Configuration/UEBuildModuleCPP.cs
--- a/Engine/Source/Programs/UnrealBuildTool/Configuration/UEBuildModuleCPP.cs
+++ b/Engine/Source/Programs/UnrealBuildTool/Configuration/UEBuildModuleCPP.cs
@@ -564,7 +564,7 @@ namespace UnrealBuildTool
 							string NameWithoutPrefix = Name.Substring(4);
 							if (CPPFilesLookup.TryGetValue(NameWithoutPrefix, out FileItem? Item))
 							{
-								Logger.LogWarning("'{0}' .gen.cpp not inlined. Add '#include UE_INLINE_GENERATED_CPP_BY_NAME({1})'", Item.Name, Name);
+								Logger.LogWarning("'{0}' .gen.cpp not inlined. Add '#include UE_INLINE_GENERATED_CPP_BY_NAME({1})'", Item.Name, NameWithoutPrefix);
 							}
 						}
 					}
