UnrealBuildTool: Use Normal priority for ParallelExecutor for asymmetric processors as BelowNormal can cause scheduling issues

#rnx

diff --git a/Engine/Source/Programs/UnrealBuildTool/Executors/ParallelExecutor.cs b/Engine/Source/Programs/UnrealBuildTool/Executors/ParallelExecutor.cs
--- a/Engine/Source/Programs/UnrealBuildTool/Executors/ParallelExecutor.cs
+++ b/Engine/Source/Programs/UnrealBuildTool/Executors/ParallelExecutor.cs
@@ -46,10 +46,10 @@ namespace UnrealBuildTool
 		/// <summary>
 		/// The priority to set for spawned processes.
 		/// Valid Settings: Idle, BelowNormal, Normal, AboveNormal, High
-		/// Default: BelowNormal
+		/// Default: BelowNormal or Normal for an Asymmetrical processor as BelowNormal can cause scheduling issues.
 		/// </summary>
 		[XmlConfigFile]
-		private static ProcessPriorityClass ProcessPriority = ProcessPriorityClass.BelowNormal;
+		private static ProcessPriorityClass ProcessPriority = Utils.IsAsymmetricalProcessor() ? ProcessPriorityClass.Normal : ProcessPriorityClass.BelowNormal;
 
 		/// <summary>
 		/// When enabled, will stop compiling targets after a compile error occurs.
diff --git a/Engine/Source/Programs/UnrealBuildTool/System/Utils.cs b/Engine/Source/Programs/UnrealBuildTool/System/Utils.cs
--- a/Engine/Source/Programs/UnrealBuildTool/System/Utils.cs
+++ b/Engine/Source/Programs/UnrealBuildTool/System/Utils.cs
@@ -1145,6 +1145,23 @@ namespace UnrealBuildTool
 			return -1;
 		}
 
+		/// <summary>
+		/// Gets if the processos has asymmetrical cores (Windows only)
+		/// </summary>
+		/// <returns></returns>
+		public static bool IsAsymmetricalProcessor()
+		{
+			int LogicalCores = GetLogicalProcessorCount();
+			int PhysicalCores = GetPhysicalProcessorCount();
+
+			if (PhysicalCores <= 0 || PhysicalCores == LogicalCores)
+			{
+				return false;
+			}
+
+			return LogicalCores != PhysicalCores * 2;
+		}
+
 		/// <summary>
 		/// Gets the total memory bytes available, based on what is known to the garbage collector.
 		/// </summary>
