UnrealBuildTool: Handle version numbers with only two components in MicrosoftPlatformSDK

#rnx

diff --git a/Engine/Source/Programs/UnrealBuildTool/Platform/Windows/MicrosoftPlatformSDK.cs b/Engine/Source/Programs/UnrealBuildTool/Platform/Windows/MicrosoftPlatformSDK.cs
--- a/Engine/Source/Programs/UnrealBuildTool/Platform/Windows/MicrosoftPlatformSDK.cs
+++ b/Engine/Source/Programs/UnrealBuildTool/Platform/Windows/MicrosoftPlatformSDK.cs
@@ -69,6 +69,16 @@ namespace UnrealBuildTool
 				return true;
 			}
 
+			Result = Regex.Match(StringValue, @"^(\d+).(\d+)");
+			if (Result.Success)
+			{
+				// 8 bits for major, 8 for minor - ignore build and patch
+				OutValue |= UInt64.Parse(Result.Groups[1].Value) << 24;
+				OutValue |= UInt64.Parse(Result.Groups[2].Value) << 16;
+
+				return true;
+			}
+
 			return false;
 		}
 
