[UnrealBuildTool]
* Fixed what seems to be a c# compiler bug.

Calling this in the ctor and GetToolChainInfo being virtual/abstract really screws things up.

LazyInfo = new Lazy<ClangToolChainInfo>(() => GetToolChainInfo());

By changing it to be a trampoline function in between like this it works:

LazyInfo = new Lazy<ClangToolChainInfo>(() => { return GetToolChainInfo() });

diff --git a/Engine/Source/Programs/UnrealBuildTool/ToolChain/ClangToolChain.cs b/Engine/Source/Programs/UnrealBuildTool/ToolChain/ClangToolChain.cs
--- a/Engine/Source/Programs/UnrealBuildTool/ToolChain/ClangToolChain.cs
+++ b/Engine/Source/Programs/UnrealBuildTool/ToolChain/ClangToolChain.cs
@@ -195,7 +195,7 @@ namespace UnrealBuildTool
 		{
 			Options = InOptions;
 
-			LazyInfo = new Lazy<ClangToolChainInfo>(() => GetToolChainInfo());
+			LazyInfo = new Lazy<ClangToolChainInfo>(() => { return GetToolChainInfo(); }); // Don't change to => GetToolChainInfo().. it doesnt produce correct code
 		}
 
 		protected abstract ClangToolChainInfo GetToolChainInfo();
