UnrealBuildTool: Additional check for if a platform is available when recommending visual studio components to install

#rnx

diff --git a/Engine/Source/Programs/UnrealBuildTool/Platform/Windows/MicrosoftPlatformSDK.Versions.cs b/Engine/Source/Programs/UnrealBuildTool/Platform/Windows/MicrosoftPlatformSDK.Versions.cs
--- a/Engine/Source/Programs/UnrealBuildTool/Platform/Windows/MicrosoftPlatformSDK.Versions.cs
+++ b/Engine/Source/Programs/UnrealBuildTool/Platform/Windows/MicrosoftPlatformSDK.Versions.cs
@@ -121,8 +121,8 @@ namespace UnrealBuildTool
 		/// </summary>
 		public static IEnumerable<string> GetVisualStudioSuggestedComponents(VCProjectFileFormat Format)
 		{
-			bool LinuxValid = InstalledPlatformInfo.IsValidPlatform(UnrealTargetPlatform.Linux);
-			bool HololensValid = InstalledPlatformInfo.IsValidPlatform(UnrealTargetPlatform.HoloLens);
+			bool LinuxValid = InstalledPlatformInfo.IsValidPlatform(UnrealTargetPlatform.Linux) && UEBuildPlatform.IsPlatformAvailable(UnrealTargetPlatform.Linux);
+			bool HololensValid = InstalledPlatformInfo.IsValidPlatform(UnrealTargetPlatform.HoloLens) && UEBuildPlatform.IsPlatformAvailable(UnrealTargetPlatform.HoloLens);
 
 			SortedSet<string> Components = new SortedSet<string>();
 			Components.UnionWith(VisualStudioSuggestedComponents);
