PR #9903: Restore ability to use XGE executor on non-Windows platforms (Contributed by [REDACTED_PERSON])

diff --git a/Engine/Source/Programs/AutomationTool/AutomationUtils/UnrealBuild.cs b/Engine/Source/Programs/AutomationTool/AutomationUtils/UnrealBuild.cs
--- a/Engine/Source/Programs/AutomationTool/AutomationUtils/UnrealBuild.cs
+++ b/Engine/Source/Programs/AutomationTool/AutomationUtils/UnrealBuild.cs
@@ -1121,23 +1121,24 @@ namespace AutomationTool
 			string Args = null;
 			if (XGETool != null) 
 			{
-				Debug.Assert(OperatingSystem.IsWindows());
-
 				Args = "\"" + TaskFilePath + "\" /Rebuild /NoLogo /ShowAgent /ShowTime";
 				if (ParseParam("StopOnErrors"))
 				{
 					Args += " /StopOnErrors";
 				}
 
-				// A bug in the UCRT can cause XGE to hang on VS2015 builds. Figure out if this hang is likely to effect this build and workaround it if able.
-				string XGEVersion = Microsoft.Win32.Registry.GetValue(@"HKEY_LOCAL_MACHINE\SOFTWARE\Wow6432Node\Xoreax\IncrediBuild\Builder", "Version", null) as string;
-				if (XGEVersion != null)
+				if (OperatingSystem.IsWindows())
 				{
-					// Per Xoreax support, subtract 1001000 from the registry value to get the build number of the installed XGE.
-					int XGEBuildNumber;
-					if (Int32.TryParse(XGEVersion, out XGEBuildNumber) && XGEBuildNumber - 1001000 >= 1659)
+					// A bug in the UCRT can cause XGE to hang on VS2015 builds. Figure out if this hang is likely to effect this build and workaround it if able.
+					string XGEVersion = Microsoft.Win32.Registry.GetValue(@"HKEY_LOCAL_MACHINE\SOFTWARE\Wow6432Node\Xoreax\IncrediBuild\Builder", "Version", null) as string;
+					if (XGEVersion != null)
 					{
-						Args += " /no_watchdog_thread";
+						// Per Xoreax support, subtract 1001000 from the registry value to get the build number of the installed XGE.
+						int XGEBuildNumber;
+						if (Int32.TryParse(XGEVersion, out XGEBuildNumber) && XGEBuildNumber - 1001000 >= 1659)
+						{
+							Args += " /no_watchdog_thread";
+						}
 					}
 				}
 			}
diff --git a/Engine/Source/Programs/UnrealBuildTool/Executors/XGE.cs b/Engine/Source/Programs/UnrealBuildTool/Executors/XGE.cs
--- a/Engine/Source/Programs/UnrealBuildTool/Executors/XGE.cs
+++ b/Engine/Source/Programs/UnrealBuildTool/Executors/XGE.cs
@@ -86,17 +86,15 @@ namespace UnrealBuildTool
 			XmlConfig.ApplyTo(this);
 		}
 
-		[SupportedOSPlatform("windows")]
 		public override string Name
 		{
 			get { return "XGE"; }
 		}
 
-		[SupportedOSPlatform("windows")]
 		public static bool TryGetXgConsoleExecutable([NotNullWhen(true)] out string? OutXgConsoleExe)
 		{
 			// Try to get the path from the registry
-			if(BuildHostPlatform.Current.Platform == UnrealTargetPlatform.Win64)
+			if (OperatingSystem.IsWindows())
 			{
 				string? XgConsoleExe;
 				if(TryGetXgConsoleExecutableFromRegistry(RegistryView.Registry32, out XgConsoleExe))
@@ -304,11 +302,6 @@ namespace UnrealBuildTool
 
 		public static bool IsAvailable(ILogger Logger)
 		{
-			if (!OperatingSystem.IsWindows())
-			{
-				return false;
-			}
-
 			string? XgConsoleExe;
 			if (!TryGetXgConsoleExecutable(out XgConsoleExe))
 			{
@@ -389,7 +382,6 @@ namespace UnrealBuildTool
 			Logger.LogInformation("XGEEXPORT: Exported '{OutFile}'", OutFile);
 		}
 
-		[SupportedOSPlatform("windows")]
 		public override bool ExecuteActions(IEnumerable<LinkedAction> Actions, ILogger Logger)
 		{
 			if (!Actions.Any())
@@ -610,7 +602,6 @@ namespace UnrealBuildTool
 		/// <param name="ActionCount"></param>
 		/// <param name="Logger"></param>
 		/// <returns>Indicates whether the tasks were successfully executed.</returns>
-		[SupportedOSPlatform("windows")]
 		bool ExecuteTaskFile(string TaskFilePath, DataReceivedEventHandler OutputEventHandler, int ActionCount, ILogger Logger)
 		{
 			// A bug in the UCRT can cause XGE to hang on VS2015 builds. Figure out if this hang is likely to effect this build and workaround it if able.
@@ -700,7 +691,6 @@ namespace UnrealBuildTool
 		/// <summary>
 		/// Executes the tasks in the specified file, parsing progress markup as part of the output.
 		/// </summary>
-		[SupportedOSPlatform("windows")]
 		bool ExecuteTaskFileWithProgressMarkup(string TaskFilePath, int NumActions, ILogger Logger)
 		{
 			using (ProgressWriter Writer = new ProgressWriter("Compiling C++ source files...", false, Logger))
