PR #9980: New Feature: UBT - Cross Platform QMake Generator (Contributed by [REDACTED_PERSON])

diff --git a/Engine/Source/Programs/UnrealBuildTool/ProjectFiles/QMake/QMakefileGenerator.cs b/Engine/Source/Programs/UnrealBuildTool/ProjectFiles/QMake/QMakefileGenerator.cs
--- a/Engine/Source/Programs/UnrealBuildTool/ProjectFiles/QMake/QMakefileGenerator.cs
+++ b/Engine/Source/Programs/UnrealBuildTool/ProjectFiles/QMake/QMakefileGenerator.cs
@@ -115,6 +115,24 @@ namespace UnrealBuildTool
 			string BuildCommand = "";
 
 			string QMakeGameProjectFile = "";
+			string QMakeFilesDirectory = ".qmake";
+			char Seperator = Path.DirectorySeparatorChar;
+			string CurrentPlatform = "";
+			string Console = "";
+			if (OperatingSystem.IsWindows())
+			{
+				CurrentPlatform = "Win64";
+			}
+			else if (OperatingSystem.IsLinux())
+			{
+				CurrentPlatform = "Linux";
+				Console = "bash";
+			}
+			else if (OperatingSystem.IsMacOS())
+			{
+				CurrentPlatform = "Mac";
+				Console = "bash";
+			}
 
 			foreach (ProjectFile CurProject in GeneratedProjectFiles)
 			{
@@ -229,11 +247,11 @@ namespace UnrealBuildTool
 				GameProjectPath = OnlyGameProject!.Directory.FullName;
 				GameProjectFile = OnlyGameProject.FullName;
 				QMakeGameProjectFile = "gameProjectFile=" + GameProjectFile + "\n";
-				BuildCommand = $"build=bash $$unrealRootPath/Engine/{Unreal.RelativeDotnetDirectory}/dotnet $$unrealRootPath/Engine/Binaries/DotNET/UnrealBuildTool/UnrealBuildTool.dll\n\n";
+				BuildCommand = $"build={Console} $$unrealRootPath{Seperator}Engine{Seperator}{Unreal.RelativeDotnetDirectory}{Seperator}dotnet $$unrealRootPath{Seperator}Engine{Seperator}Binaries{Seperator}DotNET{Seperator}UnrealBuildTool{Seperator}UnrealBuildTool.dll\n\n";
 			}
 			else
 			{
-				BuildCommand = "build=bash $$unrealRootPath/Engine/Build/BatchFiles/Linux/Build.sh\n";
+				BuildCommand = string.Format("build=bash $$unrealRootPath{0}Engine{0}Build{0}BatchFiles{0}{1}{0}Build.sh\n", Seperator, CurrentPlatform);
 			}
 
 			string UnrealRootPath = Unreal.RootDirectory.FullName;
@@ -276,11 +294,11 @@ namespace UnrealBuildTool
 				QMakeGameProjectFile +
 				BuildCommand +
 				"args=$(ARGS)\n\n" +
-				"include(" + QMakeSourcePriFileName + ")\n" +
-				"include(" + QMakeHeaderPriFileName + ")\n" +
-				"include(" + QMakeConfigPriFileName + ")\n" +
-				"include(" + QMakeIncludesFileName + ")\n" +
-				"include(" + QMakeDefinesFileName + ")\n\n"
+				"include(" + QMakeFilesDirectory + Seperator + QMakeSourcePriFileName + ")\n" +
+				"include(" + QMakeFilesDirectory + Seperator + QMakeHeaderPriFileName + ")\n" +
+				"include(" + QMakeFilesDirectory + Seperator + QMakeConfigPriFileName + ")\n" +
+				"include(" + QMakeFilesDirectory + Seperator + QMakeIncludesFileName + ")\n" +
+				"include(" + QMakeFilesDirectory + Seperator + QMakeDefinesFileName + ")\n\n"
 			);
 
 			// Create SourceFiles, HeaderFiles, and ConfigFiles sections.
@@ -292,13 +310,13 @@ namespace UnrealBuildTool
 				{
 					string SourceFileRelativeToRoot = CurSourceFile.MakeRelativeTo(Unreal.EngineDirectory);
 					// Exclude some directories that we don't compile (note that we still want Windows/Mac etc for code navigation)
-					if (!SourceFileRelativeToRoot.Contains("Source/ThirdParty/"))
+					if (!SourceFileRelativeToRoot.Contains($"Source{Seperator}ThirdParty{Seperator}"))
 					{
 						if (SourceFileRelativeToRoot.EndsWith(".cpp"))
 						{
 							if (!SourceFileRelativeToRoot.StartsWith("..") && !Path.IsPathRooted(SourceFileRelativeToRoot))
 							{
-								QMakeSourceFilesListBuilder.Append("\t\"" + "$$unrealRootPath/Engine/" + SourceFileRelativeToRoot + "\" \\\n");
+								QMakeSourceFilesListBuilder.Append("\t\"" + $"$$unrealRootPath{Seperator}Engine{Seperator}" + SourceFileRelativeToRoot + "\" \\\n");
 							}
 							else
 							{
@@ -308,7 +326,7 @@ namespace UnrealBuildTool
 								}
 								else
 								{
-									QMakeSourceFilesListBuilder.Append("\t\"$$" + GameProjectName + "RootPath/" + Utils.MakePathRelativeTo(CurSourceFile.FullName, GameProjectPath) + "\" \\\n");
+									QMakeSourceFilesListBuilder.Append("\t\"$$" + GameProjectName + $"RootPath{Seperator}" + Utils.MakePathRelativeTo(CurSourceFile.FullName, GameProjectPath) + "\" \\\n");
 								}
 							}
 						}
@@ -317,7 +335,7 @@ namespace UnrealBuildTool
 							if (!SourceFileRelativeToRoot.StartsWith("..") && !Path.IsPathRooted(SourceFileRelativeToRoot))
 							{
 								// SourceFileRelativeToRoot = "Engine/" + SourceFileRelativeToRoot;
-								QMakeHeaderFilesListBuilder.Append("\t\"" + "$$unrealRootPath/Engine/" + SourceFileRelativeToRoot + "\" \\\n");
+								QMakeHeaderFilesListBuilder.Append("\t\"" + $"$$unrealRootPath{Seperator}Engine{Seperator}" + SourceFileRelativeToRoot + "\" \\\n");
 							}
 							else
 							{
@@ -328,7 +346,7 @@ namespace UnrealBuildTool
 								}
 								else
 								{
-									QMakeHeaderFilesListBuilder.Append("\t\"$$" + GameProjectName + "RootPath/" + Utils.MakePathRelativeTo(CurSourceFile.FullName, GameProjectPath) + "\" \\\n");
+									QMakeHeaderFilesListBuilder.Append("\t\"$$" + GameProjectName + $"RootPath{Seperator}" + Utils.MakePathRelativeTo(CurSourceFile.FullName, GameProjectPath) + "\" \\\n");
 								}
 							}
 						}
@@ -337,8 +355,7 @@ namespace UnrealBuildTool
 							if (!SourceFileRelativeToRoot.StartsWith("..") && !Path.IsPathRooted(SourceFileRelativeToRoot))
 							{
 								// SourceFileRelativeToRoot = "Engine/" + SourceFileRelativeToRoot;
-								QMakeConfigFilesListBuilder.Append("\t\"" + "$$unrealRootPath/Engine/" + SourceFileRelativeToRoot + "\" \\\n");
-
+								QMakeConfigFilesListBuilder.Append("\t\"" + $"$$unrealRootPath{Seperator}Engine{Seperator}" + SourceFileRelativeToRoot + "\" \\\n");
 							}
 							else
 							{
@@ -349,7 +366,7 @@ namespace UnrealBuildTool
 								}
 								else
 								{
-									QMakeConfigFilesListBuilder.Append("\t\"$$" + GameProjectName + "RootPath/" + Utils.MakePathRelativeTo(CurSourceFile.FullName, GameProjectPath) + "\" \\\n");
+									QMakeConfigFilesListBuilder.Append("\t\"$$" + GameProjectName + $"RootPath{Seperator}" + Utils.MakePathRelativeTo(CurSourceFile.FullName, GameProjectPath) + "\" \\\n");
 								}
 							}
 						}
@@ -393,8 +410,8 @@ namespace UnrealBuildTool
 									QMakeProjectCmdArg = " -project=\"\\\"$$gameProjectFile\\\"\"";
 								}
 								string ConfName = Enum.GetName(typeof(UnrealTargetConfiguration), CurConfiguration)!;
-								QMakeFileContent.Append(String.Format("{0}-Linux-{1}.commands = $$build {0} Linux {1} {2} $$args\n", TargetName, ConfName, QMakeProjectCmdArg));
-								QMakeTargetList += "\t" + TargetName + "-Linux-" + ConfName + " \\\n"; // , TargetName, ConfName);
+								QMakeFileContent.Append(String.Format("{0}-{3}-{1}.commands = $$build {0} {3} {1} {2} $$args\n", TargetName, ConfName, QMakeProjectCmdArg, CurrentPlatform));
+								QMakeTargetList += "\t" + TargetName + "-" + CurrentPlatform + "-" + ConfName + " \\\n"; // , TargetName, ConfName);
 							}
 						}
 					}
@@ -403,21 +420,21 @@ namespace UnrealBuildTool
 					{
 						QMakeProjectCmdArg = " -project=\"\\\"$$gameProjectFile\\\"\"";
 					}
-
-					QMakeFileContent.Append(String.Format("{0}.commands = $$build {0} Linux Development {1} $$args\n\n", TargetName, QMakeProjectCmdArg));
+					QMakeFileContent.Append(String.Format("{0}.commands = $$build {0} {2} Development {1} $$args\n\n", TargetName, QMakeProjectCmdArg, CurrentPlatform));
 					QMakeTargetList += "\t" + TargetName + " \\\n";
 				}
 			}
 
 			QMakeFileContent.Append(QMakeTargetList.TrimEnd('\\'));
-
+			
 			string FullFileName = Path.Combine(PrimaryProjectPath.FullName, FileName);
-
-			string FullQMakeDefinesFileName = Path.Combine(PrimaryProjectPath.FullName, QMakeDefinesFileName);
-			string FullQMakeIncludesFileName = Path.Combine(PrimaryProjectPath.FullName, QMakeIncludesFileName);
-			string FullQMakeSourcePriFileName = Path.Combine(PrimaryProjectPath.FullName, QMakeSourcePriFileName);
-			string FullQMakeHeaderPriFileName = Path.Combine(PrimaryProjectPath.FullName, QMakeHeaderPriFileName);
-			string FullQMakeConfigPriFileName = Path.Combine(PrimaryProjectPath.FullName, QMakeConfigPriFileName);
+			string QMakeFilePath = Path.Combine(ProjectFileGenerator.PrimaryProjectPath.FullName, QMakeFilesDirectory);
+			
+			string FullQMakeDefinesFileName = Path.Combine(QMakeFilePath, QMakeDefinesFileName);
+			string FullQMakeIncludesFileName = Path.Combine(QMakeFilePath, QMakeIncludesFileName);
+			string FullQMakeSourcePriFileName = Path.Combine(QMakeFilePath, QMakeSourcePriFileName);
+			string FullQMakeHeaderPriFileName = Path.Combine(QMakeFilePath, QMakeHeaderPriFileName);
+			string FullQMakeConfigPriFileName = Path.Combine(QMakeFilePath, QMakeConfigPriFileName);
 
 			WriteFileIfChanged(FullQMakeDefinesFileName, QMakeDefinesPriFileContent.ToString(), Logger);
 			WriteFileIfChanged(FullQMakeIncludesFileName, QMakeIncludesPriFileContent.ToString(), Logger);
