[UBT]
* Added IWYU toolchain which can be enabled with -IWYU
* Added IWYU mode that can be used with -Mode=IWYU. This can be used to modify code files
* Binary of include-what-you-use that contains a ton of fixes to make ue build

diff --git a/Engine/Source/Programs/UnrealBuildTool/Configuration/ModuleRules.cs b/Engine/Source/Programs/UnrealBuildTool/Configuration/ModuleRules.cs
--- a/Engine/Source/Programs/UnrealBuildTool/Configuration/ModuleRules.cs
+++ b/Engine/Source/Programs/UnrealBuildTool/Configuration/ModuleRules.cs
@@ -59,6 +59,39 @@ namespace UnrealBuildTool
 		User
 	}
 
+	/// <summary>
+	/// To what extent a module supports include-what-you-use
+	/// </summary>
+	public enum IWYUSupport
+	{
+		/// <summary>
+		/// None means code does not even compile. IWYU needs to skip this module entirely
+		/// </summary>
+		None,
+
+		/// <summary>
+		/// Module could be modified with iwyu but we want it to stay the way it is and handle changes manually
+		/// </summary>
+		KeepAsIs,
+
+		/// <summary>
+		/// Module is parsed and processed. This means that from the outside it is stripped for includes
+		/// even though the files are not modified. This can be used to defer iwyu work on a module.
+		/// When it comes to transitive includes this module is seen as modified from the outside.
+		/// </summary>
+		KeepAsIsForNow,
+
+		/// <summary>
+		/// Same as KeepAsIsForNow but will allow iwyu to update private headers and cpp files.
+		/// </summary>
+		KeepPublicAsIsForNow,
+
+		/// <summary>
+		/// Full iwyu support. When running with -Mode=IWYU this module will be modified if needed
+		/// </summary>
+		Full,
+	}
+
 	/// <summary>
 	/// ModuleRules is a data structure that contains the rules for defining a module
 	/// </summary>
@@ -874,7 +907,14 @@ namespace UnrealBuildTool
 		/// Enforce "include what you use" rules when PCHUsage is set to ExplicitOrSharedPCH; warns when monolithic headers (Engine.h, UnrealEd.h, etc...) 
 		/// are used, and checks that source files include their matching header first.
 		/// </summary>
-		public bool bEnforceIWYU = true;
+		//[Obsolete("Deprecated in UE5.2 - Use IWYUSupport instead.")]
+		public bool bEnforceIWYU { set { if (!value) IWYUSupport = IWYUSupport.None; } }
+
+		/// <summary>
+		/// Allows "include what you use" to modify the source code when run. bEnforceIWYU must be true for this variable to matter.
+		/// 
+		/// </summary>
+		public IWYUSupport IWYUSupport = IWYUSupport.Full;
 
 		/// <summary>
 		/// Whether to add all the default include paths to the module (eg. the Source/Classes folder, subfolders under Source/Public).
diff --git a/Engine/Source/Programs/UnrealBuildTool/Configuration/UEBuildModuleCPP.cs b/Engine/Source/Programs/UnrealBuildTool/Configuration/UEBuildModuleCPP.cs
--- a/Engine/Source/Programs/UnrealBuildTool/Configuration/UEBuildModuleCPP.cs
+++ b/Engine/Source/Programs/UnrealBuildTool/Configuration/UEBuildModuleCPP.cs
@@ -429,6 +429,12 @@ namespace UnrealBuildTool
 				}
 			}
 
+			// We are building with IWYU and thismodule does not support it, early out
+			if (Target.bIWYU && Rules.IWYUSupport == IWYUSupport.None)
+			{
+				return new List<FileItem>();
+			}
+
 			// Process all of the header file dependencies for this module
 			CheckFirstIncludeMatchesEachCppFile(Target, ModuleCompileEnvironment, InputFiles.HeaderFiles, InputFiles.CPPFiles);
 
@@ -772,6 +778,43 @@ namespace UnrealBuildTool
 				}
 			}
 
+			// IWYU needs to build all headers separate from cpp files to produce proper recommendations for includes
+			if (Target.bIWYU)
+			{
+				// Find FileItems for module's pch files
+				FileItem? PrivatePchFileItem = null;
+				if (Rules.PrivatePCHHeaderFile != null)
+					PrivatePchFileItem = FileItem.GetItemByFileReference(FileReference.Combine(ModuleDirectory, Rules.PrivatePCHHeaderFile));
+				FileItem? SharedPchFileItem = null;
+				if (Rules.SharedPCHHeaderFile != null)
+					SharedPchFileItem = FileItem.GetItemByFileReference(FileReference.Combine(ModuleDirectory, Rules.SharedPCHHeaderFile));
+
+				// Collect the headers that should be built
+				List<FileItem> HeaderFileItems = new();
+				foreach (FileItem HeaderFileItem in InputFiles.HeaderFiles)
+				{
+					// We don't want to build pch files in iwyu, skip those.
+					if (HeaderFileItem == PrivatePchFileItem || HeaderFileItem == SharedPchFileItem)
+					{
+						continue;
+					}
+
+					// If file is skipped by header units it means they can't be compiled by themselves and we must skip them too
+					if (CompileEnvironment.MetadataCache.GetHeaderUnitType(HeaderFileItem) != HeaderUnitType.Valid)
+					{
+						continue;
+					}
+
+					HeaderFileItems.Add(HeaderFileItem);
+				}
+
+				if (HeaderFileItems.Count > 0)
+				{
+					// Add the compile actions
+					LinkInputFiles.AddRange(ToolChain.CompileCPPFiles(CompileEnvironment, HeaderFileItems, IntermediateDirectory, Name, Graph).ObjectFiles);
+				}
+			}
+
 			return LinkInputFiles;
 		}
 
@@ -1454,7 +1497,7 @@ namespace UnrealBuildTool
 
 					// Find the directly included files for each source file, and make sure it includes the matching header if possible
 					InvalidIncludeDirectiveMessages = new List<string>();
-					if (Rules != null && Rules.bEnforceIWYU && Target.bEnforceIWYU)
+					if (Rules != null && Rules.IWYUSupport != IWYUSupport.None && Target.bEnforceIWYU)
 					{
 						foreach (FileItem CppFile in CppFiles)
 						{
@@ -1579,7 +1622,7 @@ namespace UnrealBuildTool
 			}
 
 			// Set the macro used to check whether monolithic headers can be used
-			if (Rules.bTreatAsEngineModule && (!Rules.bEnforceIWYU || !Target.bEnforceIWYU))
+			if (Rules.bTreatAsEngineModule && (Rules.IWYUSupport == IWYUSupport.None || !Target.bEnforceIWYU))
 			{
 				Result.Definitions.Add("SUPPRESS_MONOLITHIC_HEADER_WARNINGS=1");
 			}
diff --git a/Engine/Source/Programs/UnrealBuildTool/Modes/IWYUMode.cs b/Engine/Source/Programs/UnrealBuildTool/Modes/IWYUMode.cs
new file mode 100644
--- /dev/null
+++ b/Engine/Source/Programs/UnrealBuildTool/Modes/IWYUMode.cs
@@ -0,0 +1,1304 @@
+// Copyright Epic Games, Inc. All Rights Reserved.
+
+using System;
+using System.Collections.Generic;
+using EpicGames.Core;
+using UnrealBuildBase;
+using Microsoft.Extensions.Logging;
+using System.Linq;
+using System.Text.Json;
+using System.IO;
+using System.Threading.Tasks;
+using System.Text;
+using System.Diagnostics.CodeAnalysis;
+using System.Threading;
+
+
+// IWYUMode is a mode that can be used to clean up includes in source code. It uses the clang based tool include-what-you-use (IWYU) to figure out what is needed in each .h/.cpp file
+// and then cleans up accordingly. This mode can be used to clean up unreal as well as plugins and projects on top of unreal.
+// Note, IWYU is not perfect. There are still c++ features not supported so even though it will do a good job cleaning up it might require a little bit of hands on but not much.
+// When iwyu for some reason removes includes you want to keep, there are ways to anotate the code. Look at the pragma link below.
+// 
+// IWYU Github:  [REDACTED_PROFILE_URL]/include-what-you-use 
+// IWYU Pragmas: [REDACTED_PROFILE_URL]/include-what-you-use/blob/master/docs/IWYUPragmas.md
+// 
+// Here are examples of how a commandline could look (using ushell)
+// 
+//  1. Will build lyra editor target using include-what-you-use.exe instead of clang.exe, and preview how iwyu would update module LyraGame and all the modules depending on it
+//  
+//   .build target LyraEditor linux development -- -Mode=IWYU -ModuleToUpdate=LyraGame -UpdateDependents
+//   
+//  2. Same as 1 but will check out files from p4 and modify them
+//  
+//   .build target LyraEditor linux development -- -Mode=IWYU -ModuleToUpdate=LyraGame -UpdateDependents -Write
+//   
+//  3. Will build and then update all modules that has Lyra/Plugins in its path. Note it will only include modules that are part of LyraEditor target.
+// 
+//   .build target LyraEditor linux development -- -Mode=IWYU -PathToUpdate=Lyra/Plugins -Write
+//   
+//  4. Update Niagara plugins private code (not public api) with preview
+// 
+//   .build target UnrealEditor linux development -- -Mode=IWYU -PathToUpdate=Engine/Plugins/FX/Niagara -UpdateOnlyPrivate
+
+
+namespace UnrealBuildTool
+{
+	/// <summary>
+	/// Representing an #include inside a code file.
+	/// </summary>
+	class IWYUIncludeEntry
+	{
+		/// <summary>
+		/// Full path using forward paths, eg: d:/dev/folder/file.h
+		/// </summary>
+		public string Full { get; set; } = "";
+
+		/// <summary>
+		/// Printable path for actual include. Does not include quotes or angle brackets. eg: folder/file.h
+		/// </summary>
+		public string Printable { get; set; } = "";
+
+		/// <summary>
+		/// Reference to file info. Not populated with .iwyu file
+		/// </summary>
+		public IWYUInfo? Resolved;
+
+		/// <summary>
+		/// Decides if printable should be #include quote or angle bracket
+		/// </summary>
+		public bool System { get; set; }
+	}
+
+	/// <summary>
+	/// Comparer for include entries using printable to compare
+	/// </summary>
+	class IWYUIncludeEntryPrintableComparer : IEqualityComparer<IWYUIncludeEntry>
+	{
+		public bool Equals(IWYUIncludeEntry? x, IWYUIncludeEntry? y)
+		{
+			return x!.Printable == y!.Printable;
+		}
+
+		public int GetHashCode([DisallowNull] IWYUIncludeEntry obj)
+		{
+			throw new NotImplementedException();
+		}
+	}
+
+	/// <summary>
+	/// Representing a forward declaration inside a code file. eg: class MyClass;
+	/// </summary>
+	struct IWYUForwardEntry
+	{
+		/// <summary>
+		/// The string to add to the file
+		/// </summary>
+		public string Printable { get; set; }
+
+		/// <summary>
+		/// True if forward declaration has already been seen in the file
+		/// </summary>
+		public bool Present { get; set; }
+	}
+
+	/// <summary>
+	/// Representing a file (header or source). Each .iwyu file has a list of these (most of the time only one)
+	/// </summary>
+	class IWYUInfo
+	{
+		/// <summary>
+		/// Full path of file using forward paths, eg: d:/dev/folder/file.h
+		/// </summary>
+		public string File { get; set; } = "";
+
+		/// <summary>
+		/// Includes that this file needs. This is what iwyu have decided is used.
+		/// </summary>
+		public List<IWYUIncludeEntry> Includes { get; set; } = new List<IWYUIncludeEntry>();
+
+		/// <summary>
+		/// Forward declarations that this file needs. This is what iwyu have decided is used.
+		/// </summary>
+		public List<IWYUForwardEntry> ForwardDeclarations { get; set; } = new List<IWYUForwardEntry>();
+
+		/// <summary>
+		/// Includes seen in the file. This is how it looked like on disk.
+		/// </summary>
+		public List<IWYUIncludeEntry> IncludesSeenInFile { get; set; } = new List<IWYUIncludeEntry>();
+
+		/// <summary>
+		/// Transitive includes. This is all the includes that someone gets for free when including this file
+		/// Note, this is based on what iwyu believes should be in all includes. So it will not look at "seen includes"
+		/// and only use its generated list.
+		/// </summary>
+		public Dictionary<string, string> TransitiveIncludes = new();
+
+		/// <summary>
+		/// Transitive forward declarations. Same as transitive includes
+		/// </summary>
+		public Dictionary<string, string> TransitiveForwardDeclarations = new();
+
+		/// <summary>
+		/// Which .iwyu file that produced this info. Note, it might be null for some special cases (like .generated.h files)
+		/// </summary>
+		public IWYUFile? Source;
+
+		/// <summary>
+		/// Module that this file belongs to
+		/// </summary>
+		public UEBuildModule? Module;
+	}
+
+	/// <summary>
+	/// This is what include-what-you-use.exe produces for each build. a .iwyu file that is parsed into these instances
+	/// </summary>
+	class IWYUFile
+	{
+		/// <summary>
+		/// The different headers and source files that was covered by this execution of iwyu.
+		/// </summary>
+		public List<IWYUInfo> Files { get; set; } = new ();
+
+		/// <summary>
+		/// Name of .iwyu file
+		/// </summary>
+		public string? Name;
+	}
+
+	/// <summary>
+	/// Profiles different unity sizes and prints out the different size and its timings
+	/// </summary>
+	[ToolMode("IWYU", ToolModeOptions.XmlConfig | ToolModeOptions.BuildPlatforms | ToolModeOptions.SingleInstance | ToolModeOptions.StartPrefetchingEngine | ToolModeOptions.ShowExecutionTime)]
+	class IWYUMode : ToolMode
+	{
+		/// <summary>
+		/// Will check out files from p4 and write to disk
+		/// </summary>
+		[CommandLine("-Write")]
+		public bool bWrite = false;
+
+		/// <summary>
+		/// Update only includes in cpp files and don't touch headers
+		/// </summary>
+		[CommandLine("-UpdateOnlyPrivate")]
+		public bool bUpdateOnlyPrivate = false;
+
+		/// <summary>
+		/// Which module to run IWYU on. Will also include all modules depending on this module
+		/// </summary>
+		[CommandLine("-ModuleToUpdate")]
+		public string ModuleToUpdate = "";
+
+		/// <summary>
+		/// Which directory to run IWYU on. Will search for modules that has module directory matching this string
+		/// If no module is found in -PathToUpdate it handles this for individual files instead.
+		/// Note this can be combined with -ModuleToUpdate to double filter
+		/// </summary>
+		[CommandLine("-PathToUpdate")]
+		public string PathToUpdate = "";
+
+		/// <summary>
+		/// Also update modules that are depending on the module we are updating
+		/// </summary>
+		[CommandLine("-UpdateDependents")]
+		public bool bUpdateDependents = false;
+
+		/// <summary>
+		/// Will check out files from p4 and write to disk
+		/// </summary>
+		[CommandLine("-NoP4")]
+		public bool bNoP4 = false;
+
+		/// <summary>
+		/// Allow files to not add includes if they are transitively included by other includes
+		/// </summary>
+		[CommandLine("-NoTransitiveIncludes")]
+		public bool bNoTransitiveIncludes = false;
+
+		/// <summary>
+		/// Will skip compiling before updating. Handle with care, this is dangerous since files might not match .iwyu files
+		/// </summary>
+		[CommandLine("-NoCompile")]
+		public bool bNoCompile = false;
+
+		/// <summary>
+		/// Will ignore update check that .iwyu is newer than source files.
+		/// </summary>
+		[CommandLine("-IgnoreUpToDateCheck")]
+		public bool bIgnoreUpToDateCheck = false;
+
+		/// <summary>
+		/// If set, this will keep removed includes in #if/#endif scope at the end of updated file.
+		/// Applied to non-private headers that are part of the Engine folder.
+		/// </summary>
+		[CommandLine("-DeprecateTag")]
+		public string KeepRemovedIncludesInDeprecated = "UE_ENABLE_INCLUDE_ORDER_DEPRECATED_IN_5_2";
+
+		/// <summary>
+		/// Execute the command
+		/// </summary>
+		/// <param name="Arguments">Command line arguments</param>
+		/// <returns>Exit code</returns>
+		/// <param name="Logger"></param>
+		public override int Execute(CommandLineArguments Arguments, ILogger Logger)
+		{
+			Logger.LogInformation($"====================================================");
+			Logger.LogInformation($"Running IWYU. {(bWrite ? "" : "(Preview mode. Add -Write to write modifications to files)")}");
+			Logger.LogInformation($"====================================================");
+
+			Arguments.ApplyTo(this);
+
+			// Create the build configuration object, and read the settings
+			CommandLineArguments BuildArguments = Arguments.Append(new[] { "-IWYU" }); // Add in case it is not added (it is needed for iwyu toolchain)
+			BuildConfiguration BuildConfiguration = new BuildConfiguration();
+			XmlConfig.ApplyTo(BuildConfiguration);
+			BuildArguments.ApplyTo(BuildConfiguration);
+			BuildConfiguration.MaxNestedPathLength = 220; // For now since the path is slightly longer
+
+			// Parse all the target descriptors
+			List<TargetDescriptor> TargetDescriptors = TargetDescriptor.ParseCommandLine(BuildArguments, BuildConfiguration.bUsePrecompiled, BuildConfiguration.bSkipRulesCompile, BuildConfiguration.bForceRulesCompile, Logger);
+
+			if (TargetDescriptors.Count != 1)
+			{
+				Logger.LogError($"IWYUMode can only handle command lines that produce one target (Cmdline: {Arguments.ToString()})");
+				return 0;
+			}
+
+			// Calculate file paths filter to figure out which files that IWYU will run on
+			Logger.LogInformation($"Calculating file filter for IWYU...");
+			HashSet<string> ValidPaths = new();
+			Dictionary<string, UEBuildModule?> PathToModule = new();
+			{
+				// Create target to be able to traverse all existing modules
+				UEBuildTarget Target = UEBuildTarget.Create(TargetDescriptors[0], BuildConfiguration.bSkipRulesCompile, BuildConfiguration.bForceRulesCompile, BuildConfiguration.bUsePrecompiled, Logger);
+
+				UEBuildModule? UEModuleToUpdate = null;
+				string PathToUpdateLower = "";
+
+				// Turn provided module string into UEBuildModule reference
+				if (!String.IsNullOrEmpty(ModuleToUpdate))
+				{
+					UEModuleToUpdate = Target.GetModuleByName(ModuleToUpdate);
+					if (UEModuleToUpdate == null)
+					{
+						Logger.LogError($"Can't find module with name {ModuleToUpdate}");
+						return -1;
+					}
+				}
+				
+				if (!String.IsNullOrEmpty(PathToUpdate))
+				{
+					PathToUpdateLower = PathToUpdate.ToLower();
+				}
+				else if (UEModuleToUpdate == null)
+				{
+					Logger.LogError($"Need to provide -ModuleToUpdate <modulename> or -PathToUpdate <partofpath> to run IWYU.");
+					return -1;
+				}
+
+				int ModulesToUpdateCount = 0;
+				int ModulesSkippedCount = 0;
+
+				// Traverse all modules to figure out which ones should be updated. This is based on -ModuleToUpdate and -PathToUpdate
+				foreach (var Binary in Target.Binaries)
+				{
+					foreach (var Module in Binary.Modules)
+					{
+						if (Module.Rules.Type != ModuleRules.ModuleType.CPlusPlus)
+						{
+							continue;
+						}
+
+						bool ShouldUpdate = false;
+						string ModuleDir = Module.ModuleDirectory.FullName.Replace('\\', '/');
+
+						if (!PathToModule.TryAdd(ModuleDir, Module))
+						{
+							if (PathToModule[ModuleDir] != Module)
+							{
+								Logger.LogError("FOOO");
+								return -1;
+							}
+							continue;
+						}
+						foreach (DirectoryReference AdditionalDir in Module.ModuleDirectories)
+						{
+							if (AdditionalDir != Module.ModuleDirectory)
+							{
+								PathToModule.TryAdd(AdditionalDir.FullName.Replace('\\', '/'), Module);
+							}
+						}
+
+						if (UEModuleToUpdate != null)
+						{
+							bool DependsOnModule = Module == UEModuleToUpdate;
+							if (bUpdateDependents)
+							{
+								foreach (var Dependency in Module.PublicDependencyModules!)
+								{
+									DependsOnModule = DependsOnModule || Dependency == UEModuleToUpdate;
+								}
+								foreach (var Dependency in Module.PrivateDependencyModules!)
+								{
+									DependsOnModule = DependsOnModule || Dependency == UEModuleToUpdate;
+								}
+							}
+							ShouldUpdate = DependsOnModule;
+						}
+						
+						if (!String.IsNullOrEmpty(PathToUpdateLower))
+						{
+							if (ModuleDir.ToLower().Contains(PathToUpdateLower))
+							{
+								if (UEModuleToUpdate == null)
+								{
+									ShouldUpdate = true;
+								}
+							}
+							else
+							{
+								ShouldUpdate = false;
+							}
+						}
+						else if (UEModuleToUpdate == null)
+						{
+							ShouldUpdate = true;
+						}
+
+						if (!ShouldUpdate)
+						{
+							continue;
+						}
+
+						if (Module.Rules.IWYUSupport == IWYUSupport.None)
+						{
+							++ModulesSkippedCount;
+							continue;
+						}
+
+						++ModulesToUpdateCount;
+						ValidPaths.Add(ModuleDir);
+						foreach (DirectoryReference AdditionalDir in Module.ModuleDirectories)
+						{
+							if (AdditionalDir != Module.ModuleDirectory)
+							{
+								ValidPaths.Add(AdditionalDir.FullName.Replace('\\', '/'));
+							}
+						}
+					}
+				}
+
+				if (ValidPaths.Count == 0 && !String.IsNullOrEmpty(PathToUpdate))
+				{
+					ValidPaths.Add(PathToUpdate.Replace('\\', '/'));
+					Logger.LogInformation($"Will update files matching {PathToUpdate}. Note, path is case sensitive");
+				}
+				else
+				{
+					Logger.LogInformation($"Will update {ModulesToUpdateCount} module(s) using IWYU. ({ModulesSkippedCount} skipped because of bEnforceIWYU=false)...");
+				}
+			}
+
+			Dictionary<string, IWYUInfo> Infos = new();
+			HashSet<string> GeneratedHeaderInfos = new();
+			List<IWYUInfo> GeneratedCppInfos = new();
+			IWYUInfo? ObjectMacrosInfo = null;
+			IWYUInfo? ScriptMacrosInfo = null;
+			bool ReadSuccess = true;
+
+			TargetDescriptor Descriptor = TargetDescriptors.First();
+			using (ISourceFileWorkingSet WorkingSet = new EmptySourceFileWorkingSet())
+			{
+				// Create the make file that contains all the actions we will use to find .iwyu files.
+				Logger.LogInformation($"Creating MakeFile for target...");
+				TargetMakefile Makefile = BuildMode.CreateMakefile(BuildConfiguration, Descriptor, WorkingSet, Logger);
+
+				// Use the make file to build unless -NoCompile is set.
+				if (!bNoCompile)
+				{
+					BuildMode.Build(new TargetMakefile[] { Makefile }, new List<TargetDescriptor>() { Descriptor }, BuildConfiguration, BuildOptions.None, null, Logger);
+					SourceFileMetadataCache.SaveAll();
+					CppDependencyCache.SaveAll();
+				}
+
+				// Time to parse all the .iwyu files generated from iwyu.
+				// We Do this in parallel since it involves reading a ton of .iwyu files from disk.
+				Logger.LogInformation("Parsing output from IWYU...");
+				Parallel.ForEach(Makefile.Actions, Action =>
+				{
+					if (Action.ActionType != ActionType.Compile)
+					{
+						return;
+					}
+
+					FileItem? JsonFile = Action.ProducedItems.Where(i => i.Name.Contains(".iwyu")).FirstOrDefault();
+					if (JsonFile == null)
+					{
+						Logger.LogError($"Can't find produced json file for Action {String.Join(',', Action.PrerequisiteItems.Select(i => i.Name))}");
+						ReadSuccess = false;
+						return;
+					}
+
+					string? JsonContent = File.ReadAllText(JsonFile.AbsolutePath);
+					if (JsonContent == null)
+					{
+						Logger.LogError($"Failed to read file {JsonFile.AbsolutePath}");
+						ReadSuccess = false;
+						return;
+					}
+
+					try
+					{
+						IWYUFile? IWYUFile = JsonSerializer.Deserialize<IWYUFile>(JsonContent);
+						IWYUFile!.Name = JsonFile.AbsolutePath;
+
+						// Traverse the cpp/inl/h file entries inside the .iwyu file
+						foreach (var Info in IWYUFile!.Files)
+						{
+							Info.Source = IWYUFile;
+
+							// We track .gen.cpp in a special list, they need special treatment later
+							if (Info.File.Contains(".gen.cpp"))
+							{
+								lock (GeneratedCppInfos)
+									GeneratedCppInfos.Add(Info);
+								continue;
+							}
+
+							// Ok, time to add file entry to lookup
+							lock (Infos)
+							{
+								if (!Infos.TryAdd(Info.File, Info))
+								{
+									// This is a valid scenario when Foo.cpp also registers Foo.h... and then Foo.h registers itself.
+									var ExistingEntry = Infos[Info.File];
+									if (IWYUFile.Files.Count == 1)
+									{
+										if (ExistingEntry.Source!.Files.Count == 1)
+										{
+											Logger.LogError($"{Info.File} - built twice somehow?");
+											return;
+										}
+										else
+										{
+											Infos[Info.File] = Info;
+										}
+									}
+									else
+									{
+										//bool Equals = Info.Includes.SequenceEqual(ExistingEntry.Includes, new IWYUIncludeEntryPrintableComparer());
+										//if (!Equals)
+										//{
+										//	Logger.LogWarning($"{Info.File} - mismatch found in multiple .iwyu-files");
+										//}
+									}
+								}
+							}
+
+							string FilePath = Info.File;
+							while (true)
+							{
+								int LastIndexOfSlash = FilePath.LastIndexOf('/');
+								if (LastIndexOfSlash == -1)
+									break;
+								FilePath = FilePath.Substring(0, LastIndexOfSlash);
+								if (PathToModule.TryGetValue(FilePath, out Info.Module))
+									break;
+							}
+
+							// Some special logic for headers.
+							if (Info.File.EndsWith(".h"))
+							{
+								// Find ObjectMacros and ScriptMacros. These are always included in .generated.h files which will generate IWYUInfos for further down
+								if (ObjectMacrosInfo == null && Info.File.EndsWith("UObject/ObjectMacros.h"))
+									ObjectMacrosInfo = Info;
+								if (ScriptMacrosInfo == null && Info.File.EndsWith("UObject/ScriptMacros.h"))
+									ScriptMacrosInfo = Info;
+
+								// We need to add entries for .generated.h. They are not in the iwyu files
+								foreach (var Include in Info.Includes)
+								{
+									if (Include.Printable.EndsWith(".generated.h"))
+									{
+										lock (GeneratedHeaderInfos)
+											GeneratedHeaderInfos.Add(Include.Full);
+										break;
+									}
+								}
+							}
+						}
+					}
+					catch (Exception e)
+					{
+						Logger.LogError($"Failed to parse json {JsonFile.AbsolutePath}: {e.Message} - File will be deleted");
+						File.Delete(JsonFile.AbsolutePath);
+						ReadSuccess = false;
+						return;
+					}
+				});
+			}
+
+			// Something went wrong parsing iwyu files.
+			if (!ReadSuccess)
+				return -1;
+
+			// Add all .generated.h files as entries in the lookup and explicitly add the includes they have which will never be removed
+			Logger.LogInformation($"Generating infos for .generated.h files...");
+			if (GeneratedHeaderInfos.Count > 0)
+			{
+				IWYUIncludeEntry ObjectMacrosInclude = new();
+				IWYUIncludeEntry ScriptMacrosInclude = new();
+				ObjectMacrosInclude.Printable = "UObject/ObjectMacros.h";
+				ObjectMacrosInclude.Full = ObjectMacrosInfo!.File;
+				ObjectMacrosInclude.Resolved = ObjectMacrosInfo;
+				ScriptMacrosInclude.Printable = "UObject/ScriptMacros.h";
+				ScriptMacrosInclude.Full = ScriptMacrosInfo!.File;
+				ScriptMacrosInclude.Resolved = ScriptMacrosInfo;
+
+				foreach (var Gen in GeneratedHeaderInfos)
+				{
+					IWYUInfo GenInfo = new();
+					GenInfo.File = Gen;
+					GenInfo.Includes.Add(ObjectMacrosInclude);
+					GenInfo.Includes.Add(ScriptMacrosInclude);
+					Infos.Add(Gen, GenInfo);
+				}
+			}
+
+			Logger.LogInformation($"Found {Infos.Count} IWYU entries...");
+
+			Logger.LogInformation($"Resolving Includes...");
+			Parallel.ForEach(Infos.Values, Info =>
+			{
+				if (Info.Module != null && Info.Module.Rules.IWYUSupport == IWYUSupport.KeepAsIs)
+				{
+					foreach (IWYUIncludeEntry Entry in Info.IncludesSeenInFile)
+					{
+						// Special hack, we don't want CoreMinimal to be the reason we remove includes transitively
+						if (!Entry.Printable.Contains("CoreMinimal.h"))
+						{
+							Info.Includes.Add(Entry);
+						}
+					}
+				}
+
+				foreach (var Include in Info.Includes)
+				{
+					if (Include.Resolved == null)
+						Infos.TryGetValue(Include.Full, out Include.Resolved);
+				}
+			});
+
+			Logger.LogInformation($"Generating transitive include lists and forward declaration lists...");
+			Parallel.ForEach(Infos.Values, Info =>
+			{
+				Stack<string> Stack = new();
+				CalculateTransitive(Info, Info, Stack, Info.TransitiveIncludes, Info.TransitiveForwardDeclarations);
+			});
+
+			// If we have built .gen.cpp it means that it is not inlined in another cpp file
+			// And we need to promote the includes needed to the header since we can never modify the .gen.cpp file
+			Logger.LogInformation($"Transferring needed includes from .gen.cpp to owning file...");
+			Parallel.ForEach(GeneratedCppInfos, GeneratedCpp =>
+			{
+				// First, check which files .gen.cpp will see
+				HashSet<string> SeenTransitiveIncludes = new();
+				foreach (var SeenInclude in GeneratedCpp.IncludesSeenInFile)
+				{
+					SeenTransitiveIncludes.Add(SeenInclude.Printable);
+					foreach (var Include in GeneratedCpp.Includes)
+					{
+						IWYUInfo? IncludeInfo;
+						if (SeenInclude.Printable == Include.Printable && Infos.TryGetValue(Include.Full, out IncludeInfo))
+						{
+							foreach (string I in IncludeInfo.TransitiveIncludes.Keys)
+								SeenTransitiveIncludes.Add(I);
+							break;
+						}
+					}
+				}
+
+				IWYUInfo? IncluderInfo = null;
+
+				// If there is only one file in .iwyu it means that .gen.cpp was compiled separately
+				if (GeneratedCpp.Source!.Files.Count == 1)
+				{
+					int NameStart = GeneratedCpp.File.LastIndexOf('/');
+					string HeaderName = GeneratedCpp.File.Substring(NameStart + 1, GeneratedCpp.File.Length - NameStart - ".gen.cpp".Length) + "h";
+					foreach (var Include in GeneratedCpp.Includes)
+					{
+						NameStart = Include.Full.LastIndexOf('/');
+						string IncludeName = Include.Full.Substring(NameStart + 1);
+						if (HeaderName == IncludeName)
+						{
+							Infos.TryGetValue(Include.Full, out IncluderInfo);
+							break;
+						}
+					}
+				}
+				else // this .gen.cpp is inlined in a cpp..
+				{
+					IncluderInfo = GeneratedCpp.Source.Files.FirstOrDefault(I => I.File.Contains(".cpp") && !I.File.Contains(".gen"));
+				}
+
+				if (IncluderInfo == null)
+				{
+					return;
+				}
+
+				foreach (var Include in GeneratedCpp.Includes)
+				{
+					if (SeenTransitiveIncludes.Contains(Include.Printable))
+					{
+						continue;
+					}
+
+					// TODO: Remove UObject check once we've added "IWYU pragma: keep" around the includes in ScriptMacros and ObjectMacros
+
+					if (Include.Full.Contains(".generated.h") || Include.Printable.StartsWith("UObject/"))
+					{
+						continue;
+					}
+
+					if (!IncluderInfo!.TransitiveIncludes.ContainsKey(Include.Printable))
+					{
+						IncluderInfo.Includes.Add(Include);
+					}
+				}
+			});
+
+			object? ShouldLog = bWrite ? null : new();
+
+			List<ValueTuple<IWYUInfo, List<string>>> UpdatedFiles = new(); // <FilePath>, <NewLines>
+			HashSet<IWYUFile> SkippedFiles = new();
+			int SkippedCount = 0;
+			int OutOfDateCount = 0;
+
+			Logger.LogInformation($"Updating code files (in memory)...");
+
+			uint FilesParseCount = 0;
+			Parallel.ForEach(Infos.Values, Info =>
+			{
+				if (Info.Source == null) // .generated.h is also in this list, ignore them
+					return;
+
+				// Filter out files
+				bool Included = false;
+				foreach (string ValidPath in ValidPaths)
+				{
+					if (Info.File.Contains(ValidPath))
+					{
+						Included = true;
+						break;
+					}
+				}
+				if (!Included)
+				{
+					return;
+				}
+
+				if (Info.Module?.Rules.IWYUSupport != IWYUSupport.Full)
+				{
+					return;
+				}
+
+				bool IsCpp = Info.File.EndsWith(".cpp");
+				bool IsPrivate = IsCpp || Info.File.Contains("/Private/");
+
+				// If we only want to update private files we early out for non-private headers
+				if (!IsPrivate && (bUpdateOnlyPrivate || (Info.Module?.Rules.IWYUSupport == IWYUSupport.KeepPublicAsIsForNow)))
+				{
+					return;
+				}
+
+				Interlocked.Increment(ref FilesParseCount);
+
+				string MatchingH = "";
+				if (IsCpp)
+				{
+					int LastSlash = Info.File.LastIndexOf('/') + 1;
+					MatchingH = Info.File.Substring(LastSlash, Info.File.Length - LastSlash - 4) + ".h";
+				}
+
+				Dictionary<string, string> TransitivelyIncluded = new();
+				SortedSet<string> CleanedupIncludes = new();
+				SortedSet<string> ForwardDeclarationsToAdd = new();
+
+				foreach (var Include in Info.Includes)
+				{
+					// We never remove header with name matching cpp
+					string NameWithoutPath = Include.Printable;
+					int LastSlash = NameWithoutPath.LastIndexOf("/");
+					if (LastSlash != -1)
+						NameWithoutPath = NameWithoutPath.Substring(LastSlash + 1);
+
+					string QuotedPrintable = Include.System ? $"<{Include.Printable}>" : $"\"{Include.Printable}\"";
+
+					bool Keep = true;
+					if (Info.File == Include.Full) // Sometimes IWYU outputs include to the same file if .gen.cpp is inlined and includes file with slightly different path. just skip those
+						Keep = false;
+					else if (IsCpp && MatchingH == NameWithoutPath)
+						Keep = true;
+					else if (!bNoTransitiveIncludes)
+					{
+						foreach (var Include2 in Info.Includes)
+						{
+							if (Include2.Resolved != null && Include != Include2)
+							{
+								string Key = Include.Printable; //Full;
+								string? TransitivePath;
+								if (Include2.Resolved.TransitiveIncludes!.TryGetValue(Key, out TransitivePath))
+								{
+									if (ShouldLog != null)
+										TransitivelyIncluded.TryAdd(QuotedPrintable, string.Join(" -> ", Include2.Printable, TransitivePath));
+									Keep = false;
+									break;
+								}
+							}
+						}
+					}
+					if (Keep)
+					{
+						CleanedupIncludes.Add(QuotedPrintable);
+					}
+				}
+
+				if (!IsCpp) // Ignore forward declarations for cpp files. They should never need them afaik
+				{
+					foreach (var ForwardDeclaration in Info.ForwardDeclarations)
+					{
+						bool Add = ForwardDeclaration.Present == false;
+						if (!bNoTransitiveIncludes)
+						{
+							foreach (var Include2 in Info.Includes)
+							{
+								if (Include2.Resolved != null && Include2.Resolved.TransitiveForwardDeclarations!.ContainsKey(ForwardDeclaration.Printable))
+								{
+									Add = false;
+									break;
+								}
+							}
+						}
+						if (Add)
+						{
+							ForwardDeclarationsToAdd.Add(ForwardDeclaration.Printable);
+						}
+					}
+				}
+
+
+				// Read all lines of the header/source file
+				string[] ExistingLines = File.ReadAllLines(Info.File);
+
+				SortedSet<string> LinesToRemove = new();
+				SortedSet<string> IncludesToRemove = new();
+				SortedSet<string> IncludesToAdd = new(CleanedupIncludes);
+
+				bool HasIncludes = false;
+				string? FirstForwardDeclareLine = null;
+
+				HashSet<string> SeenIncludes = new();
+				foreach (IWYUIncludeEntry SeenInclude in Info.IncludesSeenInFile)
+				{
+					if (SeenInclude.Printable[0] != '<')
+						SeenIncludes.Add($"\"{SeenInclude.Printable}\"");
+					else
+						SeenIncludes.Add(SeenInclude.Printable);
+				}
+
+
+				bool ForceKeepScope = false;
+				int LineIndex = -1;
+				foreach (var L in ExistingLines)
+				{
+					++LineIndex;
+
+					string Line = L.Trim();
+					if (!Line.StartsWith("#include"))
+					{
+						// Might be forward declaration.. 
+						if (ForwardDeclarationsToAdd.Remove(Line)) // Skip adding the ones that already exists
+						{
+							if (FirstForwardDeclareLine == null)
+							{
+								FirstForwardDeclareLine = L;
+							}
+						}
+
+						if (Line.Contains("IWYU pragma: "))
+						{
+							if (Line.Contains(": begin_keep"))
+								ForceKeepScope = true;
+							else if (Line.Contains(": end_keep"))
+								ForceKeepScope = false;
+						}
+
+						// File is autogenerated by some tool, don't mess with it
+						if (Line.Contains("AUTO GENERATED CONTENT, DO NOT MODIFY"))
+						{
+							return;
+						}
+
+						continue;
+					}
+
+					HasIncludes = true;
+
+					bool ForceKeep = false;
+
+					ReadOnlySpan<char> IncludeSpan = Line.AsSpan(8).TrimStart();
+					char LeadingIncludeChar = IncludeSpan[0];
+					if (LeadingIncludeChar != '"' && LeadingIncludeChar != '<')
+					{
+						if (IncludeSpan.IndexOf("UE_INLINE_GENERATED_CPP_BY_NAME") != -1)
+						{
+							int Open = IncludeSpan.IndexOf('(') + 1;
+							int Close = IncludeSpan.IndexOf(')');
+							string ActualInclude = $"\"{IncludeSpan.Slice(Open, Close - Open).ToString()}.gen.cpp\"";
+							IncludesToAdd.Remove(ActualInclude);
+						}
+						else if (IncludeSpan.IndexOf("COMPILED_PLATFORM_HEADER") != -1)
+						{
+							int Open = IncludeSpan.IndexOf('(') + 1;
+							int Close = IncludeSpan.IndexOf(')');
+							string ActualInclude = $"\"Linux/Linux{IncludeSpan.Slice(Open, Close - Open).ToString()}\"";
+							IncludesToAdd.Remove(ActualInclude);
+						}
+						else
+						{
+							// TODO: These are includes made through defines. IWYU should probably report these in their original shape
+							// .. so a #include MY_SPECIAL_INCLUDE is actually reported as MY_SPECIAL_INCLUDE from IWYU instead of what the define expands to
+						}
+						continue;
+					}
+					else
+					{
+						int Index = IncludeSpan.Slice(1).IndexOf(LeadingIncludeChar == '"' ? '"' : '>');
+						if (Index != -1)
+						{
+							IncludeSpan = IncludeSpan.Slice(0, Index + 2);
+						}
+
+						if (Line.Contains("IWYU pragma: "))
+							ForceKeep = true;
+					}
+
+					string Include = IncludeSpan.ToString();
+
+					if (!SeenIncludes.Contains(Include))
+						continue;
+
+					if (!ForceKeep && !ForceKeepScope && !CleanedupIncludes.Contains(Include))
+					{
+						LinesToRemove.Add(L);
+						IncludesToRemove.Add(Include);
+					}
+					else
+					{
+						IncludesToAdd.Remove(Include);
+					}
+				}
+
+				// Nothing has changed! early out of this file
+				if (IncludesToAdd.Count == 0 && LinesToRemove.Count == 0 && ForwardDeclarationsToAdd.Count == 0)
+				{
+					return;
+				}
+
+				// If code file last write time is newer than IWYU file this means that iwyu is not up to date and needs to compile before we can apply anything
+				if (!bIgnoreUpToDateCheck)
+				{
+					FileInfo IwyuFileInfo = new FileInfo(Info.Source!.Name!);
+					FileInfo CodeFileInfo = new FileInfo(Info.File!);
+					if (CodeFileInfo.LastWriteTime > IwyuFileInfo.LastWriteTime)
+					{
+						Interlocked.Increment(ref OutOfDateCount);
+						return;
+					}
+				}
+
+
+				SortedSet<string> LinesToAdd = new();
+
+				foreach (string IncludeToAdd in IncludesToAdd)
+				{
+					LinesToAdd.Add("#include " + IncludeToAdd);
+				}
+
+				if (ShouldLog != null)
+				{
+					lock (ShouldLog)
+					{
+						System.Console.WriteLine(Info.File);
+						foreach (var I in LinesToAdd)
+						{
+							System.Console.WriteLine("  +" + I);
+						}
+						foreach (var I in IncludesToRemove)
+						{
+							System.Console.Write("  -" + I);
+							string? Reason;
+							if (TransitivelyIncluded.TryGetValue(I, out Reason))
+							{
+								System.Console.Write("  (Transitively included from " + Reason + ")");
+							}
+							System.Console.WriteLine();
+						}
+						foreach (var I in ForwardDeclarationsToAdd)
+						{
+							System.Console.WriteLine("  +" + I);
+						}
+						System.Console.WriteLine();
+					}
+				}
+
+				List<string> NewLines = new(ExistingLines.Length);
+				SortedSet<String> LinesRemoved = new();
+
+				if (!HasIncludes)
+				{
+					LineIndex = 0;
+					foreach (string OldLine in ExistingLines)
+					{
+						NewLines.Add(OldLine);
+						++LineIndex;
+
+						if (!OldLine.TrimStart().StartsWith("#pragma"))
+							continue;
+						NewLines.Add("");
+						foreach (string Line in LinesToAdd)
+							NewLines.Add(Line);
+
+						NewLines.AddRange(ExistingLines.Skip(LineIndex));
+						break;
+					}
+				}
+				else
+				{
+					// This is a bit of a tricky problem to solve in a generic ways since there are lots of exceptions and hard to make assumptions
+					// Right now we make the assumption that if there are 
+					// That will be the place where we add/remove our includes.
+
+					int ContiguousNonIncludeLineCount = 0;
+					bool IsInFirstIncludeBlock = true;
+					int LastSeenIncludeBeforeCode = -1;
+					int FirstForwardDeclareLineIndex = -1;
+
+					foreach (string OldLine in ExistingLines)
+					{
+						string OldLineTrimmedStart = OldLine.TrimStart();
+
+						if (!OldLineTrimmedStart.StartsWith("#include"))
+						{
+							if (IsInFirstIncludeBlock)
+							{
+								if (!String.IsNullOrEmpty(OldLineTrimmedStart))
+								{
+									++ContiguousNonIncludeLineCount;
+								}
+								if (LastSeenIncludeBeforeCode != -1 && ContiguousNonIncludeLineCount > 10)
+								{
+									IsInFirstIncludeBlock = false;
+								}
+								if (OldLineTrimmedStart.StartsWith("#if"))
+								{
+									// This logic is a bit shaky but handle the situations where file starts with #if and ends with #endif
+									if (LastSeenIncludeBeforeCode != -1)
+										IsInFirstIncludeBlock = false;
+								}
+
+								// This need to be inside "IsInFirstIncludeBlock" check because some files have forward declares far down in the file
+								if (FirstForwardDeclareLine == OldLine)
+								{
+									FirstForwardDeclareLineIndex = NewLines.Count;
+								}
+							}
+
+							NewLines.Add(OldLine);
+							continue;
+						}
+
+						ContiguousNonIncludeLineCount = 0;
+
+						if (!LinesToRemove.Contains(OldLine))
+						{
+							// If we find #include SOME_DEFINE we assume that should be last and "end" the include block with that
+							if (!OldLineTrimmedStart.Contains('\"') && !OldLineTrimmedStart.Contains('<'))
+							{
+								IsInFirstIncludeBlock = false;
+							}
+							else if (LinesToAdd.Count != 0 && (!IsCpp || LastSeenIncludeBeforeCode != -1))
+							{
+								string LineToAdd = LinesToAdd.First();
+								if (LineToAdd.CompareTo(OldLine) < 0)
+								{
+									NewLines.Add(LineToAdd);
+									LinesToAdd.Remove(LineToAdd);
+								}
+							}
+
+
+							NewLines.Add(OldLine);
+						}
+						else
+						{
+							FirstForwardDeclareLineIndex = -1; // This should never happen, but just in case, reset since lines have changed
+							LinesRemoved.Add(OldLine);
+						}
+
+						if (IsInFirstIncludeBlock)
+							LastSeenIncludeBeforeCode = NewLines.Count - 1;
+					}
+
+					if (LinesToAdd.Count > 0)
+					{
+						int InsertPos = LastSeenIncludeBeforeCode + 1;
+						if (NewLines[LastSeenIncludeBeforeCode].Contains(".generated.h"))
+						{
+							--InsertPos;
+						}
+
+						NewLines.InsertRange(InsertPos, LinesToAdd);
+						LastSeenIncludeBeforeCode += LinesToAdd.Count;
+
+						if (FirstForwardDeclareLineIndex != -1 && FirstForwardDeclareLineIndex > LastSeenIncludeBeforeCode)
+							FirstForwardDeclareLineIndex += LinesToAdd.Count;
+					}
+
+					if (ForwardDeclarationsToAdd.Count > 0)
+					{
+						int InsertPos;
+						if (FirstForwardDeclareLineIndex == -1)
+						{
+							InsertPos = LastSeenIncludeBeforeCode + 1;
+							if (!String.IsNullOrEmpty(NewLines[InsertPos]))
+							{
+								NewLines.Insert(InsertPos++, "");
+							}
+							NewLines.Insert(InsertPos + 1, "");
+						}
+						else
+						{
+							InsertPos = FirstForwardDeclareLineIndex;
+						}
+						NewLines.InsertRange(InsertPos + 1, ForwardDeclarationsToAdd);
+					}
+				}
+
+				// If file is public, in engine and we have a deprecation tag set we will 
+				// add a deprecated include scope at the end of the file (unless scope already exists, then we'll add it inside that)
+				string EngineDir = Unreal.EngineDirectory.FullName.Replace('\\', '/');
+				if (!IsPrivate && Info.File.StartsWith(EngineDir) && !String.IsNullOrEmpty(KeepRemovedIncludesInDeprecated))
+				{
+					// Remove the includes in LinesRemoved
+					LinesRemoved.RemoveWhere(Line =>
+					{
+						ReadOnlySpan<char> IncludeSpan = Line.AsSpan(8).TrimStart();
+						char LeadingIncludeChar = IncludeSpan[0];
+						if (LeadingIncludeChar != '"' && LeadingIncludeChar != '<')
+						{
+							return false;
+						}
+						int Index = IncludeSpan.Slice(1).IndexOf(LeadingIncludeChar == '"' ? '"' : '>');
+						if (Index == -1)
+						{
+							return false;
+						}
+						IncludeSpan = IncludeSpan.Slice(1, Index);
+						return Info.TransitiveIncludes.ContainsKey(IncludeSpan.ToString());
+					});
+
+					if (LinesRemoved.Count > 0)
+					{
+						int IndexOfDeprecateScope = -1;
+						string Match = "#if " + KeepRemovedIncludesInDeprecated;
+						for (int I = NewLines.Count - 1; I != 0; --I)
+						{
+							if (NewLines[I] == Match)
+							{
+								IndexOfDeprecateScope = I + 1;
+								break;
+							}
+						}
+
+						if (IndexOfDeprecateScope == -1)
+						{
+							NewLines.Add("");
+							NewLines.Add(Match);
+							IndexOfDeprecateScope = NewLines.Count;
+							NewLines.Add("#endif");
+						}
+						else
+						{
+							// Scan the already added includes to prevent additional adds.
+						}
+
+						NewLines.InsertRange(IndexOfDeprecateScope, LinesRemoved);
+					}
+				}
+				lock (UpdatedFiles)
+					UpdatedFiles.Add(new(Info, NewLines));
+			});
+
+			Logger.LogInformation($"Parsed {FilesParseCount} and updated {UpdatedFiles.Count} files (Found {OutOfDateCount} .iwyu files out of date)");
+
+			// Wooohoo, all files are up-to-date
+			if (UpdatedFiles.Count == 0)
+			{
+				Logger.LogInformation($"All files are up to date!");
+				return 0;
+			}
+
+			// If we have been logging we can exit now since we don't want to write any files to disk
+			if (ShouldLog != null)
+				return 0;
+
+			List<System.Diagnostics.Process> P4Processes = new();
+
+			Action<string> AddP4Process = (Arguments) =>
+			{
+				System.Diagnostics.Process Process = new System.Diagnostics.Process();
+				System.Diagnostics.ProcessStartInfo StartInfo = new System.Diagnostics.ProcessStartInfo();
+				Process.StartInfo.WindowStyle = System.Diagnostics.ProcessWindowStyle.Hidden;
+				Process.StartInfo.CreateNoWindow = true;
+				Process.StartInfo.FileName = "p4.exe";
+				Process.StartInfo.Arguments = Arguments;
+				Process.Start();
+				P4Processes.Add(Process);
+			};
+
+			Func<bool> WaitForP4 = () =>
+			{
+				bool P4Success = true;
+				foreach (var P4Process in P4Processes)
+				{
+					P4Process.WaitForExit();
+					if (P4Process.ExitCode != 0)
+					{
+						P4Success = false;
+						Logger.LogError($"p4 edit failed - {P4Process.StartInfo.Arguments}");
+					}
+					P4Process.Close();
+				}
+				P4Processes.Clear();
+				return P4Success;
+			};
+
+			if (!bNoP4)
+			{
+				List<IWYUInfo> ReadOnlyFileInfos = new();
+				foreach (var (Info, NewLines) in UpdatedFiles)
+				{
+					if (new FileInfo(Info.File).IsReadOnly)
+						ReadOnlyFileInfos.Add(Info);
+				}
+
+				if (ReadOnlyFileInfos.Count > 0)
+				{
+					// Check out files in batches. This can go quite crazy if there are lots of files.
+					// Should probably revisit this code to prevent 100s of p4 processes to start at once
+					Logger.LogInformation($"Opening {UpdatedFiles.Count} files for edit in P4... ({SkippedCount} files skipped)");
+					int ShowCount = 8;
+					foreach (var Info in ReadOnlyFileInfos)
+					{
+						Logger.LogInformation($"   edit {Info.File}");
+						if (--ShowCount == 0)
+							break;
+					}
+					if (UpdatedFiles.Count > 5)
+						Logger.LogInformation($"   ... and {UpdatedFiles.Count - 5} more.");
+
+					StringBuilder P4Arguments = new();
+					int BatchSize = 10;
+					int BatchCount = 0;
+					int Index = 0;
+					foreach (var Info in ReadOnlyFileInfos)
+					{
+						if (!SkippedFiles.Contains(Info.Source!))
+						{
+							P4Arguments.Append(" \"").Append(Info.File).Append('\"');
+							++BatchCount;
+						}
+						++Index;
+						if (BatchCount == BatchSize || Index == UpdatedFiles.Count)
+						{
+							AddP4Process($"edit{P4Arguments}");
+							P4Arguments.Clear();
+							BatchCount = 0;
+						}
+					}
+
+					// Waiting for edit
+					if (!WaitForP4())
+					{
+						return -1;
+					}
+				}
+			}
+
+			bool WriteSuccess = true;
+			Logger.LogInformation($"Writing {UpdatedFiles.Count - SkippedCount} files to disk...");
+
+			foreach (var (Info, NewLines) in UpdatedFiles)
+			{
+				if (SkippedFiles.Contains(Info.Source!))
+					continue;
+				try
+				{
+					File.WriteAllLines(Info.File, NewLines);
+				}
+				catch (Exception e)
+				{
+					Logger.LogError($"Failed to write {Info.File}: {e.Message} - File will be reverted");
+					SkippedFiles.Add(Info.Source!); // In case other entries from same file is queued
+
+					AddP4Process($"revert {String.Join(' ', Info.Source!.Files.Select(f => f.File))}");
+					WriteSuccess = false;
+				}
+			}
+
+			// Waiting for reverts
+			if (!WaitForP4())
+			{
+				return -1;
+			}	
+
+			if (!WriteSuccess)
+				return -1;
+
+			Logger.LogInformation($"Done!");
+			return 0;
+		}
+
+		/// <summary>
+		/// Calculate all indirect transitive includes for a file. This list contains does not contain itself and will handle circular dependencies
+		/// </summary>
+		void CalculateTransitive(IWYUInfo Root, IWYUInfo Info, Stack<string> Stack, Dictionary<string, string> TransitiveIncludes, Dictionary<string, string> TransitiveForwardDeclarations)
+		{
+			foreach (var Include in Info.Includes)
+			{
+				string Key = Include.Printable; // Full
+				if (TransitiveIncludes.ContainsKey(Key) || Include.Resolved == Root)
+				{
+					continue;
+				}
+
+				Stack.Push(Include.Printable);
+				string TransitivePath = String.Join(" -> ", Stack.Reverse());
+				TransitiveIncludes.Add(Key, TransitivePath);
+				if (Include.Resolved != null)
+				{
+					CalculateTransitive(Root, Include.Resolved, Stack, TransitiveIncludes, TransitiveForwardDeclarations);
+				}
+				Stack.Pop();
+			}
+
+			foreach (var ForwardDeclaration in Info.ForwardDeclarations)
+			{
+				string Key = ForwardDeclaration.Printable;
+				if (!TransitiveForwardDeclarations.ContainsKey(Key))
+					TransitiveForwardDeclarations.Add(Key, Info.File);
+			}
+		}
+	}
+}
diff --git a/Engine/Source/Programs/UnrealBuildTool/Platform/Linux/UEBuildLinux.cs b/Engine/Source/Programs/UnrealBuildTool/Platform/Linux/UEBuildLinux.cs
--- a/Engine/Source/Programs/UnrealBuildTool/Platform/Linux/UEBuildLinux.cs
+++ b/Engine/Source/Programs/UnrealBuildTool/Platform/Linux/UEBuildLinux.cs
@@ -201,7 +201,16 @@ namespace UnrealBuildTool
 
 		public override string GetFolderNameForArchitecture(string Architecture)
 		{
-			return ConvertToReadableArchitecture(Architecture);
+			string Str = ConvertToReadableArchitecture(Architecture);
+
+			// TODO: This will be revisted. Unfortunately there is no way to know we are about to build with iwyu toolchain
+			// except parsing the commandline arguments.
+			bool HAAACK = false;
+			foreach (var Arg in Environment.GetCommandLineArgs())
+				HAAACK |= Arg.Contains("iwyu", StringComparison.OrdinalIgnoreCase);
+			if (HAAACK)
+				Str += "-IWYU";
+			return Str;
 		}
 
 		public override void ResetTarget(TargetRules Target)
@@ -275,6 +284,11 @@ namespace UnrealBuildTool
 			Target.bCheckSystemHeadersForModification = (BuildHostPlatform.Current.Platform == UnrealTargetPlatform.Linux);
 
 			Target.bCompileISPC = true;
+
+			if (Target.bIWYU)
+			{
+				IWYUToolChain.ValidateTarget(Target);
+			}
 		}
 
 		/// <summary>
@@ -560,6 +574,11 @@ namespace UnrealBuildTool
 		/// <returns>New toolchain instance.</returns>
 		public override UEToolChain CreateToolChain(ReadOnlyTargetRules Target)
 		{
+			if (Target.bIWYU)
+			{
+				return new IWYUToolChain(Target, Logger);
+			}
+
 			ClangToolChainOptions Options = ClangToolChainOptions.None;
 
 			if (Target.LinuxPlatform.bEnableAddressSanitizer)
diff --git a/Engine/Source/Programs/UnrealBuildTool/ToolChain/IWYUToolChain.cs b/Engine/Source/Programs/UnrealBuildTool/ToolChain/IWYUToolChain.cs
new file mode 100644
--- /dev/null
+++ b/Engine/Source/Programs/UnrealBuildTool/ToolChain/IWYUToolChain.cs
@@ -0,0 +1,242 @@
+// Copyright Epic Games, Inc. All Rights Reserved.
+
+using System;
+using System.Collections.Generic;
+using System.IO;
+using EpicGames.Core;
+using Microsoft.Extensions.Logging;
+using UnrealBuildBase;
+
+namespace UnrealBuildTool
+{
+	class IWYUToolChain : ClangToolChain
+	{
+		private List<string> CrossCompilingArguments = new();
+		private FileReference IWYUMappingFile;
+		private string RelativePathToIWYUDirectory = @"Restricted\NotForLicensees\Source\ThirdParty\IWYU";
+
+		public static void ValidateTarget(TargetRules Target)
+		{
+			Target.bDisableDebugInfo = true;
+			Target.bDisableLinking = true;
+			Target.bUsePCHFiles = false;
+			Target.bUseSharedPCHs = false;
+			Target.bUseUnityBuild = false;
+			Target.bCompileISPC = false;
+		}
+
+		public IWYUToolChain(ReadOnlyTargetRules Target, ILogger InLogger) : base(ClangToolChainOptions.None, InLogger)
+		{
+			// set up the path to our toolchains
+			IWYUMappingFile = FileReference.Combine(Unreal.EngineDirectory, RelativePathToIWYUDirectory, "ue_mapping.imp");
+			if (IWYUMappingFile == null)
+			{
+				throw new BuildException("It seems you don't have access to NotForLicensees folder.\n" +
+										 "IWYU is not yet released to the public." +
+										 "We are working on validating so we can release the modified iwyu exe");
+			}
+		}
+
+		protected override ClangToolChainInfo GetToolChainInfo()
+		{
+			var PlatformSDK = new LinuxPlatformSDK(Logger);
+			DirectoryReference? BaseLinuxPath = PlatformSDK.GetBaseLinuxPathForArchitecture(LinuxPlatform.DefaultHostArchitecture);
+			CrossCompilingArguments.Add($"--sysroot=\"{NormalizeCommandLinePath(BaseLinuxPath!)}\"");
+
+			string? InternalSdkPath = UEBuildPlatform.GetSDK(UnrealTargetPlatform.Linux)!.GetInternalSDKPath();
+			if (InternalSdkPath != null)
+			{
+				FileReference ClangPath = FileReference.Combine(BaseLinuxPath!, "bin", $"clang++{BuildHostPlatform.Current.BinarySuffix}");
+				ClangToolChainInfo CompilerToolChainInfo = new ClangToolChainInfo(ClangPath, null!, Logger);
+				CrossCompilingArguments.Add(String.Format("-isystem\"{0}\"", System.IO.Path.Combine(InternalSdkPath, "lib/clang/" + CompilerToolChainInfo.ClangVersion + "/include/").Replace("\\", "/")));
+			}
+
+			string DevPath = ""; //@"include-what-you-use-0.19\vs_projects\bin\RelWithDebInfo";
+
+			FileReference IWYUPath = FileReference.Combine(Unreal.EngineDirectory, RelativePathToIWYUDirectory, DevPath, @"include-what-you-use.exe");
+			return new ClangToolChainInfo(IWYUPath!, null!, Logger);
+		}
+
+		protected override string GetFileNameFromExtension(string AbsolutePath, string Extension)
+		{
+			string FileName = Path.GetFileName(AbsolutePath);
+			if (AbsolutePath.EndsWith(".h"))
+			{
+				string HashString = ContentHash.MD5(AbsolutePath).GetHashCode().ToString("X4");
+				FileName = Path.GetFileNameWithoutExtension(FileName) + "_" + HashString + ".h";
+
+				if (Extension != ".d" && Extension != ".o")
+				{
+					throw new NotImplementedException($"Files with extension {Extension} not handled by IWYUToolChain");
+				}
+			}
+
+			if (Extension == ".o")
+			{
+				Extension = ".iwyu";
+			}
+
+			return FileName + Extension;
+		}
+
+		public override CPPOutput CompileCPPFiles(CppCompileEnvironment CompileEnvironment, List<FileItem> InputFiles, DirectoryReference OutputDir, string ModuleName, IActionGraphBuilder Graph)
+		{
+			if (InputFiles.Count == 0)
+			{
+				return new CPPOutput();
+			}
+
+			Dictionary<FileItem, List<FileItem>> GenCppLookup = CompileEnvironment.FileInlineGenCPPMap;
+
+			CompileEnvironment = new(CompileEnvironment);
+			CompileEnvironment.Definitions.Add("SUPPRESS_MONOLITHIC_HEADER_WARNINGS=1");
+			CompileEnvironment.Definitions.Add("PLATFORM_COMPILER_IWYU=1");
+
+			// Remove this once c++20 linux compile errors are fixed
+			CompileEnvironment.CppStandard = CppStandardVersion.Cpp17;
+
+			List<string> GlobalArguments = new();
+			GetCompileArguments_Global(CompileEnvironment, GlobalArguments);
+
+			// Need to add these on cmd line.. IWYU is not parsing commands after response file expansion.
+			string CommonCommandLineArgs = " -Xiwyu --mapping_file=" + IWYUMappingFile.FullName +
+										   " -Xiwyu --prefix_header_includes=keep" +
+//										   " -Xiwyu --transitive_includes_only" +  // Since we are building headers separately (not together with their cpp) we don't need this
+										   " -Xiwyu --no_check_matching_header" +
+										   " -Xiwyu --cxx17ns";
+
+			// Create a compile action for each source file.
+			List<FileItem> IwyuFiles = new List<FileItem>();
+			foreach (FileItem SourceFile in InputFiles)
+			{
+				//if (SourceFile.HasExtension(".cpp") || !SourceFile.FullName.Contains("\\Runtime\\Engine\\"))
+				//	continue;
+
+				Action Action = CompileCPPFile(CompileEnvironment, SourceFile, OutputDir, ModuleName, Graph, GlobalArguments, new CPPOutput());
+
+				string CommandLineArgs = CommonCommandLineArgs + " -Xiwyu --write_json_path=\"" + Action.ProducedItems[0] + "\" ";
+				if (SourceFile.HasExtension(".cpp"))
+				{
+					List<FileItem>? InlinedFiles;
+					if (GenCppLookup.TryGetValue(SourceFile, out InlinedFiles))
+					{
+						foreach (FileItem InlinedFile in InlinedFiles)
+						{
+							CommandLineArgs += $"-Xiwyu --check_also={InlinedFile.FullName.Replace('\\', '/')} ";
+						}
+					}
+				}
+
+				Action.CommandArguments = CommandLineArgs + Action.CommandArguments;
+				IwyuFiles.Add(Action.ProducedItems[0]);
+			}
+
+			return new CPPOutput() { ObjectFiles = IwyuFiles };
+		}
+
+		// Include paths must be absolute because we have code gen that can produce includes that looks like this: "../../../Source/Foo/Bar.h"
+		// IWYU combines include path and include and then try to normalize the result. This could result in a path that can't be normalized because there are more
+		// ".." than there are directories to reduce. This silently fails in odd ways
+
+		protected override string NormalizeCommandLinePath(FileSystemReference Reference)
+		{
+			return Reference.FullName.Replace('\\', '/');
+		}
+
+		protected override string NormalizeCommandLinePath(FileItem Item)
+		{
+			return Item.FullName.Replace('\\', '/');
+		}
+
+		// Code below here is mostly copy-pasted from LinuxToolChain.
+		// We didn't inherit from LinuxToolChain on purpose in order to be able to add support for running iwyu against other clang platforms than linux
+
+		static string ArchitectureSpecificDefines(string Architecture)
+		{
+			string Result = "";
+
+			if (Architecture.StartsWith("x86_64") || Architecture.StartsWith("aarch64"))
+			{
+				Result += "-D_LINUX64";
+			}
+
+			return Result;
+		}
+
+		private static bool ShouldUseLibcxx(string Architecture)
+		{
+			// set UE_LINUX_USE_LIBCXX to either 0 or 1. If unset, defaults to 1.
+			string? UseLibcxxEnvVarOverride = Environment.GetEnvironmentVariable("UE_LINUX_USE_LIBCXX");
+			if (string.IsNullOrEmpty(UseLibcxxEnvVarOverride) || UseLibcxxEnvVarOverride == "1")
+			{
+				// at the moment ARM32 libc++ remains missing
+				return Architecture.StartsWith("x86_64") || Architecture.StartsWith("aarch64");
+			}
+			return false;
+		}
+
+		protected override void GetCompileArguments_Global(CppCompileEnvironment CompileEnvironment, List<string> Arguments)
+		{
+			base.GetCompileArguments_Global(CompileEnvironment, Arguments);
+
+			if (ShouldUseLibcxx(CompileEnvironment.Architecture))
+			{
+				Arguments.Add("-nostdinc++");
+				Arguments.Add(GetSystemIncludePathArgument(DirectoryReference.Combine(Unreal.EngineSourceDirectory, "ThirdParty", "Unix", "LibCxx", "include")));
+				Arguments.Add(GetSystemIncludePathArgument(DirectoryReference.Combine(Unreal.EngineSourceDirectory, "ThirdParty", "Unix", "LibCxx", "include", "c++", "v1")));
+			}
+
+			Arguments.Add("-fno-math-errno");               // do not assume that math ops have side effects
+
+			Arguments.Add(GetRTTIFlag(CompileEnvironment)); // flag for run-time type info
+
+			Arguments.Add(ArchitectureSpecificDefines(CompileEnvironment.Architecture));
+			if (true)//CrossCompiling())
+			{
+				if (!String.IsNullOrEmpty(CompileEnvironment.Architecture))
+				{
+					Arguments.Add($"-target {CompileEnvironment.Architecture}");        // Set target triple
+				}
+
+				Arguments.AddRange(CrossCompilingArguments);
+			}
+
+			if (CompileEnvironment.bHideSymbolsByDefault)
+			{
+				Arguments.Add("-fvisibility-ms-compat");
+				Arguments.Add("-fvisibility-inlines-hidden");
+			}
+		}
+
+		protected override void GetCompileArguments_FileType(CppCompileEnvironment CompileEnvironment, FileItem SourceFile, DirectoryReference OutputDir, List<string> Arguments, Action CompileAction, CPPOutput CompileResult)
+		{
+			base.GetCompileArguments_FileType(CompileEnvironment, SourceFile, OutputDir, Arguments, CompileAction, CompileResult);
+		}
+
+		protected override void GetCompileArguments_WarningsAndErrors(CppCompileEnvironment CompileEnvironment, List<string> Arguments)
+		{
+			CompileEnvironment.ShadowVariableWarningLevel = WarningLevel.Off;
+			base.GetCompileArguments_WarningsAndErrors(CompileEnvironment, Arguments);
+			Arguments.Add("-Wno-undefined-bool-conversion");
+			Arguments.Add("-Wno-deprecated-anon-enum-enum-conversion");
+			Arguments.Add("-Wno-ambiguous-reversed-operator");
+			Arguments.Add("-Wno-pragma-once-outside-header");
+		}
+
+		// Skip ISPC headers
+		public override CPPOutput GenerateISPCHeaders(CppCompileEnvironment CompileEnvironment, List<FileItem> InputFiles, DirectoryReference OutputDir, IActionGraphBuilder Graph)
+		{
+			return new CPPOutput();
+		}
+		public override CPPOutput CompileISPCFiles(CppCompileEnvironment CompileEnvironment, List<FileItem> InputFiles, DirectoryReference OutputDir, IActionGraphBuilder Graph)
+		{
+			return new CPPOutput();
+		}
+
+		// IWYU can't link
+		public override FileItem LinkFiles(LinkEnvironment LinkEnvironment, bool bBuildImportLibraryOnly, IActionGraphBuilder Graph)
+		{
+			throw new BuildException("Unable to link with IWYU toolchain.");
+		}
+	}
+}
\ No newline at end of file
