[UBT]
* IWYUToolChain - Fixed so iwyu toolchain works without autosdk.

(Implemented by Chris, reviewed by me)

diff --git a/Engine/Source/Programs/UnrealBuildTool/ToolChain/IWYUToolChain.cs b/Engine/Source/Programs/UnrealBuildTool/ToolChain/IWYUToolChain.cs
--- a/Engine/Source/Programs/UnrealBuildTool/ToolChain/IWYUToolChain.cs
+++ b/Engine/Source/Programs/UnrealBuildTool/ToolChain/IWYUToolChain.cs
@@ -41,15 +41,14 @@ namespace UnrealBuildTool
 		{
 			var PlatformSDK = new LinuxPlatformSDK(Logger);
 			DirectoryReference? BaseLinuxPath = PlatformSDK.GetBaseLinuxPathForArchitecture(LinuxPlatform.DefaultHostArchitecture);
+
 			CrossCompilingArguments.Add($"--sysroot=\"{NormalizeCommandLinePath(BaseLinuxPath!)}\"");
 
-			string? InternalSdkPath = UEBuildPlatform.GetSDK(UnrealTargetPlatform.Linux)!.GetInternalSDKPath();
-			if (InternalSdkPath != null)
-			{
-				FileReference ClangPath = FileReference.Combine(BaseLinuxPath!, "bin", $"clang++{BuildHostPlatform.Current.BinarySuffix}");
-				ClangToolChainInfo CompilerToolChainInfo = new ClangToolChainInfo(ClangPath, null!, Logger);
-				CrossCompilingArguments.Add(String.Format("-isystem\"{0}\"", System.IO.Path.Combine(InternalSdkPath, "lib/clang/" + CompilerToolChainInfo.ClangVersion + "/include/").Replace("\\", "/")));
-			}
+			FileReference ClangPath = FileReference.Combine(BaseLinuxPath!, "bin", $"clang++{BuildHostPlatform.Current.BinarySuffix}");
+			ClangToolChainInfo CompilerToolChainInfo = new ClangToolChainInfo(ClangPath, null!, Logger);
+
+			DirectoryReference SystemPath = DirectoryReference.Combine(BaseLinuxPath!, "lib", "clang", CompilerToolChainInfo.ClangVersion.ToString(), "include");
+			CrossCompilingArguments.Add(GetSystemIncludePathArgument(SystemPath));
 
 			string DevPath = ""; //@"include-what-you-use-0.19\vs_projects\bin\RelWithDebInfo";
 
