PGO - do not use FASTGENPROFILE on MSVC as it creates unstable code

#rnx

diff --git a/Engine/Source/Programs/UnrealBuildTool/Platform/Windows/VCToolChain.cs b/Engine/Source/Programs/UnrealBuildTool/Platform/Windows/VCToolChain.cs
--- a/Engine/Source/Programs/UnrealBuildTool/Platform/Windows/VCToolChain.cs
+++ b/Engine/Source/Programs/UnrealBuildTool/Platform/Windows/VCToolChain.cs
@@ -623,12 +623,6 @@ namespace UnrealBuildTool
 				Arguments.Add("/GL");
 			}
 
-			if (CompileEnvironment.bPGOOptimize)
-			{
-				// Prevent VM execution functions from being inlined - UE-172329
-				AddDefinition(Arguments, "FORCE_NO_INLINE_VM_FUNCTION=1");
-			}
-
 			//
 			//	PC
 			//
