Android architecture overhaul (version 1, more can be done!)
- Changed -arm64 and -x64 to arm64 and x64
- Removed the Arches array and other multi-architecture stuff from Adnroid compiling and linking, and using the higher level UEToolchain functionality to handle multi-architecture
- AndroidToolchain is now a proper ClangToolchain subclass instead of overriding everything it did (being a proper ISPCToolchain subclass coming later)
- Updated some of the UEToolchain side stuff (Android is the first platform to need separate Link output, but still need a single Target)
- All the architecture stuff is driven by UEBuildAndroid, not AndroidToolchain, so the make-a-temp-toolchain object just to get architectures was removed
- Removed unused stuff like bAllowMissingNDK, bExecuteCompilerThroughShell, Arches/GPUArches arrays
- Added a high-levle ConvertToPlatformArchitecture() function that UBT can use to turn "x86_64" into the achitecture that Android wants of "x64"

diff --git a/Engine/Source/Programs/AutomationTool/Android/AndroidPlatform.Automation.cs b/Engine/Source/Programs/AutomationTool/Android/AndroidPlatform.Automation.cs
--- a/Engine/Source/Programs/AutomationTool/Android/AndroidPlatform.Automation.cs
+++ b/Engine/Source/Programs/AutomationTool/Android/AndroidPlatform.Automation.cs
@@ -489,7 +489,7 @@ Log.TraceInformation("ANDROID_HOME = {0}", PathVar);
 		}
 
 		// Apk's go to project location, not necessarily where the .so is (content only packages need to output to their directory)
-		string ApkName = Path.Combine(ProjectDir, DecoratedExeName) + Architecture + ".apk";
+		string ApkName = Path.Combine(ProjectDir, DecoratedExeName) + "-" + Architecture + ".apk";
 
 		// if the source binary was UnrealGame, handle using it or switching to project name
 		if (Path.GetFileNameWithoutExtension(Params.GetProjectExeForPlatform(UnrealTargetPlatform.Android).ToString()) == "UnrealGame")
@@ -747,6 +747,12 @@ Log.TraceInformation("ANDROID_HOME = {0}", PathVar);
 		return "";
 	}
 
+	private List<string> GetDeploymentArchitectures(ProjectParams Params, DeploymentContext SC)
+	{
+		string ArchitectureList = (Params.ClientArchitecture != "") ? Params.ClientArchitecture : SC.StageTargets[0].Receipt.Architecture;
+		return ArchitectureList.Split('+').ToList();
+	}
+
 	private List<string> CollectPluginDataPaths(DeploymentContext SC)
 	{
 		// collect plugin extra data paths from target receipts
@@ -994,8 +1000,7 @@ Log.TraceInformation("ANDROID_HOME = {0}", PathVar);
 
 		UnrealTargetConfiguration TargetConfiguration = SC.StageTargetConfigurations[0];
 
-		IAndroidToolChain ToolChain = AndroidExports.CreateToolChain(Params.RawProjectPath);
-		var Architectures = ToolChain.GetAllArchitectures();
+		List<string> Architectures = GetDeploymentArchitectures(Params, SC);
 		bool bMakeSeparateApks = UnrealBuildTool.AndroidExports.ShouldMakeSeparateApks();
 		bool bBuildWithHiddenSymbolVisibility = BuildWithHiddenSymbolVisibility(SC);
 		bool bSaveSymbols = GetSaveSymbols(SC);
@@ -1785,8 +1790,7 @@ Log.TraceInformation("ANDROID_HOME = {0}", PathVar);
 		}
 
 		UnrealTargetConfiguration TargetConfiguration = SC.StageTargetConfigurations[0];
-		IAndroidToolChain ToolChain = AndroidExports.CreateToolChain(Params.RawProjectPath);
-		var Architectures = ToolChain.GetAllArchitectures();
+		IEnumerable<string> Architectures = GetDeploymentArchitectures(Params, SC);
 		bool bMakeSeparateApks = UnrealBuildTool.AndroidExports.ShouldMakeSeparateApks();
 		bool bPackageDataInsideApk = UnrealBuildTool.AndroidExports.CreateDeploymentHandler(Params.RawProjectPath, Params.ForcePackageData).GetPackageDataInsideApk();
 
@@ -2298,17 +2302,15 @@ Log.TraceInformation("ANDROID_HOME = {0}", PathVar);
 
 	private void DeployAndroidFileServer(ProjectParams Params, DeploymentContext SC, string AFSToken)
 	{
-		var AppArchitectures = AndroidExports.CreateToolChain(Params.RawProjectPath).GetAllArchitectures();
-
 		ConfigHierarchy Ini = ConfigCache.ReadHierarchy(ConfigHierarchyType.Engine, DirectoryReference.FromFile(Params.RawProjectPath), UnrealTargetPlatform.Android);
-		bool bDisablePerfHarden = false;
+		bool bDisablePerfHarden;
 		Ini.GetBool("/Script/AndroidRuntimeSettings.AndroidRuntimeSettings", "bEnableMaliPerfCounters", out bDisablePerfHarden);
 
-		bool bUseCompression = false;
-		bool bLogFiles = false;
-		bool bReportStats = false;
-		bool bUseManualIPAddress = false;
-		string ManualIPAddress = "";
+		bool bUseCompression;
+		bool bLogFiles;
+		bool bReportStats;
+		bool bUseManualIPAddress;
+		string ManualIPAddress;
 		EConnectionType ConnectionType = GetAndroidFileServerNetworkConfig(SC, out bUseCompression, out bLogFiles, out bReportStats, out bUseManualIPAddress, out ManualIPAddress);
 
 		AndroidFileClient.OptimalADB adb = new AndroidFileClient.OptimalADB();
@@ -2324,7 +2326,7 @@ Log.TraceInformation("ANDROID_HOME = {0}", PathVar);
 			{
 				string CookFlavor = SC.FinalCookPlatform.IndexOf("_") > 0 ? SC.FinalCookPlatform.Substring(SC.FinalCookPlatform.IndexOf("_")) : "";
 				string SOName = GetSONameWithoutArchitecture(Params, SC.StageExecutables[0]);
-				Deploy.SetAndroidPluginData(AppArchitectures, CollectPluginDataPaths(SC));
+				Deploy.SetAndroidPluginData(GetDeploymentArchitectures(Params, SC), CollectPluginDataPaths(SC));
 				Deploy.PrepForUATPackageOrDeploy(Params.RawProjectPath, Params.ShortProjectName, SC.ProjectRoot, SOName, SC.LocalRoot + "/Engine", Params.Distribution, CookFlavor, SC.StageTargets[0].Receipt.Configuration, true, false);
 			}
 
@@ -2893,8 +2895,6 @@ Log.TraceInformation("ANDROID_HOME = {0}", PathVar);
 
 	private void DeployADB(ProjectParams Params, DeploymentContext SC)
     {
-		var AppArchitectures = AndroidExports.CreateToolChain(Params.RawProjectPath).GetAllArchitectures();
-
 		ConfigHierarchy Ini = ConfigCache.ReadHierarchy(ConfigHierarchyType.Engine, DirectoryReference.FromFile(Params.RawProjectPath), UnrealTargetPlatform.Android);
 		bool bDisablePerfHarden = false;
 		Ini.GetBool("/Script/AndroidRuntimeSettings.AndroidRuntimeSettings", "bEnableMaliPerfCounters", out bDisablePerfHarden);
@@ -2910,7 +2910,7 @@ Log.TraceInformation("ANDROID_HOME = {0}", PathVar);
             {
                 string CookFlavor = SC.FinalCookPlatform.IndexOf("_") > 0 ? SC.FinalCookPlatform.Substring(SC.FinalCookPlatform.IndexOf("_")) : "";
 				string SOName = GetSONameWithoutArchitecture(Params, SC.StageExecutables[0]);
-				Deploy.SetAndroidPluginData(AppArchitectures, CollectPluginDataPaths(SC));
+				Deploy.SetAndroidPluginData(GetDeploymentArchitectures(Params, SC), CollectPluginDataPaths(SC));
                 Deploy.PrepForUATPackageOrDeploy(Params.RawProjectPath, Params.ShortProjectName, SC.ProjectRoot, SOName, SC.LocalRoot + "/Engine", Params.Distribution, CookFlavor, SC.StageTargets[0].Receipt.Configuration, true, false);
             }
 
@@ -3719,7 +3719,8 @@ Log.TraceInformation("ANDROID_HOME = {0}", PathVar);
 			return "";
 		}
 
-		var AppArchitectures = AndroidExports.CreateToolChain(Params.RawProjectPath).GetAllArchitectures();
+		// @todo get TargetName?
+		IEnumerable<string> AppArchitectures = PlatformExports.GetProjectArchitectures(UnrealTargetPlatform.Android, Params.RawProjectPath, null, true, Params.Distribution);
 
 		// ask the device
 		IProcessResult ABIResult = RunAdbCommand(Params, DeviceName, " shell getprop ro.product.cpu.abi", null, ERunOptions.AppMustExist);
@@ -3730,40 +3731,11 @@ Log.TraceInformation("ANDROID_HOME = {0}", PathVar);
 		// if the architecture wasn't built, look for a backup
 		if (!AppArchitectures.Contains(DeviceArch))
 		{
-			// go from 64 to 32-bit
-			if (DeviceArch == "-arm64")
+			// Houdini emulation can run arm64 on intel
+			if (DeviceArch == "-x64")
 			{
-				DeviceArch = "-armv7";
+				DeviceArch = "-arm64";
 			}
-			// go from 64 to 32-bit
-			else if (DeviceArch == "-x64")
-			{
-				if (!AppArchitectures.Contains("-x86"))
-				{
-					DeviceArch = "-x86";
-				}
-				// if it didn't have 32-bit x86, look for 64-bit arm for emulation
-				// @todo android 64-bit: x86_64 most likely can't emulate arm64 at this ponit
-// 				else if (Array.IndexOf(AppArchitectures, "-arm64") == -1)
-// 				{
-// 					DeviceArch = "-arm64";
-// 				}
-				// finally try for 32-bit arm emulation (Houdini)
-				else
-				{
-					DeviceArch = "-armv7";
-				}
-			}
-			// use armv7 (with Houdini emulation)
-			else if (DeviceArch == "-x86")
-			{
-				DeviceArch = "-armv7";
-			}
-            else
-            {
-                // future-proof by dropping back to armv7 for unknown
-                DeviceArch = "-armv7";
-            }
 		}
 
 		// if after the fallbacks, we still don't have it, we can't continue
diff --git a/Engine/Source/Programs/UnrealBuildTool/Configuration/TargetDescriptor.cs b/Engine/Source/Programs/UnrealBuildTool/Configuration/TargetDescriptor.cs
--- a/Engine/Source/Programs/UnrealBuildTool/Configuration/TargetDescriptor.cs
+++ b/Engine/Source/Programs/UnrealBuildTool/Configuration/TargetDescriptor.cs
@@ -367,6 +367,7 @@ namespace UnrealBuildTool
 
 				// Parse the architecture parameter, or get the default for the platform (and letting the platform clean it up
 				IEnumerable<string> ParamArchitectures = Arguments.GetValues("-Architecture=", '+')
+					.Union(Arguments.GetValues("-Architectures=", '+'))
 					.Select(x => BuildPlatform.ConvertToPlatformArchitecture(x))
 					.Distinct();
 
diff --git a/Engine/Source/Programs/UnrealBuildTool/Platform/Android/AndroidExports.cs b/Engine/Source/Programs/UnrealBuildTool/Platform/Android/AndroidExports.cs
--- a/Engine/Source/Programs/UnrealBuildTool/Platform/Android/AndroidExports.cs
+++ b/Engine/Source/Programs/UnrealBuildTool/Platform/Android/AndroidExports.cs
@@ -15,12 +15,6 @@ namespace UnrealBuildTool
 	/// </summary>
 	public interface IAndroidToolChain
 	{
-		/// <summary>
-		/// Finds the list of supported architectures
-		/// </summary>
-		/// <returns>The targeted architectures</returns>
-		List<string> GetAllArchitectures();
-
 		/// <summary>
 		/// Returns the Android NDK Level
 		/// </summary>
@@ -83,25 +77,6 @@ namespace UnrealBuildTool
 	/// </summary>
 	public static class AndroidExports
 	{
-		/// <summary>
-		/// 
-		/// </summary>
-		/// <param name="ProjectFile"></param>
-		/// <returns></returns>
-		public static IAndroidToolChain CreateToolChain(FileReference ProjectFile)
-		{
-			return new AndroidToolChain(ProjectFile, null, null, Log.Logger);
-		}
-
-		/// <summary>
-		/// 
-		/// </summary>
-		/// <returns></returns>
-		public static IAndroidToolChain CreateTempToolChain()
-		{
-			return new AndroidToolChain(null, null, null, Log.Logger);
-		}
-
 		/// <summary>
 		/// 
 		/// </summary>
@@ -140,7 +115,7 @@ namespace UnrealBuildTool
 		/// <param name="Logger">Logger for output</param>
 		public static void StripSymbols(FileReference SourceFile, FileReference TargetFile, ILogger Logger)
 		{
-			AndroidToolChain ToolChain = new AndroidToolChain(null, null, null, Logger);
+			AndroidToolChain ToolChain = new AndroidToolChain(null, Logger);
 			ToolChain.StripSymbols(SourceFile, TargetFile, Logger);
 		}
 	}
diff --git a/Engine/Source/Programs/UnrealBuildTool/Platform/Android/AndroidToolChain.cs b/Engine/Source/Programs/UnrealBuildTool/Platform/Android/AndroidToolChain.cs
--- a/Engine/Source/Programs/UnrealBuildTool/Platform/Android/AndroidToolChain.cs
+++ b/Engine/Source/Programs/UnrealBuildTool/Platform/Android/AndroidToolChain.cs
@@ -18,6 +18,7 @@ namespace UnrealBuildTool
 		// Minimum NDK API level to allow
 		public const int MinimumNDKAPILevel = 26;
 
+		// this is architectures with the dash, which we match in filenames that have inlined arch name
 		public static readonly string[] AllCpuSuffixes =
 		{
 			"-arm64",
@@ -28,8 +29,8 @@ namespace UnrealBuildTool
 		public static readonly Dictionary<string, string> ShortArchNames = new Dictionary<string, string>()
 		{
 			{ "", "" },
-			{ "-arm64", "a8" },
-			{ "-x64", "x6" },
+			{ "arm64", "a8" },
+			{ "x64", "x6" },
 		};
 
 		public enum ClangSanitizer
@@ -56,44 +57,38 @@ namespace UnrealBuildTool
 		}
 
 		protected FileReference? ProjectFile;
-		private List<string> AdditionalArches;
-		private List<string> AdditionalGPUArches;
-		protected bool bExecuteCompilerThroughShell;
 
 		// Version string from the Android specific build of clang. E.g in Android (6317467 based on r365631c1) clang version 9.0.8
 		// this would be 6317467)
 		protected static string? AndroidClangBuild;
 
-		// the list of architectures we will compile for
-		protected List<string>? Arches = null;
-
 		// the "-android" suffix paths here are vcpkg triplets for the android platform
 		static private Dictionary<string, string[]> AllArchNames = new Dictionary<string, string[]> {
-			{ "-arm64", new string[] { "arm64", "arm64-v8a", "arm64-android" } },
-			{ "-x64",   new string[] { "x64", "x86_64", "x64-android" } }
+			{ "arm64", new string[] { "arm64", "arm64-v8a", "arm64-android" } },
+			{ "x64",   new string[] { "x64", "x86_64", "x64-android" } }
 		};
 
 		// architecture paths to use for filtering include and lib paths
 		static private Dictionary<string, string[]> AllFilterArchNames = new Dictionary<string, string[]> {
-			{ "-armv7", new string[] { "armv7", "armeabi-v7a", "arm-android" } },
-			{ "-arm64", new string[] { "arm64", "arm64-v8a", "arm64-android" } },
-			{ "-x86",   new string[] { "x86", "x86-android" } },
-			{ "-x64",   new string[] { "x64", "x86_64", "x64-android" } }
+			{ "armv7", new string[] { "armv7", "armeabi-v7a", "arm-android" } },
+			{ "arm64", new string[] { "arm64", "arm64-v8a", "arm64-android" } },
+			{ "x86",   new string[] { "x86", "x86-android" } },
+			{ "x64",   new string[] { "x64", "x86_64", "x64-android" } }
 		};
 
 		static private Dictionary<string, string[]> LibrariesToSkip = new Dictionary<string, string[]> {
-			{ "-arm64", new string[] { "nvToolsExt", "nvToolsExtStub", "vorbisenc", } },
-			{ "-x64",   new string[] { "nvToolsExt", "nvToolsExtStub", "oculus", "OVRPlugin", "vrapi", "ovrkernel", "systemutils", "openglloader", "ovrplatformloader", "gpg", "vorbisenc", } }
+			{ "arm64", new string[] { "nvToolsExt", "nvToolsExtStub", "vorbisenc", } },
+			{ "x64",   new string[] { "nvToolsExt", "nvToolsExtStub", "oculus", "OVRPlugin", "vrapi", "ovrkernel", "systemutils", "openglloader", "ovrplatformloader", "gpg", "vorbisenc", } }
 		};
 
 		static private Dictionary<string, string[]> ModulesToSkip = new Dictionary<string, string[]> {
-			{ "-arm64", new string[] {  } },
-			{ "-x64",   new string[] { "OnlineSubsystemOculus", "OculusHMD", "OculusMR", "OnlineSubsystemGooglePlay" } }
+			{ "arm64", new string[] {  } },
+			{ "x64",   new string[] { "OnlineSubsystemOculus", "OculusHMD", "OculusMR", "OnlineSubsystemGooglePlay" } }
 		};
 
 		static private Dictionary<string, string[]> GeneratedModulesToSkip = new Dictionary<string, string[]> {
-			{ "-arm64", new string[] {  } },
-			{ "-x64",   new string[] { "OculusEntitlementCallbackProxy", "OculusCreateSessionCallbackProxy", "OculusFindSessionsCallbackProxy", "OculusIdentityCallbackProxy", "OculusNetConnection", "OculusNetDriver", "OnlineSubsystemOculus_init" } }
+			{ "arm64", new string[] {  } },
+			{ "x64",   new string[] { "OculusEntitlementCallbackProxy", "OculusCreateSessionCallbackProxy", "OculusFindSessionsCallbackProxy", "OculusIdentityCallbackProxy", "OculusNetConnection", "OculusNetDriver", "OnlineSubsystemOculus_init" } }
 		};
 
 		public string? NDKToolchainVersion;
@@ -121,55 +116,22 @@ namespace UnrealBuildTool
 			return NDKVersionInt >= 220000;
 		}
 
-		[CommandLine("-Architectures=", ListSeparator = '+')]
-		public List<string> ArchitectureArg = new List<string>();
-
-		public AndroidToolChain(FileReference? InProjectFile, IReadOnlyList<string>? InAdditionalArches, IReadOnlyList<string>? InAdditionalGPUArches, ILogger InLogger)
-			: this(InProjectFile, InAdditionalArches, InAdditionalGPUArches, false, ClangToolChainOptions.None, InLogger)
-		{
-		}
-
-		public AndroidToolChain(FileReference? InProjectFile, IReadOnlyList<string>? InAdditionalArches, IReadOnlyList<string>? InAdditionalGPUArches, ClangToolChainOptions ToolchainOptions, ILogger InLogger)
-			: this(InProjectFile, InAdditionalArches, InAdditionalGPUArches, false, ToolchainOptions, InLogger)
+		public AndroidToolChain(FileReference? InProjectFile, ILogger InLogger)
+			: this(InProjectFile, ClangToolChainOptions.None, InLogger)
 		{
 		}
 
-		protected AndroidToolChain(FileReference? InProjectFile, IReadOnlyList<string>? InAdditionalArches, IReadOnlyList<string>? InAdditionalGPUArches, bool bAllowMissingNDK, ClangToolChainOptions ToolchainOptions, ILogger InLogger)
+		public AndroidToolChain(FileReference? InProjectFile, ClangToolChainOptions ToolchainOptions, ILogger InLogger)
 			: base(ToolchainOptions, InLogger)
 		{
 			Options = ToolchainOptions;
 			ProjectFile = InProjectFile;
-			AdditionalArches = new List<string>();
-			AdditionalGPUArches = new List<string>();
-
-			if (InAdditionalArches != null)
-			{
-				AdditionalArches.AddRange(InAdditionalArches);
-			}
-
-			if (InAdditionalGPUArches != null)
-			{
-				AdditionalGPUArches.AddRange(InAdditionalGPUArches);
-			}
-
-			// by default tools chains don't parse arguments, but we want to be able to check the -architectures flag defined above. This is
-			// only necessary when AndroidToolChain is used during UAT
-			CommandLine.ParseArguments(Environment.GetCommandLineArgs(), this, Logger);
-
-			if (AdditionalArches.Count == 0 && ArchitectureArg.Count > 0)
-			{
-				AdditionalArches.AddRange(ArchitectureArg);
-			}
 
 			string? NDKPath = Environment.GetEnvironmentVariable("NDKROOT");
 
 			// don't register if we don't have an NDKROOT specified
 			if (String.IsNullOrEmpty(NDKPath))
 			{
-				if (bAllowMissingNDK)
-				{
-					return;
-				}
 				throw new BuildException("NDKROOT is not specified; cannot use Android toolchain.");
 			}
 
@@ -259,7 +221,7 @@ namespace UnrealBuildTool
 			int NDKApiLevel64Int = GetNdkApiLevelInt(MinimumNDKAPILevel);
 
 			string GCCToolchainPath = Path.Combine(NDKPath, @"toolchains/llvm", ArchitecturePath);
-			string SysrootPath = Path.Combine(NDKPath, @"toolchains/llvm", ArchitecturePath, "sysroot");
+			string SysrootPath = Path.Combine(NDKPath, @"toolchains/llvm", ArchitecturePath, "sysroot").Replace("\\", "/");
 
 			// toolchain params (note: use ANDROID=1 same as we define it)
 			ToolchainLinkParamsArm64 = " --target=aarch64-none-linux-android" + NDKApiLevel64Int + " --gcc-toolchain=\"" + GCCToolchainPath + "\" --sysroot=\"" + SysrootPath + "\" -DANDROID=1";
@@ -278,52 +240,12 @@ namespace UnrealBuildTool
 			ReadElfPath = Path.Combine(NDKPath, @"toolchains/llvm", ArchitecturePath, @"bin/llvm-readelf" + ExeExtension);
 		}
 
+
 		protected override ClangToolChainInfo GetToolChainInfo()
 		{
 			return new ClangToolChainInfo(FileReference.FromString(ClangPath)!, FileReference.FromString(ArPathArm64)!, Logger);
 		}
 
-		public virtual void ParseArchitectures()
-		{
-			// look in ini settings for what platforms to compile for
-			ConfigHierarchy Ini = ConfigCache.ReadHierarchy(ConfigHierarchyType.Engine, DirectoryReference.FromFile(ProjectFile), UnrealTargetPlatform.Android);
-			Arches = new List<string>();
-			bool bBuild = true;
-			bool bUnsupportedBinaryBuildArch = false;
-
-			if (Ini.GetBool("/Script/AndroidRuntimeSettings.AndroidRuntimeSettings", "bBuildForArm64", out bBuild) && bBuild
-				|| (AdditionalArches != null && (AdditionalArches.Contains("arm64", StringComparer.OrdinalIgnoreCase) || AdditionalArches.Contains("-arm64", StringComparer.OrdinalIgnoreCase))))
-			{
-				Arches.Add("-arm64");
-			}
-			if (Ini.GetBool("/Script/AndroidRuntimeSettings.AndroidRuntimeSettings", "bBuildForx8664", out bBuild) && bBuild
-				|| (AdditionalArches != null && (AdditionalArches.Contains("x64", StringComparer.OrdinalIgnoreCase) || AdditionalArches.Contains("-x64", StringComparer.OrdinalIgnoreCase))))
-			{
-				if (File.Exists(Path.Combine(Unreal.EngineDirectory.FullName, "Build", "InstalledBuild.txt")))
-				{
-					bUnsupportedBinaryBuildArch = true;
-					Log.TraceWarningOnce("Please install source to build for x86_64 (-x64); ignoring this architecture target.");
-				}
-				else
-				{
-					Arches.Add("-x64");
-				}
-			}
-
-			// force armv7 if something went wrong
-			if (Arches.Count == 0)
-			{
-				if (bUnsupportedBinaryBuildArch)
-				{
-					throw new BuildException("Only architectures unsupported by binary-only engine selected.");
-				}
-				else
-				{
-					Arches.Add("-arm64");
-				}
-			}
-		}
-
 		static public string GetGLESVersion(bool bBuildForES31)
 		{
 			string GLESversion = "0x00030000";
@@ -376,23 +298,6 @@ namespace UnrealBuildTool
 			return Path.Combine(LinkEnvironment.IntermediateDirectory!.FullName, "ExportSymbols.ldscript");
 		}
 
-		public override void SetUpGlobalEnvironment(ReadOnlyTargetRules Target)
-		{
-			base.SetUpGlobalEnvironment(Target);
-
-			ParseArchitectures();
-		}
-
-		public List<string> GetAllArchitectures()
-		{
-			if (Arches == null)
-			{
-				ParseArchitectures();
-			}
-
-			return Arches!;
-		}
-
 		public int GetNdkApiLevelInt(int MinNdk = MinimumNDKAPILevel)
 		{
 			string NDKVersion = GetNdkApiLevel();
@@ -537,163 +442,80 @@ namespace UnrealBuildTool
 			return "android-" + MaxPlatform.ToString();
 		}
 
-		protected virtual string GetCLArguments_Global(CppCompileEnvironment CompileEnvironment, string Architecture)
-		{
-			string Result = "";
-
-			switch (Architecture)
-			{
-				case "-arm64": Result += ToolchainParamsArm64; break;
-				case "-x64": Result += ToolchainParamsx64; break;
-				default: Result += ToolchainParamsArm64; break;
-			}
 
-			// build up the commandline common to C and C++
-			Result += " -c";
-			Result += " -no-canonical-prefixes";
-			Result += " -fdiagnostics-format=msvc";
-			Result += " -Wall";
-			Result += " -Wdelete-non-virtual-dtor";
-			Result += " -fno-PIE";
-
-			Result += " -Wno-unused-variable";
-            if (CompilerVersionGreaterOrEqual(13, 0, 0))
-            {
-                Result += " -Wno-unused-but-set-variable";
-                Result += " -Wno-unused-but-set-parameter";
-                Result += " -Wno-ordered-compare-function-pointers";
-            }
+		protected override void GetCompileArguments_IncludePaths(CppCompileEnvironment CompileEnvironment, List<string> Arguments)
+		{
+			// remove paths that are not meant for this architecture
+			// @todo can remove this when we only add paths properly for architecture
+			IEnumerable<DirectoryReference> FilteredPaths = CompileEnvironment.UserIncludePaths.Where(x => IsDirectoryForArch(x.FullName, CompileEnvironment.Architecture));
+			Arguments.AddRange(FilteredPaths.Select(IncludePath => GetUserIncludePathArgument(IncludePath)));
 			
-			if (CompilerVersionGreaterOrEqual(14, 0, 2))
-			{
-				Result += " -Wno-bitwise-instead-of-logical";
-			}
-
-			// this will hide the warnings about static functions in headers that aren't used in every single .cpp file
-			Result += " -Wno-unused-function";
-			// this hides the "enumeration value 'XXXXX' not handled in switch [-Wswitch]" warnings - we should maybe remove this at some point and add UE_LOG(, Fatal, ) to default cases
-			Result += " -Wno-switch";
-			// this hides the "warning : comparison of unsigned expression < 0 is always false" type warnings due to constant comparisons, which are possible with template arguments
-			Result += " -Wno-tautological-compare";
-			//This will prevent the issue of warnings for unused private variables.
-			Result += " -Wno-unused-private-field";
-			Result += " -Wno-local-type-template-args"; // engine triggers this
-			Result += " -Wno-return-type-c-linkage";    // needed for PhysX
-			Result += " -Wno-reorder";                  // member initialization order
-			Result += " -Wno-unknown-pragmas";          // probably should kill this one, sign of another issue in PhysX?
-			Result += " -Wno-invalid-offsetof";         // needed to suppress warnings about using offsetof on non-POD types.
-			Result += " -Wno-logical-op-parentheses";   // needed for external headers we can't change
-
-			if (BuildWithHiddenSymbolVisibility(CompileEnvironment))
-			{
-				// Result += " -fvisibility-ms-compat -fvisibility-inlines-hidden"; // This hides all symbols by default but exports all type info (vtable/rtti) for a non-monolithic setup
-				Result += " -fvisibility=hidden -fvisibility-inlines-hidden"; // Symbols default to hidden.
-				//TODO: add when Android's clang will support this
-				//Resul += " -fvisibility-inlines-hidden-static-local-var";
-			}
-
-			if (CompilerVersionGreaterOrEqual(12, 0, 0))
-			{
-				// The 'this' vs nullptr comparisons get optimized away for newer versions of Clang, which is undesirable until we refactor these checks.
-				Result += " -fno-delete-null-pointer-checks";
-			}
-
-			if (CompileEnvironment.DeprecationWarningLevel == WarningLevel.Error)
-			{
-				Result += " -Werror=deprecated-declarations";
-			}
-
-			if (CompileEnvironment.ShadowVariableWarningLevel != WarningLevel.Off)
-			{
-				Result += " -Wshadow" + ((CompileEnvironment.ShadowVariableWarningLevel == WarningLevel.Error) ? "" : " -Wno-error=shadow");
-			}
-
-			if (CompileEnvironment.bEnableUndefinedIdentifierWarnings)
-			{
-				Result += " -Wundef" + (CompileEnvironment.bUndefinedIdentifierWarningsAsErrors ? "" : " -Wno-error=undef");
-			}
-
-			// new for clang4.5 warnings:
-			if (CompilerVersionGreaterOrEqual(3, 5, 0))
-			{
-				Result += " -Wno-undefined-bool-conversion"; // 'this' pointer cannot be null in well-defined C++ code; pointer may be assumed to always convert to true (if (this))
-
-				// we use this feature to allow static FNames.
-				Result += " -Wno-gnu-string-literal-operator-template";
-			}
+			FilteredPaths = CompileEnvironment.SystemIncludePaths.Where(x => IsDirectoryForArch(x.FullName, CompileEnvironment.Architecture)); 
+			Arguments.AddRange(FilteredPaths.Select(IncludePath => GetSystemIncludePathArgument(IncludePath)));
+		}
 
-			if (CompilerVersionGreaterOrEqual(3, 6, 0))
-			{
-				Result += " -Wno-unused-local-typedef";             // clang is being overly strict here? PhysX headers trigger this.
-				Result += " -Wno-inconsistent-missing-override";    // these have to be suppressed for UE 4.8, should be fixed later.
-			}
+		/// <inheritdoc/>
+		protected override void GetCompileArguments_WarningsAndErrors(CppCompileEnvironment CompileEnvironment, List<string> Arguments)
+		{
+			base.GetCompileArguments_WarningsAndErrors(CompileEnvironment, Arguments);
 
-			if (CompilerVersionGreaterOrEqual(3, 8, 275480))
-			{
-				Result += " -Wno-undefined-var-template";           // not really a good warning to disable
-				Result += " -Wno-nonportable-include-path";         // not all of these are real
-			}
 
-            if (CompilerVersionGreaterOrEqual(4, 0, 0))
-			{
-				Result += " -Wno-unused-lambda-capture";            // probably should fix the code
-																	//				Result += " -Wno-nonportable-include-path";         // not all of these are real
-			}
+			// @todo unlikely all needed
+			Arguments.Add("-Wno-local-type-template-args"); // engine triggers this
+			Arguments.Add("-Wno-return-type-c-linkage");    // needed for PhysX
+			Arguments.Add("-Wno-reorder");                  // member initialization order
+			Arguments.Add("-Wno-logical-op-parentheses");   // needed for external headers we can't change
+			Arguments.Add("-Wno-nonportable-include-path"); // not all of these are real
+		}
 
-			// shipping builds will cause this warning with "ensure", so disable only in those case
-			if (CompileEnvironment.Configuration == CppConfiguration.Shipping)
-			{
-				Result += " -Wno-unused-value";
-			}
+		/// <inheritdoc/>
+		protected override void GetCompileArguments_Optimizations(CppCompileEnvironment CompileEnvironment, List<string> Arguments)
+		{
+			// @todo the base has PGO options that are a little different than what we want - it would be nice to merge these
 
-			// debug info
-			if (CompileEnvironment.bCreateDebugInfo)
+			if (CompilerVersionGreaterOrEqual(12, 0, 0))
 			{
-				Result += " -g2 -gdwarf-4";
+				// We have 'this' vs nullptr comparisons that get optimized away for newer versions of Clang, which is undesirable until we refactor these checks.
+				Arguments.Add("-fno-delete-null-pointer-checks");
 			}
 
 			// optimization level
 			if (!CompileEnvironment.bOptimizeCode)
 			{
-				Result += " -O0";
+				Arguments.Add("-O0");
 			}
 			else
 			{
 				if (CompileEnvironment.OptimizationLevel == OptimizationMode.Size)
 				{
-					Result += " -Oz";
+					Arguments.Add("-Oz");
 				}
 				else if (CompileEnvironment.OptimizationLevel == OptimizationMode.SizeAndSpeed)
 				{
-					Result += " -Os";
-					if (Architecture == "-arm64")
+					Arguments.Add("-Os");
+					if (CompileEnvironment.Architecture == "arm64")
 					{
-						Result += " -moutline";
+						Arguments.Add("-moutline");
 					}
 				}
 				else
 				{
-					Result += " -O3";
+					Arguments.Add("-O3");
 				}
 			}
 
-			// FORTIFY default
-			Result += " -D_FORTIFY_SOURCE=2";
-
-			// Optionally enable exception handling (off by default since it generates extra code needed to propagate exceptions)
-			if (CompileEnvironment.bEnableExceptions)
-			{
-				Result += " -fexceptions";
-				Result += " -DPLATFORM_EXCEPTIONS_DISABLED=0";
-			}
-			else
+			if (CompileEnvironment.bAllowLTCG)
 			{
-				Result += " -fno-exceptions";
-				Result += " -DPLATFORM_EXCEPTIONS_DISABLED=1";
+				if ((Options & ClangToolChainOptions.EnableThinLTO) != 0)
+				{
+					Arguments.Add("-flto=thin");
+				}
+				else
+				{
+					Arguments.Add("-flto");
+				}
 			}
 
-			Result += $" {GetRTTIFlag(CompileEnvironment)}";
-
 			// Profile Guided Optimization (PGO) and Link Time Optimization (LTO)
 			if (CompileEnvironment.bPGOOptimize)
 			{
@@ -703,111 +525,186 @@ namespace UnrealBuildTool
 				//
 				// Disable these warnings. They are far too verbose.
 				//
-				Result += " -Wno-profile-instr-out-of-date";
-				Result += " -Wno-profile-instr-unprofiled";
-
+				Arguments.Add("-Wno-profile-instr-out-of-date");
+				Arguments.Add("-Wno-profile-instr-unprofiled");
 				// apparently there can be hashing conflicts with PGO which can result in:
 				// 'Function control flow change detected (hash mismatch)' warnings. 
-				Result += " -Wno-backend-plugin";
-
-				Result += string.Format(" -fprofile-use=\"{0}.profdata\"", Path.Combine(CompileEnvironment.PGODirectory!, CompileEnvironment.PGOFilenamePrefix!));
+				Arguments.Add("-Wno-backend-plugin");
+				Arguments.Add(string.Format(" -fprofile-use=\"{0}.profdata\"", Path.Combine(CompileEnvironment.PGODirectory!, CompileEnvironment.PGOFilenamePrefix!)));
 			}
 			else if (CompileEnvironment.bPGOProfile)
 			{
-				// Always enable LTO when generating PGO profile data.
 				// Android supports only LLVM IR-based for instrumentation-based PGO.
-				Result += " -fprofile-generate";
+				Arguments.Add("-fprofile-generate");
 				//  for sampling-based profile collection to generate minimal debug information:
-				Result += " -gline-tables-only";
+				Arguments.Add("-gline-tables-only");
 			}
 
-			if (CompileEnvironment.bAllowLTCG)
+			if (!CompileEnvironment.bUseInlining)
 			{
-				if ((Options & ClangToolChainOptions.EnableThinLTO) != 0)
-				{
-					Result += " -flto=thin";
-				}
-				else
-				{
-					Result += " -flto";
-				}
+				Arguments.Add("-fno-inline-functions");
 			}
+		}
 
-			if (Architecture == "-arm64")
-			{
-				Result += " -funwind-tables";           // Just generates any needed static data, affects no code
-				Result += " -fno-strict-aliasing";      // Prevents unwanted or invalid optimizations that could produce incorrect code
-				Result += " -fPIC";                     // Generates position-independent code (PIC) suitable for use in a shared library
-				Result += " -fno-short-enums";          // Do not allocate to an enum type only as many bytes as it needs for the declared range of possible values
-				Result += " -D__arm64__";               // for some reason this isn't defined and needed for PhysX
+		/// <inheritdoc/>
+		protected override void GetCompileArguments_Debugging(CppCompileEnvironment CompileEnvironment, List<string> Arguments)
+		{
+			base.GetCompileArguments_Debugging(CompileEnvironment, Arguments);
 
-				Result += " -march=armv8-a";
-				//Result += " -mfloat-abi=softfp";
-				//Result += " -mfpu=vfpv3-d16";			//@todo android: UE3 was just vfp. arm7a should all support v3 with 16 registers
+			// debug info
+			if (CompileEnvironment.bCreateDebugInfo)
+			{
+				Arguments.Add("-g2");
+				Arguments.Add("-gdwarf-4");
+			}
 
-				// Add flags for on-device debugging
-				if (CompileEnvironment.Configuration == CppConfiguration.Debug)
-				{
-					Result += " -fno-omit-frame-pointer";   // Disable removing the save/restore frame pointer for better debugging
-					if (CompilerVersionGreaterOrEqual(3, 6, 0))
-					{
-						Result += " -fno-function-sections";    // Improve breakpoint location
-					}
-				}
+			if (!DisableStackProtector())
+			{
+				Arguments.Add("-fstack-protector-strong");  // Emits extra code to check for buffer overflows
+			}
 
-				// Some switches interfere with on-device debugging
-				if (CompileEnvironment.Configuration != CppConfiguration.Debug && !DisableFunctionDataSections())
+			// Add flags for on-device debugging
+			if (CompileEnvironment.Configuration == CppConfiguration.Debug)
+			{
+				Arguments.Add("-fno-omit-frame-pointer");   // Disable removing the save/restore frame pointer for better debugging
+				if (CompilerVersionGreaterOrEqual(3, 6, 0))
 				{
-					Result += " -ffunction-sections";   // Places each function in its own section of the output file, linker may be able to perform opts to improve locality of reference
-					Result += " -fdata-sections";		// Places each data item in its own section of the output file, linker may be able to perform opts to improve locality of reference
+					Arguments.Add(" -fno-function-sections");    // Improve breakpoint location
 				}
-
-				Result += " -fsigned-char";             // Treat chars as signed //@todo android: any concerns about ABI compatibility with libs here?
 			}
-			else if (Architecture == "-x64")
+
+			// Some switches interfere with on-device debugging
+			if (CompileEnvironment.Configuration != CppConfiguration.Debug && !DisableFunctionDataSections())
 			{
-				Result += " -funwind-tables";           // Just generates any needed static data, affects no code
-				Result += " -fPIC";                     // Generates position-independent code (PIC) suitable for use in a shared library
-				Result += " -fno-omit-frame-pointer";
-				Result += " -fno-strict-aliasing";
-				Result += " -fno-short-enums";
-				Result += " -march=atom";
-				Result += " -msse4.2";
+				Arguments.Add("-ffunction-sections");   // Places each function in its own section of the output file, linker may be able to perform opts to improve locality of reference
+				Arguments.Add("-fdata-sections");       // Places each data item in its own section of the output file, linker may be able to perform opts to improve locality of reference
 			}
 
+		}
+
+		/// <inheritdoc/>
+		protected override void GetCompilerArguments_Sanitizers(CppCompileEnvironment CompileEnvironment, List<string> Arguments)
+		{
+			// TODO: Reconcile with base
+			//base.GetCompilerArguments_Sanitizers(CompileEnvironment, Arguments);
+
 			ClangSanitizer Sanitizer = BuildWithSanitizer();
 			if (Sanitizer != ClangSanitizer.None)
 			{
-				Result += " -fsanitize=" + GetCompilerOption(Sanitizer);
+				Arguments.Add("-fsanitize=" + GetCompilerOption(Sanitizer));
 
 				if (Sanitizer == ClangSanitizer.Address || Sanitizer == ClangSanitizer.HwAddress)
 				{
-					Result += " -fno-omit-frame-pointer -DRUNNING_WITH_ASAN=1";
+					Arguments.Add("-fno-omit-frame-pointer -DRUNNING_WITH_ASAN=1");
 				}
 			}
 
-			if (!DisableStackProtector())
+
+
+			//string? SanitizerMode = Environment.GetEnvironmentVariable("ENABLE_ADDRESS_SANITIZER");
+			//if ((SanitizerMode != null && SanitizerMode == "YES") || (Options.HasFlag(ClangToolChainOptions.EnableAddressSanitizer)))
+			//{
+			//	Arguments.Add("-fsanitize=address -fno-omit-frame-pointer -DFORCE_ANSI_ALLOCATOR=1");
+			//}
+
+			//string? UndefSanitizerMode = Environment.GetEnvironmentVariable("ENABLE_UNDEFINED_BEHAVIOR_SANITIZER");
+			//if ((UndefSanitizerMode != null && UndefSanitizerMode == "YES") || (Options.HasFlag(ClangToolChainOptions.EnableUndefinedBehaviorSanitizer)))
+			//{
+			//	Arguments.Add("-fsanitize=undefined -fno-sanitize=bounds,enum,return,float-divide-by-zero");
+			//}
+
+			//if (Options.HasFlag(ClangToolChainOptions.EnableThreadSanitizer))
+			//{
+			//	Arguments.Add("-fsanitize=thread");
+			//}
+
+			//if (Options.HasFlag(ClangToolChainOptions.EnableMemorySanitizer))
+			//{
+			//	Arguments.Add("-fsanitize=memory");
+			//}
+		}
+
+		/// <inheritdoc/>
+		protected override void GetCompileArguments_Global(CppCompileEnvironment CompileEnvironment, List<string> Arguments)
+		{
+			base.GetCompileArguments_Global(CompileEnvironment, Arguments);
+
+			// NDK setup (enforce minimum API level)
+			int NDKApiLevel64Int = GetNdkApiLevelInt(21);       // deliberately use 21 minimum to force NDKApiLevel64Bit to update if below minimum
+			string NDKApiLevel64Bit = GetNdkApiLevel();
+			if (NDKApiLevel64Int < MinimumNDKAPILevel)
+			{
+				NDKApiLevel64Int = MinimumNDKAPILevel;
+				NDKApiLevel64Bit = "android-" + MinimumNDKAPILevel;
+			}
+
+			Log.TraceInformationOnce("Compiling Native 64-bit code with NDK API '{0}'", NDKApiLevel64Bit);
+
+			string NativeGluePath = Path.GetFullPath(GetNativeGluePath());
+
+
+			if (BuildWithHiddenSymbolVisibility(CompileEnvironment))
 			{
-				Result += " -fstack-protector-strong";  // Emits extra code to check for buffer overflows
+				Arguments.Add("-fvisibility=hidden");
+				Arguments.Add("-fvisibility-inlines-hidden");
+				//TODO: add when Android's clang will support this
+				//Arguments.Add("-fvisibility-inlines-hidden-static-local-var");
 			}
 
-			// Add analysis flags to the argument list.
-			if (IsAnalyzing(CompileEnvironment))
+			Arguments.Add(GetRTTIFlag(CompileEnvironment));
+			Arguments.Add("-no-canonical-prefixes");
+			Arguments.Add("-fno-PIE");
+			Arguments.Add("-funwind-tables");           // Just generates any needed static data, affects no code
+			Arguments.Add("-fPIC");                     // Generates position-independent code (PIC) suitable for use in a shared library
+			Arguments.Add("-fno-strict-aliasing");      // Prevents unwanted or invalid optimizations that could produce incorrect code
+			Arguments.Add("-fno-short-enums");          // Do not allocate to an enum type only as many bytes as it needs for the declared range of possible values
+			Arguments.Add("-fforce-emit-vtables");      // Helps with devirtualization
+			Arguments.Add("-D_FORTIFY_SOURCE=2");       // FORTIFY default
+			Arguments.Add($"-DPLATFORM_USED_NDK_VERSION_INTEGER={NDKApiLevel64Int}");       // NDK version
+			Arguments.Add("-DPLATFORM_64BITS=1");       // NDK version
+
+
+			if (CompileEnvironment.bCompileISPC)
 			{
-				List<string> Arguments = new List<string>();
-				GetCompileArguments_Analyze(CompileEnvironment, Arguments);
-				for (int i = 0; i < Arguments.Count; i++)
-				{
-					Result += " " + Arguments[i];
-				}
+				Arguments.Add("-DINTEL_ISPC=1");
 			}
 
-			Result += " -fforce-emit-vtables";      // Helps with devirtualization
-			//Result += " -fstrict-vtable-pointers";	// Assumes that vtable pointer will never change, i.e. object of a different type would not be constracted in place of another object. Helps with devirtualization
+			if (CompileEnvironment.Architecture == "arm64")
+			{
+				Arguments.Add(ToolchainParamsArm64);
+				Arguments.Add("-D__arm64__");            // for some reason this isn't defined and needed for PhysX
+				Arguments.Add("-DPLATFORM_ANDROID_ARM64=1");
+				Arguments.Add("-march=armv8-a");
+				Arguments.Add("-fsigned-char");             // Treat chars as signed //@todo android: any concerns about ABI compatibility with libs here?
+			}
+			else if (CompileEnvironment.Architecture == "x64")
+			{
+				Arguments.Add(ToolchainParamsx64);
+				Arguments.Add("-fno-omit-frame-pointer");
+				Arguments.Add("-march=atom");
+				Arguments.Add("-DPLATFORM_ANDROID_X64=1");
+			}
+		}
 
-			return Result;
+		/// <inheritdoc/>
+		protected override void GetCompileArguments_FileType(CppCompileEnvironment CompileEnvironment, FileItem SourceFile, DirectoryReference OutputDir, List<string> Arguments, Action CompileAction, CPPOutput CompileResult)
+		{
+			base.GetCompileArguments_FileType(CompileEnvironment, SourceFile, OutputDir, Arguments, CompileAction, CompileResult);
+
+			string Extension = Path.GetExtension(SourceFile.AbsolutePath).ToUpperInvariant();
+			if (Extension == ".C")
+			{
+				if (SourceFile.AbsolutePath.Equals(GetNativeGluePath()))
+				{
+					// Remove visibility settings for android native glue. Since it doesn't decorate with visibility attributes.
+					Arguments.RemoveAll(x => x.StartsWith("-fvisibility"));
+				}
+				// remove any forced include/PCH stuff - mostly for the force-added .c files in Launch as those will attempt to have the PCH used that was made with .cpp language
+				Arguments.RemoveAll(x => x.StartsWith("-include"));
+			}
 		}
 
+
 		protected virtual string GetLinkArguments(LinkEnvironment LinkEnvironment, string Architecture)
 		{
 			string Result = "";
@@ -828,12 +725,12 @@ namespace UnrealBuildTool
 				Result += " -Wl,--strip-debug";
 			}
 
-			if (Architecture == "-x64")
+			if (Architecture == "x64")
 			{
 				Result += ToolchainLinkParamsx64;
 				Result += " -march=atom";
 			}
-			else // if (Architecture == "-arm64")
+			else // if (Architecture == "arm64")
 			{
 				Result += ToolchainLinkParamsArm64;
 				Result += " -march=armv8-a";
@@ -1025,7 +922,7 @@ namespace UnrealBuildTool
 			}
 
 			// if another architecture is in the filename, reject it
-			foreach (string ComboName in Arches!)
+			foreach (string ComboName in AllCpuSuffixes)
 			{
 				if (ComboName != Arch)
 				{
@@ -1060,7 +957,7 @@ namespace UnrealBuildTool
 			}
 		}
 
-		void GenerateEmptyLinkFunctionsForRemovedModules(List<FileItem> SourceFiles, string ModuleName, DirectoryReference OutputDirectory, IActionGraphBuilder Graph, ILogger Logger)
+		void GenerateEmptyLinkFunctionsForRemovedModules(List<FileItem> SourceFiles, string Arch, string ModuleName, DirectoryReference OutputDirectory, IActionGraphBuilder Graph, ILogger Logger)
 		{
 			// Only add to Launch module
 			if (!ModuleName.Equals("Launch"))
@@ -1074,25 +971,22 @@ namespace UnrealBuildTool
 			List<string> Result = new List<string>();
 			Result.Add("#include \"CoreTypes.h\"");
 			Result.Add("");
-			foreach (string Arch in Arches!)
+			switch (Arch)
 			{
-				switch (Arch)
-				{
-					case "-arm64": Result.Add("#if PLATFORM_ANDROID_ARM64"); break;
-					case "-x64": Result.Add("#if PLATFORM_ANDROID_X64"); break;
-					default: Result.Add("#if PLATFORM_ANDROID_ARM"); break;
-				}
+				case "arm64": Result.Add("#if PLATFORM_ANDROID_ARM64"); break;
+				case "x64": Result.Add("#if PLATFORM_ANDROID_X64"); break;
+				default: Result.Add("#if PLATFORM_ANDROID_ARM"); break;
+			}
 
-				foreach (string ModName in ModulesToSkip[Arch])
-				{
-					Result.Add("  void EmptyLinkFunctionForStaticInitialization" + ModName + "(){}");
-				}
-				foreach (string ModName in GeneratedModulesToSkip[Arch])
-				{
-					Result.Add("  void EmptyLinkFunctionForGeneratedCode" + ModName + "(){}");
-				}
-				Result.Add("#endif");
+			foreach (string ModName in ModulesToSkip[Arch])
+			{
+				Result.Add("  void EmptyLinkFunctionForStaticInitialization" + ModName + "(){}");
+			}
+			foreach (string ModName in GeneratedModulesToSkip[Arch])
+			{
+				Result.Add("  void EmptyLinkFunctionForGeneratedCode" + ModName + "(){}");
 			}
+			Result.Add("#endif");
 
 			Graph.CreateIntermediateTextFile(LinkerExceptionsCPPFilename, Result);
 
@@ -1101,10 +995,10 @@ namespace UnrealBuildTool
 
 		// cache the location of NDK tools
 		protected static string? ClangPath;
-		protected static string? ToolchainParamsArm64;
-		protected static string? ToolchainParamsx64;
-		protected static string? ToolchainLinkParamsArm64;
-		protected static string? ToolchainLinkParamsx64;
+		protected static string ToolchainParamsArm64 = "";
+		protected static string ToolchainParamsx64 = "";
+		protected static string ToolchainLinkParamsArm64 = "";
+		protected static string ToolchainLinkParamsx64 = "";
 		protected static string? ArPathArm64;
 		protected static string? ArPathx64;
 		protected static string? ReadElfPath;
@@ -1115,465 +1009,353 @@ namespace UnrealBuildTool
 
 			switch (UnrealArch)
 			{
-				case "-arm64": StripPath = ArPathArm64!; break;
-				case "-x64": StripPath = ArPathx64!; break;
+				case "arm64": StripPath = ArPathArm64!; break;
+				case "x64": StripPath = ArPathx64!; break;
 				default: StripPath = ArPathArm64!; break;
 			}
 			return StripPath.Replace("-ar", "-strip");
 		}
 
-		static private bool bHasHandledLaunchModule = false;
+		private bool bHasHandledLaunchModule = false;
 		protected override CPPOutput CompileCPPFiles(CppCompileEnvironment CompileEnvironment, List<FileItem> InputFiles, DirectoryReference OutputDir, string ModuleName, IActionGraphBuilder Graph)
 		{
-			if (Arches!.Count == 0)
+			if (ShouldSkipModule(ModuleName, CompileEnvironment.Architecture))
 			{
-				throw new BuildException("At least one architecture (arm64, x64, etc) needs to be selected in the project settings to build");
+				return new CPPOutput();
 			}
 
-			CPPOutput Result = new CPPOutput();
 
-			// Skip if nothing to do
-			if (InputFiles.Count == 0)
+			// Deal with Launch module special if first time seen
+			if (!bHasHandledLaunchModule && (ModuleName.Equals("Launch") || ModuleName.Equals("AndroidLauncher")))
 			{
-				return Result;
-			}
+				// Directly added NDK files for NDK extensions
+				ModifySourceFiles(CompileEnvironment, InputFiles, ModuleName);
+				// Deal with dynamic modules removed by architecture
+				GenerateEmptyLinkFunctionsForRemovedModules(InputFiles, CompileEnvironment.Architecture, ModuleName, OutputDir, Graph, Logger);
 
-			/*
-			Trace.TraceInformation("CompileCPPFiles: Module={0}, SourceFiles={1}", ModuleName, SourceFiles.Count);
-			foreach (string Arch in Arches)
-			{
-				Trace.TraceInformation("  Arch: {0}", Arch);
+				bHasHandledLaunchModule = true;
 			}
-			foreach (FileItem SourceFile in SourceFiles)
+
+
+			return base.CompileCPPFiles(CompileEnvironment, InputFiles, OutputDir, ModuleName, Graph);
+		}
+
+		static public string InlineArchName(string Pathname, string Arch, bool bUseShortNames = false)
+		{
+			string FinalArch = "-" + Arch;
+			if (bUseShortNames)
 			{
-				Trace.TraceInformation("  {0}", SourceFile.AbsolutePath);
+				FinalArch = ShortArchNames[FinalArch];
 			}
-			*/
+			return Path.Combine(Path.GetDirectoryName(Pathname)!, Path.GetFileNameWithoutExtension(Pathname) + FinalArch + Path.GetExtension(Pathname));
+		}
 
-			// NDK setup (enforce minimum API level)
-			int NDKApiLevel64Int = GetNdkApiLevelInt(21);		// deliberately use 21 minimum to force NDKApiLevel64Bit to update if below minimum
-			string NDKApiLevel64Bit = GetNdkApiLevel();
-			if (NDKApiLevel64Int < MinimumNDKAPILevel)
+		public string RemoveArchName(string Pathname)
+		{
+			// remove all architecture names
+			foreach (string Arch in AllCpuSuffixes)
 			{
-				NDKApiLevel64Int = MinimumNDKAPILevel;
-				NDKApiLevel64Bit = "android-" + MinimumNDKAPILevel;
+				Pathname = Path.Combine(Path.GetDirectoryName(Pathname)!, Path.GetFileName(Pathname).Replace(Arch, ""));
 			}
+			return Pathname;
+		}
 
-			Log.TraceInformationOnce("Compiling Native 64-bit code with NDK API '{0}'", NDKApiLevel64Bit);
+		static public DirectoryReference InlineArchIncludeFolder(DirectoryReference PathRef, string Arch)
+		{
+			return DirectoryReference.Combine(PathRef, "include", Arch.Replace("-", ""));
+		}
 
-			string BaseArguments = "";
+		public override CPPOutput GenerateISPCHeaders(CppCompileEnvironment CompileEnvironment, List<FileItem> InputFiles, DirectoryReference OutputDir, IActionGraphBuilder Graph)
+		{
+			CPPOutput Result = new CPPOutput();
 
-			if (CompileEnvironment.PrecompiledHeaderAction != PrecompiledHeaderAction.Create)
+			if (!CompileEnvironment.bCompileISPC)
 			{
-				BaseArguments += " -Werror";
+				return Result;
 			}
 
-			string NativeGluePath = Path.GetFullPath(GetNativeGluePath());
+			List<string> CompileTargets = GetISPCCompileTargets(CompileEnvironment.Platform, CompileEnvironment.Architecture);
 
-			// Deal with Launch module special if first time seen
-			if (!bHasHandledLaunchModule && (ModuleName.Equals("Launch") || ModuleName.Equals("AndroidLauncher")))
+			CompileEnvironment.UserIncludePaths.Add(InlineArchIncludeFolder(OutputDir, CompileEnvironment.Architecture));
+
+			foreach (FileItem ISPCFile in InputFiles)
 			{
-				// Directly added NDK files for NDK extensions
-				ModifySourceFiles(CompileEnvironment, InputFiles, ModuleName);
-				// Deal with dynamic modules removed by architecture
-				GenerateEmptyLinkFunctionsForRemovedModules(InputFiles, ModuleName, OutputDir, Graph, Logger);
+				Action CompileAction = Graph.CreateAction(ActionType.Compile);
+				CompileAction.CommandDescription = "Compile";
+				CompileAction.WorkingDirectory = Unreal.EngineSourceDirectory;
+				CompileAction.CommandPath = new FileReference(GetISPCHostCompilerPath(BuildHostPlatform.Current.Platform));
+				CompileAction.StatusDescription = Path.GetFileName(ISPCFile.AbsolutePath);
+				CompileAction.CommandVersion = GetISPCHostCompilerVersion(BuildHostPlatform.Current.Platform).ToString();
 
-				bHasHandledLaunchModule = true;
-			}
+				// Disable remote execution to workaround mismatched case on XGE
+				CompileAction.bCanExecuteRemotely = false;
 
-			// Add preprocessor definitions to the argument list.
-			foreach (string Definition in CompileEnvironment.Definitions)
-			{
-				// We must change the \" to escaped double quotes to save it properly for clang .rsp
-				BaseArguments += string.Format(" -D \"{0}\"", Definition.Replace("\"", RuntimePlatform.IsWindows ? "\"\"" : "\\\""));
-			}
+				List<string> Arguments = new List<string>();
 
-			//LUMIN_MERGE
-			//string NDKRoot = Environment.GetEnvironmentVariable("NDKROOT").Replace("\\", "/");
+				// Add the ISPC obj file as a prerequisite of the action.
+				Arguments.Add(String.Format(" \"{0}\"", ISPCFile.AbsolutePath));
 
-			string BasePCHName = "";
-			//UEBuildPlatform BuildPlatform = UEBuildPlatform.GetBuildPlatform(CompileEnvironment.Platform);
-			string PCHExtension = ".gch";
-			if (CompileEnvironment.PrecompiledHeaderAction == PrecompiledHeaderAction.Include)
-			{
-				BasePCHName = RemoveArchName(CompileEnvironment.PrecompiledHeaderFile!.AbsolutePath).Replace(PCHExtension, "");
-			}
+				// Add the ISPC h file to the produced item list.
+				FileItem ISPCIncludeHeaderFile = FileItem.GetItemByFileReference(
+					FileReference.Combine(
+						InlineArchIncludeFolder(OutputDir, CompileEnvironment.Architecture),
+						Path.GetFileName(ISPCFile.AbsolutePath) + ".generated.dummy.h"
+						)
+					);
 
-			// Create a compile action for each source file.
-			foreach (string Arch in Arches)
-			{
-				if (ShouldSkipModule(ModuleName, Arch))
+				// Add the ISPC file to be compiled.
+				Arguments.Add(String.Format("-h \"{0}\"", ISPCIncludeHeaderFile));
+
+				// Build target string. No comma on last
+				string TargetString = "";
+				foreach (string Target in CompileTargets)
 				{
-					continue;
+					if (Target == CompileTargets.Last())
+					{
+						TargetString += Target;
+					}
+					else
+					{
+						TargetString += Target + ",";
+					}
 				}
 
-				// which toolchain to use
-				string Arguments = GetCLArguments_Global(CompileEnvironment, Arch) + BaseArguments;
+				// Build target triplet
+				Arguments.Add(String.Format("--target-os=\"{0}\"", GetISPCOSTarget(CompileEnvironment.Platform)));
+				Arguments.Add(String.Format("--arch=\"{0}\"", GetISPCArchTarget(CompileEnvironment.Platform, CompileEnvironment.Architecture)));
+				Arguments.Add(String.Format("--target=\"{0}\"", TargetString));
 
-				switch (Arch)
+				Arguments.Add("--pic");
+
+				// Include paths. Don't use AddIncludePath() here, since it uses the full path and exceeds the max command line length.
+				foreach (DirectoryReference IncludePath in CompileEnvironment.UserIncludePaths)
 				{
-					case "-x64": Arguments += " -DPLATFORM_64BITS=1 -DPLATFORM_ANDROID_X64=1 -DPLATFORM_USED_NDK_VERSION_INTEGER=" + NDKApiLevel64Int.ToString(); break;
-					case "-arm64": Arguments += " -DPLATFORM_64BITS=1 -DPLATFORM_ANDROID_ARM64=1 -DPLATFORM_USED_NDK_VERSION_INTEGER=" + NDKApiLevel64Int.ToString(); break;
-					default: throw new BuildException($"Unknown arch {Arch}");
+					Arguments.Add(String.Format("-I\"{0}\"", IncludePath));
 				}
 
-				if(CompileEnvironment.bCompileISPC)
+				// System include paths.
+				foreach (DirectoryReference SystemIncludePath in CompileEnvironment.SystemIncludePaths)
 				{
-					Arguments += " -DINTEL_ISPC=1";
+					Arguments.Add(String.Format("-I\"{0}\"", SystemIncludePath));
 				}
 
-				// which PCH file to include
-				string PCHArguments = "";
-				if (CompileEnvironment.PrecompiledHeaderAction == PrecompiledHeaderAction.Include)
-				{
-					// add the platform-specific PCH reference
-					PCHArguments += string.Format(" -include-pch \"{0}\"", InlineArchName(BasePCHName, Arch) + PCHExtension);
+				// Generate the included header dependency list
+				FileItem DependencyListFile = FileItem.GetItemByFileReference(FileReference.Combine(OutputDir, InlineArchName(Path.GetFileName(ISPCFile.AbsolutePath) + ".d", CompileEnvironment.Architecture, true)));
+				Arguments.Add(String.Format("-M -MF \"{0}\"", DependencyListFile.AbsolutePath.Replace('\\', '/')));
+				CompileAction.DependencyListFile = DependencyListFile;
+				CompileAction.ProducedItems.Add(DependencyListFile);
 
-					// Add the precompiled header file's path to the include path so Clang can find it.
-					// This needs to be before the other include paths to ensure Clang uses it instead of the source header file.
-					PCHArguments += string.Format(" -include \"{0}\"", BasePCHName);
+				CompileAction.ProducedItems.Add(ISPCIncludeHeaderFile);
 
-					// Validate PCH inputs by content if mtime check fails
-					if (CompilerVersionGreaterOrEqual(11, 0, 0))
-					{
-						PCHArguments += " -fpch-validate-input-files-content";
-					}
-				}
+				CompileAction.CommandArguments = String.Join(" ", Arguments);
 
-				// Add include paths to the argument list (filtered by architecture)
-				foreach (DirectoryReference IncludePath in CompileEnvironment.SystemIncludePaths)
+				// Add the source file and its included files to the prerequisite item list.
+				CompileAction.PrerequisiteItems.Add(ISPCFile);
+				CompileAction.StatusDescription = string.Format("{0} [{1}]", Path.GetFileName(ISPCFile.AbsolutePath), CompileEnvironment.Architecture.Replace("-", ""));
+
+				FileItem ISPCFinalHeaderFile = FileItem.GetItemByFileReference(
+					FileReference.Combine(
+						InlineArchIncludeFolder(OutputDir, CompileEnvironment.Architecture),
+						Path.GetFileName(ISPCFile.AbsolutePath) + ".generated.h"
+						)
+					);
+
+				// Fix interrupted build issue by copying header after generation completes
+				FileReference SourceFile = ISPCIncludeHeaderFile.Location;
+				FileReference TargetFile = ISPCFinalHeaderFile.Location;
+
+				FileItem SourceFileItem = FileItem.GetItemByFileReference(SourceFile);
+				FileItem TargetFileItem = FileItem.GetItemByFileReference(TargetFile);
+
+				Action CopyAction = Graph.CreateAction(ActionType.BuildProject);
+				CopyAction.CommandDescription = "Copy";
+				CopyAction.CommandPath = BuildHostPlatform.Current.Shell;
+				if (BuildHostPlatform.Current.ShellType == ShellType.Cmd)
 				{
-					if (IsDirectoryForArch(IncludePath.FullName, Arch))
-					{
-						Arguments += string.Format(" -I\"{0}\"", IncludePath);
-					}
+					CopyAction.CommandArguments = String.Format("/C \"copy /Y \"{0}\" \"{1}\" 1>nul\"", SourceFile, TargetFile);
 				}
-				foreach (DirectoryReference IncludePath in CompileEnvironment.UserIncludePaths)
+				else
 				{
-					if (IsDirectoryForArch(IncludePath.FullName, Arch))
-					{
-						Arguments += string.Format(" -I\"{0}\"", IncludePath);
-					}
+					CopyAction.CommandArguments = String.Format("-c 'cp -f \"{0}\" \"{1}\"'", SourceFile.FullName, TargetFile.FullName);
 				}
+				CopyAction.WorkingDirectory = Unreal.EngineSourceDirectory;
+				CopyAction.PrerequisiteItems.Add(SourceFileItem);
+				CopyAction.ProducedItems.Add(TargetFileItem);
+				CopyAction.StatusDescription = TargetFileItem.Location.GetFileName();
+				CopyAction.bCanExecuteRemotely = false;
+				CopyAction.bShouldOutputStatusDescription = false;
 
-				foreach (FileItem SourceFile in InputFiles)
-				{
-					Action CompileAction = Graph.CreateAction(ActionType.Compile);
-					CompileAction.PrerequisiteItems.AddRange(CompileEnvironment.ForceIncludeFiles);
-					CompileAction.PrerequisiteItems.AddRange(CompileEnvironment.AdditionalPrerequisites);
+				Result.GeneratedHeaderFiles.Add(TargetFileItem);
 
-					string FileArguments = "";
-					bool bIsPlainCFile = Path.GetExtension(SourceFile.AbsolutePath).ToUpperInvariant() == ".C";
-					bool bDisableShadowWarning = false;
+				Logger.LogDebug("   ISPC Generating Header {StatusDescription}: \"{CommandPath}\" {CommandArguments}", CompileAction.StatusDescription, CompileAction.CommandPath, CompileAction.CommandArguments);
+			}
 
-					// Add C or C++ specific compiler arguments.
-					if (bIsPlainCFile)
-					{
-						List<string> CArguments = new();
-						GetCompileArguments_C(CompileEnvironment, CArguments);
-						FileArguments += $" {string.Join(' ', CArguments)}";
+			return Result;
+		}
+		
+		public override CPPOutput CompileISPCFiles(CppCompileEnvironment CompileEnvironment, List<FileItem> InputFiles, DirectoryReference OutputDir, IActionGraphBuilder Graph)
+		{
+			CPPOutput Result = new CPPOutput();
 
-						// remove shadow variable warnings for externally included files
-						if (!SourceFile.Location.IsUnderDirectory(Unreal.RootDirectory))
-						{
-							bDisableShadowWarning = true;
-						}
-					}
-					else if (CompileEnvironment.PrecompiledHeaderAction == PrecompiledHeaderAction.Create)
-					{
-						List<string> CreatePCHArguments = new();
-						GetCompileArguments_PCH(CompileEnvironment, CreatePCHArguments);
-						FileArguments += $" {string.Join(' ', CreatePCHArguments)}";
-					}
-					else
-					{
-						List<string> CPPArguments = new();
-						GetCompileArguments_CPP(CompileEnvironment, CPPArguments);
-						FileArguments += $" {string.Join(' ', CPPArguments)}";
+			if (!CompileEnvironment.bCompileISPC)
+			{
+				return Result;
+			}
 
-						// only use PCH for .cpp files
-						FileArguments += PCHArguments;
-					}
+			List<string> CompileTargets = GetISPCCompileTargets(CompileEnvironment.Platform, CompileEnvironment.Architecture);
 
-					foreach (FileItem ForceIncludeFile in CompileEnvironment.ForceIncludeFiles)
-					{
-						FileArguments += string.Format(" -include \"{0}\"", ForceIncludeFile.Location);
-					}
+			foreach (FileItem ISPCFile in InputFiles)
+			{
+				Action CompileAction = Graph.CreateAction(ActionType.Compile);
+				CompileAction.CommandDescription = "Compile";
+				CompileAction.WorkingDirectory = Unreal.EngineSourceDirectory;
+				CompileAction.CommandPath = new FileReference(GetISPCHostCompilerPath(BuildHostPlatform.Current.Platform));
+				CompileAction.StatusDescription = Path.GetFileName(ISPCFile.AbsolutePath);
+
+				// Disable remote execution to workaround mismatched case on XGE
+				CompileAction.bCanExecuteRemotely = false;
+
+				List<string> Arguments = new List<string>();
+
+				// Add the ISPC file to be compiled.
+				Arguments.Add(String.Format(" \"{0}\"", ISPCFile.AbsolutePath));
+
+				List<FileItem> CompiledISPCObjFiles = new List<FileItem>();
+				List<FileItem> FinalISPCObjFiles = new List<FileItem>();
+				string TargetString = "";
 
-					// Add the C++ source file and its included files to the prerequisite item list.
-					CompileAction.PrerequisiteItems.Add(SourceFile);
+				foreach (string Target in CompileTargets)
+				{
+					string ObjTarget = Target;
 
-					List<FileItem>? InlinedFiles;
-					if (CompileEnvironment.FileInlineGenCPPMap.TryGetValue(SourceFile, out InlinedFiles))
+					if (Target.Contains("-"))
 					{
-						CompileAction.PrerequisiteItems.AddRange(InlinedFiles);
+						// Remove lane width and gang size from obj file name
+						ObjTarget = Target.Split('-')[0];
 					}
 
-					if (CompileEnvironment.PrecompiledHeaderAction == PrecompiledHeaderAction.Create && !bIsPlainCFile)
-					{
-						// Add the precompiled header file to the produced item list.
-						FileItem PrecompiledHeaderFile = FileItem.GetItemByFileReference(
-							FileReference.Combine(
-								OutputDir,
-								Path.GetFileName(InlineArchName(SourceFile.AbsolutePath, Arch) + PCHExtension)
-								)
-							);
+					FileItem CompiledISPCObjFile;
+					FileItem FinalISPCObjFile;
 
-						CompileAction.ProducedItems.Add(PrecompiledHeaderFile);
-						Result.PrecompiledHeaderFile = PrecompiledHeaderFile;
+					if (CompileTargets.Count > 1)
+					{
+						CompiledISPCObjFile = FileItem.GetItemByFileReference(
+						FileReference.Combine(
+							OutputDir,
+							Path.GetFileNameWithoutExtension(InlineArchName(Path.GetFileName(ISPCFile.AbsolutePath) + ".o", CompileEnvironment.Architecture, true)) + "_" + ObjTarget + ".o"
+							)
+						);
 
-						// Add the parameters needed to compile the precompiled header file to the command-line.
-						FileArguments += string.Format(" -o \"{0}\"", PrecompiledHeaderFile.AbsolutePath);
+						FinalISPCObjFile = FileItem.GetItemByFileReference(
+						FileReference.Combine(
+							OutputDir,
+							Path.GetFileName(ISPCFile.AbsolutePath) + "_" + ObjTarget + InlineArchName(".o", CompileEnvironment.Architecture, true)
+							)
+						);
 					}
 					else
 					{
-						if (CompileEnvironment.PrecompiledHeaderAction == PrecompiledHeaderAction.Include)
-						{
-							FileItem ArchPrecompiledHeaderFile = FileItem.GetItemByPath(InlineArchName(BasePCHName, Arch) + PCHExtension);
-							CompileAction.PrerequisiteItems.Add(ArchPrecompiledHeaderFile);
-						}
-
-						string ObjectFileExtension;
-						if(CompileEnvironment.AdditionalArguments != null && CompileEnvironment.AdditionalArguments.Contains("-emit-llvm"))
-						{
-							ObjectFileExtension = ".bc";
-						}
-						else
-						{
-							ObjectFileExtension = ".o";
-						}
-
-						// Add the object file to the produced item list.
-						FileItem ObjectFile = FileItem.GetItemByFileReference(
+						CompiledISPCObjFile = FileItem.GetItemByFileReference(
 							FileReference.Combine(
 								OutputDir,
-								InlineArchName(Path.GetFileName(SourceFile.AbsolutePath) + ObjectFileExtension, Arch, true)
+								InlineArchName(Path.GetFileName(ISPCFile.AbsolutePath) + ".o", CompileEnvironment.Architecture, true)
 								)
 							);
-						CompileAction.ProducedItems.Add(ObjectFile);
-						Result.ObjectFiles.Add(ObjectFile);
-
-						FileArguments += string.Format(" -o \"{0}\"", ObjectFile.AbsolutePath);
-					}
-
-					// Add the source file path to the command-line.
-					FileArguments += string.Format(" \"{0}\"", SourceFile.AbsolutePath);
-
-					// Generate the timing info
-					if (CompileEnvironment.bPrintTimingInfo)
-					{
-						bool ShortArchName = !(CompileEnvironment.PrecompiledHeaderAction == PrecompiledHeaderAction.Create && !bIsPlainCFile);
-						FileItem TraceFile;
-						if (CompileEnvironment.PrecompiledHeaderAction == PrecompiledHeaderAction.Create && !bIsPlainCFile)
-						{
-							TraceFile = FileItem.GetItemByFileReference(FileReference.Combine(OutputDir, InlineArchName(Path.GetFileName(SourceFile.AbsolutePath), Arch) + ".json"));
-						}
-						else
-						{
-							TraceFile = FileItem.GetItemByFileReference(FileReference.Combine(OutputDir, InlineArchName(Path.GetFileName(SourceFile.AbsolutePath) + ".json", Arch, true)));
-						}
-						FileArguments += " -ftime-trace";
-						CompileAction.ProducedItems.Add(TraceFile);
-					}
-
-					// Generate the included header dependency list
-					FileItem DependencyListFile = FileItem.GetItemByFileReference(FileReference.Combine(OutputDir, InlineArchName(Path.GetFileName(SourceFile.AbsolutePath) + ".d", Arch, true)));
-					FileArguments += string.Format(" -MD -MF\"{0}\"", DependencyListFile.AbsolutePath.Replace('\\', '/'));
-					CompileAction.DependencyListFile = DependencyListFile;
-					CompileAction.ProducedItems.Add(DependencyListFile);
-
-					// Build a full argument list
-					string AllArguments = Arguments + FileArguments + CompileEnvironment.AdditionalArguments;
-
-					if (SourceFile.AbsolutePath.Equals(NativeGluePath))
-					{
-						// Remove visibility settings for android native glue. Since it doesn't decorate with visibility attributes.
-						//AllArguments = AllArguments.Replace("-fvisibility-ms-compat -fvisibility-inlines-hidden", "");
-						AllArguments = AllArguments.Replace("-fvisibility=hidden -fvisibility-inlines-hidden", "");
-					}
-
-					AllArguments = Utils.ExpandVariables(AllArguments);
-					AllArguments = AllArguments.Replace("\\", "/");
 
-					// Remove shadow warning for this file if requested
-					if (bDisableShadowWarning)
-					{
-						int WarningIndex = AllArguments.IndexOf(" -Wshadow");
-						if (WarningIndex > 0)
-						{
-							AllArguments = AllArguments.Remove(WarningIndex, 9);
-						}
+						FinalISPCObjFile = CompiledISPCObjFile;
 					}
 
-					// Create the response file
-					FileReference ResponseFileName = CompileAction.ProducedItems[0].Location + ".rsp";
-					FileItem ResponseFileItem = Graph.CreateIntermediateTextFile(ResponseFileName, new List<string> { AllArguments });
-					string ResponseArgument = string.Format("@\"{0}\"", ResponseFileName);
+					// Add the ISA specific ISPC obj files to the produced item list.
+					CompiledISPCObjFiles.Add(CompiledISPCObjFile);
+					FinalISPCObjFiles.Add(FinalISPCObjFile);
 
-					CompileAction.WorkingDirectory = Unreal.EngineSourceDirectory;
-					CompileAction.CommandDescription = IsAnalyzing(CompileEnvironment) ? "Analyze" : "Compile";
-					if (bExecuteCompilerThroughShell)
+					// Build target string. No comma on last
+					if (Target == CompileTargets.Last())
 					{
-						SetupActionToExecuteCompilerThroughShell(ref CompileAction, ClangPath!, ResponseArgument, CompileAction.CommandDescription);
+						TargetString += Target;
 					}
 					else
 					{
-						CompileAction.CommandPath = new FileReference(ClangPath!);
-						CompileAction.CommandArguments = ResponseArgument;
+						TargetString += Target + ",";
 					}
-					CompileAction.PrerequisiteItems.Add(ResponseFileItem);
-					CompileAction.CommandVersion = AndroidClangBuild!;
-
-					CompileAction.StatusDescription = string.Format("{0} [{1}]", Path.GetFileName(SourceFile.AbsolutePath), Arch.Replace("-", ""));
-
-					// VC++ always outputs the source file name being compiled, so we don't need to emit this ourselves
-					CompileAction.bShouldOutputStatusDescription = true;
-
-					// Don't farm out creation of pre-compiled headers as it is the critical path task.
-					CompileAction.bCanExecuteRemotely =
-						CompileEnvironment.PrecompiledHeaderAction != PrecompiledHeaderAction.Create ||
-						CompileEnvironment.bAllowRemotelyCompiledPCHs;
 				}
-			}
-
-			return Result;
-		}
-
-		public override FileItem? LinkFiles(LinkEnvironment LinkEnvironment, bool bBuildImportLibraryOnly, IActionGraphBuilder Graph)
-		{
-			return null;
-		}
-
-		static public string InlineArchName(string Pathname, string Arch, bool bUseShortNames = false)
-		{
-			string FinalArch = Arch;
-			if (bUseShortNames)
-			{
-				FinalArch = ShortArchNames[FinalArch];
-			}
-			return Path.Combine(Path.GetDirectoryName(Pathname)!, Path.GetFileNameWithoutExtension(Pathname) + FinalArch + Path.GetExtension(Pathname));
-		}
 
-		public string RemoveArchName(string Pathname)
-		{
-			// remove all architecture names
-			foreach (string Arch in GetAllArchitectures())
-			{
-				Pathname = Path.Combine(Path.GetDirectoryName(Pathname)!, Path.GetFileName(Pathname).Replace(Arch, ""));
-			}
-			return Pathname;
-		}
+				// Add the common ISPC obj file to the produced item list.
+				FileItem CompiledISPCObjFileNoISA = FileItem.GetItemByFileReference(
+					FileReference.Combine(
+						OutputDir,
+						InlineArchName(Path.GetFileName(ISPCFile.AbsolutePath) + ".o", CompileEnvironment.Architecture, true)
+						)
+					);
 
-		static public DirectoryReference InlineArchIncludeFolder(DirectoryReference PathRef, string Arch)
-		{
-			return DirectoryReference.Combine(PathRef, "include", Arch.Replace("-", ""));
-		}
+				CompiledISPCObjFiles.Add(CompiledISPCObjFileNoISA);
+				FinalISPCObjFiles.Add(CompiledISPCObjFileNoISA);
 
-		public override CPPOutput GenerateISPCHeaders(CppCompileEnvironment CompileEnvironment, List<FileItem> InputFiles, DirectoryReference OutputDir, IActionGraphBuilder Graph)
-		{
-			if (Arches!.Count == 0)
-			{
-				throw new BuildException("At least one architecture (armv7, x86, etc) needs to be selected in the project settings to build");
-			}
+				// Add the output ISPC obj file
+				Arguments.Add(String.Format("-o \"{0}\"", CompiledISPCObjFileNoISA));
 
-			CPPOutput Result = new CPPOutput();
+				// Build target triplet
+				Arguments.Add(String.Format("--target-os=\"{0}\"", GetISPCOSTarget(CompileEnvironment.Platform)));
+				Arguments.Add(String.Format("--arch=\"{0}\"", GetISPCArchTarget(CompileEnvironment.Platform, CompileEnvironment.Architecture)));
+				Arguments.Add(String.Format("--target=\"{0}\"", TargetString));
 
-			if (!CompileEnvironment.bCompileISPC)
-			{
-				return Result;
-			}
-
-			foreach (string Arch in Arches)
-			{
-				List<string> CompileTargets = GetISPCCompileTargets(CompileEnvironment.Platform, Arch);
-
-				CompileEnvironment.UserIncludePaths.Add(InlineArchIncludeFolder(OutputDir, Arch));
-
-				foreach (FileItem ISPCFile in InputFiles)
+				if (CompileEnvironment.Configuration == CppConfiguration.Debug)
 				{
-					Action CompileAction = Graph.CreateAction(ActionType.Compile);
-					CompileAction.CommandDescription = "Compile";
-					CompileAction.WorkingDirectory = Unreal.EngineSourceDirectory;
-					CompileAction.CommandPath = new FileReference(GetISPCHostCompilerPath(BuildHostPlatform.Current.Platform));
-					CompileAction.StatusDescription = Path.GetFileName(ISPCFile.AbsolutePath);
-					CompileAction.CommandVersion = GetISPCHostCompilerVersion(BuildHostPlatform.Current.Platform).ToString();
-
-					// Disable remote execution to workaround mismatched case on XGE
-					CompileAction.bCanExecuteRemotely = false;
-
-					List<string> Arguments = new List<string>();
-
-					// Add the ISPC obj file as a prerequisite of the action.
-					Arguments.Add(String.Format(" \"{0}\"", ISPCFile.AbsolutePath));
-
-					// Add the ISPC h file to the produced item list.
-					FileItem ISPCIncludeHeaderFile = FileItem.GetItemByFileReference(
-						FileReference.Combine(
-							InlineArchIncludeFolder(OutputDir, Arch),
-							Path.GetFileName(ISPCFile.AbsolutePath) + ".generated.dummy.h"
-							)
-						);
+					Arguments.Add("-g -O0");
+				}
+				else
+				{
+					Arguments.Add("-O2");
+				}
 
-					// Add the ISPC file to be compiled.
-					Arguments.Add(String.Format("-h \"{0}\"", ISPCIncludeHeaderFile));
+				Arguments.Add("--pic");
 
-					// Build target string. No comma on last
-					string TargetString = "";
-					foreach (string Target in CompileTargets)
+				// Add include paths to the argument list (filtered by architecture)
+				foreach (DirectoryReference IncludePath in CompileEnvironment.SystemIncludePaths)
+				{
+					if (IsDirectoryForArch(IncludePath.FullName, CompileEnvironment.Architecture))
 					{
-						if (Target == CompileTargets.Last())
-						{
-							TargetString += Target;
-						}
-						else
-						{
-							TargetString += Target + ",";
-						}
+						Arguments.Add(string.Format(" -I\"{0}\"", IncludePath));
 					}
-
-					// Build target triplet
-					Arguments.Add(String.Format("--target-os=\"{0}\"", GetISPCOSTarget(CompileEnvironment.Platform)));
-					Arguments.Add(String.Format("--arch=\"{0}\"", GetISPCArchTarget(CompileEnvironment.Platform, Arch)));
-					Arguments.Add(String.Format("--target=\"{0}\"", TargetString));
-
-					Arguments.Add("--pic");
-
-					// Include paths. Don't use AddIncludePath() here, since it uses the full path and exceeds the max command line length.
-					foreach (DirectoryReference IncludePath in CompileEnvironment.UserIncludePaths)
+				}
+				foreach (DirectoryReference IncludePath in CompileEnvironment.UserIncludePaths)
+				{
+					if (IsDirectoryForArch(IncludePath.FullName, CompileEnvironment.Architecture))
 					{
-						Arguments.Add(String.Format("-I\"{0}\"", IncludePath));
+						Arguments.Add(string.Format(" -I\"{0}\"", IncludePath));
 					}
+				}
 
-					// System include paths.
-					foreach (DirectoryReference SystemIncludePath in CompileEnvironment.SystemIncludePaths)
-					{
-						Arguments.Add(String.Format("-I\"{0}\"", SystemIncludePath));
-					}
+				// Preprocessor definitions.
+				foreach (string Definition in CompileEnvironment.Definitions)
+				{
+					Arguments.Add(String.Format(" -D\"{0}\"", Definition));
+				}
 
-					// Generate the included header dependency list
-					FileItem DependencyListFile = FileItem.GetItemByFileReference(FileReference.Combine(OutputDir, InlineArchName(Path.GetFileName(ISPCFile.AbsolutePath) + ".d", Arch, true)));
-					Arguments.Add(String.Format("-M -MF \"{0}\"", DependencyListFile.AbsolutePath.Replace('\\', '/')));
-					CompileAction.DependencyListFile = DependencyListFile;
-					CompileAction.ProducedItems.Add(DependencyListFile);
+				// Consume the included header dependency list
+				FileItem DependencyListFile = FileItem.GetItemByFileReference(FileReference.Combine(OutputDir, InlineArchName(Path.GetFileName(ISPCFile.AbsolutePath) + ".d", CompileEnvironment.Architecture, true)));
+				CompileAction.DependencyListFile = DependencyListFile;
+				CompileAction.PrerequisiteItems.Add(DependencyListFile);
 
-					CompileAction.ProducedItems.Add(ISPCIncludeHeaderFile);
+				CompileAction.ProducedItems.AddRange(CompiledISPCObjFiles);
 
-					CompileAction.CommandArguments = String.Join(" ", Arguments);
+				CompileAction.CommandArguments = String.Join(" ", Arguments);
 
-					// Add the source file and its included files to the prerequisite item list.
-					CompileAction.PrerequisiteItems.Add(ISPCFile);
-					CompileAction.StatusDescription = string.Format("{0} [{1}]", Path.GetFileName(ISPCFile.AbsolutePath), Arch.Replace("-", ""));
+				// Add the source file and its included files to the prerequisite item list.
+				CompileAction.PrerequisiteItems.Add(ISPCFile);
 
-					FileItem ISPCFinalHeaderFile = FileItem.GetItemByFileReference(
-						FileReference.Combine(
-							InlineArchIncludeFolder(OutputDir, Arch),
-							Path.GetFileName(ISPCFile.AbsolutePath) + ".generated.h"
-							)
-						);
+				CompileAction.StatusDescription = string.Format("{0} [{1}]", Path.GetFileName(ISPCFile.AbsolutePath), CompileEnvironment.Architecture.Replace("-", ""));
 
-					// Fix interrupted build issue by copying header after generation completes
-					FileReference SourceFile = ISPCIncludeHeaderFile.Location;
-					FileReference TargetFile = ISPCFinalHeaderFile.Location;
+				for(int i = 0; i < CompiledISPCObjFiles.Count; i++)
+				{
+					// ISPC compiler can't add suffix on the end of the arch, so copy to put into what linker expects
+					FileReference SourceFile = CompiledISPCObjFiles[i].Location;
+					FileReference TargetFile = FinalISPCObjFiles[i].Location;
+
+					if (SourceFile.Equals(TargetFile))
+					{
+						continue;
+					}
 
 					FileItem SourceFileItem = FileItem.GetItemByFileReference(SourceFile);
 					FileItem TargetFileItem = FileItem.GetItemByFileReference(TargetFile);
@@ -1595,511 +1377,284 @@ namespace UnrealBuildTool
 					CopyAction.StatusDescription = TargetFileItem.Location.GetFileName();
 					CopyAction.bCanExecuteRemotely = false;
 					CopyAction.bShouldOutputStatusDescription = false;
+				}
 
-					Result.GeneratedHeaderFiles.Add(TargetFileItem);
+				Result.ObjectFiles.AddRange(FinalISPCObjFiles);
 
-					Logger.LogDebug("   ISPC Generating Header {StatusDescription}: \"{CommandPath}\" {CommandArguments}", CompileAction.StatusDescription, CompileAction.CommandPath, CompileAction.CommandArguments);
-				}
+				Logger.LogDebug("   ISPC Compiling {StatusDescription}: \"{CommandPath}\" {CommandArguments}", CompileAction.StatusDescription, CompileAction.CommandPath, CompileAction.CommandArguments);
 			}
 
 			return Result;
 		}
-		
-		public override CPPOutput CompileISPCFiles(CppCompileEnvironment CompileEnvironment, List<FileItem> InputFiles, DirectoryReference OutputDir, IActionGraphBuilder Graph)
+
+		public override FileItem? LinkFiles(LinkEnvironment LinkEnvironment, bool bBuildImportLibraryOnly, IActionGraphBuilder Graph)
 		{
-			if (Arches!.Count == 0)
+			if (!LinkEnvironment.bIsBuildingLibrary)
 			{
-				throw new BuildException("At least one architecture (armv7, x86, etc) needs to be selected in the project settings to build");
+				// @todo Lumin: will this add them multiple times?
+				ModifyLibraries(LinkEnvironment);
 			}
 
-			CPPOutput Result = new CPPOutput();
+			string Arch = LinkEnvironment.Architecture;
 
-			if (!CompileEnvironment.bCompileISPC)
-			{
-				return Result;
-			}
+			// Create an action that invokes the linker.
+			Action LinkAction = Graph.CreateAction(ActionType.Link);
+			LinkAction.WorkingDirectory = Unreal.EngineSourceDirectory;
 
-			foreach (string Arch in Arches)
+			if (LinkEnvironment.bIsBuildingLibrary)
 			{
-				List<string> CompileTargets = GetISPCCompileTargets(CompileEnvironment.Platform, Arch);
-
-				foreach (FileItem ISPCFile in InputFiles)
+				switch (Arch)
 				{
-					Action CompileAction = Graph.CreateAction(ActionType.Compile);
-					CompileAction.CommandDescription = "Compile";
-					CompileAction.WorkingDirectory = Unreal.EngineSourceDirectory;
-					CompileAction.CommandPath = new FileReference(GetISPCHostCompilerPath(BuildHostPlatform.Current.Platform));
-					CompileAction.StatusDescription = Path.GetFileName(ISPCFile.AbsolutePath);
-
-					// Disable remote execution to workaround mismatched case on XGE
-					CompileAction.bCanExecuteRemotely = false;
-
-					List<string> Arguments = new List<string>();
-
-					// Add the ISPC file to be compiled.
-					Arguments.Add(String.Format(" \"{0}\"", ISPCFile.AbsolutePath));
-
-					List<FileItem> CompiledISPCObjFiles = new List<FileItem>();
-					List<FileItem> FinalISPCObjFiles = new List<FileItem>();
-					string TargetString = "";
-
-					foreach (string Target in CompileTargets)
-					{
-						string ObjTarget = Target;
-
-						if (Target.Contains("-"))
-						{
-							// Remove lane width and gang size from obj file name
-							ObjTarget = Target.Split('-')[0];
-						}
-
-						FileItem CompiledISPCObjFile;
-						FileItem FinalISPCObjFile;
-
-						if (CompileTargets.Count > 1)
-						{
-							CompiledISPCObjFile = FileItem.GetItemByFileReference(
-							FileReference.Combine(
-								OutputDir,
-								Path.GetFileNameWithoutExtension(InlineArchName(Path.GetFileName(ISPCFile.AbsolutePath) + ".o", Arch, true)) + "_" + ObjTarget + ".o"
-								)
-							);
-
-							FinalISPCObjFile = FileItem.GetItemByFileReference(
-							FileReference.Combine(
-								OutputDir,
-								Path.GetFileName(ISPCFile.AbsolutePath) + "_" + ObjTarget + InlineArchName(".o", Arch, true)
-								)
-							);
-						}
-						else
-						{
-							CompiledISPCObjFile = FileItem.GetItemByFileReference(
-								FileReference.Combine(
-									OutputDir,
-									InlineArchName(Path.GetFileName(ISPCFile.AbsolutePath) + ".o", Arch, true)
-									)
-								);
-
-							FinalISPCObjFile = CompiledISPCObjFile;
-						}
-
-						// Add the ISA specific ISPC obj files to the produced item list.
-						CompiledISPCObjFiles.Add(CompiledISPCObjFile);
-						FinalISPCObjFiles.Add(FinalISPCObjFile);
-
-						// Build target string. No comma on last
-						if (Target == CompileTargets.Last())
-						{
-							TargetString += Target;
-						}
-						else
-						{
-							TargetString += Target + ",";
-						}
-					}
-
-					// Add the common ISPC obj file to the produced item list.
-					FileItem CompiledISPCObjFileNoISA = FileItem.GetItemByFileReference(
-						FileReference.Combine(
-							OutputDir,
-							InlineArchName(Path.GetFileName(ISPCFile.AbsolutePath) + ".o", Arch, true)
-							)
-						);
-
-					CompiledISPCObjFiles.Add(CompiledISPCObjFileNoISA);
-					FinalISPCObjFiles.Add(CompiledISPCObjFileNoISA);
-
-					// Add the output ISPC obj file
-					Arguments.Add(String.Format("-o \"{0}\"", CompiledISPCObjFileNoISA));
-
-					// Build target triplet
-					Arguments.Add(String.Format("--target-os=\"{0}\"", GetISPCOSTarget(CompileEnvironment.Platform)));
-					Arguments.Add(String.Format("--arch=\"{0}\"", GetISPCArchTarget(CompileEnvironment.Platform, Arch)));
-					Arguments.Add(String.Format("--target=\"{0}\"", TargetString));
-
-					if (CompileEnvironment.Configuration == CppConfiguration.Debug)
-					{
-						Arguments.Add("-g -O0");
-					}
-					else
-					{
-						Arguments.Add("-O2");
-					}
-
-					Arguments.Add("--pic");
-
-					// Add include paths to the argument list (filtered by architecture)
-					foreach (DirectoryReference IncludePath in CompileEnvironment.SystemIncludePaths)
-					{
-						if (IsDirectoryForArch(IncludePath.FullName, Arch))
-						{
-							Arguments.Add(string.Format(" -I\"{0}\"", IncludePath));
-						}
-					}
-					foreach (DirectoryReference IncludePath in CompileEnvironment.UserIncludePaths)
-					{
-						if (IsDirectoryForArch(IncludePath.FullName, Arch))
-						{
-							Arguments.Add(string.Format(" -I\"{0}\"", IncludePath));
-						}
-					}
-
-					// Preprocessor definitions.
-					foreach (string Definition in CompileEnvironment.Definitions)
-					{
-						Arguments.Add(String.Format(" -D\"{0}\"", Definition));
-					}
-
-					// Consume the included header dependency list
-					FileItem DependencyListFile = FileItem.GetItemByFileReference(FileReference.Combine(OutputDir, InlineArchName(Path.GetFileName(ISPCFile.AbsolutePath) + ".d", Arch, true)));
-					CompileAction.DependencyListFile = DependencyListFile;
-					CompileAction.PrerequisiteItems.Add(DependencyListFile);
-
-					CompileAction.ProducedItems.AddRange(CompiledISPCObjFiles);
-
-					CompileAction.CommandArguments = String.Join(" ", Arguments);
+					case "arm64": LinkAction.CommandPath = new FileReference(ArPathArm64!); break;
+					case "x64": LinkAction.CommandPath = new FileReference(ArPathx64!); break;
+				}
+			}
+			else
+			{
+				LinkAction.CommandPath = new FileReference(ClangPath!);
+			}
 
-					// Add the source file and its included files to the prerequisite item list.
-					CompileAction.PrerequisiteItems.Add(ISPCFile);
+			DirectoryReference LinkerPath = LinkAction.WorkingDirectory;
 
-					CompileAction.StatusDescription = string.Format("{0} [{1}]", Path.GetFileName(ISPCFile.AbsolutePath), Arch.Replace("-", ""));
+			LinkAction.WorkingDirectory = LinkEnvironment.IntermediateDirectory!;
 
-					for(int i = 0; i < CompiledISPCObjFiles.Count; i++)
-					{
-						// ISPC compiler can't add suffix on the end of the arch, so copy to put into what linker expects
-						FileReference SourceFile = CompiledISPCObjFiles[i].Location;
-						FileReference TargetFile = FinalISPCObjFiles[i].Location;
+			// Get link arguments.
+			LinkAction.CommandArguments = LinkEnvironment.bIsBuildingLibrary ? GetArArguments(LinkEnvironment) : GetLinkArguments(LinkEnvironment, Arch);
 
-						if (SourceFile.Equals(TargetFile))
-						{
-							continue;
-						}
+			// Add the output file as a production of the link action.
+			FileItem OutputFile;
+			OutputFile = FileItem.GetItemByFileReference(LinkEnvironment.OutputFilePaths.Where(x => x.FullName.Contains(LinkEnvironment.Architecture)).First());
+			LinkAction.ProducedItems.Add(OutputFile);
+			LinkAction.StatusDescription = string.Format("{0}", Path.GetFileName(OutputFile.AbsolutePath));
+			LinkAction.CommandVersion = AndroidClangBuild!;
 
-						FileItem SourceFileItem = FileItem.GetItemByFileReference(SourceFile);
-						FileItem TargetFileItem = FileItem.GetItemByFileReference(TargetFile);
+			// LinkAction.bPrintDebugInfo = true;
 
-						Action CopyAction = Graph.CreateAction(ActionType.BuildProject);
-						CopyAction.CommandDescription = "Copy";
-						CopyAction.CommandPath = BuildHostPlatform.Current.Shell;
-						if (BuildHostPlatform.Current.ShellType == ShellType.Cmd)
-						{
-							CopyAction.CommandArguments = String.Format("/C \"copy /Y \"{0}\" \"{1}\" 1>nul\"", SourceFile, TargetFile);
-						}
-						else
-						{
-							CopyAction.CommandArguments = String.Format("-c 'cp -f \"{0}\" \"{1}\"'", SourceFile.FullName, TargetFile.FullName);
-						}
-						CopyAction.WorkingDirectory = Unreal.EngineSourceDirectory;
-						CopyAction.PrerequisiteItems.Add(SourceFileItem);
-						CopyAction.ProducedItems.Add(TargetFileItem);
-						CopyAction.StatusDescription = TargetFileItem.Location.GetFileName();
-						CopyAction.bCanExecuteRemotely = false;
-						CopyAction.bShouldOutputStatusDescription = false;
-					}
-
-					Result.ObjectFiles.AddRange(FinalISPCObjFiles);
-
-					Logger.LogDebug("   ISPC Compiling {StatusDescription}: \"{CommandPath}\" {CommandArguments}", CompileAction.StatusDescription, CompileAction.CommandPath, CompileAction.CommandArguments);
-				}
+			// Add the output file to the command-line.
+			if (LinkEnvironment.bIsBuildingLibrary)
+			{
+				LinkAction.CommandArguments += string.Format(" \"{0}\"", OutputFile.AbsolutePath);
 			}
-
-			return Result;
-		}
-
-		public override FileItem[] LinkAllFiles(LinkEnvironment LinkEnvironment, bool bBuildImportLibraryOnly, IActionGraphBuilder Graph)
-		{
-			List<FileItem> Outputs = new List<FileItem>();
-
-			if (!LinkEnvironment.bIsBuildingLibrary)
+			else
 			{
-				// @todo Lumin: will this add them multiple times?
-				ModifyLibraries(LinkEnvironment);
+				LinkAction.CommandArguments += string.Format(" -o \"{0}\"", OutputFile.AbsolutePath);
 			}
 
-			for (int ArchIndex = 0; ArchIndex < Arches!.Count; ArchIndex++)
+			// Add the input files to a response file, and pass the response file on the command-line.
+			List<string> InputFileNames = new List<string>();
+			foreach (FileItem InputFile in LinkEnvironment.InputFiles)
 			{
-				string Arch = Arches[ArchIndex];
-
-				int OutputPathIndex = ArchIndex;
-
-				// Android will have an array of outputs
-				if (!LinkEnvironment.bIsBuildingDLL && // DLL compiles don't have the Arch embedded in the name
-					(LinkEnvironment.OutputFilePaths.Count < OutputPathIndex ||
-					!LinkEnvironment.OutputFilePaths[OutputPathIndex].GetFileNameWithoutExtension().EndsWith(Arch)))
-				{
-					throw new BuildException("The OutputFilePaths array didn't match the Arches array in AndroidToolChain.LinkAllFiles");
-				}
-
-				// Create an action that invokes the linker.
-				Action LinkAction = Graph.CreateAction(ActionType.Link);
-				LinkAction.WorkingDirectory = Unreal.EngineSourceDirectory;
-
-				if (LinkEnvironment.bIsBuildingLibrary)
+				string InputPath;
+				if (InputFile.Location.IsUnderDirectory(LinkEnvironment.IntermediateDirectory!))
 				{
-					switch (Arch)
-					{
-						case "-arm64": LinkAction.CommandPath = new FileReference(ArPathArm64!); break;
-						case "-x64": LinkAction.CommandPath = new FileReference(ArPathx64!); break;
-					}
+					InputPath = InputFile.Location.MakeRelativeTo(LinkEnvironment.IntermediateDirectory!);
 				}
 				else
 				{
-					LinkAction.CommandPath = new FileReference(ClangPath!);
+					InputPath = InputFile.Location.FullName;
 				}
+				InputFileNames.Add(string.Format("\"{0}\"", InputPath.Replace('\\', '/')));
 
-				DirectoryReference LinkerPath = LinkAction.WorkingDirectory;
-
-				LinkAction.WorkingDirectory = LinkEnvironment.IntermediateDirectory!;
-
-				// Get link arguments.
-				LinkAction.CommandArguments = LinkEnvironment.bIsBuildingLibrary ? GetArArguments(LinkEnvironment) : GetLinkArguments(LinkEnvironment, Arch);
-
-				// Add the output file as a production of the link action.
-				FileItem OutputFile;
-				OutputFile = FileItem.GetItemByFileReference(LinkEnvironment.OutputFilePaths[OutputPathIndex]);
-				Outputs.Add(OutputFile);
-				LinkAction.ProducedItems.Add(OutputFile);
-				LinkAction.StatusDescription = string.Format("{0}", Path.GetFileName(OutputFile.AbsolutePath));
-				LinkAction.CommandVersion = AndroidClangBuild!;
+				LinkAction.PrerequisiteItems.Add(InputFile);
+			}
 
-				// LinkAction.bPrintDebugInfo = true;
+			string LinkResponseArguments = "";
 
-				// Add the output file to the command-line.
-				if (LinkEnvironment.bIsBuildingLibrary)
-				{
-					LinkAction.CommandArguments += string.Format(" \"{0}\"", OutputFile.AbsolutePath);
-				}
-				else
-				{
-					LinkAction.CommandArguments += string.Format(" -o \"{0}\"", OutputFile.AbsolutePath);
-				}
+			// libs don't link in other libs
+			if (!LinkEnvironment.bIsBuildingLibrary)
+			{
+				// Make a list of library paths to search
+				List<string> AdditionalLibraryPaths = new List<string>();
+				List<string> AdditionalLibraries = new List<string>();
 
-				// Add the input files to a response file, and pass the response file on the command-line.
-				List<string> InputFileNames = new List<string>();
-				foreach (FileItem InputFile in LinkEnvironment.InputFiles)
+				// Add the library paths to the additional path list
+				foreach (DirectoryReference LibraryPath in LinkEnvironment.SystemLibraryPaths)
 				{
-					// make sure it's for current Arch
-					if (Path.GetFileNameWithoutExtension(InputFile.AbsolutePath).EndsWith(ShortArchNames[Arch]))
+					// LinkerPaths could be relative or absolute
+					string AbsoluteLibraryPath = Utils.ExpandVariables(LibraryPath.FullName);
+					if (IsDirectoryForArch(AbsoluteLibraryPath, Arch))
 					{
-						string InputPath;
-						if (InputFile.Location.IsUnderDirectory(LinkEnvironment.IntermediateDirectory!))
+						// environment variables aren't expanded when using the $( style
+						if (Path.IsPathRooted(AbsoluteLibraryPath) == false)
 						{
-							InputPath = InputFile.Location.MakeRelativeTo(LinkEnvironment.IntermediateDirectory!);
+							AbsoluteLibraryPath = Path.Combine(LinkerPath.FullName, AbsoluteLibraryPath);
 						}
-						else
+						AbsoluteLibraryPath = Utils.CollapseRelativeDirectories(AbsoluteLibraryPath);
+						if (!AdditionalLibraryPaths.Contains(AbsoluteLibraryPath))
 						{
-							InputPath = InputFile.Location.FullName;
+							AdditionalLibraryPaths.Add(AbsoluteLibraryPath);
 						}
-						InputFileNames.Add(string.Format("\"{0}\"", InputPath.Replace('\\', '/')));
-
-						LinkAction.PrerequisiteItems.Add(InputFile);
 					}
 				}
 
-				string LinkResponseArguments = "";
-
-				// libs don't link in other libs
-				if (!LinkEnvironment.bIsBuildingLibrary)
+				// discover additional libraries and their paths
+				foreach (string SystemLibrary in LinkEnvironment.SystemLibraries)
 				{
-					// Make a list of library paths to search
-					List<string> AdditionalLibraryPaths = new List<string>();
-					List<string> AdditionalLibraries = new List<string>();
-
-					// Add the library paths to the additional path list
-					foreach (DirectoryReference LibraryPath in LinkEnvironment.SystemLibraryPaths)
+					if (!ShouldSkipLib(SystemLibrary, Arch))
 					{
-						// LinkerPaths could be relative or absolute
-						string AbsoluteLibraryPath = Utils.ExpandVariables(LibraryPath.FullName);
-						if (IsDirectoryForArch(AbsoluteLibraryPath, Arch))
+						if (String.IsNullOrEmpty(Path.GetDirectoryName(SystemLibrary)))
 						{
-							// environment variables aren't expanded when using the $( style
-							if (Path.IsPathRooted(AbsoluteLibraryPath) == false)
+							if (SystemLibrary.StartsWith("lib"))
 							{
-								AbsoluteLibraryPath = Path.Combine(LinkerPath.FullName, AbsoluteLibraryPath);
+								AdditionalLibraries.Add(SystemLibrary);
 							}
-							AbsoluteLibraryPath = Utils.CollapseRelativeDirectories(AbsoluteLibraryPath);
-							if (!AdditionalLibraryPaths.Contains(AbsoluteLibraryPath))
+							else
 							{
-								AdditionalLibraryPaths.Add(AbsoluteLibraryPath);
+								AdditionalLibraries.Add("lib" + SystemLibrary);
 							}
 						}
 					}
-
-					// discover additional libraries and their paths
-					foreach (string SystemLibrary in LinkEnvironment.SystemLibraries)
+				}
+				foreach (FileReference Library in LinkEnvironment.Libraries)
+				{
+					if (!ShouldSkipLib(Library.FullName, Arch))
 					{
-						if (!ShouldSkipLib(SystemLibrary, Arch))
+						string AbsoluteLibraryPath = Path.GetDirectoryName(Library.FullName)!;
+						LinkAction.PrerequisiteItems.Add(FileItem.GetItemByFileReference(Library));
+
+						string Lib = Path.GetFileNameWithoutExtension(Library.FullName);
+						if (Lib.StartsWith("lib"))
 						{
-							if (String.IsNullOrEmpty(Path.GetDirectoryName(SystemLibrary)))
+							AdditionalLibraries.Add(Lib);
+							if (!AdditionalLibraryPaths.Contains(AbsoluteLibraryPath))
 							{
-								if (SystemLibrary.StartsWith("lib"))
-								{
-									AdditionalLibraries.Add(SystemLibrary);
-								}
-								else
-								{
-									AdditionalLibraries.Add("lib" + SystemLibrary);
-								}
+								AdditionalLibraryPaths.Add(AbsoluteLibraryPath);
 							}
 						}
-					}
-					foreach (FileReference Library in LinkEnvironment.Libraries)
-					{
-						if (!ShouldSkipLib(Library.FullName, Arch))
+						else
 						{
-							string AbsoluteLibraryPath = Path.GetDirectoryName(Library.FullName)!;
-							LinkAction.PrerequisiteItems.Add(FileItem.GetItemByFileReference(Library));
-
-							string Lib = Path.GetFileNameWithoutExtension(Library.FullName);
-							if (Lib.StartsWith("lib"))
-							{
-								AdditionalLibraries.Add(Lib);
-								if (!AdditionalLibraryPaths.Contains(AbsoluteLibraryPath))
-								{
-									AdditionalLibraryPaths.Add(AbsoluteLibraryPath);
-								}
-							}
-							else
-							{
-								AdditionalLibraries.Add(Library.FullName);
-							}
+							AdditionalLibraries.Add(AbsoluteLibraryPath);
+						}
 
-							if (!DisableLibCppSharedDependencyValidation() && ReadElfPath != null)
+						if (!DisableLibCppSharedDependencyValidation() && ReadElfPath != null)
+						{
+							string? Output = Utils.RunLocalProcessAndReturnStdOut(ReadElfPath, "--dynamic \"" + Library.FullName + "\"");
+							if (Output != null)
 							{
-								string? Output = Utils.RunLocalProcessAndReturnStdOut(ReadElfPath, "--dynamic \"" + Library.FullName + "\"");
-								if (Output != null)
+								if (Output.Contains("libc++_shared.so"))
 								{
-									if (Output.Contains("libc++_shared.so"))
+									if (IsNewNDKModel())
+									{
+										throw new BuildException("Lib {0} depends on libc++_shared.so. There are known incompatibility issues when linking libc++_shared.so with Unreal Engine built with NDK22+." +
+											" Please rebuild your dependencies with static libc++!", Lib);
+									}
+									else
 									{
-										if (IsNewNDKModel())
-										{
-											throw new BuildException("Lib {0} depends on libc++_shared.so. There are known incompatibility issues when linking libc++_shared.so with Unreal Engine built with NDK22+." +
-												" Please rebuild your dependencies with static libc++!", Lib);
-										}
-										else
-										{
-											Log.TraceWarning("Lib {0} depends on libc++_shared.so. Unreal Engine is designed to be linked with libs that are built against static libc++ only. Please rebuild your dependencies with static libc++!", Lib);
-										}
+										Log.TraceWarning("Lib {0} depends on libc++_shared.so. Unreal Engine is designed to be linked with libs that are built against static libc++ only. Please rebuild your dependencies with static libc++!", Lib);
 									}
 								}
 							}
 						}
 					}
+				}
+
+				// add the library paths to response
+				foreach (string LibaryPath in AdditionalLibraryPaths)
+				{
+					LinkResponseArguments += string.Format(" -L\"{0}\"", LibaryPath);
+				}
 
-					// add the library paths to response
-					foreach (string LibaryPath in AdditionalLibraryPaths)
+				// add libraries in a library group
+				LinkResponseArguments += string.Format(" -Wl,--start-group");
+				foreach (string AdditionalLibrary in AdditionalLibraries)
+				{
+					if (AdditionalLibrary.StartsWith("lib"))
 					{
-						LinkResponseArguments += string.Format(" -L\"{0}\"", LibaryPath);
+						LinkResponseArguments += string.Format(" \"-l{0}\"", AdditionalLibrary.Substring(3));
 					}
-
-					// add libraries in a library group
-					LinkResponseArguments += string.Format(" -Wl,--start-group");
-					foreach (string AdditionalLibrary in AdditionalLibraries)
+					else
 					{
-						if (AdditionalLibrary.StartsWith("lib"))
-						{
-							LinkResponseArguments += string.Format(" \"-l{0}\"", AdditionalLibrary.Substring(3));
-						}
-						else
-						{
-							LinkResponseArguments += string.Format(" \"{0}\"", AdditionalLibrary);
-						}
+						LinkResponseArguments += string.Format(" \"{0}\"", AdditionalLibrary);
 					}
-					LinkResponseArguments += string.Format(" -Wl,--end-group");
+				}
+				LinkResponseArguments += string.Format(" -Wl,--end-group");
 
-					// Write the MAP file to the output directory.
-					if (LinkEnvironment.bCreateMapFile)
-					{
-						FileReference MAPFilePath = FileReference.Combine(LinkEnvironment.OutputDirectory!, Path.GetFileNameWithoutExtension(OutputFile.AbsolutePath) + ".map");
-						FileItem MAPFile = FileItem.GetItemByFileReference(MAPFilePath);
-						LinkResponseArguments += String.Format(" -Wl,--cref -Wl,-Map,\"{0}\"", MAPFilePath);
-						LinkAction.ProducedItems.Add(MAPFile);
+				// Write the MAP file to the output directory.
+				if (LinkEnvironment.bCreateMapFile)
+				{
+					FileReference MAPFilePath = FileReference.Combine(LinkEnvironment.OutputDirectory!, Path.GetFileNameWithoutExtension(OutputFile.AbsolutePath) + ".map");
+					FileItem MAPFile = FileItem.GetItemByFileReference(MAPFilePath);
+					LinkResponseArguments += String.Format(" -Wl,--cref -Wl,-Map,\"{0}\"", MAPFilePath);
+					LinkAction.ProducedItems.Add(MAPFile);
 
-						// Export a list of object file paths, so we can locate the object files referenced by the map file
-						ExportObjectFilePaths(LinkEnvironment, Path.ChangeExtension(MAPFilePath.FullName, ".objpaths"));
-					}
+					// Export a list of object file paths, so we can locate the object files referenced by the map file
+					ExportObjectFilePaths(LinkEnvironment, Path.ChangeExtension(MAPFilePath.FullName, ".objpaths"));
 				}
+			}
 
-				// Add the additional arguments specified by the environment.
-				LinkResponseArguments += LinkEnvironment.AdditionalArguments;
+			// Add the additional arguments specified by the environment.
+			LinkResponseArguments += LinkEnvironment.AdditionalArguments;
 
-				// Write out a response file
-				FileReference ResponseFileName = GetResponseFileName(LinkEnvironment, OutputFile);
-				InputFileNames.Add(LinkResponseArguments.Replace("\\", "/"));
+			// Write out a response file
+			FileReference ResponseFileName = GetResponseFileName(LinkEnvironment, OutputFile);
+			InputFileNames.Add(LinkResponseArguments.Replace("\\", "/"));
 
-				FileItem ResponseFileItem = Graph.CreateIntermediateTextFile(ResponseFileName, InputFileNames);
+			FileItem ResponseFileItem = Graph.CreateIntermediateTextFile(ResponseFileName, InputFileNames);
 
-				LinkAction.CommandArguments += string.Format(" @\"{0}\"", ResponseFileName);
-				LinkAction.PrerequisiteItems.Add(ResponseFileItem);
+			LinkAction.CommandArguments += string.Format(" @\"{0}\"", ResponseFileName);
+			LinkAction.PrerequisiteItems.Add(ResponseFileItem);
 
-				// Fix up the paths in commandline
-				LinkAction.CommandArguments = LinkAction.CommandArguments.Replace("\\", "/");
+			// Fix up the paths in commandline
+			LinkAction.CommandArguments = LinkAction.CommandArguments.Replace("\\", "/");
 
-				// Only execute linking on the local PC.
-				LinkAction.bCanExecuteRemotely = false;
+			// Only execute linking on the local PC.
+			LinkAction.bCanExecuteRemotely = false;
 
-				string VersionScriptFileItem = GetVersionScriptFilename(LinkEnvironment);
-				LinkAction.PrerequisiteItems.Add(FileItem.GetItemByPath(VersionScriptFileItem));
+			string VersionScriptFileItem = GetVersionScriptFilename(LinkEnvironment);
+			LinkAction.PrerequisiteItems.Add(FileItem.GetItemByPath(VersionScriptFileItem));
 				
-				if(bExecuteCompilerThroughShell)
-				{
-					SetupActionToExecuteCompilerThroughShell(ref LinkAction, LinkAction.CommandPath.FullName, LinkAction.CommandArguments, "Link");
-				}
-
-				//Logger.LogInformation("Link: {LinkActionCommandPathFullName} {LinkActionCommandArguments}", LinkAction.CommandPath.FullName, LinkAction.CommandArguments);
+			Logger.LogInformation("Link: {LinkActionCommandPathFullName} {LinkActionCommandArguments}", LinkAction.CommandPath.FullName, LinkAction.CommandArguments);
 
-				// Windows can run into an issue with too long of a commandline when clang tries to call ld to link.
-				// To work around this we call clang to just get the command it would execute and generate a
-				// second response file to directly call ld with the right arguments instead of calling through clang.
+			// Windows can run into an issue with too long of a commandline when clang tries to call ld to link.
+			// To work around this we call clang to just get the command it would execute and generate a
+			// second response file to directly call ld with the right arguments instead of calling through clang.
 /* disable while tracking down some linker errors this introduces
-				if (RuntimePlatform.IsWindows)
+			if (RuntimePlatform.IsWindows)
+			{
+				// capture the actual link command without running it
+				ProcessStartInfo StartInfo = new ProcessStartInfo();
+				StartInfo.WorkingDirectory = LinkEnvironment.IntermediateDirectory.FullName;
+				StartInfo.FileName = LinkAction.CommandPath;
+				StartInfo.Arguments = "-### " + LinkAction.CommandArguments;
+				StartInfo.UseShellExecute = false;
+				StartInfo.CreateNoWindow = true;
+				StartInfo.RedirectStandardError = true;
+
+				LinkerCommandline = "";
+
+				Process Proc = new Process();
+				Proc.StartInfo = StartInfo;
+				Proc.ErrorDataReceived += new DataReceivedEventHandler(OutputReceivedForLinker);
+				Proc.Start();
+				Proc.BeginErrorReadLine();
+				Proc.WaitForExit(5000);
+
+				LinkerCommandline = LinkerCommandline.Trim();
+
+				// the command should be in quotes; if not we'll just use clang to link as usual
+				int FirstQuoteIndex = LinkerCommandline.IndexOf('"');
+				if (FirstQuoteIndex >= 0)
 				{
-					// capture the actual link command without running it
-					ProcessStartInfo StartInfo = new ProcessStartInfo();
-					StartInfo.WorkingDirectory = LinkEnvironment.IntermediateDirectory.FullName;
-					StartInfo.FileName = LinkAction.CommandPath;
-					StartInfo.Arguments = "-### " + LinkAction.CommandArguments;
-					StartInfo.UseShellExecute = false;
-					StartInfo.CreateNoWindow = true;
-					StartInfo.RedirectStandardError = true;
-
-					LinkerCommandline = "";
-
-					Process Proc = new Process();
-					Proc.StartInfo = StartInfo;
-					Proc.ErrorDataReceived += new DataReceivedEventHandler(OutputReceivedForLinker);
-					Proc.Start();
-					Proc.BeginErrorReadLine();
-					Proc.WaitForExit(5000);
-
-					LinkerCommandline = LinkerCommandline.Trim();
-
-					// the command should be in quotes; if not we'll just use clang to link as usual
-					int FirstQuoteIndex = LinkerCommandline.IndexOf('"');
-					if (FirstQuoteIndex >= 0)
+					int SecondQuoteIndex = LinkerCommandline.Substring(FirstQuoteIndex + 1).IndexOf('"');
+					if (SecondQuoteIndex >= 0)
 					{
-						int SecondQuoteIndex = LinkerCommandline.Substring(FirstQuoteIndex + 1).IndexOf('"');
-						if (SecondQuoteIndex >= 0)
-						{
-							LinkAction.CommandPath = LinkerCommandline.Substring(FirstQuoteIndex + 1, SecondQuoteIndex - FirstQuoteIndex);
-							LinkAction.CommandArguments = LinkerCommandline.Substring(FirstQuoteIndex + SecondQuoteIndex + 3);
+						LinkAction.CommandPath = LinkerCommandline.Substring(FirstQuoteIndex + 1, SecondQuoteIndex - FirstQuoteIndex);
+						LinkAction.CommandArguments = LinkerCommandline.Substring(FirstQuoteIndex + SecondQuoteIndex + 3);
 
-							// replace double backslashes
-							LinkAction.CommandPath = LinkAction.CommandPath.Replace("\\\\", "/");
+						// replace double backslashes
+						LinkAction.CommandPath = LinkAction.CommandPath.Replace("\\\\", "/");
 
-							// now create a response file for the full command using ld directly
-							FileReference FinalResponseFileName = FileReference.Combine(LinkEnvironment.IntermediateDirectory, OutputFile.Location.GetFileName() + ".responseFinal");
-							FileItem FinalResponseFileItem = Graph.CreateIntermediateTextFile(FinalResponseFileName, LinkAction.CommandArguments);
-							LinkAction.CommandArguments = string.Format("@\"{0}\"", FinalResponseFileName);
-							LinkAction.PrerequisiteItems.Add(FinalResponseFileItem);
-						}
+						// now create a response file for the full command using ld directly
+						FileReference FinalResponseFileName = FileReference.Combine(LinkEnvironment.IntermediateDirectory, OutputFile.Location.GetFileName() + ".responseFinal");
+						FileItem FinalResponseFileItem = Graph.CreateIntermediateTextFile(FinalResponseFileName, LinkAction.CommandArguments);
+						LinkAction.CommandArguments = string.Format("@\"{0}\"", FinalResponseFileName);
+						LinkAction.PrerequisiteItems.Add(FinalResponseFileItem);
 					}
 				}
-*/
 			}
+*/
 
-			return Outputs.ToArray();
+			return OutputFile;
 		}
 
 		// captures stderr from clang
@@ -2200,27 +1755,6 @@ namespace UnrealBuildTool
 			Utils.RunLocalProcessAndLogOutput(StartInfo, Logger);
 		}
 
-		protected virtual void SetupActionToExecuteCompilerThroughShell(ref Action CompileOrLinkAction, string CommandPath, string CommandArguments, string CommandDescription)
-		{
-			string QuotedCommandPath = CommandPath;
-			if (CommandPath.Contains(' '))
-			{
-				QuotedCommandPath = "'" + CommandPath + "'";
-			}
-	
-			if (BuildHostPlatform.Current.ShellType == ShellType.Cmd)
-			{
-				CompileOrLinkAction.CommandArguments = String.Format("/c \"{0} {1}\"", QuotedCommandPath, CommandArguments);
-			}
-			else
-			{
-				CompileOrLinkAction.CommandArguments = String.Format("-c \'{0} {1}\'", QuotedCommandPath, CommandArguments);
-			}
-
-			CompileOrLinkAction.CommandPath = BuildHostPlatform.Current.Shell;
-			CompileOrLinkAction.CommandDescription = CommandDescription;
-		}
-
 		public ClangSanitizer BuildWithSanitizer()
 		{
 			if (Options.HasFlag(ClangToolChainOptions.EnableAddressSanitizer))
diff --git a/Engine/Source/Programs/UnrealBuildTool/Platform/Android/UEBuildAndroid.cs b/Engine/Source/Programs/UnrealBuildTool/Platform/Android/UEBuildAndroid.cs
--- a/Engine/Source/Programs/UnrealBuildTool/Platform/Android/UEBuildAndroid.cs
+++ b/Engine/Source/Programs/UnrealBuildTool/Platform/Android/UEBuildAndroid.cs
@@ -2,11 +2,12 @@
 
 using System;
 using System.Collections.Generic;
-using System.Text;
+using System.Linq;
 using System.IO;
 using EpicGames.Core;
 using Microsoft.Extensions.Logging;
 using Microsoft.Extensions.Logging.Abstractions;
+using UnrealBuildBase;
 
 namespace UnrealBuildTool
 {
@@ -15,18 +16,6 @@ namespace UnrealBuildTool
 	/// </summary>
 	public partial class AndroidTargetRules
 	{
-		/// <summary>
-		/// Lists Architectures that you want to build
-		/// </summary>
-		[CommandLine("-Architectures=", ListSeparator = '+')]
-		public List<string> Architectures = new List<string>();
-
-		/// <summary>
-		/// Lists GPU Architectures that you want to build (mostly used for mobile etc.)
-		/// </summary>
-		[CommandLine("-GPUArchitectures=", ListSeparator = '+')]
-		public List<string> GPUArchitectures = new List<string>();
-
 		/// <summary>
 		/// Enables address sanitizer (ASan)
 		/// </summary>
@@ -83,16 +72,6 @@ namespace UnrealBuildTool
 		#region Read-only accessor properties 
 		#pragma warning disable CS1591
 
-		public IReadOnlyList<string> Architectures
-		{
-			get { return Inner.Architectures.AsReadOnly(); }
-		}
-
-		public IReadOnlyList<string> GPUArchitectures
-		{
-			get { return Inner.GPUArchitectures.AsReadOnly(); }
-		}
-
 		public bool bEnableAddressSanitizer
 		{
 			get { return Inner.bEnableAddressSanitizer; }
@@ -141,6 +120,16 @@ namespace UnrealBuildTool
 		{
 		}
 
+		public override string GetFolderNameForArchitecture(string Architecture)
+		{
+			if (Architecture.Contains("+"))
+			{
+				return Architecture;
+			}
+				
+			return Architecture.Contains("arm") ? "a" : "x";
+		}
+
 		public override void ResetTarget(TargetRules Target)
 		{
 			ValidateTarget(Target);
@@ -170,6 +159,16 @@ namespace UnrealBuildTool
 
 			Target.bCompileRecast = true;
 			Target.bCompileISPC = false;
+
+
+			// disable plugins by architecture
+			if (Target.Architecture.ToLower() == "x64" && Target.Name != "UnrealHeaderTool")
+			{
+				Target.DisablePlugins.AddRange(new string[]
+				{
+
+				});
+			}
 		}
 
 		public override bool CanUseXGE()
@@ -183,6 +182,17 @@ namespace UnrealBuildTool
 			return Architecture == "-x64";
 		}
 
+		public override bool NeedsSingleTargetForMultiArchitecture(IEnumerable<string> Architectures)
+		{
+			return true;
+		}
+
+		public override bool RequiresArchitectureSuffix()
+		{
+			// Any -architecture argument passed to UBT only affects the contents of the binaries, not their naming
+			return false;
+		}
+		
 		public override bool IsBuildProduct(string FileName, string[] NamePrefixes, string[] NameSuffixes)
 		{
 			return IsBuildProductWithArch(FileName, NamePrefixes, NameSuffixes, ".so")
@@ -227,6 +237,73 @@ namespace UnrealBuildTool
 			return new string [] {};
 		}
 
+		public override string ConvertToPlatformArchitecture(string Arch)
+		{
+			switch (Arch.ToLower())
+			{
+				case "arm64":
+				case "a8":
+				case "a":
+				case "-arm64": 
+					return "arm64";
+
+				case "x64":
+				case "x86_64":
+				case "x6":
+				case "x":
+				case "-x64": 
+					return "x64";
+			}
+			throw new BuildException($"Unknown architecture {Arch}");
+		}
+
+		private List<string>? ActiveArches = null;
+		private FileReference? ActiveArchesProject = null;
+		public override IEnumerable<string> GetProjectArchitectures(FileReference? ProjectFile, string? TargetName, bool bGetAllSupported, bool bIsDistributionMode)
+		{
+			if (ActiveArches == null || ProjectFile != ActiveArchesProject)
+			{
+				ActiveArches = new();
+				ActiveArchesProject = ProjectFile;
+
+				// look in ini settings for what platforms to compile for
+				ConfigHierarchy Ini = ConfigCache.ReadHierarchy(ConfigHierarchyType.Engine, DirectoryReference.FromFile(ProjectFile), UnrealTargetPlatform.Android);
+				bool bBuild;
+				bool bUnsupportedBinaryBuildArch = false;
+
+				if (Ini.GetBool("/Script/AndroidRuntimeSettings.AndroidRuntimeSettings", "bBuildForArm64", out bBuild) && bBuild)
+				{
+					ActiveArches.Add("arm64");
+				}
+				if (Ini.GetBool("/Script/AndroidRuntimeSettings.AndroidRuntimeSettings", "bBuildForx8664", out bBuild) && bBuild)
+				{
+					if (File.Exists(Path.Combine(Unreal.EngineDirectory.FullName, "Build", "InstalledBuild.txt")))
+					{
+						bUnsupportedBinaryBuildArch = true;
+						Log.TraceWarningOnce("Please install source to build for x86_64 (-x64); ignoring this architecture target.");
+					}
+					else
+					{
+						ActiveArches.Add("x64");
+					}
+				}
+
+				// force arm64 if something went wrong
+				if (ActiveArches.Count == 0)
+				{
+					if (bUnsupportedBinaryBuildArch)
+					{
+						throw new BuildException("Only architectures unsupported by binary-only engine selected.");
+					}
+					else
+					{
+						ActiveArches.Add("arm64");
+					}
+				}
+			}
+			return ActiveArches;
+		}
+
 		public override void FindAdditionalBuildProductsToClean(ReadOnlyTargetRules Target, List<FileReference> FilesToDelete, List<DirectoryReference> DirectoriesToDelete)
 		{
 			base.FindAdditionalBuildProductsToClean(Target, FilesToDelete, DirectoriesToDelete);
@@ -340,13 +417,9 @@ namespace UnrealBuildTool
 
 		public override List<FileReference> FinalizeBinaryPaths(FileReference BinaryName, FileReference? ProjectFile, ReadOnlyTargetRules Target)
 		{
-			AndroidToolChain ToolChain = (AndroidToolChain)CreateToolChain(Target);
-
-			List<string> Architectures = ToolChain.GetAllArchitectures();
-
 			// make multiple output binaries
 			List<FileReference> AllBinaries = new List<FileReference>();
-			foreach (string Architecture in Architectures)
+			foreach (string Architecture in Target.Architecture.Split('+'))
 			{
 				string BinaryPath;
 				if (Target.bShouldCompileAsDLL)
@@ -388,7 +461,7 @@ namespace UnrealBuildTool
 			string NDKPath = Environment.GetEnvironmentVariable("NDKROOT")!;
 			NDKPath = NDKPath.Replace("\"", "");
 
-			AndroidToolChain ToolChain = new AndroidToolChain(Target.ProjectFile, Target.AndroidPlatform.Architectures, Target.AndroidPlatform.GPUArchitectures, Logger);
+			AndroidToolChain ToolChain = new AndroidToolChain(Target.ProjectFile, Logger);
 
 			// figure out the NDK version
 			string? NDKToolchainVersion = SDK.GetInstalledVersion();
@@ -540,14 +613,14 @@ namespace UnrealBuildTool
 			{
 				Options |= ClangToolChainOptions.EnableThinLTO;
 			}
-			return new AndroidToolChain(Target.ProjectFile, Target.AndroidPlatform.Architectures, Target.AndroidPlatform.GPUArchitectures, Options, Logger);
+			return new AndroidToolChain(Target.ProjectFile, Options, Logger);
 		}
 		public virtual UEToolChain CreateTempToolChainForProject(FileReference? ProjectFile)
 		{
 			AndroidTargetRules TargetRules = new AndroidTargetRules();
 			CommandLine.ParseArguments(Environment.GetCommandLineArgs(), TargetRules, Logger);
 			ClangToolChainOptions Options = CreateToolChainOptions(TargetRules);
-			return new AndroidToolChain(ProjectFile, null, null, Options, Logger);
+			return new AndroidToolChain(ProjectFile, Options, Logger);
 		}
 
 		/// <inheritdoc/>
diff --git a/Engine/Source/Programs/UnrealBuildTool/Platform/Android/UEDeployAndroid.cs b/Engine/Source/Programs/UnrealBuildTool/Platform/Android/UEDeployAndroid.cs
--- a/Engine/Source/Programs/UnrealBuildTool/Platform/Android/UEDeployAndroid.cs
+++ b/Engine/Source/Programs/UnrealBuildTool/Platform/Android/UEDeployAndroid.cs
@@ -93,9 +93,12 @@ namespace UnrealBuildTool
 		private bool OculusMobilePluginEnabled = false;
 		private bool GoogleVRPluginEnabled = false;
 		private bool EOSSDKPluginEnabled = false;
+		private List<string>? Architectures = null;
 
 		public void SetAndroidPluginData(List<string> Architectures, List<string> inPluginExtraData)
 		{
+			this.Architectures = new List<string>(Architectures);
+
 			List<string> NDKArches = new List<string>();
 			foreach (string NDKArch in Architectures)
 			{
@@ -1003,8 +1006,8 @@ namespace UnrealBuildTool
 
 			switch (UnrealArch)
 			{
-				case "-arm64": return CachedStoreVersion + CachedStoreVersionOffsetArm64;
-				case "-x64": return CachedStoreVersion + CachedStoreVersionOffsetX8664;
+				case "arm64": return CachedStoreVersion + CachedStoreVersionOffsetArm64;
+				case "x64": return CachedStoreVersion + CachedStoreVersionOffsetX8664;
 			}
 
 			return CachedStoreVersion;
@@ -1291,8 +1294,8 @@ namespace UnrealBuildTool
 		{
 			switch (UnrealArch)
 			{
-				case "-arm64":  return "arm64-v8a";
-				case "-x64":	return "x86_64";
+				case "arm64":  return "arm64-v8a";
+				case "x64":	return "x86_64";
 
 				default: throw new BuildException("Unknown Unreal architecture {0}", UnrealArch);
 			}
@@ -1302,14 +1305,14 @@ namespace UnrealBuildTool
 		{
 			switch (NDKArch)
 			{
-				case "arm64-v8a":   return "-arm64";
-				case "arm64":       return "-arm64";
+				case "arm64-v8a":
+				case "arm64":       return "arm64";
 				case "x86_64":
-				case "x64":			return "-x64";
+				case "x64":			return "x64";
 					
 	//				default: throw new BuildException("Unknown NDK architecture '{0}'", NDKArch);
 				// future-proof by returning arm64 for unknown
-				default:            return "-arm64";
+				default:            return "arm64";
 			}
 		}
 
@@ -1933,7 +1936,8 @@ namespace UnrealBuildTool
 			return true;
 		}
 
-		private string GetAllBuildSettings(AndroidToolChain ToolChain, UnrealPluginLanguage UPL, bool bForDistribution, bool bMakeSeparateApks, bool bPackageDataInsideApk, bool bDisableVerifyOBBOnStartUp, bool bUseExternalFilesDir, string TemplatesHashCode)
+		private string GetAllBuildSettings(AndroidToolChain ToolChain, UnrealPluginLanguage UPL, bool bForDistribution, bool bMakeSeparateApks, bool bPackageDataInsideApk, 
+			bool bDisableVerifyOBBOnStartUp, bool bUseExternalFilesDir, string TemplatesHashCode)
 		{
 			// make the settings string - this will be char by char compared against last time
 			StringBuilder CurrentSettings = new StringBuilder();
@@ -1994,8 +1998,7 @@ namespace UnrealBuildTool
 				}
 			}
 
-			List<string> Arches = ToolChain.GetAllArchitectures();
-			foreach (string Arch in Arches)
+			foreach (string Arch in Architectures!)
 			{
 				CurrentSettings.AppendFormat("Arch={0}{1}", Arch, Environment.NewLine);
 			}
@@ -2027,7 +2030,7 @@ namespace UnrealBuildTool
 
 			// get a list of the ini settings in UPL files that may affect the build
 			// architecture doesn't matter here since this node does not use init logic
-			string UPLBuildSettings = UPL.ProcessPluginNode("arm64-v8a", "registerBuildSettings", "");
+			string UPLBuildSettings = UPL.ProcessPluginNode(GetNDKArch(Architectures[0]), "registerBuildSettings", "");
 			foreach (string Line in UPLBuildSettings.Split('\n', StringSplitOptions.RemoveEmptyEntries))
 			{
 				string SectionName = Line.Trim();
@@ -2074,14 +2077,12 @@ namespace UnrealBuildTool
 			return CurrentSettings.ToString();
 		}
 
-		private bool CheckDependencies(AndroidToolChain ToolChain, string ProjectName, string ProjectDirectory, string UnrealBuildFilesPath, string GameBuildFilesPath, string EngineDirectory, List<string> SettingsFiles,
+		private bool CheckDependencies(List<string> Architectures, string ProjectName, string ProjectDirectory, string UnrealBuildFilesPath, string GameBuildFilesPath, string EngineDirectory, List<string> SettingsFiles,
 			string CookFlavor, string OutputPath, bool bMakeSeparateApks, bool bPackageDataInsideApk)
 		{
-			List<string> Arches = ToolChain.GetAllArchitectures();
-
 			// check all input files (.so, java files, .ini files, etc)
 			bool bAllInputsCurrent = true;
-			foreach (string Arch in Arches)
+			foreach (string Arch in Architectures)
 			{
 				string SourceSOName = AndroidToolChain.InlineArchName(OutputPath, Arch);
 				// if the source binary was UnrealGame, replace it with the new project name, when re-packaging a binary only build
@@ -2623,7 +2624,7 @@ namespace UnrealBuildTool
 			bool bAddDensity = (SDKLevelInt >= 24) && (MinSDKVersion >= 17);
 
 			// disable Oculus Mobile if not supported platform (in this case only armv7 for now)
-			if (UnrealArch != "-armv7" && UnrealArch != "-arm64")
+			if (UnrealArch != "arm64")
 			{
 				if (bPackageForOculusMobile)
 				{
@@ -3994,9 +3995,11 @@ namespace UnrealBuildTool
 		{
 			Logger.LogInformation("");
 			Logger.LogInformation("===={Time}====PREPARING TO MAKE APK=================================================================", DateTime.Now.ToString());
-
-			// Get list of all architecture and GPU targets for build
-			List<string> Arches = ToolChain.GetAllArchitectures();
+			
+			if (Architectures == null)
+			{
+				throw new BuildException("Called MakeApk without first calling SetAndroidPluginData");
+			}
 
 			// we do not need to really build an engine UnrealGame.apk so short-circuit it
 			if (!ForceAPKGeneration && ProjectName == "UnrealGame" && OutputPath.Replace("\\", "/").Contains("/Engine/Binaries/Android/") && Path.GetFileNameWithoutExtension(OutputPath).StartsWith("UnrealGame"))
@@ -4083,7 +4086,7 @@ namespace UnrealBuildTool
 
 			// get a list of unique NDK architectures enabled for build
 			List<string> NDKArches = new List<string>();
-			foreach (string Arch in Arches)
+			foreach (string Arch in Architectures)
 			{
 				string NDKArch = GetNDKArch(Arch);
 				if (!NDKArches.Contains(NDKArch))
@@ -4182,11 +4185,11 @@ namespace UnrealBuildTool
 				RequiredOBBFiles.Add(PatchFileLocation);
 			}
 
-			if (Arches.Count > 0)
+			if (Architectures.Count > 0)
 			{
 				// Generate the OBBData.java file if out of date (can skip rewriting it if packaging inside Apk in some cases)
 				// Note: this may be replaced per architecture later if store version is different
-				WriteJavaOBBDataFile(UnrealOBBDataFileName, PackageName, RequiredOBBFiles, CookFlavor, bPackageDataInsideApk, Arches[0].Substring(1));
+				WriteJavaOBBDataFile(UnrealOBBDataFileName, PackageName, RequiredOBBFiles, CookFlavor, bPackageDataInsideApk, Architectures[0]);
 			}
 
 			// Make sure any existing proguard file in project is NOT used (back it up)
@@ -4206,7 +4209,7 @@ namespace UnrealBuildTool
 			// Sometimes old files get left behind if things change, so we'll do a clean up pass
 			foreach (string NDKArch in NDKArches)
 			{
-				string UnrealBuildPath = Path.Combine(IntermediateAndroidPath, GetUnrealArch(NDKArch).Substring(1).Replace("-", "_"));
+				string UnrealBuildPath = Path.Combine(IntermediateAndroidPath, GetUnrealArch(NDKArch).Replace("-", "_"));
 
 				string CleanUpBaseDir = Path.Combine(ProjectDirectory, "Build", "Android", "src");
 				string ImmediateBaseDir = Path.Combine(UnrealBuildPath, "src");
@@ -4335,7 +4338,7 @@ namespace UnrealBuildTool
 				JavaFiles.AddRange(from t in templates select t.SourceFile);
 				JavaFiles.AddRange(from t in templates select t.DestinationFile);
 
-				bBuildSettingsMatch = CheckDependencies(ToolChain, ProjectName, ProjectDirectory, UnrealBuildFilesPath, GameBuildFilesPath,
+				bBuildSettingsMatch = CheckDependencies(Architectures, ProjectName, ProjectDirectory, UnrealBuildFilesPath, GameBuildFilesPath,
 					EngineDirectory, JavaFiles, CookFlavor, OutputPath, bMakeSeparateApks, bPackageDataInsideApk);
 
 			}
@@ -4362,13 +4365,13 @@ namespace UnrealBuildTool
 			bool bRequiresOBB = RequiresOBB(bDisallowPackagingDataInApk, ObbFileLocation);
 			if (!bBuildSettingsMatch)
 			{
-				BuildList = from Arch in Arches
+				BuildList = from Arch in Architectures
 							let manifest = GenerateManifest(ToolChain, ProjectName, InTargetType, EngineDirectory, bForDistribution, bPackageDataInsideApk, GameBuildFilesPath, bRequiresOBB, bDisableVerifyOBBOnStartUp, Arch, CookFlavor, bUseExternalFilesDir, Configuration.ToString(), SDKLevelInt, bSkipGradleBuild, bEnableBundle)
 							select Tuple.Create(Arch, manifest);
 			}
 			else
 			{
-				BuildList = from Arch in Arches
+				BuildList = from Arch in Architectures
 							let manifestFile = Path.Combine(IntermediateAndroidPath, Arch + "_AndroidManifest.xml")
 							let manifest = GenerateManifest(ToolChain, ProjectName, InTargetType, EngineDirectory, bForDistribution, bPackageDataInsideApk, GameBuildFilesPath, bRequiresOBB, bDisableVerifyOBBOnStartUp, Arch, CookFlavor, bUseExternalFilesDir, Configuration.ToString(), SDKLevelInt, bSkipGradleBuild, bEnableBundle)
 							let OldManifest = File.Exists(manifestFile) ? File.ReadAllText(manifestFile) : ""
@@ -4395,7 +4398,7 @@ namespace UnrealBuildTool
 					string Manifest = build.Item2;
 					string NDKArch = GetNDKArch(Arch);
 					
-					string UnrealBuildPath = Path.Combine(IntermediateAndroidPath, Arch.Substring(1).Replace("-", "_"));
+					string UnrealBuildPath = Path.Combine(IntermediateAndroidPath, Arch.Replace("-", "_"));
 
 					Logger.LogInformation("\n===={Time}====PREPARING NATIVE CODE====={Arch}============================================================", DateTime.Now.ToString(), Arch);
 
@@ -4590,7 +4593,7 @@ namespace UnrealBuildTool
 
 				Logger.LogInformation("\n===={Time}====PREPARING NATIVE CODE====={Arch}============================================================", DateTime.Now.ToString(), Arch);
 
-				string UnrealBuildPath = Path.Combine(IntermediateAndroidPath, Arch.Substring(1).Replace("-", "_"));
+				string UnrealBuildPath = Path.Combine(IntermediateAndroidPath, Arch.Replace("-", "_"));
 
 				// If we are packaging for Amazon then we need to copy the  file to the correct location
 				Logger.LogInformation("bPackageDataInsideApk = {bPackageDataInsideApk}", bPackageDataInsideApk);
@@ -4964,7 +4967,7 @@ namespace UnrealBuildTool
 					string Manifest = build.Item2;
 					string NDKArch = GetNDKArch(Arch);
 
-					string UnrealBuildPath = Path.Combine(IntermediateAndroidPath, Arch.Substring(1).Replace("-", "_"));
+					string UnrealBuildPath = Path.Combine(IntermediateAndroidPath, Arch.Replace("-", "_"));
 					string UnrealBuildGradlePath = Path.Combine(UnrealBuildPath, "gradle");
 
 					if (!Directory.Exists(UnrealBuildGradlePath))
@@ -5297,7 +5300,7 @@ namespace UnrealBuildTool
 
 						Logger.LogInformation("\n===={Time}====GENERATING BUNDLE====={Arch}================================================================", DateTime.Now.ToString(), Arch);
 
-						string UnrealBuildPath = Path.Combine(IntermediateAndroidPath, Arch.Substring(1).Replace("-", "_"));
+						string UnrealBuildPath = Path.Combine(IntermediateAndroidPath, Arch.Replace("-", "_"));
 						string UnrealBuildGradlePath = Path.Combine(UnrealBuildPath, "gradle");
 
 						string GradleScriptPath = Path.Combine(UnrealBuildGradlePath, "gradlew");
@@ -5388,7 +5391,7 @@ namespace UnrealBuildTool
 			AndroidToolChain ToolChain = (AndroidToolChain)((AndroidPlatform)UEBuildPlatform.GetBuildPlatform(Receipt.Platform)).CreateTempToolChainForProject(Receipt.ProjectFile);
 
 			// get the receipt
-			SetAndroidPluginData(ToolChain.GetAllArchitectures(), CollectPluginDataPaths(Receipt));
+			SetAndroidPluginData(Receipt.Architecture.Split('+').ToList(), CollectPluginDataPaths(Receipt));
 
 			bool bShouldCompileAsDll = Receipt.HasValueForAdditionalProperty("CompileAsDll", "true");
 
@@ -5481,7 +5484,16 @@ namespace UnrealBuildTool
 			AndroidTargetRules TargetRules = new AndroidTargetRules();
 			CommandLine.ParseArguments(Environment.GetCommandLineArgs(), TargetRules, Logger);
 			ClangToolChainOptions Options = AndroidPlatform.CreateToolChainOptions(TargetRules);
-			AndroidToolChain ToolChain = new AndroidToolChain(ProjectFile, null, null, Options, Logger);
+			AndroidToolChain ToolChain = new AndroidToolChain(ProjectFile, Options, Logger);
+
+			// hunt down the receipt that compiled the .so
+			//FileReference ReceiptFilename = TargetReceipt.GetDefaultPath(ProjectDirectory, ProjectName, SC.StageTargetPlatform.IniPlatformType, Configuration, Architecture);
+			//if (!FileReference.Exists(ReceiptFilename))
+			//{
+			//	ReceiptFilename = TargetReceipt.GetDefaultPath(Unreal.EngineDirectory, ReceiptName, SC.StageTargetPlatform.IniPlatformType, Configuration, Architecture);
+			//}
+			FileReference ReceiptFilename = new FileReference(ExecutablePath).ChangeExtension(".target");
+			TargetReceipt Receipt = TargetReceipt.Read(ReceiptFilename);
 
 			SavePackageInfo(ProjectName, ProjectDirectory.FullName, Type, bSkipGradleBuild);
 
diff --git a/Engine/Source/Programs/UnrealBuildTool/ToolChain/ClangToolChain.cs b/Engine/Source/Programs/UnrealBuildTool/ToolChain/ClangToolChain.cs
--- a/Engine/Source/Programs/UnrealBuildTool/ToolChain/ClangToolChain.cs
+++ b/Engine/Source/Programs/UnrealBuildTool/ToolChain/ClangToolChain.cs
@@ -995,7 +995,8 @@ namespace UnrealBuildTool
 		{
 			Action CompileAction = Graph.CreateAction(ActionType.Compile);
 
-			List<string> FileArguments = new();
+			// copy the global arguments into the file arguments, so GetCompileArguments_FileType can remove entries if needed (special case but can be important)
+			List<string> FileArguments = new(GlobalArguments);
 
 			// Add C or C++ specific compiler arguments.
 			GetCompileArguments_FileType(CompileEnvironment, SourceFile, OutputDir, FileArguments, CompileAction, Result);
@@ -1005,10 +1006,7 @@ namespace UnrealBuildTool
 
 			// Creates the path to the response file using the name of the output file and creates its contents.
 			FileReference ResponseFileName = new FileReference(TargetFile.AbsolutePath + ".response");
-			List<string> ResponseFileContents = new();
-			ResponseFileContents.AddRange(GlobalArguments);
-			ResponseFileContents.AddRange(FileArguments);
-			ResponseFileContents = ExpandResponseFileContents(ResponseFileContents);
+			List<string> ResponseFileContents = ExpandResponseFileContents(FileArguments);
 
 			if (RuntimePlatform.IsWindows)
 			{
diff --git a/Engine/Source/Programs/UnrealBuildTool/ToolChain/ISPCToolChain.cs b/Engine/Source/Programs/UnrealBuildTool/ToolChain/ISPCToolChain.cs
--- a/Engine/Source/Programs/UnrealBuildTool/ToolChain/ISPCToolChain.cs
+++ b/Engine/Source/Programs/UnrealBuildTool/ToolChain/ISPCToolChain.cs
@@ -41,11 +41,9 @@ namespace UnrealBuildTool
 			{
 				switch (Arch)
 				{
-					case "-armv7": ISPCTargets.Add("neon"); break; // Assumes NEON is in use
-					case "-arm64": ISPCTargets.Add("neon"); break;
-					case "-x86": ISPCTargets.Add("sse4"); break;
-					case "-x64": ISPCTargets.Add("sse4"); break;
-					default: Logger.LogWarning("Invalid Android architecture for ISPC. At least one architecture (armv7, x86, etc) needs to be selected in the project settings to build"); break;
+					case "arm64": ISPCTargets.Add("neon"); break;
+					case "x64": ISPCTargets.AddRange(new string[] { "sse4" }); break;
+					default: Logger.LogWarning("Invalid Android architecture for ISPC. At least one architecture (arm64, x64) needs to be selected in the project settings to build"); break;
 				}
 			}
 			else if (Platform == UnrealTargetPlatform.IOS)
@@ -121,11 +119,9 @@ namespace UnrealBuildTool
 			{
 				switch (Arch)
 				{
-					case "-armv7": ISPCArch += "arm"; break; // Assumes NEON is in use
-					case "-arm64": ISPCArch += "aarch64"; break;
-					case "-x86": ISPCArch += "x86"; break;
-					case "-x64": ISPCArch += "x86-64"; break;
-					default: Logger.LogWarning("Invalid Android architecture for ISPC. At least one architecture (armv7, x86, etc) needs to be selected in the project settings to build"); break;
+					case "arm64": ISPCArch += "aarch64"; break;
+					case "x64": ISPCArch += "x86-64"; break;
+					default: Logger.LogWarning("Invalid Android architecture for ISPC. At least one architecture (arm64, x64) needs to be selected in the project settings to build"); break;
 				}
 			}
 			else if (Platform == UnrealTargetPlatform.IOS)
