UnrealBuildTool: Suppress unreachable static analysis warning

#rnx

diff --git a/Engine/Source/Programs/UnrealBuildTool/Executors/XGE.cs b/Engine/Source/Programs/UnrealBuildTool/Executors/XGE.cs
--- a/Engine/Source/Programs/UnrealBuildTool/Executors/XGE.cs
+++ b/Engine/Source/Programs/UnrealBuildTool/Executors/XGE.cs
@@ -602,6 +602,7 @@ namespace UnrealBuildTool
 		/// <param name="ActionCount"></param>
 		/// <param name="Logger"></param>
 		/// <returns>Indicates whether the tasks were successfully executed.</returns>
+		[SuppressMessage("Interoperability", "CA1416:Validate platform compatibility", Justification = "Registry only checked on Windows HostPlatform")]
 		bool ExecuteTaskFile(string TaskFilePath, DataReceivedEventHandler OutputEventHandler, int ActionCount, ILogger Logger)
 		{
 			// A bug in the UCRT can cause XGE to hang on VS2015 builds. Figure out if this hang is likely to effect this build and workaround it if able.
