Add a TargetRule to remove FixedTimeStep/Deterministic mode from a project.

diff --git a/Engine/Source/Programs/UnrealBuildTool/Configuration/TargetRules.cs b/Engine/Source/Programs/UnrealBuildTool/Configuration/TargetRules.cs
--- a/Engine/Source/Programs/UnrealBuildTool/Configuration/TargetRules.cs
+++ b/Engine/Source/Programs/UnrealBuildTool/Configuration/TargetRules.cs
@@ -1178,6 +1178,12 @@ namespace UnrealBuildTool
 		[RequiresUniqueBuildEnvironment]
 		public bool bWithDirectXMath = false;
 
+		/// <summary>
+		/// Whether to enable the support for FixedTimeStep in the engine
+		/// </summary>
+		[RequiresUniqueBuildEnvironment]
+		public bool bWithFixedTimeStepSupport = true;
+
 		/// <summary>
 		/// Whether to turn on logging for test/shipping builds.
 		/// </summary>
@@ -3058,6 +3064,11 @@ namespace UnrealBuildTool
 			get { return Inner.bWithDirectXMath; }
 		}
 
+		public bool bWithFixedTimeStepSupport
+		{
+			get { return Inner.bWithFixedTimeStepSupport; }
+		}
+
 		public bool bUseLoggingInShipping
 		{
 			get { return Inner.bUseLoggingInShipping; }
diff --git a/Engine/Source/Programs/UnrealBuildTool/Configuration/UEBuildTarget.cs b/Engine/Source/Programs/UnrealBuildTool/Configuration/UEBuildTarget.cs
--- a/Engine/Source/Programs/UnrealBuildTool/Configuration/UEBuildTarget.cs
+++ b/Engine/Source/Programs/UnrealBuildTool/Configuration/UEBuildTarget.cs
@@ -4117,6 +4117,15 @@ namespace UnrealBuildTool
 				GlobalCompileEnvironment.Definitions.Add("WITH_PERFCOUNTERS=0");
 			}
 
+			if (Rules.bWithFixedTimeStepSupport)
+			{
+				GlobalCompileEnvironment.Definitions.Add("WITH_FIXED_TIME_STEP_SUPPORT=1");
+			}
+			else
+			{
+				GlobalCompileEnvironment.Definitions.Add("WITH_FIXED_TIME_STEP_SUPPORT=0");
+			}
+
 			if (Rules.bUseLoggingInShipping)
 			{
 				GlobalCompileEnvironment.Definitions.Add("USE_LOGGING_IN_SHIPPING=1");
diff --git a/Engine/Source/Runtime/Core/Public/Misc/App.h b/Engine/Source/Runtime/Core/Public/Misc/App.h
--- a/Engine/Source/Runtime/Core/Public/Misc/App.h
+++ b/Engine/Source/Runtime/Core/Public/Misc/App.h
@@ -495,7 +495,11 @@ public:
 	 */
 	static bool UseFixedTimeStep()
 	{
+#if WITH_FIXED_TIME_STEP_SUPPORT
 		return bUseFixedTimeStep;
+#else
+		return false;
+#endif
 	}
 
 	/**
diff --git a/Engine/Source/Runtime/Engine/Private/UnrealEngine.cpp b/Engine/Source/Runtime/Engine/Private/UnrealEngine.cpp
--- a/Engine/Source/Runtime/Engine/Private/UnrealEngine.cpp
+++ b/Engine/Source/Runtime/Engine/Private/UnrealEngine.cpp
@@ -2478,7 +2478,11 @@ void UEngine::UpdateTimeAndHandleMaxTickRate()
 	bool bTimeWasManipulatedDebug = bTimeWasManipulated;	//Just used for logging of previous frame
 
 															// Figure out whether we want to use real or fixed time step.
+#if WITH_FIXED_TIME_STEP_SUPPORT
 	const bool bUseFixedTimeStep = FApp::IsBenchmarking() || FApp::UseFixedTimeStep();
+#else
+	const bool bUseFixedTimeStep = false;
+#endif
 	static bool bPreviousUseFixedTimeStep = bUseFixedTimeStep;
 
 	// Updates logical last time to match logical current time from last tick
diff --git a/Engine/Source/Runtime/Launch/Private/LaunchEngineLoop.cpp b/Engine/Source/Runtime/Launch/Private/LaunchEngineLoop.cpp
--- a/Engine/Source/Runtime/Launch/Private/LaunchEngineLoop.cpp
+++ b/Engine/Source/Runtime/Launch/Private/LaunchEngineLoop.cpp
@@ -2414,12 +2414,14 @@ int32 FEngineLoop::PreInitPreStartupScreen(const TCHAR* CmdLine)
 	FApp::SetBenchmarking(false);
 #endif // !UE_BUILD_SHIPPING
 
+#if WITH_FIXED_TIME_STEP_SUPPORT
 	// "-Deterministic" is a shortcut for "-UseFixedTimeStep -FixedSeed"
 	bool bDeterministic = FParse::Param(FCommandLine::Get(), TEXT("Deterministic"));
 
 	FApp::SetUseFixedTimeStep(bDeterministic || FParse::Param(FCommandLine::Get(), TEXT("UseFixedTimeStep")));
 
 	FApp::bUseFixedSeed = bDeterministic || FApp::IsBenchmarking() || FParse::Param(FCommandLine::Get(), TEXT("FixedSeed"));
+#endif
 
 	// Initialize random number generator.
 	{
