Change CppStandard to CStandard for *.m files

diff --git a/Engine/Source/Programs/UnrealBuildTool/ToolChain/ClangToolChain.cs b/Engine/Source/Programs/UnrealBuildTool/ToolChain/ClangToolChain.cs
--- a/Engine/Source/Programs/UnrealBuildTool/ToolChain/ClangToolChain.cs
+++ b/Engine/Source/Programs/UnrealBuildTool/ToolChain/ClangToolChain.cs
@@ -405,7 +405,7 @@ namespace UnrealBuildTool
 			Arguments.Add("-x objective-c");
 			Arguments.Add("-fobjc-abi-version=2");
 			Arguments.Add("-fobjc-legacy-dispatch");
-			GetCppStandardCompileArgument(CompileEnvironment, Arguments);
+			GetCStandardCompileArgument(CompileEnvironment, Arguments);
 		}
 
 		protected virtual void GetCompileArguments_PCH(CppCompileEnvironment CompileEnvironment, List<string> Arguments)
@@ -1130,4 +1130,4 @@ namespace UnrealBuildTool
 		{
 		}
 	}
-}
\ No newline at end of file
+}
