Allow for stack reserve and commit size control on all Microsoft platforms (expanded from Windows only previously)

- don't attempt to set stack reserve size of 0 as we use this as our "use the system default" value in the configs

diff --git a/Engine/Source/Programs/UnrealBuildTool/Platform/Windows/VCToolChain.cs b/Engine/Source/Programs/UnrealBuildTool/Platform/Windows/VCToolChain.cs
--- a/Engine/Source/Programs/UnrealBuildTool/Platform/Windows/VCToolChain.cs
+++ b/Engine/Source/Programs/UnrealBuildTool/Platform/Windows/VCToolChain.cs
@@ -1161,15 +1161,21 @@ namespace UnrealBuildTool
 
 				// Explicitly declare that the executable is compatible with Data Execution Prevention.
 				Arguments.Add("/NXCOMPAT");
+			}
 
-				// Set the default stack size.
-				if (LinkEnvironment.DefaultStackSizeCommit > 0)
-				{
-					Arguments.Add("/STACK:" + LinkEnvironment.DefaultStackSize + "," + LinkEnvironment.DefaultStackSizeCommit);
-				}
-				else
+			// Set the default stack size.
+			if (LinkEnvironment.Platform.IsInGroup(UnrealPlatformGroup.Microsoft))
+			{
+				if (LinkEnvironment.DefaultStackSize > 0)
 				{
-					Arguments.Add("/STACK:" + LinkEnvironment.DefaultStackSize);
+					if (LinkEnvironment.DefaultStackSizeCommit > 0)
+					{
+						Arguments.Add("/STACK:" + LinkEnvironment.DefaultStackSize + "," + LinkEnvironment.DefaultStackSizeCommit);
+					}
+					else
+					{
+						Arguments.Add("/STACK:" + LinkEnvironment.DefaultStackSize);
+					}
 				}
 			}
 
