- Fixed legacy xcode projects when made for TVOS only - it was making Mac only projects

diff --git a/Engine/Source/Programs/UnrealBuildTool/ProjectFiles/Xcode/XcodeProjectLegacy.cs b/Engine/Source/Programs/UnrealBuildTool/ProjectFiles/Xcode/XcodeProjectLegacy.cs
--- a/Engine/Source/Programs/UnrealBuildTool/ProjectFiles/Xcode/XcodeProjectLegacy.cs
+++ b/Engine/Source/Programs/UnrealBuildTool/ProjectFiles/Xcode/XcodeProjectLegacy.cs
@@ -1110,9 +1110,13 @@ namespace UnrealBuildTool.XcodeProjectLegacy
 		private void AppendNativeTargetBuildConfiguration(StringBuilder Content, XcodeBuildConfig Config, string ConfigGuid, FileReference? ProjectFile, ILogger Logger)
 		{
 			bool bMacOnly = true;
-			if (Config.ProjectTarget!.TargetRules != null && XcodeProjectFileGenerator.XcodePlatforms.Contains(UnrealTargetPlatform.IOS))
+			if (Config.ProjectTarget!.TargetRules != null)
 			{
-				if (Config.ProjectTarget.SupportedPlatforms.Contains(UnrealTargetPlatform.IOS))
+				if (XcodeProjectFileGenerator.XcodePlatforms.Contains(UnrealTargetPlatform.IOS) && Config.ProjectTarget.SupportedPlatforms.Contains(UnrealTargetPlatform.IOS))
+				{
+					bMacOnly = false;
+				}
+				if (XcodeProjectFileGenerator.XcodePlatforms.Contains(UnrealTargetPlatform.TVOS) && Config.ProjectTarget.SupportedPlatforms.Contains(UnrealTargetPlatform.TVOS))
 				{
 					bMacOnly = false;
 				}
@@ -1297,9 +1301,13 @@ namespace UnrealBuildTool.XcodeProjectLegacy
 		private void AppendLegacyTargetBuildConfiguration(StringBuilder Content, XcodeBuildConfig Config, string ConfigGuid, FileReference? ProjectFile)
 		{
 			bool bMacOnly = true;
-			if (Config.ProjectTarget!.TargetRules != null && XcodeProjectFileGenerator.XcodePlatforms.Contains(UnrealTargetPlatform.IOS))
+			if (Config.ProjectTarget!.TargetRules != null)
 			{
-				if (Config.ProjectTarget.SupportedPlatforms.Contains(UnrealTargetPlatform.IOS))
+				if (XcodeProjectFileGenerator.XcodePlatforms.Contains(UnrealTargetPlatform.IOS) && Config.ProjectTarget.SupportedPlatforms.Contains(UnrealTargetPlatform.IOS))
+				{
+					bMacOnly = false;
+				}
+				if (XcodeProjectFileGenerator.XcodePlatforms.Contains(UnrealTargetPlatform.TVOS) && Config.ProjectTarget.SupportedPlatforms.Contains(UnrealTargetPlatform.TVOS))
 				{
 					bMacOnly = false;
 				}
