IWYU iteration changes
- Adding -IWYUHeadersOnly command line option to restrict IWYU to just header files
- Adding IWYU support for -Module=XXX

diff --git a/Engine/Source/Programs/UnrealBuildTool/Configuration/TargetRules.cs b/Engine/Source/Programs/UnrealBuildTool/Configuration/TargetRules.cs
--- a/Engine/Source/Programs/UnrealBuildTool/Configuration/TargetRules.cs
+++ b/Engine/Source/Programs/UnrealBuildTool/Configuration/TargetRules.cs
@@ -1281,6 +1281,12 @@ namespace UnrealBuildTool
 		[CommandLine("-IWYU")]
 		public bool bIWYU = false;
 
+		/// <summary>
+		/// Tells "include what you use" to only compile header files.
+		/// </summary>
+		[CommandLine("-IWYUHeadersOnly")]
+		public bool bIWYUHeadersOnly = false;
+
 		/// <summary>
 		/// Enforce "include what you use" rules; warns if monolithic headers (Engine.h, UnrealEd.h, etc...) are used, and checks that source files include their matching header first.
 		/// </summary>
@@ -3155,6 +3161,11 @@ namespace UnrealBuildTool
 			get { return Inner.bIWYU; }
 		}
 
+		public bool bIWYUHeadersOnly
+		{
+			get { return Inner.bIWYUHeadersOnly; }
+		}
+
 		public bool bEnforceIWYU
 		{
 			get { return Inner.bEnforceIWYU; }
diff --git a/Engine/Source/Programs/UnrealBuildTool/Configuration/UEBuildModuleCPP.cs b/Engine/Source/Programs/UnrealBuildTool/Configuration/UEBuildModuleCPP.cs
--- a/Engine/Source/Programs/UnrealBuildTool/Configuration/UEBuildModuleCPP.cs
+++ b/Engine/Source/Programs/UnrealBuildTool/Configuration/UEBuildModuleCPP.cs
@@ -810,6 +810,11 @@ namespace UnrealBuildTool
 
 				if (HeaderFileItems.Count > 0)
 				{
+					if (Target.bIWYUHeadersOnly)
+					{
+						LinkInputFiles.Clear();
+					}
+
 					// Add the compile actions
 					LinkInputFiles.AddRange(ToolChain.CompileAllCPPFiles(CompileEnvironment, HeaderFileItems, IntermediateDirectory, Name, Graph).ObjectFiles);
 				}
diff --git a/Engine/Source/Programs/UnrealBuildTool/Modes/IWYUMode.cs b/Engine/Source/Programs/UnrealBuildTool/Modes/IWYUMode.cs
--- a/Engine/Source/Programs/UnrealBuildTool/Modes/IWYUMode.cs
+++ b/Engine/Source/Programs/UnrealBuildTool/Modes/IWYUMode.cs
@@ -284,6 +284,15 @@ namespace UnrealBuildTool
 						Logger.LogError($"Can't find module with name {ModuleToUpdate}");
 						return -1;
 					}
+
+					foreach (string OnlyModuleName in TargetDescriptors[0].OnlyModuleNames)
+					{
+						if (string.Compare(OnlyModuleName, ModuleToUpdate, StringComparison.OrdinalIgnoreCase) != 0)
+						{
+							Logger.LogError($"ModuleToUpdate '{ModuleToUpdate}' was not in list of specified modules: {string.Join(", ", TargetDescriptors[0].OnlyModuleNames)}");
+							return -1;
+						}
+					}
 				}
 				
 				if (!String.IsNullOrEmpty(PathToUpdate))
@@ -421,28 +430,41 @@ namespace UnrealBuildTool
 					CppDependencyCache.SaveAll();
 				}
 
+				HashSet<FileItem> OutputItems = new HashSet<FileItem>();
+				if (Descriptor.OnlyModuleNames.Count > 0)
+				{
+					foreach (string OnlyModuleName in Descriptor.OnlyModuleNames)
+					{
+						FileItem[]? OutputItemsForModule;
+						if (!Makefile.ModuleNameToOutputItems.TryGetValue(OnlyModuleName, out OutputItemsForModule))
+						{
+							throw new BuildException("Unable to find output items for module '{0}'", OnlyModuleName);
+						}
+						OutputItems.UnionWith(OutputItemsForModule);
+					}
+				}
+				else
+				{
+					// Use all the output items from the target
+					OutputItems.UnionWith(Makefile.OutputItems);
+				}
+
 				// Time to parse all the .iwyu files generated from iwyu.
 				// We Do this in parallel since it involves reading a ton of .iwyu files from disk.
 				Logger.LogInformation("Parsing output from IWYU...");
-				Parallel.ForEach(Makefile.Actions, Action =>
+				Parallel.ForEach(OutputItems, OutputItem =>
 				{
-					if (Action.ActionType != ActionType.Compile)
+					if (!OutputItem.Name.EndsWith(".iwyu"))
 					{
 						return;
 					}
 
-					FileItem? JsonFile = Action.ProducedItems.Where(i => i.Name.Contains(".iwyu")).FirstOrDefault();
-					if (JsonFile == null)
-					{
-						Logger.LogError($"Can't find produced json file for Action {String.Join(',', Action.PrerequisiteItems.Select(i => i.Name))}");
-						ReadSuccess = false;
-						return;
-					}
+					string IWYUFilePath = OutputItem.AbsolutePath;
 
-					string? JsonContent = File.ReadAllText(JsonFile.AbsolutePath);
+					string? JsonContent = File.ReadAllText(IWYUFilePath);
 					if (JsonContent == null)
 					{
-						Logger.LogError($"Failed to read file {JsonFile.AbsolutePath}");
+						Logger.LogError($"Failed to read file {IWYUFilePath}");
 						ReadSuccess = false;
 						return;
 					}
@@ -450,7 +472,7 @@ namespace UnrealBuildTool
 					try
 					{
 						IWYUFile? IWYUFile = JsonSerializer.Deserialize<IWYUFile>(JsonContent);
-						IWYUFile!.Name = JsonFile.AbsolutePath;
+						IWYUFile!.Name = IWYUFilePath;
 
 						// Traverse the cpp/inl/h file entries inside the .iwyu file
 						foreach (var Info in IWYUFile!.Files)
@@ -530,8 +552,8 @@ namespace UnrealBuildTool
 					}
 					catch (Exception e)
 					{
-						Logger.LogError($"Failed to parse json {JsonFile.AbsolutePath}: {e.Message} - File will be deleted");
-						File.Delete(JsonFile.AbsolutePath);
+						Logger.LogError($"Failed to parse json {IWYUFilePath}: {e.Message} - File will be deleted");
+						File.Delete(IWYUFilePath);
 						ReadSuccess = false;
 						return;
 					}
