Fixing -Module=XXX -nolink on Modular builds. The mapping of modules->cpp files wasn't being updated when linking was disabled.

diff --git a/Engine/Source/Programs/UnrealBuildTool/Configuration/UEBuildBinary.cs b/Engine/Source/Programs/UnrealBuildTool/Configuration/UEBuildBinary.cs
--- a/Engine/Source/Programs/UnrealBuildTool/Configuration/UEBuildBinary.cs
+++ b/Engine/Source/Programs/UnrealBuildTool/Configuration/UEBuildBinary.cs
@@ -197,6 +197,12 @@ namespace UnrealBuildTool
 			// If linking is disabled, our build products are just the compiled object files
 			if (Target.bDisableLinking)
 			{
+				if (Target.LinkType == TargetLinkType.Modular)
+				{
+					// Make sure we record these files as outputs for this particular module.
+					Graph.SetOutputItemsForModule(PrimaryModule.Name, BinaryLinkEnvironment.InputFiles.ToArray());
+				}
+
 				return BinaryLinkEnvironment.InputFiles;
 			}
 
