- Fixed .c compiling with -nopch

diff --git a/Engine/Source/Programs/UnrealBuildTool/Platform/Android/AndroidToolChain.cs b/Engine/Source/Programs/UnrealBuildTool/Platform/Android/AndroidToolChain.cs
--- a/Engine/Source/Programs/UnrealBuildTool/Platform/Android/AndroidToolChain.cs
+++ b/Engine/Source/Programs/UnrealBuildTool/Platform/Android/AndroidToolChain.cs
@@ -699,8 +699,8 @@ namespace UnrealBuildTool
 					// Remove visibility settings for android native glue. Since it doesn't decorate with visibility attributes.
 					Arguments.RemoveAll(x => x.StartsWith("-fvisibility"));
 				}
-				// remove any forced include/PCH stuff - mostly for the force-added .c files in Launch as those will attempt to have the PCH used that was made with .cpp language
-				Arguments.RemoveAll(x => x.StartsWith("-include"));
+				// remove any PCH includes - mostly for the force-added .c files in Launch as those will attempt to have the PCH used that was made with .cpp language
+				Arguments.RemoveAll(x => x.StartsWith("-include-pch"));
 			}
 		}
 
