IWYU changes
- Use the first module from -Module= if -ModuleToUpdate= wasn't used on the command line
- Removing ToLower calls in favor of case insensitive comparisons

diff --git a/Engine/Source/Programs/UnrealBuildTool/Modes/IWYUMode.cs b/Engine/Source/Programs/UnrealBuildTool/Modes/IWYUMode.cs
--- a/Engine/Source/Programs/UnrealBuildTool/Modes/IWYUMode.cs
+++ b/Engine/Source/Programs/UnrealBuildTool/Modes/IWYUMode.cs
@@ -234,6 +234,21 @@ namespace UnrealBuildTool
 		[CommandLine("-DeprecateTag")]
 		public string KeepRemovedIncludesInDeprecated = "UE_ENABLE_INCLUDE_ORDER_DEPRECATED_IN_5_2";
 
+		private string? GetModuleToUpdateName(TargetDescriptor Descriptor)
+		{
+			if (!String.IsNullOrEmpty(ModuleToUpdate))
+			{
+				return ModuleToUpdate;
+			}
+
+			if (Descriptor.OnlyModuleNames.Count > 0)
+			{
+				return Descriptor.OnlyModuleNames.First();
+			}
+
+			return null;
+		}
+
 		/// <summary>
 		/// Execute the command
 		/// </summary>
@@ -264,6 +279,20 @@ namespace UnrealBuildTool
 				return 0;
 			}
 
+			string? ModuleToUpdateName = GetModuleToUpdateName(TargetDescriptors[0]);
+
+			if (!string.IsNullOrEmpty(ModuleToUpdateName))
+			{
+				foreach (string OnlyModuleName in TargetDescriptors[0].OnlyModuleNames)
+				{
+					if (string.Compare(OnlyModuleName, ModuleToUpdateName, StringComparison.OrdinalIgnoreCase) != 0)
+					{
+						Logger.LogError($"ModuleToUpdate '{ModuleToUpdateName}' was not in list of specified modules: {string.Join(", ", TargetDescriptors[0].OnlyModuleNames)}");
+						return -1;
+					}
+				}
+			}
+
 			// Calculate file paths filter to figure out which files that IWYU will run on
 			Logger.LogInformation($"Calculating file filter for IWYU...");
 			HashSet<string> ValidPaths = new();
@@ -273,33 +302,19 @@ namespace UnrealBuildTool
 				UEBuildTarget Target = UEBuildTarget.Create(TargetDescriptors[0], BuildConfiguration.bSkipRulesCompile, BuildConfiguration.bForceRulesCompile, BuildConfiguration.bUsePrecompiled, Logger);
 
 				UEBuildModule? UEModuleToUpdate = null;
-				string PathToUpdateLower = "";
 
 				// Turn provided module string into UEBuildModule reference
-				if (!String.IsNullOrEmpty(ModuleToUpdate))
+				if (!String.IsNullOrEmpty(ModuleToUpdateName))
 				{
-					UEModuleToUpdate = Target.GetModuleByName(ModuleToUpdate);
+					UEModuleToUpdate = Target.GetModuleByName(ModuleToUpdateName);
 					if (UEModuleToUpdate == null)
 					{
-						Logger.LogError($"Can't find module with name {ModuleToUpdate}");
+						Logger.LogError($"Can't find module with name {ModuleToUpdateName}");
 						return -1;
 					}
-
-					foreach (string OnlyModuleName in TargetDescriptors[0].OnlyModuleNames)
-					{
-						if (string.Compare(OnlyModuleName, ModuleToUpdate, StringComparison.OrdinalIgnoreCase) != 0)
-						{
-							Logger.LogError($"ModuleToUpdate '{ModuleToUpdate}' was not in list of specified modules: {string.Join(", ", TargetDescriptors[0].OnlyModuleNames)}");
-							return -1;
-						}
-					}
 				}
 				
-				if (!String.IsNullOrEmpty(PathToUpdate))
-				{
-					PathToUpdateLower = PathToUpdate.ToLower();
-				}
-				else if (UEModuleToUpdate == null)
+				if (String.IsNullOrEmpty(PathToUpdate) && UEModuleToUpdate == null)
 				{
 					Logger.LogError($"Need to provide -ModuleToUpdate <modulename> or -PathToUpdate <partofpath> to run IWYU.");
 					return -1;
@@ -355,9 +370,9 @@ namespace UnrealBuildTool
 							ShouldUpdate = DependsOnModule;
 						}
 						
-						if (!String.IsNullOrEmpty(PathToUpdateLower))
+						if (!String.IsNullOrEmpty(PathToUpdate))
 						{
-							if (ModuleDir.ToLower().Contains(PathToUpdateLower))
+							if (ModuleDir.Contains(PathToUpdate, StringComparison.OrdinalIgnoreCase))
 							{
 								if (UEModuleToUpdate == null)
 								{
