[UnrealBuildTool - IWYUMode]
* New iwyu binary that outputs "missing includes" as in includes that couldnt be seen from a specific file but that iwyu believes should be included (Definition.h for example is included through response file)
* Fixed so files that has Definitions.h as missing include and has no other includes add HAL/Platform.h. This handles situations with headers only having _API as a dependency
* Added -PathToUpdateFile which can be used to provide all the paths to update through a text file.
* Added so -PathToUpdate can take in multiple semi colon separated paths
* Added special handling for files that only has one include which is some custom define turning into an include.
* Added so if processing of any file fail we early out with an error

diff --git a/Engine/Source/Programs/UnrealBuildTool/Modes/IWYUMode.cs b/Engine/Source/Programs/UnrealBuildTool/Modes/IWYUMode.cs
--- a/Engine/Source/Programs/UnrealBuildTool/Modes/IWYUMode.cs
+++ b/Engine/Source/Programs/UnrealBuildTool/Modes/IWYUMode.cs
@@ -101,6 +101,18 @@ namespace UnrealBuildTool
 		public bool Present { get; set; }
 	}
 
+	/// <summary>
+	/// Representing an include that the code file needs that is missing in the include list
+	/// Note, in cpp files it is listing the includes that are in the matching h file
+	/// </summary>
+	struct IWYUMissingInclude
+	{
+		/// <summary>
+		/// Full path using forward paths, eg: d:/dev/folder/file.h
+		/// </summary>
+		public string Full { get; set; }
+	}
+
 	/// <summary>
 	/// Representing a file (header or source). Each .iwyu file has a list of these (most of the time only one)
 	/// </summary>
@@ -126,6 +138,11 @@ namespace UnrealBuildTool
 		/// </summary>
 		public List<IWYUIncludeEntry> IncludesSeenInFile { get; set; } = new List<IWYUIncludeEntry>();
 
+		/// <summary>
+		/// Includes that didnt end up in the Includes list but is needed by the code file
+		/// </summary>
+		public List<IWYUMissingInclude> MissingIncludes { get; set; } = new List<IWYUMissingInclude>();
+
 		/// <summary>
 		/// Transitive includes. This is all the includes that someone gets for free when including this file
 		/// Note, this is based on what iwyu believes should be in all includes. So it will not look at "seen includes"
@@ -193,10 +210,18 @@ namespace UnrealBuildTool
 		/// Which directory to run IWYU on. Will search for modules that has module directory matching this string
 		/// If no module is found in -PathToUpdate it handles this for individual files instead.
 		/// Note this can be combined with -ModuleToUpdate to double filter
+		/// PathToUpdate supports multiple paths using semi colon separation
 		/// </summary>
 		[CommandLine("-PathToUpdate")]
 		public string PathToUpdate = "";
 
+		/// <summary>
+		/// Same as PathToUpdate but provide the list of paths in a text file instead.
+		/// Add all the desired paths on a separate line
+		/// </summary>
+		[CommandLine("-PathToUpdateFile")]
+		public string PathToUpdateFile = "";
+
 		/// <summary>
 		/// Also update modules that are depending on the module we are updating
 		/// </summary>
@@ -293,6 +318,24 @@ namespace UnrealBuildTool
 				}
 			}
 
+			string[] PathToUpdateList = Array.Empty<string>();
+			if (!String.IsNullOrEmpty(PathToUpdate))
+			{
+				PathToUpdateList = PathToUpdate.Split(";");
+			}
+			else if (!String.IsNullOrEmpty(PathToUpdateFile))
+			{
+				PathToUpdateList = File.ReadAllLines(PathToUpdateFile);
+			}
+
+			if (PathToUpdateList.Length > 0)
+			{
+				for (int I = 0; I != PathToUpdateList.Length; ++I)
+				{
+					PathToUpdateList[I] = PathToUpdateList[I].Replace('\\', '/');
+				}
+			}
+			
 			// Calculate file paths filter to figure out which files that IWYU will run on
 			Logger.LogInformation($"Calculating file filter for IWYU...");
 			HashSet<string> ValidPaths = new();
@@ -314,7 +357,7 @@ namespace UnrealBuildTool
 					}
 				}
 				
-				if (String.IsNullOrEmpty(PathToUpdate) && UEModuleToUpdate == null)
+				if (PathToUpdateList.Length == 0 && UEModuleToUpdate == null)
 				{
 					Logger.LogError($"Need to provide -ModuleToUpdate <modulename> or -PathToUpdate <partofpath> to run IWYU.");
 					return -1;
@@ -370,9 +413,15 @@ namespace UnrealBuildTool
 							ShouldUpdate = DependsOnModule;
 						}
 						
-						if (!String.IsNullOrEmpty(PathToUpdate))
+						if (PathToUpdateList.Length > 0)
 						{
-							if (ModuleDir.Contains(PathToUpdate, StringComparison.OrdinalIgnoreCase))
+							bool Match = false;
+							for (int I=0; I!= PathToUpdateList.Length; ++I)
+							{
+								Match = Match || ModuleDir.Contains(PathToUpdateList[I], StringComparison.OrdinalIgnoreCase);
+							}
+
+							if (Match)
 							{
 								if (UEModuleToUpdate == null)
 								{
@@ -412,9 +461,12 @@ namespace UnrealBuildTool
 					}
 				}
 
-				if (ValidPaths.Count == 0 && !String.IsNullOrEmpty(PathToUpdate))
+				if (ValidPaths.Count == 0 && PathToUpdateList.Length > 0)
 				{
-					ValidPaths.Add(PathToUpdate.Replace('\\', '/'));
+					foreach (string Path in PathToUpdateList)
+					{
+						ValidPaths.Add(Path);
+					}
 					Logger.LogInformation($"Will update files matching {PathToUpdate}. Note, path is case sensitive");
 				}
 				else
@@ -428,7 +480,8 @@ namespace UnrealBuildTool
 			List<IWYUInfo> GeneratedCppInfos = new();
 			IWYUInfo? ObjectMacrosInfo = null;
 			IWYUInfo? ScriptMacrosInfo = null;
-			bool ReadSuccess = true;
+			IWYUInfo? PlatformInfo = null;
+			int ReadSuccess = 1;
 
 			TargetDescriptor Descriptor = TargetDescriptors.First();
 			using (ISourceFileWorkingSet WorkingSet = new EmptySourceFileWorkingSet())
@@ -480,7 +533,7 @@ namespace UnrealBuildTool
 					if (JsonContent == null)
 					{
 						Logger.LogError($"Failed to read file {IWYUFilePath}");
-						ReadSuccess = false;
+						Interlocked.Exchange(ref ReadSuccess, 0);
 						return;
 					}
 
@@ -549,8 +602,10 @@ namespace UnrealBuildTool
 								// Find ObjectMacros and ScriptMacros. These are always included in .generated.h files which will generate IWYUInfos for further down
 								if (ObjectMacrosInfo == null && Info.File.EndsWith("UObject/ObjectMacros.h"))
 									ObjectMacrosInfo = Info;
-								if (ScriptMacrosInfo == null && Info.File.EndsWith("UObject/ScriptMacros.h"))
+								else if (ScriptMacrosInfo == null && Info.File.EndsWith("UObject/ScriptMacros.h"))
 									ScriptMacrosInfo = Info;
+								else if (PlatformInfo == null && Info.File.EndsWith("HAL/Platform.h"))
+									PlatformInfo = Info;
 
 								// We need to add entries for .generated.h. They are not in the iwyu files
 								foreach (var Include in Info.Includes)
@@ -569,14 +624,14 @@ namespace UnrealBuildTool
 					{
 						Logger.LogError($"Failed to parse json {IWYUFilePath}: {e.Message} - File will be deleted");
 						File.Delete(IWYUFilePath);
-						ReadSuccess = false;
+						Interlocked.Exchange(ref ReadSuccess, 0);
 						return;
 					}
 				});
 			}
 
 			// Something went wrong parsing iwyu files.
-			if (!ReadSuccess)
+			if (ReadSuccess == 0)
 				return -1;
 
 			// Add all .generated.h files as entries in the lookup and explicitly add the includes they have which will never be removed
@@ -607,6 +662,7 @@ namespace UnrealBuildTool
 			Logger.LogInformation($"Resolving Includes...");
 			Parallel.ForEach(Infos.Values, Info =>
 			{
+				// If KeepAsIs we transfer all "seen includes" into the include list.
 				if (Info.Module != null && Info.Module.Rules.IWYUSupport == IWYUSupport.KeepAsIs)
 				{
 					foreach (IWYUIncludeEntry Entry in Info.IncludesSeenInFile)
@@ -619,6 +675,21 @@ namespace UnrealBuildTool
 					}
 				}
 
+				// Definitions.h is automatically added in the reponse file and iwyu sees it as not included
+				// If there are no includes but we depend on Definitions.h (which is missing) we add Platform.h because it is most likely a <module>_API entry in the file
+				if (Info.Includes.Count == 0)
+				{
+					bool IsUsingDefinitionsH = false;
+					foreach (var Missing in Info.MissingIncludes)
+					{
+						IsUsingDefinitionsH = IsUsingDefinitionsH || Missing.Full.Contains("Definitions.h");
+					}
+					if (IsUsingDefinitionsH)
+					{
+						Info.Includes.Add(new IWYUIncludeEntry() { Full = PlatformInfo!.File, Printable = "HAL/Platform.h", Resolved = PlatformInfo });
+					}
+				}
+
 				foreach (var Include in Info.Includes)
 				{
 					if (Include.Resolved == null)
@@ -714,6 +785,8 @@ namespace UnrealBuildTool
 			Logger.LogInformation($"Updating code files (in memory)...");
 
 			uint FilesParseCount = 0;
+			int ProcessSuccess = 1;
+
 			Parallel.ForEach(Infos.Values, Info =>
 			{
 				if (Info.Source == null) // .generated.h is also in this list, ignore them
@@ -845,6 +918,7 @@ namespace UnrealBuildTool
 
 
 				bool ForceKeepScope = false;
+				bool ErrorOnMoreIncludes = false;
 				int LineIndex = -1;
 				foreach (var L in ExistingLines)
 				{
@@ -879,6 +953,13 @@ namespace UnrealBuildTool
 						continue;
 					}
 
+					if (ErrorOnMoreIncludes)
+					{
+						Interlocked.Exchange(ref ProcessSuccess, 0);
+						Logger.LogError($"{Info.File} - Found special include using macro and did not expect more includes in this file");
+						return;
+					}
+
 					HasIncludes = true;
 
 					bool ForceKeep = false;
@@ -905,6 +986,12 @@ namespace UnrealBuildTool
 						{
 							// TODO: These are includes made through defines. IWYU should probably report these in their original shape
 							// .. so a #include MY_SPECIAL_INCLUDE is actually reported as MY_SPECIAL_INCLUDE from IWYU instead of what the define expands to
+							// For now, let's assume that if there is one line in the file that is an include 
+							if (IncludesToAdd.Count == 1)
+							{
+								IncludesToAdd.Clear();
+								ErrorOnMoreIncludes = true;
+							}
 						}
 						continue;
 					}
@@ -1177,6 +1264,10 @@ namespace UnrealBuildTool
 					UpdatedFiles.Add(new(Info, NewLines));
 			});
 
+			// Something went wrong processing code files.
+			if (ProcessSuccess == 0)
+				return -1;
+
 			Logger.LogInformation($"Parsed {FilesParseCount} and updated {UpdatedFiles.Count} files (Found {OutOfDateCount} .iwyu files out of date)");
 
 			// Wooohoo, all files are up-to-date
