Don't throw cryptic exception in UBT if there is an include with an unterminated string
#rnx

diff --git a/Engine/Source/Programs/UnrealBuildTool/System/SourceFileMetadataCache.cs b/Engine/Source/Programs/UnrealBuildTool/System/SourceFileMetadataCache.cs
--- a/Engine/Source/Programs/UnrealBuildTool/System/SourceFileMetadataCache.cs
+++ b/Engine/Source/Programs/UnrealBuildTool/System/SourceFileMetadataCache.cs
@@ -309,7 +309,10 @@ namespace UnrealBuildTool
 					if (FirstQuotation != -1 && (FirstAngleBracket == -1 || FirstQuotation < FirstAngleBracket)) // Handle #include <foo.h> // Some text with "
 					{
 						int SecondQuotation = Line.IndexOf('"', FirstQuotation + 1);
-						Includes.Add(Line.Substring(FirstQuotation + 1, SecondQuotation - FirstQuotation - 1));
+						if (SecondQuotation != -1)
+						{
+							Includes.Add(Line.Substring(FirstQuotation + 1, SecondQuotation - FirstQuotation - 1));
+						}
 					}
 				}
 
