When using -ModuleToUpdate, match the folder path exactly. This prevents cases like -ModuleToUpdate=RHI from matching RHICore.

diff --git a/Engine/Source/Programs/UnrealBuildTool/Modes/IWYUMode.cs b/Engine/Source/Programs/UnrealBuildTool/Modes/IWYUMode.cs
--- a/Engine/Source/Programs/UnrealBuildTool/Modes/IWYUMode.cs
+++ b/Engine/Source/Programs/UnrealBuildTool/Modes/IWYUMode.cs
@@ -274,6 +274,16 @@ namespace UnrealBuildTool
 			return null;
 		}
 
+		private string AdjustModulePathForMatching(string ModulePath)
+		{
+			string NewModulePath = ModulePath.Replace('\\', '/');
+			if (!NewModulePath.EndsWith('/'))
+			{
+				NewModulePath = NewModulePath + '/';
+			}
+			return NewModulePath;
+		}
+
 		/// <summary>
 		/// Execute the command
 		/// </summary>
@@ -450,12 +460,15 @@ namespace UnrealBuildTool
 						}
 
 						++ModulesToUpdateCount;
-						ValidPaths.Add(ModuleDir);
+
+						// When adding to ValidPaths, make sure the Module directory ends in a / so we match the exact folder later
+						ValidPaths.Add(AdjustModulePathForMatching(ModuleDir));
+
 						foreach (DirectoryReference AdditionalDir in Module.ModuleDirectories)
 						{
 							if (AdditionalDir != Module.ModuleDirectory)
 							{
-								ValidPaths.Add(AdditionalDir.FullName.Replace('\\', '/'));
+								ValidPaths.Add(AdjustModulePathForMatching(AdditionalDir.FullName));
 							}
 						}
 					}
@@ -793,16 +806,7 @@ namespace UnrealBuildTool
 					return;
 
 				// Filter out files
-				bool Included = false;
-				foreach (string ValidPath in ValidPaths)
-				{
-					if (Info.File.Contains(ValidPath))
-					{
-						Included = true;
-						break;
-					}
-				}
-				if (!Included)
+				if (!ValidPaths.Any(ValidPath => Info.File.Contains(ValidPath, StringComparison.OrdinalIgnoreCase)))
 				{
 					return;
 				}
