UnrealBuildTool: Stabilize file path for action confilct

#rnx

diff --git a/Engine/Source/Programs/UnrealBuildTool/System/ActionGraph.cs b/Engine/Source/Programs/UnrealBuildTool/System/ActionGraph.cs
--- a/Engine/Source/Programs/UnrealBuildTool/System/ActionGraph.cs
+++ b/Engine/Source/Programs/UnrealBuildTool/System/ActionGraph.cs
@@ -242,8 +242,8 @@ namespace UnrealBuildTool
 
 			string AJson = JsonSerializer.Serialize(A, Options);
 			string BJson = JsonSerializer.Serialize(B, Options);
-			string AJsonPath = Path.Combine(Path.GetTempPath(), "UnrealBuildTool", Path.ChangeExtension(Path.GetRandomFileName(), "json"));
-			string BJsonPath = Path.Combine(Path.GetTempPath(), "UnrealBuildTool", Path.ChangeExtension(Path.GetRandomFileName(), "json"));
+			string AJsonPath = Path.Combine(Path.GetTempPath(), "UnrealBuildTool", $"{A.StatusDescription}-{IoHash.Compute(Encoding.Default.GetBytes(AJson)).GetHashCode():X}") + ".json";
+			string BJsonPath = Path.Combine(Path.GetTempPath(), "UnrealBuildTool", $"{B.StatusDescription}-{IoHash.Compute(Encoding.Default.GetBytes(BJson)).GetHashCode():X}") + ".json";
 
 			Directory.CreateDirectory(Path.Combine(Path.GetTempPath(), "UnrealBuildTool"));
 			File.WriteAllText(AJsonPath, AJson);
