[UBT]
* Deprecated ModuleRules bEnforceIWYU. IWYUSupport property should be used instead

diff --git a/Engine/Source/Programs/UnrealBuildTool/Configuration/ModuleRules.cs b/Engine/Source/Programs/UnrealBuildTool/Configuration/ModuleRules.cs
--- a/Engine/Source/Programs/UnrealBuildTool/Configuration/ModuleRules.cs
+++ b/Engine/Source/Programs/UnrealBuildTool/Configuration/ModuleRules.cs
@@ -907,7 +907,7 @@ namespace UnrealBuildTool
 		/// Enforce "include what you use" rules when PCHUsage is set to ExplicitOrSharedPCH; warns when monolithic headers (Engine.h, UnrealEd.h, etc...) 
 		/// are used, and checks that source files include their matching header first.
 		/// </summary>
-		//[Obsolete("Deprecated in UE5.2 - Use IWYUSupport instead.")]
+		[Obsolete("Deprecated in UE5.2 - Use IWYUSupport instead.")]
 		public bool bEnforceIWYU { set { if (!value) IWYUSupport = IWYUSupport.None; } }
 
 		/// <summary>
