UnrealBuildTool: Store Action file lists in SortedSets rather than Lists, since these collections should never have duplicate entries and sorting them helps when checking for action conflicts

Requires implementing IComparable\IEquatable in FileItem (and I also implemented them in DirectoryItem for consistency)

#rnx

diff --git a/Engine/Source/Programs/Shared/EpicGames.Build/System/DirectoryItem.cs b/Engine/Source/Programs/Shared/EpicGames.Build/System/DirectoryItem.cs
--- a/Engine/Source/Programs/Shared/EpicGames.Build/System/DirectoryItem.cs
+++ b/Engine/Source/Programs/Shared/EpicGames.Build/System/DirectoryItem.cs
@@ -12,7 +12,7 @@ namespace UnrealBuildBase
 	/// <summary>
 	/// Stores the state of a directory. May or may not exist.
 	/// </summary>
-	public class DirectoryItem
+	public class DirectoryItem : IComparable<DirectoryItem>, IEquatable<DirectoryItem>
 	{
 		/// <summary>
 		/// Full path to the directory on disk
@@ -353,6 +353,73 @@ namespace UnrealBuildBase
 			}
 			File.WriteAllLines(OutFile, AllFiles);
 		}
+
+		#region IComparable, IEquatbale
+		public int CompareTo(DirectoryItem? other)
+		{
+			return FullName.CompareTo(other?.FullName);
+		}
+
+		public bool Equals(DirectoryItem? other)
+		{
+			return FullName.Equals(other?.FullName);
+		}
+
+		public override bool Equals(object? obj)
+		{
+			if (ReferenceEquals(this, obj))
+			{
+				return true;
+			}
+
+			if (ReferenceEquals(obj, null))
+			{
+				return false;
+			}
+
+			return Equals((DirectoryItem?)obj);
+		}
+
+		public override int GetHashCode()
+		{
+			return FullName.GetHashCode();
+		}
+
+		public static bool operator ==(DirectoryItem? left, DirectoryItem? right)
+		{
+			if (ReferenceEquals(left, null))
+			{
+				return ReferenceEquals(right, null);
+			}
+
+			return left.Equals(right);
+		}
+
+		public static bool operator !=(DirectoryItem? left, DirectoryItem? right)
+		{
+			return !(left == right);
+		}
+
+		public static bool operator <(DirectoryItem? left, DirectoryItem? right)
+		{
+			return ReferenceEquals(left, null) ? !ReferenceEquals(right, null) : left.CompareTo(right) < 0;
+		}
+
+		public static bool operator <=(DirectoryItem? left, DirectoryItem? right)
+		{
+			return ReferenceEquals(left, null) || left.CompareTo(right) <= 0;
+		}
+
+		public static bool operator >(DirectoryItem? left, DirectoryItem? right)
+		{
+			return !ReferenceEquals(left, null) && left.CompareTo(right) > 0;
+		}
+
+		public static bool operator >=(DirectoryItem? left, DirectoryItem? right)
+		{
+			return ReferenceEquals(left, null) ? ReferenceEquals(right, null) : left.CompareTo(right) >= 0;
+		}
+		#endregion
 	}
 
 	/// <summary>
diff --git a/Engine/Source/Programs/Shared/EpicGames.Build/System/FileItem.cs b/Engine/Source/Programs/Shared/EpicGames.Build/System/FileItem.cs
--- a/Engine/Source/Programs/Shared/EpicGames.Build/System/FileItem.cs
+++ b/Engine/Source/Programs/Shared/EpicGames.Build/System/FileItem.cs
@@ -14,7 +14,7 @@ namespace UnrealBuildBase
 	/// Represents a file on disk that is used as an input or output of a build action. FileItem instances are unique for a given path. Use FileItem.GetItemByFileReference 
 	/// to get the FileItem for a specific path.
 	/// </summary>
-	public class FileItem
+	public class FileItem : IComparable<FileItem>, IEquatable<FileItem>
 	{
 		/// <summary>
 		/// The directory containing this file
@@ -267,6 +267,73 @@ namespace UnrealBuildBase
 		{
 			return AbsolutePath;
 		}
+
+		#region IComparable, IEquatbale
+		public int CompareTo(FileItem? other)
+		{
+			return FullName.CompareTo(other?.FullName);
+		}
+
+		public bool Equals(FileItem? other)
+		{
+			return FullName.Equals(other?.FullName);
+		}
+
+		public override bool Equals(object? obj)
+		{
+			if (ReferenceEquals(this, obj))
+			{
+				return true;
+			}
+
+			if (ReferenceEquals(obj, null))
+			{
+				return false;
+			}
+
+			return Equals((FileItem?)obj);
+		}
+
+		public override int GetHashCode()
+		{
+			return FullName.GetHashCode();
+		}
+
+		public static bool operator ==(FileItem? left, FileItem? right)
+		{
+			if (ReferenceEquals(left, null))
+			{
+				return ReferenceEquals(right, null);
+			}
+
+			return left.Equals(right);
+		}
+
+		public static bool operator !=(FileItem? left, FileItem? right)
+		{
+			return !(left == right);
+		}
+
+		public static bool operator <(FileItem? left, FileItem? right)
+		{
+			return ReferenceEquals(left, null) ? !ReferenceEquals(right, null) : left.CompareTo(right) < 0;
+		}
+
+		public static bool operator <=(FileItem? left, FileItem? right)
+		{
+			return ReferenceEquals(left, null) || left.CompareTo(right) <= 0;
+		}
+
+		public static bool operator >(FileItem? left, FileItem? right)
+		{
+			return !ReferenceEquals(left, null) && left.CompareTo(right) > 0;
+		}
+
+		public static bool operator >=(FileItem? left, FileItem? right)
+		{
+			return ReferenceEquals(left, null) ? ReferenceEquals(right, null) : left.CompareTo(right) >= 0;
+		}
+		#endregion
 	}
 
 	/// <summary>
diff --git a/Engine/Source/Programs/Shared/EpicGames.Core/BinaryArchiveReader.cs b/Engine/Source/Programs/Shared/EpicGames.Core/BinaryArchiveReader.cs
--- a/Engine/Source/Programs/Shared/EpicGames.Core/BinaryArchiveReader.cs
+++ b/Engine/Source/Programs/Shared/EpicGames.Core/BinaryArchiveReader.cs
@@ -390,6 +390,55 @@ namespace EpicGames.Core
 			}
 		}
 
+		/// <summary>
+		/// Reads a sortedset of items
+		/// </summary>
+		/// <typeparam name="T">The element type for the sortedset</typeparam>
+		/// <param name="readElement">Delegate used to read a single element</param>
+		/// <returns>SortedSet of items</returns>
+		public SortedSet<T>? ReadSortedSet<T>(Func<T> readElement)
+		{
+			int count = ReadInt();
+			if (count < 0)
+			{
+				return null;
+			}
+			else
+			{
+				SortedSet<T> result = new SortedSet<T>();
+				for (int idx = 0; idx < count; idx++)
+				{
+					result.Add(readElement());
+				}
+				return result;
+			}
+		}
+
+		/// <summary>
+		/// Reads a sortedset of items
+		/// </summary>
+		/// <typeparam name="T">The element type for the sortedset</typeparam>
+		/// <param name="readElement">Delegate used to read a single element</param>
+		/// <param name="comparer">Comparison function for the set</param>
+		/// <returns>SortedSet of items</returns>
+		public SortedSet<T>? ReadSortedSet<T>(Func<T> readElement, IComparer<T> comparer)
+		{
+			int count = ReadInt();
+			if (count < 0)
+			{
+				return null;
+			}
+			else
+			{
+				SortedSet<T> result = new SortedSet<T>(comparer);
+				for (int idx = 0; idx < count; idx++)
+				{
+					result.Add(readElement());
+				}
+				return result;
+			}
+		}
+
 		/// <summary>
 		/// Reads a dictionary of items
 		/// </summary>
diff --git a/Engine/Source/Programs/Shared/EpicGames.Core/BinaryArchiveWriter.cs b/Engine/Source/Programs/Shared/EpicGames.Core/BinaryArchiveWriter.cs
--- a/Engine/Source/Programs/Shared/EpicGames.Core/BinaryArchiveWriter.cs
+++ b/Engine/Source/Programs/Shared/EpicGames.Core/BinaryArchiveWriter.cs
@@ -425,15 +425,37 @@ namespace EpicGames.Core
             }
         }
 
-        /// <summary>
-        /// Writes a dictionary of items
-        /// </summary>
-        /// <typeparam name="TK">Type of the dictionary key</typeparam>
-        /// <typeparam name="TV">Type of the dictionary value</typeparam>
-        /// <param name="dictionary">The dictionary to write</param>
-        /// <param name="writeKey">Delegate used to read a single key</param>
-        /// <param name="writeValue">Delegate used to read a single value</param>
-        public void WriteDictionary<TK, TV>(IReadOnlyDictionary<TK, TV> dictionary, Action<TK> writeKey, Action<TV> writeValue) where TK : notnull
+		/// <summary>
+		/// Writes a sortedset of items
+		/// </summary>
+		/// <typeparam name="T">The element type for the sortedset</typeparam>
+		/// <param name="set">The set to write</param>
+		/// <param name="writeElement">Delegate used to read a single element</param>
+		public void WriteSortedSet<T>(SortedSet<T> set, Action<T> writeElement)
+		{
+			if (set == null)
+			{
+				WriteInt(-1);
+			}
+			else
+			{
+				WriteInt(set.Count);
+				foreach (T element in set)
+				{
+					writeElement(element);
+				}
+			}
+		}
+
+		/// <summary>
+		/// Writes a dictionary of items
+		/// </summary>
+		/// <typeparam name="TK">Type of the dictionary key</typeparam>
+		/// <typeparam name="TV">Type of the dictionary value</typeparam>
+		/// <param name="dictionary">The dictionary to write</param>
+		/// <param name="writeKey">Delegate used to read a single key</param>
+		/// <param name="writeValue">Delegate used to read a single value</param>
+		public void WriteDictionary<TK, TV>(IReadOnlyDictionary<TK, TV> dictionary, Action<TK> writeKey, Action<TV> writeValue) where TK : notnull
         {
             if (dictionary == null)
             {
diff --git a/Engine/Source/Programs/UnrealBuildTool/Configuration/UEBuildTarget.cs b/Engine/Source/Programs/UnrealBuildTool/Configuration/UEBuildTarget.cs
--- a/Engine/Source/Programs/UnrealBuildTool/Configuration/UEBuildTarget.cs
+++ b/Engine/Source/Programs/UnrealBuildTool/Configuration/UEBuildTarget.cs
@@ -2350,8 +2350,8 @@ namespace UnrealBuildTool
 				WriteMetadataAction.StatusDescription = StatusDescription;
 				WriteMetadataAction.bCanExecuteRemotely = false;
 				WriteMetadataAction.bUseActionHistory = false; // Different files for each target; do not want to invalidate based on this.
-				WriteMetadataAction.ProducedItems.AddRange(ProducedItems.Select(x => FileItem.GetItemByFileReference(x)));
-				WriteMetadataAction.PrerequisiteItems.AddRange(PrerequisiteItems);
+				WriteMetadataAction.ProducedItems.UnionWith(ProducedItems.Select(x => FileItem.GetItemByFileReference(x)));
+				WriteMetadataAction.PrerequisiteItems.UnionWith(PrerequisiteItems);
 
 				if (Info.Version == null && Info.VersionFile != null)
 				{
diff --git a/Engine/Source/Programs/UnrealBuildTool/Modes/BuildMode.cs b/Engine/Source/Programs/UnrealBuildTool/Modes/BuildMode.cs
--- a/Engine/Source/Programs/UnrealBuildTool/Modes/BuildMode.cs
+++ b/Engine/Source/Programs/UnrealBuildTool/Modes/BuildMode.cs
@@ -587,7 +587,7 @@ namespace UnrealBuildTool
 								if (PrerequisiteAction.ActionType != ActionType.GatherModuleDependencies)
 								{
 									Action NewAction = new Action(PrerequisiteAction.Inner);
-									NewAction.PrerequisiteItems.AddRange(CompiledModuleInterfaces);
+									NewAction.PrerequisiteItems.UnionWith(CompiledModuleInterfaces);
 									PrerequisiteAction.Inner = NewAction;
 								}
 							}
diff --git a/Engine/Source/Programs/UnrealBuildTool/Platform/Android/AndroidToolChain.cs b/Engine/Source/Programs/UnrealBuildTool/Platform/Android/AndroidToolChain.cs
--- a/Engine/Source/Programs/UnrealBuildTool/Platform/Android/AndroidToolChain.cs
+++ b/Engine/Source/Programs/UnrealBuildTool/Platform/Android/AndroidToolChain.cs
@@ -1337,7 +1337,7 @@ namespace UnrealBuildTool
 				CompileAction.DependencyListFile = DependencyListFile;
 				CompileAction.PrerequisiteItems.Add(DependencyListFile);
 
-				CompileAction.ProducedItems.AddRange(CompiledISPCObjFiles);
+				CompileAction.ProducedItems.UnionWith(CompiledISPCObjFiles);
 
 				CompileAction.CommandArguments = String.Join(" ", Arguments);
 
diff --git a/Engine/Source/Programs/UnrealBuildTool/Platform/IOS/IOSToolChain.cs b/Engine/Source/Programs/UnrealBuildTool/Platform/IOS/IOSToolChain.cs
--- a/Engine/Source/Programs/UnrealBuildTool/Platform/IOS/IOSToolChain.cs
+++ b/Engine/Source/Programs/UnrealBuildTool/Platform/IOS/IOSToolChain.cs
@@ -1,4 +1,4 @@
-﻿// Copyright Epic Games, Inc. All Rights Reserved.
+// Copyright Epic Games, Inc. All Rights Reserved.
 
 
 using System;
@@ -509,7 +509,7 @@ namespace UnrealBuildTool
 			foreach (FileItem SourceFile in InputFiles)
 			{
 				Action CompileAction = CompileCPPFile(CompileEnvironment, SourceFile, OutputDir, ModuleName, Graph, GlobalArguments, Result);
-				CompileAction.PrerequisiteItems.AddRange(FrameworkTokenFiles);
+				CompileAction.PrerequisiteItems.UnionWith(FrameworkTokenFiles);
 			}
 			return Result;
 		}
@@ -804,7 +804,7 @@ namespace UnrealBuildTool
 			GenDebugAction.StatusDescription = GenDebugAction.CommandArguments;// string.Format("Generating debug info for {0}", Path.GetFileName(Executable.AbsolutePath));
 			GenDebugAction.bCanExecuteRemotely = false;
 
-			return GenDebugAction.ProducedItems; // (ProjectSettings.bGeneratedSYMBundle ? ZipOutputFile : OutputFile);
+			return GenDebugAction.ProducedItems.ToList(); // (ProjectSettings.bGeneratedSYMBundle ? ZipOutputFile : OutputFile);
 		}
 
 		/// <summary>
@@ -1229,7 +1229,7 @@ namespace UnrealBuildTool
 				StripAction.CommandPath = BuildHostPlatform.Current.Shell;
 				StripAction.CommandArguments = String.Format("-c \"\\\"{0}strip\\\" {1} \\\"{2}\\\" && touch \\\"{3}\\\"\"", Settings.Value.ToolchainDir, StripArguments, Executable.Location, StripCompleteFile);
 				StripAction.PrerequisiteItems.Add(Executable);
-				StripAction.PrerequisiteItems.AddRange(OutputFiles);
+				StripAction.PrerequisiteItems.UnionWith(OutputFiles);
 				StripAction.ProducedItems.Add(StripCompleteFile);
 				StripAction.StatusDescription = String.Format("Stripping symbols from {0}", Executable.AbsolutePath);
 				StripAction.bCanExecuteRemotely = false;
@@ -1323,9 +1323,9 @@ namespace UnrealBuildTool
 				Action PostBuildSyncAction = Graph.CreateRecursiveAction<IOSPostBuildSyncMode>(ActionType.CreateAppBundle, PostBuildSyncArguments);
 				PostBuildSyncAction.WorkingDirectory = Unreal.EngineSourceDirectory;
 				PostBuildSyncAction.PrerequisiteItems.Add(Executable);
-				PostBuildSyncAction.PrerequisiteItems.AddRange(OutputFiles);
+				PostBuildSyncAction.PrerequisiteItems.UnionWith(OutputFiles);
 				PostBuildSyncAction.ProducedItems.Add(FileItem.GetItemByFileReference(GetStagedExecutablePath(Executable.Location, Target.Name)));
-				PostBuildSyncAction.DeleteItems.AddRange(PostBuildSyncAction.ProducedItems);
+				PostBuildSyncAction.DeleteItems.UnionWith(PostBuildSyncAction.ProducedItems);
 				PostBuildSyncAction.StatusDescription = "Executing PostBuildSync";
 				PostBuildSyncAction.bCanExecuteRemotely = false;
 
diff --git a/Engine/Source/Programs/UnrealBuildTool/Platform/Linux/LinuxToolChain.cs b/Engine/Source/Programs/UnrealBuildTool/Platform/Linux/LinuxToolChain.cs
--- a/Engine/Source/Programs/UnrealBuildTool/Platform/Linux/LinuxToolChain.cs
+++ b/Engine/Source/Programs/UnrealBuildTool/Platform/Linux/LinuxToolChain.cs
@@ -709,7 +709,8 @@ namespace UnrealBuildTool
 					Arguments.Add("-pie");
 				}
 				else if (bSuppressPIE)
-				{ 					Arguments.Add("-Wl,-no-pie");
+				{
+					Arguments.Add("-Wl,-no-pie");
 				}
 			}
 
@@ -1462,7 +1463,7 @@ namespace UnrealBuildTool
 					RelinkAction.CommandDescription = "Relink";
 					RelinkAction.bCanExecuteRemotely = false;
 					RelinkAction.ProducedItems.Clear();
-					RelinkAction.PrerequisiteItems = new List<FileItem>(LinkAction.PrerequisiteItems);
+					RelinkAction.PrerequisiteItems = new SortedSet<FileItem>(LinkAction.PrerequisiteItems);
 					foreach (FileItem Dependency in EngineAndGameLibrariesFiles)
 					{
 						RelinkAction.PrerequisiteItems.Add(Dependency);
diff --git a/Engine/Source/Programs/UnrealBuildTool/Platform/Mac/MacToolChain.cs b/Engine/Source/Programs/UnrealBuildTool/Platform/Mac/MacToolChain.cs
--- a/Engine/Source/Programs/UnrealBuildTool/Platform/Mac/MacToolChain.cs
+++ b/Engine/Source/Programs/UnrealBuildTool/Platform/Mac/MacToolChain.cs
@@ -569,7 +569,7 @@ namespace UnrealBuildTool
 				OutputFiles.Add(OutputFileItem);
 
 				Action LipoAction = Graph.CreateAction(ActionType.Link);
-				LipoAction.PrerequisiteItems.AddRange(PerArchOutputFiles);
+				LipoAction.PrerequisiteItems.UnionWith(PerArchOutputFiles);
 				LipoAction.ProducedItems.Add(OutputFileItem);
 				LipoAction.WorkingDirectory = GetMacDevSrcRoot();
 				LipoAction.CommandPath = FileReference.Combine(Info.Clang.Directory, "lipo");
@@ -936,7 +936,7 @@ namespace UnrealBuildTool
 			LinkAction.ProducedItems.Add(OutputFile);
 
 			// Delete all items we produce
-			LinkAction.DeleteItems.AddRange(LinkAction.ProducedItems);
+			LinkAction.DeleteItems.UnionWith(LinkAction.ProducedItems);
 
 			return OutputFile;
 		}
diff --git a/Engine/Source/Programs/UnrealBuildTool/Platform/Windows/PVSToolChain.cs b/Engine/Source/Programs/UnrealBuildTool/Platform/Windows/PVSToolChain.cs
--- a/Engine/Source/Programs/UnrealBuildTool/Platform/Windows/PVSToolChain.cs
+++ b/Engine/Source/Programs/UnrealBuildTool/Platform/Windows/PVSToolChain.cs
@@ -812,7 +812,7 @@ namespace UnrealBuildTool
 
 				AnalyzeAction.PrerequisiteItems.Add(ConfigFileItem);
 				AnalyzeAction.PrerequisiteItems.Add(PreprocessedFileItem);
-				AnalyzeAction.PrerequisiteItems.AddRange(InputFiles); // Add the InputFiles as PrerequisiteItems so that in SingleFileCompile mode the PVSAnalyze step is not filtered out
+				AnalyzeAction.PrerequisiteItems.UnionWith(InputFiles); // Add the InputFiles as PrerequisiteItems so that in SingleFileCompile mode the PVSAnalyze step is not filtered out
 				AnalyzeAction.ProducedItems.Add(OutputFileItem);
 				AnalyzeAction.DeleteItems.Add(OutputFileItem); // PVS Studio will append by default, so need to delete produced items
 				AnalyzeAction.bCanExecuteRemotely = false;
@@ -854,11 +854,11 @@ namespace UnrealBuildTool
 			AnalyzeAction.CommandArguments = $"\"{Unreal.UnrealBuildToolDllPath}\" -Mode=PVSGather -Input=\"{InputFileListItem.Location}\" -Output=\"{OutputFile}\" ";
 			AnalyzeAction.WorkingDirectory = Unreal.EngineSourceDirectory;
 			AnalyzeAction.PrerequisiteItems.Add(InputFileListItem);
-			AnalyzeAction.PrerequisiteItems.AddRange(Makefile.OutputItems);
-			AnalyzeAction.PrerequisiteItems.AddRange(CompileSourceFiles);
+			AnalyzeAction.PrerequisiteItems.UnionWith(Makefile.OutputItems);
+			AnalyzeAction.PrerequisiteItems.UnionWith(CompileSourceFiles);
 			AnalyzeAction.ProducedItems.Add(FileItem.GetItemByFileReference(OutputFile));
 			AnalyzeAction.ProducedItems.Add(FileItem.GetItemByPath(OutputFile.FullName + "_does_not_exist")); // Force the gather step to always execute
-			AnalyzeAction.DeleteItems.AddRange(AnalyzeAction.ProducedItems);
+			AnalyzeAction.DeleteItems.UnionWith(AnalyzeAction.ProducedItems);
 
 			Makefile.OutputItems.AddRange(AnalyzeAction.ProducedItems);
 		}
diff --git a/Engine/Source/Programs/UnrealBuildTool/Platform/Windows/VCToolChain.cs b/Engine/Source/Programs/UnrealBuildTool/Platform/Windows/VCToolChain.cs
--- a/Engine/Source/Programs/UnrealBuildTool/Platform/Windows/VCToolChain.cs
+++ b/Engine/Source/Programs/UnrealBuildTool/Platform/Windows/VCToolChain.cs
@@ -1661,7 +1661,7 @@ namespace UnrealBuildTool
 						AggregateTimingInfoAction.StatusDescription = $"Aggregating {TimingJsonFiles.Count} Timing File(s)";
 						AggregateTimingInfoAction.bCanExecuteRemotely = false;
 						AggregateTimingInfoAction.bCanExecuteRemotelyWithSNDBS = false;
-						AggregateTimingInfoAction.PrerequisiteItems.AddRange(TimingJsonFiles);
+						AggregateTimingInfoAction.PrerequisiteItems.UnionWith(TimingJsonFiles);
 
 						FileItem AggregateOutputFile = FileItem.GetItemByFileReference(FileReference.Combine(Makefile.ProjectIntermediateDirectory, $"{Target.Name}.cta"));
 						AggregateTimingInfoAction.ProducedItems.Add(AggregateOutputFile);
@@ -1686,7 +1686,7 @@ namespace UnrealBuildTool
 						AggregateTimingInfoAction.StatusDescription = $"Aggregating {TimingJsonFiles.Count} Timing File(s)";
 						AggregateTimingInfoAction.bCanExecuteRemotely = false;
 						AggregateTimingInfoAction.bCanExecuteRemotelyWithSNDBS = false;
-						AggregateTimingInfoAction.PrerequisiteItems.AddRange(TimingJsonFiles);
+						AggregateTimingInfoAction.PrerequisiteItems.UnionWith(TimingJsonFiles);
 
 						AggregateTimingInfoAction.ProducedItems.Add(AggregateOutputFile);
 						AggregateTimingInfoAction.ProducedItems.Add(HeadersOutputFile);
@@ -1704,7 +1704,7 @@ namespace UnrealBuildTool
 						ArchiveTimingInfoAction.StatusDescription = $"Archiving {TimingJsonFiles.Count} Timing File(s)";
 						ArchiveTimingInfoAction.bCanExecuteRemotely = false;
 						ArchiveTimingInfoAction.bCanExecuteRemotelyWithSNDBS = false;
-						ArchiveTimingInfoAction.PrerequisiteItems.AddRange(TimingJsonFiles);
+						ArchiveTimingInfoAction.PrerequisiteItems.UnionWith(TimingJsonFiles);
 
 						ArchiveTimingInfoAction.ProducedItems.Add(ArchiveOutputFile);
 						Makefile.OutputItems.AddRange(ArchiveTimingInfoAction.ProducedItems);
@@ -1720,7 +1720,7 @@ namespace UnrealBuildTool
 							CompileScoreExtractorAction.StatusDescription = $"Extracting CompileScore";
 							CompileScoreExtractorAction.bCanExecuteRemotely = false;
 							CompileScoreExtractorAction.bCanExecuteRemotelyWithSNDBS = false;
-							CompileScoreExtractorAction.PrerequisiteItems.AddRange(TimingJsonFiles);
+							CompileScoreExtractorAction.PrerequisiteItems.UnionWith(TimingJsonFiles);
 							CompileScoreExtractorAction.CommandPath = ScoreDataExtractor;
 							CompileScoreExtractorAction.CommandArguments = $"-clang -verbosity 0 -timelinepack 1000000 -extract -i \"{NormalizeCommandLinePath(Makefile.ProjectIntermediateDirectory)}\" -o \"{NormalizeCommandLinePath(CompileScoreOutput)}\"";
 
@@ -1804,8 +1804,8 @@ namespace UnrealBuildTool
 				CompileAction.WorkingDirectory = Unreal.EngineSourceDirectory;
 				CompileAction.CommandPath = EnvVars.ResourceCompilerPath;
 				CompileAction.StatusDescription = Path.GetFileName(RCFile.AbsolutePath);
-				CompileAction.PrerequisiteItems.AddRange(CompileEnvironment.ForceIncludeFiles);
-				CompileAction.PrerequisiteItems.AddRange(CompileEnvironment.AdditionalPrerequisites);
+				CompileAction.PrerequisiteItems.UnionWith(CompileEnvironment.ForceIncludeFiles);
+				CompileAction.PrerequisiteItems.UnionWith(CompileEnvironment.AdditionalPrerequisites);
 
 				// Resource tool can run remotely if possible
 				CompileAction.bCanExecuteRemotely = true;
@@ -2301,20 +2301,20 @@ namespace UnrealBuildTool
 			}
 			LinkAction.CommandArguments = $"@\"{ResponseFileName}\"";
 			LinkAction.CommandVersion = EnvVars.ToolChainVersion.ToString();
-			LinkAction.ProducedItems.AddRange(ProducedItems);
-			LinkAction.PrerequisiteItems.AddRange(PrerequisiteItems);
+			LinkAction.ProducedItems.UnionWith(ProducedItems);
+			LinkAction.PrerequisiteItems.UnionWith(PrerequisiteItems);
 			LinkAction.StatusDescription = Path.GetFileName(OutputFile.AbsolutePath);
 
 			// VS 15.3+ does not touch lib files if they do not contain any modifications, but we need to ensure the timestamps are updated to avoid repeatedly building them.
 			if (bBuildImportLibraryOnly || (LinkEnvironment.bHasExports && !bIsBuildingLibraryOrImportLibrary))
 			{
-				LinkAction.DeleteItems.AddRange(LinkAction.ProducedItems.Where(x => x.Location.HasExtension(".lib") || x.Location.HasExtension(".exp")));
+				LinkAction.DeleteItems.UnionWith(LinkAction.ProducedItems.Where(x => x.Location.HasExtension(".lib") || x.Location.HasExtension(".exp")));
 			}
 
 			// Delete PDB files for all produced items, since incremental updates are slower than full ones.
 			if (!LinkEnvironment.bUseIncrementalLinking)
 			{
-				LinkAction.DeleteItems.AddRange(LinkAction.ProducedItems.Where(x => x.Location.HasExtension(".pdb")));
+				LinkAction.DeleteItems.UnionWith(LinkAction.ProducedItems.Where(x => x.Location.HasExtension(".pdb")));
 			}
 
 			// Tell the action that we're building an import library here and it should conditionally be
diff --git a/Engine/Source/Programs/UnrealBuildTool/System/Action.cs b/Engine/Source/Programs/UnrealBuildTool/System/Action.cs
--- a/Engine/Source/Programs/UnrealBuildTool/System/Action.cs
+++ b/Engine/Source/Programs/UnrealBuildTool/System/Action.cs
@@ -146,17 +146,17 @@ namespace UnrealBuildTool
 		/// <summary>
 		/// Every file this action depends on.  These files need to exist and be up to date in order for this action to even be considered
 		/// </summary>
-		public List<FileItem> PrerequisiteItems { get; set; } = new List<FileItem>();
+		public SortedSet<FileItem> PrerequisiteItems { get; set; } = new SortedSet<FileItem>();
 
 		/// <summary>
 		/// The files that this action produces after completing
 		/// </summary>
-		public List<FileItem> ProducedItems { get; set; } = new List<FileItem>();
+		public SortedSet<FileItem> ProducedItems { get; set; } = new SortedSet<FileItem>();
 
 		/// <summary>
 		/// Items that should be deleted before running this action
 		/// </summary>
-		public List<FileItem> DeleteItems { get; set; } = new List<FileItem>();
+		public SortedSet<FileItem> DeleteItems { get; set; } = new SortedSet<FileItem>();
 
 		/// <summary>
 		/// For C++ source files, specifies a dependency list file used to check changes to header files
@@ -240,9 +240,9 @@ namespace UnrealBuildTool
 		public Action(IExternalAction InOther)
 		{
 			ActionType = InOther.ActionType;
-			PrerequisiteItems = new List<FileItem>(InOther.PrerequisiteItems);
-			ProducedItems = new List<FileItem>(InOther.ProducedItems);
-			DeleteItems = new List<FileItem>(InOther.DeleteItems);
+			PrerequisiteItems = new SortedSet<FileItem>(InOther.PrerequisiteItems);
+			ProducedItems = new SortedSet<FileItem>(InOther.ProducedItems);
+			DeleteItems = new SortedSet<FileItem>(InOther.DeleteItems);
 			DependencyListFile = InOther.DependencyListFile;
 			WorkingDirectory = InOther.WorkingDirectory;
 			CommandPath = InOther.CommandPath;
@@ -272,9 +272,9 @@ namespace UnrealBuildTool
 			bIsGCCCompiler = Reader.ReadBool();
 			bShouldOutputStatusDescription = Reader.ReadBool();
 			bProducesImportLibrary = Reader.ReadBool();
-			PrerequisiteItems = Reader.ReadList(() => Reader.ReadFileItem())!;
-			ProducedItems = Reader.ReadList(() => Reader.ReadFileItem())!;
-			DeleteItems = Reader.ReadList(() => Reader.ReadFileItem())!;
+			PrerequisiteItems = Reader.ReadSortedSet(() => Reader.ReadFileItem())!;
+			ProducedItems = Reader.ReadSortedSet(() => Reader.ReadFileItem())!;
+			DeleteItems = Reader.ReadSortedSet(() => Reader.ReadFileItem())!;
 			DependencyListFile = Reader.ReadFileItem();
 			bUseActionHistory = Reader.ReadBool();
 		}
@@ -296,9 +296,9 @@ namespace UnrealBuildTool
 			Writer.WriteBool(bIsGCCCompiler);
 			Writer.WriteBool(bShouldOutputStatusDescription);
 			Writer.WriteBool(bProducesImportLibrary);
-			Writer.WriteList(PrerequisiteItems, Item => Writer.WriteFileItem(Item));
-			Writer.WriteList(ProducedItems, Item => Writer.WriteFileItem(Item));
-			Writer.WriteList(DeleteItems, Item => Writer.WriteFileItem(Item));
+			Writer.WriteSortedSet(PrerequisiteItems, Item => Writer.WriteFileItem(Item));
+			Writer.WriteSortedSet(ProducedItems, Item => Writer.WriteFileItem(Item));
+			Writer.WriteSortedSet(DeleteItems, Item => Writer.WriteFileItem(Item));
 			Writer.WriteFileItem(DependencyListFile);
 			Writer.WriteBool(bUseActionHistory);
 		}
@@ -380,19 +380,19 @@ namespace UnrealBuildTool
 			string[]? PrerequisiteItems;
 			if (Object.TryGetStringArrayField("PrerequisiteItems", out PrerequisiteItems))
 			{
-				Action.PrerequisiteItems.AddRange(PrerequisiteItems.Select(x => FileItem.GetItemByPath(x)));
+				Action.PrerequisiteItems.UnionWith(PrerequisiteItems.Select(x => FileItem.GetItemByPath(x)));
 			}
 
 			string[]? ProducedItems;
 			if (Object.TryGetStringArrayField("ProducedItems", out ProducedItems))
 			{
-				Action.ProducedItems.AddRange(ProducedItems.Select(x => FileItem.GetItemByPath(x)));
+				Action.ProducedItems.UnionWith(ProducedItems.Select(x => FileItem.GetItemByPath(x)));
 			}
 
 			string[]? DeleteItems;
 			if (Object.TryGetStringArrayField("DeleteItems", out DeleteItems))
 			{
-				Action.DeleteItems.AddRange(DeleteItems.Select(x => FileItem.GetItemByPath(x)));
+				Action.DeleteItems.UnionWith(DeleteItems.Select(x => FileItem.GetItemByPath(x)));
 			}
 
 			string? DependencyListFile;
diff --git a/Engine/Source/Programs/UnrealBuildTool/System/HotReload.cs b/Engine/Source/Programs/UnrealBuildTool/System/HotReload.cs
--- a/Engine/Source/Programs/UnrealBuildTool/System/HotReload.cs
+++ b/Engine/Source/Programs/UnrealBuildTool/System/HotReload.cs
@@ -1063,16 +1063,17 @@ namespace UnrealBuildTool
 					}
 
 					// Update this action's list of prerequisite items too
-					for (int ItemIndex = 0; ItemIndex < NewAction.PrerequisiteItems.Count; ++ItemIndex)
+					List<FileItem> UpdatePrerequisiteItems = new List<FileItem>(NewAction.PrerequisiteItems);
+					for (int ItemIndex = 0; ItemIndex < UpdatePrerequisiteItems.Count; ++ItemIndex)
 					{
-						FileItem OriginalPrerequisiteItem = NewAction.PrerequisiteItems[ItemIndex];
+						FileItem OriginalPrerequisiteItem = UpdatePrerequisiteItems[ItemIndex];
 						string NewPrerequisiteItemFilePath = ReplaceBaseFileName(OriginalPrerequisiteItem.AbsolutePath, OriginalFileNameWithoutExtension, NewFileNameWithoutExtension);
 
 						if (OriginalPrerequisiteItem.AbsolutePath != NewPrerequisiteItemFilePath)
 						{
 							// OK, the prerequisite item's file name changed so we'll update it to point to our new file
 							FileItem NewPrerequisiteItem = FileItem.GetItemByPath(NewPrerequisiteItemFilePath);
-							NewAction.PrerequisiteItems[ItemIndex] = NewPrerequisiteItem;
+							UpdatePrerequisiteItems[ItemIndex] = NewPrerequisiteItem;
 
 							// Keep track of it so we can fix up dependencies in a second pass afterwards
 							AffectedOriginalFileItemAndNewFileItemMap.Add(OriginalPrerequisiteItem, NewPrerequisiteItem);
@@ -1093,34 +1094,39 @@ namespace UnrealBuildTool
 							}
 						}
 					}
+					NewAction.PrerequisiteItems = new SortedSet<FileItem>(UpdatePrerequisiteItems);
 				}
 
 				// Update this action's list of produced items too
-				for (int ItemIndex = 0; ItemIndex < NewAction.ProducedItems.Count; ++ItemIndex)
+				List<FileItem> UpdateProducedItems = new List<FileItem>(NewAction.ProducedItems);
+				for (int ItemIndex = 0; ItemIndex < UpdateProducedItems.Count; ++ItemIndex)
 				{
-					FileItem OriginalProducedItem = NewAction.ProducedItems[ItemIndex];
+					FileItem OriginalProducedItem = UpdateProducedItems[ItemIndex];
 
 					string NewProducedItemFilePath = ReplaceBaseFileName(OriginalProducedItem.AbsolutePath, OriginalFileNameWithoutExtension, NewFileNameWithoutExtension);
 					if (OriginalProducedItem.AbsolutePath != NewProducedItemFilePath)
 					{
 						// OK, the produced item's file name changed so we'll update it to point to our new file
 						FileItem NewProducedItem = FileItem.GetItemByPath(NewProducedItemFilePath);
-						NewAction.ProducedItems[ItemIndex] = NewProducedItem;
+						UpdateProducedItems[ItemIndex] = NewProducedItem;
 
 						// Keep track of it so we can fix up dependencies in a second pass afterwards
 						AffectedOriginalFileItemAndNewFileItemMap.Add(OriginalProducedItem, NewProducedItem);
 					}
 				}
+				NewAction.ProducedItems = new SortedSet<FileItem>(UpdateProducedItems);
 
 				// Fix up the list of items to delete too
-				for(int Idx = 0; Idx < NewAction.DeleteItems.Count; Idx++)
+				List<FileItem> UpdateDeleteItems = new List<FileItem>(NewAction.DeleteItems);
+				for (int Idx = 0; Idx < UpdateDeleteItems.Count; Idx++)
 				{
 					FileItem? NewItem;
-					if(AffectedOriginalFileItemAndNewFileItemMap.TryGetValue(NewAction.DeleteItems[Idx], out NewItem))
+					if(AffectedOriginalFileItemAndNewFileItemMap.TryGetValue(UpdateDeleteItems[Idx], out NewItem))
 					{
-						NewAction.DeleteItems[Idx] = NewItem;
+						UpdateDeleteItems[Idx] = NewItem;
 					}
 				}
+				NewAction.DeleteItems = new SortedSet<FileItem>(UpdateDeleteItems);
 
 				// The status description of the item has the file name, so we'll update it too
 				NewAction.StatusDescription = ReplaceBaseFileName(Action.StatusDescription, OriginalFileNameWithoutExtension, NewFileNameWithoutExtension);
@@ -1141,17 +1147,19 @@ namespace UnrealBuildTool
 			foreach (LinkedAction Action in Actions)
 			{
 				Action NewAction = new Action(Action.Inner);
-				for (int ItemIndex = 0; ItemIndex < NewAction.PrerequisiteItems.Count; ++ItemIndex)
+				List<FileItem> UpdatePrerequisiteItems = new List<FileItem>(NewAction.PrerequisiteItems);
+				for (int ItemIndex = 0; ItemIndex < UpdatePrerequisiteItems.Count; ++ItemIndex)
 				{
-					FileItem OriginalFileItem = NewAction.PrerequisiteItems[ItemIndex];
+					FileItem OriginalFileItem = UpdatePrerequisiteItems[ItemIndex];
 
 					FileItem? NewFileItem;
 					if (AffectedOriginalFileItemAndNewFileItemMap.TryGetValue(OriginalFileItem, out NewFileItem))
 					{
 						// OK, looks like we need to replace this file item because we've renamed the file
-						NewAction.PrerequisiteItems[ItemIndex] = NewFileItem;
+						UpdatePrerequisiteItems[ItemIndex] = NewFileItem;
 					}
 				}
+				NewAction.PrerequisiteItems = new SortedSet<FileItem>(UpdatePrerequisiteItems);
 				Action.Inner = NewAction;
 			}
 
@@ -1274,7 +1282,7 @@ namespace UnrealBuildTool
 
 						Action NewAction = new Action(Action.Inner);
 
-						NewAction.PrerequisiteItems.RemoveAll(x => x.Location == TargetInfoFile);
+						NewAction.PrerequisiteItems.RemoveWhere(x => x.Location == TargetInfoFile);
 						NewAction.PrerequisiteItems.Add(FileItem.GetItemByFileReference(HotReloadTargetInfoFile));
 
 						NewAction.CommandArguments = Arguments.Substring(0, FileNameIdx) + HotReloadTargetInfoFile + Arguments.Substring(FileNameEndIdx);
diff --git a/Engine/Source/Programs/UnrealBuildTool/ToolChain/ClangToolChain.cs b/Engine/Source/Programs/UnrealBuildTool/ToolChain/ClangToolChain.cs
--- a/Engine/Source/Programs/UnrealBuildTool/ToolChain/ClangToolChain.cs
+++ b/Engine/Source/Programs/UnrealBuildTool/ToolChain/ClangToolChain.cs
@@ -243,7 +243,7 @@ namespace UnrealBuildTool
 					AggregateTimingInfoAction.StatusDescription = $"Aggregating {TimingJsonFiles.Count} Timing File(s)";
 					AggregateTimingInfoAction.bCanExecuteRemotely = false;
 					AggregateTimingInfoAction.bCanExecuteRemotelyWithSNDBS = false;
-					AggregateTimingInfoAction.PrerequisiteItems.AddRange(TimingJsonFiles);
+					AggregateTimingInfoAction.PrerequisiteItems.UnionWith(TimingJsonFiles);
 
 					AggregateTimingInfoAction.ProducedItems.Add(AggregateOutputFile);
 					AggregateTimingInfoAction.ProducedItems.Add(HeadersOutputFile);
@@ -261,7 +261,7 @@ namespace UnrealBuildTool
 					ArchiveTimingInfoAction.StatusDescription = $"Archiving {TimingJsonFiles.Count} Timing File(s)";
 					ArchiveTimingInfoAction.bCanExecuteRemotely = false;
 					ArchiveTimingInfoAction.bCanExecuteRemotelyWithSNDBS = false;
-					ArchiveTimingInfoAction.PrerequisiteItems.AddRange(TimingJsonFiles);
+					ArchiveTimingInfoAction.PrerequisiteItems.UnionWith(TimingJsonFiles);
 
 					ArchiveTimingInfoAction.ProducedItems.Add(ArchiveOutputFile);
 					Makefile.OutputItems.AddRange(ArchiveTimingInfoAction.ProducedItems);
@@ -277,7 +277,7 @@ namespace UnrealBuildTool
 						CompileScoreExtractorAction.StatusDescription = $"Extracting CompileScore";
 						CompileScoreExtractorAction.bCanExecuteRemotely = false;
 						CompileScoreExtractorAction.bCanExecuteRemotelyWithSNDBS = false;
-						CompileScoreExtractorAction.PrerequisiteItems.AddRange(TimingJsonFiles);
+						CompileScoreExtractorAction.PrerequisiteItems.UnionWith(TimingJsonFiles);
 						CompileScoreExtractorAction.CommandPath = ScoreDataExtractor;
 						CompileScoreExtractorAction.CommandArguments = $"-clang -verbosity 0 -timelinepack 1000000 -extract -i \"{NormalizeCommandLinePath(Makefile.ProjectIntermediateDirectory)}\" -o \"{NormalizeCommandLinePath(CompileScoreOutput)}\"";
 
@@ -885,14 +885,14 @@ namespace UnrealBuildTool
 		protected virtual void GetCompileArguments_FileType(CppCompileEnvironment CompileEnvironment, FileItem SourceFile, DirectoryReference OutputDir, List<string> Arguments, Action CompileAction, CPPOutput CompileResult)
 		{
 			// Add the C++ source file and its included files to the prerequisite item list.
-			CompileAction.PrerequisiteItems.AddRange(CompileEnvironment.ForceIncludeFiles);
-			CompileAction.PrerequisiteItems.AddRange(CompileEnvironment.AdditionalPrerequisites);
+			CompileAction.PrerequisiteItems.UnionWith(CompileEnvironment.ForceIncludeFiles);
+			CompileAction.PrerequisiteItems.UnionWith(CompileEnvironment.AdditionalPrerequisites);
 			CompileAction.PrerequisiteItems.Add(SourceFile);
 
 			List<FileItem>? InlinedFiles;
 			if (CompileEnvironment.FileInlineGenCPPMap.TryGetValue(SourceFile, out InlinedFiles))
 			{
-				CompileAction.PrerequisiteItems.AddRange(InlinedFiles);
+				CompileAction.PrerequisiteItems.UnionWith(InlinedFiles);
 			}
 
 			string Extension = Path.GetExtension(SourceFile.AbsolutePath).ToUpperInvariant();
@@ -1051,7 +1051,7 @@ namespace UnrealBuildTool
 			if (PreprocessDepends)
 			{
 				Action PrepassAction = Graph.CreateAction(ActionType.Compile);
-				PrepassAction.PrerequisiteItems.AddRange(CompileAction.PrerequisiteItems);
+				PrepassAction.PrerequisiteItems.UnionWith(CompileAction.PrerequisiteItems);
 				PrepassAction.PrerequisiteItems.Remove(CompilerResponseFileItem);
 				PrepassAction.CommandDescription = "Preprocess Depends";
 				PrepassAction.StatusDescription = CompileAction.StatusDescription;
@@ -1074,10 +1074,10 @@ namespace UnrealBuildTool
 				PreprocessFileArguments.Remove("-ftime-trace");
 				PreprocessFileArguments.Remove(GetOutputFileArgument(CompileAction.ProducedItems.First()));
 
-				PrepassAction.DeleteItems.AddRange(PrepassAction.ProducedItems);
+				PrepassAction.DeleteItems.UnionWith(PrepassAction.ProducedItems);
 
 				// Gets the target file so we can get the correct output path.
-				FileItem PreprocessTargetFile = PrepassAction.ProducedItems[0];
+				FileItem PreprocessTargetFile = PrepassAction.ProducedItems.First();
 
 				// Creates the path to the response file using the name of the output file and creates its contents.
 				FileReference PreprocessResponseFileName = new FileReference(PreprocessTargetFile.AbsolutePath + ".response");
diff --git a/Engine/Source/Programs/UnrealBuildTool/ToolChain/ISPCToolChain.cs b/Engine/Source/Programs/UnrealBuildTool/ToolChain/ISPCToolChain.cs
--- a/Engine/Source/Programs/UnrealBuildTool/ToolChain/ISPCToolChain.cs
+++ b/Engine/Source/Programs/UnrealBuildTool/ToolChain/ISPCToolChain.cs
@@ -682,7 +682,7 @@ namespace UnrealBuildTool
 				CompileAction.DependencyListFile = DependencyListFile;
 				CompileAction.PrerequisiteItems.Add(DependencyListFile);
 
-				CompileAction.ProducedItems.AddRange(CompiledISPCObjFiles);
+				CompileAction.ProducedItems.UnionWith(CompiledISPCObjFiles);
 
 				FileReference ResponseFileName = new FileReference(CompiledISPCObjFileNoISA.AbsolutePath + ".response");
 				FileItem ResponseFileItem = Graph.CreateIntermediateTextFile(ResponseFileName, Arguments.Select(x => Utils.ExpandVariables(x)));
diff --git a/Engine/Source/Programs/UnrealBuildTool/ToolChain/IWYUToolChain.cs b/Engine/Source/Programs/UnrealBuildTool/ToolChain/IWYUToolChain.cs
--- a/Engine/Source/Programs/UnrealBuildTool/ToolChain/IWYUToolChain.cs
+++ b/Engine/Source/Programs/UnrealBuildTool/ToolChain/IWYUToolChain.cs
@@ -3,6 +3,7 @@
 using System;
 using System.Collections.Generic;
 using System.IO;
+using System.Linq;
 using EpicGames.Core;
 using Microsoft.Extensions.Logging;
 using UnrealBuildBase;
@@ -113,7 +114,7 @@ namespace UnrealBuildTool
 
 				Action Action = CompileCPPFile(CompileEnvironment, SourceFile, OutputDir, ModuleName, Graph, GlobalArguments, new CPPOutput());
 
-				string CommandLineArgs = CommonCommandLineArgs + " -Xiwyu --write_json_path=\"" + Action.ProducedItems[0] + "\" ";
+				string CommandLineArgs = CommonCommandLineArgs + " -Xiwyu --write_json_path=\"" + Action.ProducedItems.First() + "\" ";
 				if (SourceFile.HasExtension(".cpp"))
 				{
 					List<FileItem>? InlinedFiles;
@@ -127,7 +128,7 @@ namespace UnrealBuildTool
 				}
 
 				Action.CommandArguments = CommandLineArgs + Action.CommandArguments;
-				IwyuFiles.Add(Action.ProducedItems[0]);
+				IwyuFiles.Add(Action.ProducedItems.First());
 			}
 
 			return new CPPOutput() { ObjectFiles = IwyuFiles };
