UnrealBuildTool: Fix Action conflict when writing out metadata by ensuring the .version file is written first, then write the .modules and .target files in separate steps that load that version file

#rnx

diff --git a/Engine/Source/Programs/UnrealBuildTool/Configuration/UEBuildTarget.cs b/Engine/Source/Programs/UnrealBuildTool/Configuration/UEBuildTarget.cs
--- a/Engine/Source/Programs/UnrealBuildTool/Configuration/UEBuildTarget.cs
+++ b/Engine/Source/Programs/UnrealBuildTool/Configuration/UEBuildTarget.cs
@@ -2171,9 +2171,13 @@ namespace UnrealBuildTool
 						}
 					}
 
-					WriteMetadataTargetInfo EngineInfo = new WriteMetadataTargetInfo(null, VersionFile, Receipt.Version, null, null, EngineFileToManifest, null);
-					FileReference EngineInfoFile = FileReference.Combine(ProjectIntermediateDirectory, "EngineMetadata.dat");
-					CreateWriteMetadataAction(MakefileBuilder, VersionFile.GetFileName(), EngineInfoFile, EngineInfo, EnginePrereqItems);
+					WriteMetadataTargetInfo VersionInfo = new WriteMetadataTargetInfo(null, VersionFile, Receipt.Version, null, null, null, null);
+					FileReference VersionInfoFile = FileReference.Combine(VersionFile.IsUnderDirectory(Unreal.EngineDirectory) ? EngineIntermediateDirectory : ProjectIntermediateDirectory, "VersionMetadata.dat");
+					CreateWriteMetadataAction(MakefileBuilder, VersionFile.GetFileName(), VersionInfoFile, VersionInfo, EnginePrereqItems);
+
+					WriteMetadataTargetInfo EngineInfo = new WriteMetadataTargetInfo(null, VersionFile, null, null, null, EngineFileToManifest, null);
+					FileReference EngineInfoFile = FileReference.Combine(EngineIntermediateDirectory, "EngineMetadata.dat");
+					CreateWriteMetadataAction(MakefileBuilder, $"{TargetName}.modules", EngineInfoFile, EngineInfo, EnginePrereqItems);
 
 					WriteMetadataTargetInfo TargetInfo = new WriteMetadataTargetInfo(ProjectFile, VersionFile, null, ReceiptFileName, Receipt, TargetFileToManifest, FileNameToLoadOrderManifest);
 					FileReference TargetInfoFile = FileReference.Combine(ProjectIntermediateDirectory, "TargetMetadata.dat");
diff --git a/Engine/Source/Programs/UnrealBuildTool/Modes/WriteMetadataMode.cs b/Engine/Source/Programs/UnrealBuildTool/Modes/WriteMetadataMode.cs
--- a/Engine/Source/Programs/UnrealBuildTool/Modes/WriteMetadataMode.cs
+++ b/Engine/Source/Programs/UnrealBuildTool/Modes/WriteMetadataMode.cs
@@ -11,6 +11,7 @@ using System.Threading.Tasks;
 using EpicGames.Core;
 using UnrealBuildBase;
 using Microsoft.Extensions.Logging;
+using System.Diagnostics;
 
 namespace UnrealBuildTool
 {
@@ -65,14 +66,14 @@ namespace UnrealBuildTool
 		/// <param name="Receipt"></param>
 		/// <param name="FileToManifest"></param>
 		/// <param name="FileToLoadOrderManifest"></param>
-		public WriteMetadataTargetInfo(FileReference? ProjectFile, FileReference? VersionFile, BuildVersion? Version, FileReference? ReceiptFile, TargetReceipt? Receipt, Dictionary<FileReference, ModuleManifest> FileToManifest, Dictionary<FileReference, LoadOrderManifest>? FileToLoadOrderManifest)
+		public WriteMetadataTargetInfo(FileReference? ProjectFile, FileReference? VersionFile, BuildVersion? Version, FileReference? ReceiptFile, TargetReceipt? Receipt, Dictionary<FileReference, ModuleManifest>? FileToManifest, Dictionary<FileReference, LoadOrderManifest>? FileToLoadOrderManifest)
 		{
 			this.ProjectFile = ProjectFile;
 			this.VersionFile = VersionFile;
 			this.Version = Version;
 			this.ReceiptFile = ReceiptFile;
 			this.Receipt = Receipt;
-			this.FileToManifest = FileToManifest;
+			this.FileToManifest = FileToManifest ?? new Dictionary<FileReference, ModuleManifest>();
 			this.FileToLoadOrderManifest = FileToLoadOrderManifest ?? new Dictionary<FileReference, LoadOrderManifest>();
 		}
 	}
