Always set an output path for UBT regardless of the configuration/platform it is built for

The "dotnet build" command run by BuildUBT.bat does not specify a value for the platform property. MSBuild can automatically select "x64" as a platform, in particular if the VS Dev Command prompt is used or the user has run any of the vcvars.bat as these will set a "Platform" environment variable. As there are no property groups for this platform in UnrealBuildTool.cs, the build result was not ending up where RunUBT.bat expected it to be. This was manifesting as either a stale Binaries/DotNet/UnrealBuildTool/UnrealBuildTool.dll, or no DLL at all.

#rnx

diff --git a/Engine/Source/Programs/UnrealBuildTool/UnrealBuildTool.csproj b/Engine/Source/Programs/UnrealBuildTool/UnrealBuildTool.csproj
--- a/Engine/Source/Programs/UnrealBuildTool/UnrealBuildTool.csproj
+++ b/Engine/Source/Programs/UnrealBuildTool/UnrealBuildTool.csproj
@@ -5,6 +5,7 @@
     <TargetFramework>net6.0</TargetFramework>
     <Configuration Condition=" '$(Configuration)' == '' ">Development</Configuration>
     <OutputType>Exe</OutputType>
+    <OutputPath>..\..\..\Binaries\DotNET\UnrealBuildTool</OutputPath>
     <AppendTargetFrameworkToOutputPath>false</AppendTargetFrameworkToOutputPath>
     <GenerateAssemblyInfo>false</GenerateAssemblyInfo>
     <GenerateTargetFrameworkAttribute>false</GenerateTargetFrameworkAttribute>
@@ -23,7 +24,6 @@
   </PropertyGroup>
   <PropertyGroup Condition=" '$(Configuration)|$(Platform)' == 'Development|AnyCPU' ">
     <Optimize>True</Optimize>
-    <OutputPath>..\..\..\Binaries\DotNET\UnrealBuildTool</OutputPath>
     <DocumentationFile>..\..\..\Binaries\DotNET\UnrealBuildTool\UnrealBuildTool.xml</DocumentationFile>
   </PropertyGroup>
   <PropertyGroup Condition=" '$(Configuration)|$(Platform)' == 'Release|AnyCPU' ">
@@ -33,17 +33,14 @@
     <ErrorReport>prompt</ErrorReport>
     <AllowUnsafeBlocks>false</AllowUnsafeBlocks>
     <Prefer32Bit>false</Prefer32Bit>
-	<OutputPath>..\..\..\Binaries\DotNET\UnrealBuildTool</OutputPath>
     <DocumentationFile>..\..\..\Binaries\DotNET\UnrealBuildTool\UnrealBuildTool.xml</DocumentationFile>
   </PropertyGroup>
   <PropertyGroup Condition=" '$(Configuration)|$(Platform)' == 'Debug|AnyCPU' ">
-    <OutputPath>..\..\..\Binaries\DotNET\UnrealBuildTool</OutputPath>
     <DocumentationFile>..\..\..\Binaries\DotNET\UnrealBuildTool\UnrealBuildTool.xml</DocumentationFile>
     <DebugSymbols>true</DebugSymbols>
     <DebugType>full</DebugType>
   </PropertyGroup>
   <PropertyGroup Condition=" '$(Configuration)|$(Platform)' == 'Analyze|AnyCPU' ">
-    <OutputPath>..\..\..\Binaries\DotNET\UnrealBuildTool</OutputPath>
     <DocumentationFile>..\..\..\Binaries\DotNET\UnrealBuildTool\UnrealBuildTool.xml</DocumentationFile>
     <RunAnalyzersDuringBuild>True</RunAnalyzersDuringBuild>
     <TreatWarningsAsErrors>false</TreatWarningsAsErrors>
