[Backout] - CL23756123
Original CL Desc
-----------------------------------------------------------------
fix tvos signature #rb josh.adams #jira UE-174364 #preflight none

diff --git a/Engine/Source/Programs/UnrealBuildTool/ProjectFiles/Xcode/XcodeProjectLegacy.cs b/Engine/Source/Programs/UnrealBuildTool/ProjectFiles/Xcode/XcodeProjectLegacy.cs
--- a/Engine/Source/Programs/UnrealBuildTool/ProjectFiles/Xcode/XcodeProjectLegacy.cs
+++ b/Engine/Source/Programs/UnrealBuildTool/ProjectFiles/Xcode/XcodeProjectLegacy.cs
@@ -1110,13 +1110,9 @@ namespace UnrealBuildTool.XcodeProjectLegacy
 		private void AppendNativeTargetBuildConfiguration(StringBuilder Content, XcodeBuildConfig Config, string ConfigGuid, FileReference? ProjectFile, ILogger Logger)
 		{
 			bool bMacOnly = true;
-			if (Config.ProjectTarget!.TargetRules != null)
+			if (Config.ProjectTarget!.TargetRules != null && XcodeProjectFileGenerator.ProjectFilePlatform.HasFlag(XcodeProjectFileGenerator.XcodeProjectFilePlatform.iOS))
 			{
-				if (XcodeProjectFileGenerator.XcodePlatforms.Contains(UnrealTargetPlatform.IOS) && Config.ProjectTarget.SupportedPlatforms.Contains(UnrealTargetPlatform.IOS))
-				{
-					bMacOnly = false;
-				}
-				if (XcodeProjectFileGenerator.XcodePlatforms.Contains(UnrealTargetPlatform.TVOS) && Config.ProjectTarget.SupportedPlatforms.Contains(UnrealTargetPlatform.TVOS))
+				if (Config.ProjectTarget.SupportedPlatforms.Contains(UnrealTargetPlatform.IOS))
 				{
 					bMacOnly = false;
 				}
@@ -1301,13 +1297,9 @@ namespace UnrealBuildTool.XcodeProjectLegacy
 		private void AppendLegacyTargetBuildConfiguration(StringBuilder Content, XcodeBuildConfig Config, string ConfigGuid, FileReference? ProjectFile)
 		{
 			bool bMacOnly = true;
-			if (Config.ProjectTarget!.TargetRules != null)
+			if (Config.ProjectTarget!.TargetRules != null && XcodeProjectFileGenerator.ProjectFilePlatform.HasFlag(XcodeProjectFileGenerator.XcodeProjectFilePlatform.iOS))
 			{
-				if (XcodeProjectFileGenerator.XcodePlatforms.Contains(UnrealTargetPlatform.IOS) && Config.ProjectTarget.SupportedPlatforms.Contains(UnrealTargetPlatform.IOS))
-				{
-					bMacOnly = false;
-				}
-				if (XcodeProjectFileGenerator.XcodePlatforms.Contains(UnrealTargetPlatform.TVOS) && Config.ProjectTarget.SupportedPlatforms.Contains(UnrealTargetPlatform.TVOS))
+				if (Config.ProjectTarget.SupportedPlatforms.Contains(UnrealTargetPlatform.IOS))
 				{
 					bMacOnly = false;
 				}
