- Fixed an issue with modern xcode project generation with Programs that have uproject files

diff --git a/Engine/Source/Programs/UnrealBuildTool/ProjectFiles/Xcode/XcodeProject.cs b/Engine/Source/Programs/UnrealBuildTool/ProjectFiles/Xcode/XcodeProject.cs
--- a/Engine/Source/Programs/UnrealBuildTool/ProjectFiles/Xcode/XcodeProject.cs
+++ b/Engine/Source/Programs/UnrealBuildTool/ProjectFiles/Xcode/XcodeProject.cs
@@ -309,7 +309,7 @@ namespace UnrealBuildTool.XcodeProjectXcconfig
 			}
 
 			// now that we have a UProject file (or not), update the FileCollection RootDirectory to point to it
-			ProjectFile.FileCollection.RootDirectory = UProjectFileLocation == null ? (Unreal.EngineDirectory) : UProjectFileLocation.Directory;
+			ProjectFile.FileCollection.SetUProjectLocation(UProjectFileLocation);
 
 			return UProjectFileLocation;
 		}
@@ -2047,7 +2047,8 @@ namespace UnrealBuildTool.XcodeProjectXcconfig
 		public XcodeFileCollection(XcodeProjectFile ProjectFile)
 		{
 			ProjectDirectory = ProjectFile.ProjectFilePath.Directory;
-			RootDirectory = ProjectFile.UnrealData.UProjectFileLocation == null ? (Unreal.EngineDirectory) : ProjectFile.UnrealData.UProjectFileLocation.Directory;
+			RootDirectory = Unreal.EngineDirectory;
+			SetUProjectLocation(ProjectFile.UnrealData.UProjectFileLocation);
 		}
 
 		public XcodeFileCollection(XcodeFileCollection Other)
@@ -2069,6 +2070,11 @@ namespace UnrealBuildTool.XcodeProjectXcconfig
 
 		}
 
+		public void SetUProjectLocation(FileReference? UProjectLocation)
+		{
+			RootDirectory = UProjectLocation == null || UProjectLocation.ContainsName("Programs", 0) ? (Unreal.EngineDirectory) : UProjectLocation!.Directory;
+		}
+
 		public void ProcessFile(XcodeSourceFile File, bool bIsForBuild)
 		{
 			// remember all buildable files
