Remove C++ UHT as the UHT to be used to generate headers.  This does not remove the ability to build C++ UHT as a target.

diff --git a/Engine/Source/Programs/UnrealBuildTool/Configuration/BuildConfiguration.cs b/Engine/Source/Programs/UnrealBuildTool/Configuration/BuildConfiguration.cs
--- a/Engine/Source/Programs/UnrealBuildTool/Configuration/BuildConfiguration.cs
+++ b/Engine/Source/Programs/UnrealBuildTool/Configuration/BuildConfiguration.cs
@@ -120,8 +120,8 @@ namespace UnrealBuildTool
 		/// <summary>
 		/// If true, do not build UHT, assume it is already built.
 		/// </summary>
-		[CommandLine("-NoBuildUHT")]
-		[XmlConfigFile(Category = "UEBuildConfiguration")]
+		[Obsolete]
+		[XmlConfigFile(Category = "UEBuildConfiguration", Deprecated = true)]
 		public bool bDoNotBuildUHT = false;
 
 		/// <summary>
@@ -141,27 +141,24 @@ namespace UnrealBuildTool
 		/// <summary>
 		/// If true, the Debug version of UnrealHeaderTool will be built and run instead of the Development version.
 		/// </summary>
-		[XmlConfigFile(Category = "UEBuildConfiguration")]
+		[Obsolete]
+		[XmlConfigFile(Category = "UEBuildConfiguration", Deprecated = true)]
 		public bool bForceDebugUnrealHeaderTool = false;
 
 		/// <summary>
 		/// If true, use C# UHT internal to UBT
 		/// </summary>
-		[XmlConfigFile(Category = "UEBuildConfiguration")]
+		[Obsolete]
+		[XmlConfigFile(Category = "UEBuildConfiguration", Deprecated = true)]
 		public bool bUseBuiltInUnrealHeaderTool = true;
 
 		/// <summary>
 		/// If true, generate warnings when C++ UHT is used
 		/// </summary>
-		[XmlConfigFile(Category = "UEBuildConfiguration")]
+		[Obsolete]
+		[XmlConfigFile(Category = "UEBuildConfiguration", Deprecated = true)]
 		public bool bWarnOnCppUnrealHeaderTool = true;
 
-		/// <summary>
-		/// If true, force C++ version UHT to run as just a warning
-		/// </summary>
-		[CommandLine("-UseCppUHT")]
-		public bool bUseCppUHT = false;
-
 		/// <summary>
 		/// Whether to skip compiling rules assemblies and just assume they are valid
 		/// </summary>
diff --git a/Engine/Source/Programs/UnrealBuildTool/Configuration/UEBuildTarget.cs b/Engine/Source/Programs/UnrealBuildTool/Configuration/UEBuildTarget.cs
--- a/Engine/Source/Programs/UnrealBuildTool/Configuration/UEBuildTarget.cs
+++ b/Engine/Source/Programs/UnrealBuildTool/Configuration/UEBuildTarget.cs
@@ -1044,16 +1044,6 @@ namespace UnrealBuildTool
 		/// </summary>
 		public DirectoryReference EngineIntermediateDirectory;
 
-		/// <summary>
-		/// Identifies whether the project contains a script plugin. This will cause UHT to be rebuilt, even in installed builds.
-		/// </summary>
-		public bool bHasProjectScriptPlugin;
-
-		/// <summary>
-		/// If true, then a script plugin was found that didn't have a matching UBT plugin
-		/// </summary>
-		public bool bHasRequiredProjectScriptPlugin;
-
 		/// <summary>
 		/// All plugins which are built for this target
 		/// </summary>
@@ -1084,11 +1074,6 @@ namespace UnrealBuildTool
 		/// </summary>
 		public List<FileReference>? EnabledUhtPlugins;
 
-		/// <summary>
-		/// Collection of all required C++ UHT plugins without a C# version
-		/// </summary>
-		public List<string>? RequiredUhtCppPlugins;
-
 		/// <summary>
 		/// All application binaries; may include binaries not built by this target.
 		/// </summary>
@@ -1871,8 +1856,8 @@ namespace UnrealBuildTool
 			// Create the makefile
 			string ExternalMetadata = UEBuildPlatform.GetBuildPlatform(Platform).GetExternalBuildMetadata(ProjectFile);
 			TargetMakefile Makefile = new TargetMakefile(ExternalMetadata, Binaries[0].OutputFilePaths[0], ReceiptFileName, ProjectIntermediateDirectory, TargetType, 
-				Rules.ConfigValueTracker, bDeployAfterCompile, bHasProjectScriptPlugin, UbtPlugins?.ToArray(), EnabledUbtPlugins?.ToArray(), 
-				EnabledUhtPlugins?.ToArray(), RequiredUhtCppPlugins?.ToArray());
+				Rules.ConfigValueTracker, bDeployAfterCompile, UbtPlugins?.ToArray(), EnabledUbtPlugins?.ToArray(), 
+				EnabledUhtPlugins?.ToArray());
 			Makefile.IsTestTarget = Rules.IsTestTarget;
 			TargetMakefileBuilder MakefileBuilder = new TargetMakefileBuilder(Makefile, Logger);
 
@@ -3786,7 +3771,6 @@ namespace UnrealBuildTool
 				PluginInfo Info = Plugin.Info;
 
 				// Check to see if this plugin has an associated UBT plugin
-				bool bUbtPluginFound = false;
 				if (BuiltPlugins != null)
 				{
 					foreach ((FileReference ProjectFile, FileReference TargetAssembly) BuiltPlugin in BuiltPlugins)
@@ -3796,27 +3780,12 @@ namespace UnrealBuildTool
 							if (UhtTables.IsUhtPlugin(BuiltPlugin.TargetAssembly.FullName))
 							{
 								EnabledUhtPlugins!.Add(BuiltPlugin.TargetAssembly);
-								bUbtPluginFound = true;
 								break;
 							}
 						}
 					}
 				}
 
-				// If this plugin supports UHT, flag that we need to pass the project file to UHT when generating headers
-				if (Info.Descriptor.SupportedPrograms != null && Info.Descriptor.SupportedPrograms.Contains("UnrealHeaderTool"))
-				{
-					bHasProjectScriptPlugin = true;
-					if (!bUbtPluginFound)
-					{
-						if (RequiredUhtCppPlugins == null)
-						{
-							RequiredUhtCppPlugins = new();
-						}
-						RequiredUhtCppPlugins.Add(Plugin.Name);
-					}
-				}
-
 				if (EnabledUhtPlugins != null)
 				{
 					EnabledUhtPlugins.SortBy(P => P.FullName);
diff --git a/Engine/Source/Programs/UnrealBuildTool/Modes/CleanMode.cs b/Engine/Source/Programs/UnrealBuildTool/Modes/CleanMode.cs
--- a/Engine/Source/Programs/UnrealBuildTool/Modes/CleanMode.cs
+++ b/Engine/Source/Programs/UnrealBuildTool/Modes/CleanMode.cs
@@ -76,40 +76,6 @@ namespace UnrealBuildTool
 				Logger.LogInformation("~~~~~~~~~~~~~~~~~~~~~~~~~~~~~~~~~~~~~~~~~~~~~~~~~~~~~~~~~~~~~~~~~~~~~~~~~~~~~~~~~~~~~~~~~~~~");
 			}
 
-			// Also add implicit descriptors for cleaning UnrealBuildTool
-			if (!BuildConfiguration.bDoNotBuildUHT)
-			{
-				const string UnrealHeaderToolTarget = "UnrealHeaderTool";
-
-				// Get a list of project files to clean UHT for
-				List<FileReference?> ProjectFiles = new List<FileReference?>();
-				foreach (TargetDescriptor TargetDesc in TargetDescriptors)
-				{
-					if (TargetDesc.Name != UnrealHeaderToolTarget && !RemoteMac.HandlesTargetPlatform(TargetDesc.Platform))
-					{
-						if (ProjectFiles.Count == 0)
-						{
-							ProjectFiles.Add(null);
-						}
-						if (TargetDesc.ProjectFile != null && !ProjectFiles.Contains(TargetDesc.ProjectFile))
-						{
-							ProjectFiles.Add(TargetDesc.ProjectFile);
-						}
-					}
-				}
-
-				// Add descriptors for cleaning UHT with all these projects
-				if (ProjectFiles.Count > 0)
-				{
-					UnrealTargetConfiguration Configuration = BuildConfiguration.bForceDebugUnrealHeaderTool ? UnrealTargetConfiguration.Debug : UnrealTargetConfiguration.Development;
-					string Architecture = UEBuildPlatform.GetBuildPlatform(BuildHostPlatform.Current.Platform).GetDefaultArchitecture(null);
-					foreach (FileReference? ProjectFile in ProjectFiles)
-					{
-						TargetDescriptors.Add(new TargetDescriptor(ProjectFile, UnrealHeaderToolTarget, BuildHostPlatform.Current.Platform, Configuration, Architecture, null));
-					}
-				}
-			}
-
 			// Output the list of targets that we're cleaning
 			Logger.LogInformation("Cleaning {TargetNames} binaries...", StringUtils.FormatList(TargetDescriptors.Select(x => x.Name).Distinct()));
 
diff --git a/Engine/Source/Programs/UnrealBuildTool/System/ExternalExecution.cs b/Engine/Source/Programs/UnrealBuildTool/System/ExternalExecution.cs
--- a/Engine/Source/Programs/UnrealBuildTool/System/ExternalExecution.cs
+++ b/Engine/Source/Programs/UnrealBuildTool/System/ExternalExecution.cs
@@ -647,16 +647,10 @@ namespace UnrealBuildTool
 		/// Gets the path to the receipt for UHT
 		/// </summary>
 		/// <returns>Path to the UHT receipt</returns>
-		public static FileReference GetHeaderToolReceiptFile(FileReference? ProjectFile, UnrealTargetConfiguration Configuration, bool bHasProjectScriptPlugin, FileReference[]? EnabledUhtPlugins)
+		public static FileReference GetHeaderToolReceiptFile(FileReference? ProjectFile, FileReference[]? EnabledUhtPlugins)
 		{
-			if ((bHasProjectScriptPlugin || (EnabledUhtPlugins != null && EnabledUhtPlugins.Length > 0)) && ProjectFile != null)
-			{
-				return TargetReceipt.GetDefaultPath(ProjectFile.Directory, "UnrealHeaderTool", BuildHostPlatform.Current.Platform, Configuration, "");
-			}
-			else
-			{
-				return TargetReceipt.GetDefaultPath(Unreal.EngineDirectory, "UnrealHeaderTool", BuildHostPlatform.Current.Platform, Configuration, "");
-			}
+			DirectoryReference BaseDir = EnabledUhtPlugins != null && EnabledUhtPlugins.Length > 0 && ProjectFile != null ? ProjectFile.Directory : Unreal.EngineDirectory;
+			return TargetReceipt.GetDefaultPath(BaseDir, "UnrealHeaderTool", BuildHostPlatform.Current.Platform, UnrealTargetConfiguration.Development, "");
 		}
 
 		/// <summary>
@@ -979,53 +973,15 @@ namespace UnrealBuildTool
 		/// </summary>
 		public static void ExecuteHeaderToolIfNecessary(BuildConfiguration BuildConfiguration, FileReference? ProjectFile, TargetMakefile Makefile, string TargetName, ISourceFileWorkingSet WorkingSet, ILogger Logger)
 		{
-			bool bIsBuildingUHT = TargetName.Equals("UnrealHeaderTool", StringComparison.InvariantCultureIgnoreCase);
 
-			if (bIsBuildingUHT)
+			// No need to run UHT on itself
+			if (TargetName.Equals("UnrealHeaderTool", StringComparison.InvariantCultureIgnoreCase))
 			{
 				Logger.LogInformation("DEPRECATED: C++ UHT is being built by request.  C++ UHT has been deprecated and will be removed in 5.2");
-			}
-			else if (BuildConfiguration.bUseCppUHT)
-			{
-				Logger.LogInformation("DEPRECATED: C++ UHT being used because '-UseCppUHT'has been specified.  C++ UHT has been deprecated and will be removed in 5.2");
-			}
-			else
-			{
-				if (!BuildConfiguration.bUseBuiltInUnrealHeaderTool)
-				{
-					if (BuildConfiguration.bWarnOnCppUnrealHeaderTool)
-					{
-						Logger.LogWarning("WARNING: C++ UHT being used because 'bUseBuiltInUnrealHeaderTool' is false.  C++ UHT has been deprecated and will be removed in 5.2");
-					}
-					else
-					{
-						Logger.LogInformation("DEPRECATED: C++ UHT being used because 'bUseBuiltInUnrealHeaderTool' is false.  C++ UHT has been deprecated and will be removed in 5.2");
-					}
-				}
-				if (Makefile.RequiredUhtCppPlugins?.Length > 0)
-				{
-					foreach (string PluginName in Makefile.RequiredUhtCppPlugins)
-					{
-						if (BuildConfiguration.bWarnOnCppUnrealHeaderTool)
-						{
-							Logger.LogWarning("WARNING: C++ UHT being used because '{PluginName}' does not have a C# version.  C++ UHT has been deprecated and will be removed in 5.2", PluginName);
-						}
-						else
-						{
-							Logger.LogInformation("DEPRECATED: C++ UHT being used because '{PluginName}' does not have a C# version.  C++ UHT has been deprecated and will be removed in 5.2", PluginName);
-						}
-					}
-				}
+				return;
 			}
 
-			if (!BuildConfiguration.bUseBuiltInUnrealHeaderTool || BuildConfiguration.bUseCppUHT || Makefile.RequiredUhtCppPlugins?.Length > 0 || bIsBuildingUHT)
-			{
-				ExecuteExternalHeaderToolIfNecessary(BuildConfiguration, ProjectFile, Makefile, TargetName, WorkingSet, Logger);
-			}
-			else
-			{
-				ExecuteInternalHeaderToolIfNecessary(BuildConfiguration, ProjectFile, Makefile, TargetName, WorkingSet, Logger);
-			}
+			ExecuteHeaderToolIfNecessaryInternal(BuildConfiguration, ProjectFile, Makefile, TargetName, WorkingSet, Logger);
 		}
 
 		/// <summary>
@@ -1122,220 +1078,7 @@ namespace UnrealBuildTool
 		/// Builds and runs the header tool and touches the header directories.
 		/// Performs any early outs if headers need no changes, given the UObject modules, tool path, game name, and configuration
 		/// </summary>
-		private static void ExecuteExternalHeaderToolIfNecessary(BuildConfiguration BuildConfiguration, FileReference? ProjectFile, TargetMakefile Makefile, string TargetName, ISourceFileWorkingSet WorkingSet, ILogger Logger)
-		{
-			if (ProgressWriter.bWriteMarkup)
-			{
-				Logger.LogInformation("@progress push 5%");
-			}
-			using (ProgressWriter Progress = new ProgressWriter("Generating code...", false, Logger))
-			{
-				// We never want to try to execute the header tool when we're already trying to build it!
-				bool bIsBuildingUHT = TargetName.Equals("UnrealHeaderTool", StringComparison.InvariantCultureIgnoreCase);
-
-				string RootLocalPath = Unreal.RootDirectory.FullName;
-
-				UnrealTargetConfiguration UHTConfig = BuildConfiguration.bForceDebugUnrealHeaderTool ? UnrealTargetConfiguration.Debug : UnrealTargetConfiguration.Development;
-
-				// Figure out the receipt path
-				FileReference HeaderToolReceipt = GetHeaderToolReceiptFile(ProjectFile, UHTConfig, Makefile.bHasProjectScriptPlugin, Makefile.EnabledUhtPlugins);
-
-				// check if UHT is out of date
-				DateTime HeaderToolTimestampUtc = DateTime.MaxValue;
-				bool bHaveHeaderTool = !bIsBuildingUHT && GetHeaderToolTimestampUtc(HeaderToolReceipt, Logger, out HeaderToolTimestampUtc);
-
-				// ensure the headers are up to date
-				bool bUHTNeedsToRun = false;
-				if(!bHaveHeaderTool)
-				{
-					bUHTNeedsToRun = true;
-				}
-				else if(BuildConfiguration.bForceHeaderGeneration)
-				{
-					bUHTNeedsToRun = true;
-				}
-				else if(AreGeneratedCodeFilesOutOfDate(BuildConfiguration, Makefile.UObjectModules, HeaderToolTimestampUtc, Logger))
-				{
-					bUHTNeedsToRun = true;
-				}
-
-				// Check we're not using a different version of UHT
-				FileReference ModuleInfoFileName = GetUHTModuleInfoFileName(Makefile, TargetName);
-				FileReference ToolInfoFile = ModuleInfoFileName.ChangeExtension(".uhtpath");
-				if(!bUHTNeedsToRun)
-				{
-					if(!FileReference.Exists(ToolInfoFile))
-					{
-						bUHTNeedsToRun = true;
-					}
-					else if(FileReference.ReadAllText(ToolInfoFile) != HeaderToolReceipt.FullName)
-					{
-						bUHTNeedsToRun = true;
-					}
-				}
-
-				// Get the file containing dependencies for the generated code
-				FileReference ExternalDependenciesFile = GetUHTDepsFileName(ModuleInfoFileName);
-				if (AreExternalDependenciesOutOfDate(ExternalDependenciesFile))
-				{
-					bUHTNeedsToRun = true;
-					bHaveHeaderTool = false; // Force UHT to build until dependency checking is fast enough to run all the time
-				}
-
-				if (!bIsBuildingUHT && bUHTNeedsToRun)
-				{
-					// Always build UnrealHeaderTool if header regeneration is required, unless we're running within an installed ecosystem or hot-reloading
-					if ((!Unreal.IsEngineInstalled() || Makefile.bHasProjectScriptPlugin) &&
-						!BuildConfiguration.bDoNotBuildUHT)
-					{
-						// If it is out of date or not there it will be built.
-						// If it is there and up to date, it will add 0.1 seconds or less to the build time.
-
-						// Which desktop platform do we need to compile UHT for?
-						UnrealTargetPlatform Platform = BuildHostPlatform.Current.Platform;
-
-						// NOTE: We force Development configuration for UHT so that it runs quickly, even when compiling debug, unless we say so explicitly
-						UnrealTargetConfiguration Configuration;
-						if (BuildConfiguration.bForceDebugUnrealHeaderTool)
-						{
-							Configuration = UnrealTargetConfiguration.Debug;
-						}
-						else
-						{
-							Configuration = UnrealTargetConfiguration.Development;
-						}
-
-						// Get the default architecture
-						string Architecture = UEBuildPlatform.GetBuildPlatform(Platform).GetDefaultArchitecture(null);
-
-						// Add UHT plugins to UBT command line as external plugins
-						FileReference? ScriptProjectFile = null;
-						if(Makefile.bHasProjectScriptPlugin && ProjectFile != null)
-						{
-							ScriptProjectFile = ProjectFile;
-						}
-
-						// Create the target descriptor
-						TargetDescriptor TargetDescriptor = new TargetDescriptor(ScriptProjectFile, "UnrealHeaderTool", Platform, Configuration, Architecture, null);
-						TargetDescriptor.bQuiet = true;
-
-						using (GlobalTracer.Instance.BuildSpan("Building UnrealHeaderTool").StartActive())
-						{
-							BuildMode.Build(new List<TargetDescriptor>{ TargetDescriptor }, BuildConfiguration, WorkingSet, BuildOptions.None, null, Logger);
-						}
-					}
-
-					Progress.Write(1, 3);
-
-					string ActualTargetName = String.IsNullOrEmpty(TargetName) ? "UE5" : TargetName;
-					Logger.LogInformation("Parsing headers for {ActualTargetName}", ActualTargetName);
-
-					FileReference HeaderToolPath = GetHeaderToolPath(HeaderToolReceipt);
-					if (!FileReference.Exists(HeaderToolPath))
-					{
-						throw new BuildException("Unable to generate headers because UnrealHeaderTool binary was not found ({0}).", HeaderToolPath);
-					}
-
-					WriteUHTManifest(Makefile, TargetName, ModuleInfoFileName, ExternalDependenciesFile);
-
-					string CmdLine = (ProjectFile != null) ? "\"" + ProjectFile.FullName + "\"" : TargetName;
-					CmdLine += " \"" + ModuleInfoFileName + "\" -LogCmds=\"loginit warning, logexit warning, logdatabase error\" -Unattended -WarningsAsErrors";
-
-					if (Makefile.UHTAdditionalArguments != null)
-					{
-						foreach (string UHTArgument in Makefile.UHTAdditionalArguments)
-						{
-							CmdLine += " " + UHTArgument;
-						}
-					}
-
-					if (Log.OutputFile != null)
-					{
-						string LogFileName = Log.OutputFile.GetFileNameWithoutExtension();
-						LogFileName = (LogFileName.StartsWith("UBT") ? "UHT" + LogFileName.Substring(3) : LogFileName + "_UHT") + ".txt";
-						LogFileName = FileReference.Combine(Log.OutputFile.Directory, LogFileName).ToString();
-
-						CmdLine += " -abslog=\"" + LogFileName + "\"";
-					}
-
-					if (Unreal.IsEngineInstalled())
-					{
-						CmdLine += " -installed";
-					}
-
-					if (BuildConfiguration.bFailIfGeneratedCodeChanges)
-					{
-						CmdLine += " -FailIfGeneratedCodeChanges";
-					}
-
-					Logger.LogInformation("  Running UnrealHeaderTool {CmdLine}", CmdLine);
-
-					Stopwatch s = new Stopwatch();
-					s.Start();
-					IScope Timer = GlobalTracer.Instance.BuildSpan("Executing UnrealHeaderTool").StartActive();
-					CompilationResult UHTResult = (CompilationResult)RunExternalNativeExecutable(ExternalExecution.GetHeaderToolPath(HeaderToolReceipt), CmdLine, Logger);
-					Timer.Span.Finish();
-					s.Stop();
-
-					if (UHTResult != CompilationResult.Succeeded)
-					{
-						// On Linux and Mac, the shell will return 128+signal number exit codes if UHT gets a signal (e.g. crashes or is interrupted)
-						if ((BuildHostPlatform.Current.Platform == UnrealTargetPlatform.Linux ||
-							BuildHostPlatform.Current.Platform == UnrealTargetPlatform.Mac) &&
-							(int)(UHTResult) >= 128
-							)
-						{
-							// SIGINT is 2, so 128 + SIGINT is 130
-							UHTResult = ((int)(UHTResult) == 130) ? CompilationResult.Canceled : CompilationResult.CrashOrAssert;
-						}
-
-						if (BuildHostPlatform.Current.Platform == UnrealTargetPlatform.Win64 &&
-							(int)(UHTResult) < 0)
-						{
-							Logger.LogError("UnrealHeaderTool failed with exit code 0x{Result:X} - check that Unreal Engine prerequisites are installed.", (int)UHTResult);
-						}
-
-						throw new CompilationResultException(UHTResult);
-					}
-
-					Logger.LogInformation("Reflection code generated for {Target} in {Time} seconds", ActualTargetName, s.Elapsed.TotalSeconds);
-
-					// Update the tool info file
-					DirectoryReference.CreateDirectory(ToolInfoFile.Directory);
-					FileReference.WriteAllText(ToolInfoFile, HeaderToolReceipt.FullName);
-
-					// Now that UHT has successfully finished generating code, we need to update all cached FileItems in case their last write time has changed.
-					// Otherwise UBT might not detect changes UHT made.
-					using (GlobalTracer.Instance.BuildSpan("ExternalExecution.ResetCachedHeaderInfo()").StartActive())
-					{
-						ResetCachedHeaderInfo(Makefile.UObjectModules);
-					}
-				}
-				else
-				{
-					Logger.LogDebug("Generated code is up to date.");
-				}
-
-				Progress.Write(2, 3);
-
-				using (GlobalTracer.Instance.BuildSpan("ExternalExecution.UpdateTimestamps()").StartActive())
-				{
-					UpdateTimestamps(Makefile.UObjectModules);
-				}
-
-				Progress.Write(3, 3);
-			}
-			if (ProgressWriter.bWriteMarkup)
-			{
-				Logger.LogInformation("@progress pop");
-			}
-		}
-
-		/// <summary>
-		/// Builds and runs the header tool and touches the header directories.
-		/// Performs any early outs if headers need no changes, given the UObject modules, tool path, game name, and configuration
-		/// </summary>
-		private static void ExecuteInternalHeaderToolIfNecessary(BuildConfiguration BuildConfiguration, FileReference? ProjectFile, 
+		private static void ExecuteHeaderToolIfNecessaryInternal(BuildConfiguration BuildConfiguration, FileReference? ProjectFile, 
 			TargetMakefile Makefile, string TargetName, ISourceFileWorkingSet WorkingSet, ILogger Logger)
 		{
 			if (ProgressWriter.bWriteMarkup)
@@ -1346,10 +1089,8 @@ namespace UnrealBuildTool
 			{
 				string RootLocalPath = Unreal.RootDirectory.FullName;
 
-				UnrealTargetConfiguration UHTConfig = BuildConfiguration.bForceDebugUnrealHeaderTool ? UnrealTargetConfiguration.Debug : UnrealTargetConfiguration.Development;
-
 				// Figure out the receipt path
-				FileReference HeaderToolReceipt = GetHeaderToolReceiptFile(ProjectFile, UHTConfig, Makefile.bHasProjectScriptPlugin, Makefile.EnabledUhtPlugins);
+				FileReference HeaderToolReceipt = GetHeaderToolReceiptFile(ProjectFile, Makefile.EnabledUhtPlugins);
 
 				// Get UHT assembly timestamp
 				DateTime CompositeTimestamp = new FileInfo(Assembly.GetExecutingAssembly().Location).LastWriteTimeUtc;
diff --git a/Engine/Source/Programs/UnrealBuildTool/System/TargetMakefile.cs b/Engine/Source/Programs/UnrealBuildTool/System/TargetMakefile.cs
--- a/Engine/Source/Programs/UnrealBuildTool/System/TargetMakefile.cs
+++ b/Engine/Source/Programs/UnrealBuildTool/System/TargetMakefile.cs
@@ -151,7 +151,7 @@ namespace UnrealBuildTool
 		/// <summary>
 		/// The version number to write
 		/// </summary>
-		public const int CurrentVersion = 33;
+		public const int CurrentVersion = 34;
 
 		/// <summary>
 		/// The time at which the makefile was created
@@ -208,16 +208,6 @@ namespace UnrealBuildTool
 		/// </summary>
 		public bool bDeployAfterCompile;
 
-		/// <summary>
-		/// Whether the project has a script plugin. UHT needs to know this to detect which manifest to use for checking out-of-datedness.
-		/// </summary>
-		public bool bHasProjectScriptPlugin;
-
-		/// <summary>
-		/// If true, there exists script plugins that don't have a matching UBT plugin
-		/// </summary>
-		public bool bHasRequiredProjectScriptPlugin;
-
 		/// <summary>
 		/// Collection of all located UBT C# plugins regardless of if they are in an enabled plugin (project file)
 		/// </summary>
@@ -233,11 +223,6 @@ namespace UnrealBuildTool
 		/// </summary>
 		public FileReference[]? EnabledUhtPlugins;
 
-		/// <summary>
-		/// Collection of UHT C++ plugins that have no C# version
-		/// </summary>
-		public string[]? RequiredUhtCppPlugins;
-
 		/// <summary>
 		/// The array of command-line arguments. The makefile will be invalidated whenever these change.
 		/// </summary>
@@ -356,14 +341,12 @@ namespace UnrealBuildTool
 		/// <param name="TargetType">The type of target</param>
 		/// <param name="ConfigValueTracker">Set of dependencies on config files</param>
 		/// <param name="bDeployAfterCompile">Whether to deploy the target after compiling</param>
-		/// <param name="bHasProjectScriptPlugin">Whether the target has a project script plugin</param>
 		/// <param name="UbtPlugins">Collection of UBT plugins</param>
 		/// <param name="EnabledUbtPlugins">Collection of UBT plugins</param>
 		/// <param name="EnabledUhtPlugins">Collection of UBT plugins for UHT</param>
-		/// <param name="RequiredUhtCppPlugins">Collection of required C++ UHT plugins that don't have a C# version</param>
 		public TargetMakefile(string ExternalMetadata, FileReference ExecutableFile, FileReference ReceiptFile, DirectoryReference ProjectIntermediateDirectory, 
-			TargetType TargetType, ConfigValueTracker ConfigValueTracker, bool bDeployAfterCompile, bool bHasProjectScriptPlugin,
-			FileReference[]? UbtPlugins, FileReference[]? EnabledUbtPlugins, FileReference[]? EnabledUhtPlugins, string[]? RequiredUhtCppPlugins)
+			TargetType TargetType, ConfigValueTracker ConfigValueTracker, bool bDeployAfterCompile, 
+			FileReference[]? UbtPlugins, FileReference[]? EnabledUbtPlugins, FileReference[]? EnabledUhtPlugins)
 		{
 			this.CreateTimeUtc = UnrealBuildTool.StartTimeUtc;
 			this.ModifiedTimeUtc = CreateTimeUtc;
@@ -375,11 +358,9 @@ namespace UnrealBuildTool
 			this.TargetType = TargetType;
 			this.ConfigValueTracker = ConfigValueTracker;
 			this.bDeployAfterCompile = bDeployAfterCompile;
-			this.bHasProjectScriptPlugin = bHasProjectScriptPlugin;
 			this.UbtPlugins = UbtPlugins;
 			this.EnabledUbtPlugins = EnabledUbtPlugins;
 			this.EnabledUhtPlugins = EnabledUhtPlugins;
-			this.RequiredUhtCppPlugins = RequiredUhtCppPlugins;
 			this.Actions = new List<IExternalAction>();
 			this.OutputItems = new List<FileItem>();
 			this.ModuleNameToOutputItems = new Dictionary<string, FileItem[]>(StringComparer.OrdinalIgnoreCase);
@@ -416,11 +397,9 @@ namespace UnrealBuildTool
 			IsTestTarget = Reader.ReadBool();
 			ConfigValueTracker = new ConfigValueTracker(Reader);
 			bDeployAfterCompile = Reader.ReadBool();
-			bHasProjectScriptPlugin = Reader.ReadBool();
 			UbtPlugins = Reader.ReadArray(() => Reader.ReadFileReference())!;
 			EnabledUbtPlugins = Reader.ReadArray(() => Reader.ReadFileReference())!;
 			EnabledUhtPlugins = Reader.ReadArray(() => Reader.ReadFileReference())!;
-			RequiredUhtCppPlugins = Reader.ReadArray(() => Reader.ReadString())!;
 			AdditionalArguments = Reader.ReadArray(() => Reader.ReadString())!;
 			UHTAdditionalArguments = Reader.ReadArray(() => Reader.ReadString())!;
 			PreBuildScripts = Reader.ReadArray(() => Reader.ReadFileReference())!;
@@ -461,11 +440,9 @@ namespace UnrealBuildTool
 			Writer.WriteBool(IsTestTarget);
 			ConfigValueTracker.Write(Writer);
 			Writer.WriteBool(bDeployAfterCompile);
-			Writer.WriteBool(bHasProjectScriptPlugin);
 			Writer.WriteArray(UbtPlugins, Item => Writer.WriteFileReference(Item));
 			Writer.WriteArray(EnabledUbtPlugins, Item => Writer.WriteFileReference(Item));
 			Writer.WriteArray(EnabledUhtPlugins, Item => Writer.WriteFileReference(Item));
-			Writer.WriteArray(RequiredUhtCppPlugins, Item => Writer.WriteString(Item));
 			Writer.WriteArray(AdditionalArguments, Item => Writer.WriteString(Item));
 			Writer.WriteArray(UHTAdditionalArguments, Item => Writer.WriteString(Item));
 			Writer.WriteArray(PreBuildScripts, Item => Writer.WriteFileReference(Item));
