UnrealBuildTool: Prevent action conflict for WriteMetadata actions

This is a hack to maintain parity with the BuildGraph executor which allowed exporting mulitple conflicing WriteMetadata actions.

#rnx

diff --git a/Engine/Source/Programs/UnrealBuildTool/Modes/BuildMode.cs b/Engine/Source/Programs/UnrealBuildTool/Modes/BuildMode.cs
--- a/Engine/Source/Programs/UnrealBuildTool/Modes/BuildMode.cs
+++ b/Engine/Source/Programs/UnrealBuildTool/Modes/BuildMode.cs
@@ -417,7 +417,14 @@ namespace UnrealBuildTool
 			if ((Options & BuildOptions.SkipBuild) == 0)
 			{
 				// Make sure that none of the actions conflict with any other (producing output files differently, etc...)
-				ActionGraph.CheckForConflicts(Makefiles.SelectMany(x => x.Actions), Logger);
+				using (GlobalTracer.Instance.BuildSpan("ActionGraph.CheckForConflicts").StartActive())
+				{
+					// TODO: Skipping conflicts for WriteMetadata is a hack to maintain parity with the BuildGraph executor which allowed
+					// exporting mulitple conflicing WriteMetadata actions.
+					// This is technically buggy behavior and should be properly fixed in UEBuildTarget.cs
+					IEnumerable<IExternalAction> CheckActions = Makefiles.SelectMany(x => x.Actions).Where(x => x.ActionType != ActionType.WriteMetadata);
+					ActionGraph.CheckForConflicts(CheckActions, Logger);
+				}
 
 				// Check we don't exceed the nominal max path length
 				using (GlobalTracer.Instance.BuildSpan("ActionGraph.CheckPathLengths").StartActive())
