- Fixed Run Target name in Modern xcode to not have the _Plaform suffix

diff --git a/Engine/Source/Programs/UnrealBuildTool/ProjectFiles/Xcode/XcodeProject.cs b/Engine/Source/Programs/UnrealBuildTool/ProjectFiles/Xcode/XcodeProject.cs
--- a/Engine/Source/Programs/UnrealBuildTool/ProjectFiles/Xcode/XcodeProject.cs
+++ b/Engine/Source/Programs/UnrealBuildTool/ProjectFiles/Xcode/XcodeProject.cs
@@ -1192,30 +1192,29 @@ namespace UnrealBuildTool.XcodeProjectXcconfig
 				case Type.Run_App:
 					ProductType = "com.apple.product-type.application";
 					TargetTypeName = "PBXNativeTarget";
-					ConfigName = "Run";
+					ConfigName = "_Run";
 					break;
 				case Type.Run_Tool:
 					ProductType = "com.apple.product-type.tool";
 					TargetTypeName = "PBXNativeTarget";
-					ConfigName = "Run";
+					ConfigName = "_Run";
 					break;
 				case Type.Build:
 					ProductType = "com.apple.product-type.library.static";
 					TargetTypeName = "PBXLegacyTarget";
-					ConfigName = "Build";
+					ConfigName = "_Build";
 					break;
 				case Type.Index:
 					ProductType = "com.apple.product-type.library.static";
 					TargetTypeName = "PBXNativeTarget";
-					ConfigName = "Index";
+					ConfigName = "_Index";
 					break;
 				default:
 					throw new BuildException($"Unhandled target type {Type}");
 			}
 
 			// set up names
-			ConfigName = (OverrideName == null) ? ConfigName : OverrideName;
-			Name = $"{UnrealData.ProductName}_{ConfigName}";
+			Name = UnrealData.ProductName + (OverrideName ?? ConfigName);
 		}
 
 		public void AddDependency(XcodeTarget Target, XcodeProject Project)
@@ -1283,7 +1282,8 @@ namespace UnrealBuildTool.XcodeProjectXcconfig
 		public UnrealTargetPlatform Platform;
 
 		public XcodeRunTarget(XcodeProject Project, UnrealTargetPlatform Platform, XcodeBuildTarget? BuildTarget)
-			: base(Project.UnrealData.bIsAppBundle ? XcodeTarget.Type.Run_App : XcodeTarget.Type.Run_Tool, Project.UnrealData, Platform.ToString())
+			: base(Project.UnrealData.bIsAppBundle ? XcodeTarget.Type.Run_App : XcodeTarget.Type.Run_Tool, Project.UnrealData, 
+				  XcodeProjectFileGenerator.PerPlatformMode == XcodePerPlatformMode.OneWorkspacePerPlatform ? "" : $"_{Platform}")
 		{
 			this.Platform = Platform;
 			BuildConfigList = new XcodeBuildConfigList(Platform, Name, Project.UnrealData.AllConfigs);
