UnrealBuildTool: Support msvc analyzer extensions

Enable by setting ModuleRules.bStaticAnalyzerExtensions = true.

It is recommended to use ModuleRules.StaticAnalyzerRulesets to filter the warnings that are reported as this will enable a large number of additional warnings.

Please see the following links for more details:
https://learn.microsoft.com/en-us/cpp/code-quality/using-the-cpp-core-guidelines-checkers
https://learn.microsoft.com/en-us/cpp/code-quality/using-rule-sets-to-specify-the-cpp-rules-to-run

diff --git a/Engine/Source/Programs/UnrealBuildTool/Configuration/ModuleRules.cs b/Engine/Source/Programs/UnrealBuildTool/Configuration/ModuleRules.cs
--- a/Engine/Source/Programs/UnrealBuildTool/Configuration/ModuleRules.cs
+++ b/Engine/Source/Programs/UnrealBuildTool/Configuration/ModuleRules.cs
@@ -794,6 +794,19 @@ namespace UnrealBuildTool
 		/// </summary>
 		public bool bDisableStaticAnalysis = false;
 
+		/// <summary>
+		/// Enable additional analyzer extension warnings using the EspXEngine plugin. This is only supported for MSVC.
+		/// See https://learn.microsoft.com/en-us/cpp/code-quality/using-the-cpp-core-guidelines-checkers
+		/// This will add a large number of warnings by default. It's recommended to use StaticAnalyzerRulesets if this is enabled.
+		/// </summary>
+		public bool bStaticAnalyzerExtensions = false;
+
+		/// <summary>
+		/// The static analyzer rulesets that should be used to filter warnings. This is only supported for MSVC.
+		/// See https://learn.microsoft.com/en-us/cpp/code-quality/using-rule-sets-to-specify-the-cpp-rules-to-run
+		/// </summary>
+		public HashSet<FileReference> StaticAnalyzerRulesets = new HashSet<FileReference>();
+
 		/// <summary>
 		/// The static analyzer checkers that should be enabled rather than the defaults. This is only supported for Clang.
 		/// See https://clang.llvm.org/docs/analyzer/checkers.html for a full list. Or run:
diff --git a/Engine/Source/Programs/UnrealBuildTool/Configuration/UEBuildModuleCPP.cs b/Engine/Source/Programs/UnrealBuildTool/Configuration/UEBuildModuleCPP.cs
--- a/Engine/Source/Programs/UnrealBuildTool/Configuration/UEBuildModuleCPP.cs
+++ b/Engine/Source/Programs/UnrealBuildTool/Configuration/UEBuildModuleCPP.cs
@@ -1596,6 +1596,8 @@ namespace UnrealBuildTool
 			Result.ShadowVariableWarningLevel = Rules.ShadowVariableWarningLevel;
 			Result.UnsafeTypeCastWarningLevel = Rules.UnsafeTypeCastWarningLevel;
 			Result.bDisableStaticAnalysis = Rules.bDisableStaticAnalysis;
+			Result.bStaticAnalyzerExtensions = Rules.bStaticAnalyzerExtensions;
+			Result.StaticAnalyzerRulesets = Rules.StaticAnalyzerRulesets;
 			Result.StaticAnalyzerCheckers = Rules.StaticAnalyzerCheckers;
 			Result.StaticAnalyzerDisabledCheckers = Rules.StaticAnalyzerDisabledCheckers;
 			Result.StaticAnalyzerAdditionalCheckers = Rules.StaticAnalyzerAdditionalCheckers;
diff --git a/Engine/Source/Programs/UnrealBuildTool/Platform/Windows/VCToolChain.cs b/Engine/Source/Programs/UnrealBuildTool/Platform/Windows/VCToolChain.cs
--- a/Engine/Source/Programs/UnrealBuildTool/Platform/Windows/VCToolChain.cs
+++ b/Engine/Source/Programs/UnrealBuildTool/Platform/Windows/VCToolChain.cs
@@ -337,6 +337,17 @@ namespace UnrealBuildTool
 			{
 				Arguments.Add("/analyze");
 
+				if (CompileEnvironment.bStaticAnalyzerExtensions)
+				{
+					FileReference EspxEngine = FileReference.Combine(EnvVars.CompilerPath.Directory, "EspXEngine.dll");
+					Arguments.Add($"/analyze:plugin\"{NormalizeCommandLinePath(EspxEngine)}\"");
+				}
+
+				foreach (FileReference Ruleset in CompileEnvironment.StaticAnalyzerRulesets)
+				{
+					Arguments.Add($"/analyze:ruleset\"{NormalizeCommandLinePath(Ruleset)}\"");
+				}
+
 				if (Target.WindowsPlatform.Compiler >= WindowsCompiler.VisualStudio2022)
 				{
 					// Ignore warnings in external headers
diff --git a/Engine/Source/Programs/UnrealBuildTool/System/CppCompileEnvironment.cs b/Engine/Source/Programs/UnrealBuildTool/System/CppCompileEnvironment.cs
--- a/Engine/Source/Programs/UnrealBuildTool/System/CppCompileEnvironment.cs
+++ b/Engine/Source/Programs/UnrealBuildTool/System/CppCompileEnvironment.cs
@@ -298,6 +298,19 @@ namespace UnrealBuildTool
 		/// </summary>
 		public bool bDisableStaticAnalysis = false;
 
+		/// <summary>
+		/// Enable additional analyzer extension warnings using the EspXEngine plugin. This is only supported for MSVC.
+		/// See https://learn.microsoft.com/en-us/cpp/code-quality/using-the-cpp-core-guidelines-checkers
+		/// This will add a large number of warnings by default. It's recommended to use StaticAnalyzerRulesets if this is enabled.
+		/// </summary>
+		public bool bStaticAnalyzerExtensions = false;
+
+		/// <summary>
+		/// The static analyzer rulesets that should be used to filter warnings. This is only supported for MSVC.
+		/// See https://learn.microsoft.com/en-us/cpp/code-quality/using-rule-sets-to-specify-the-cpp-rules-to-run
+		/// </summary>
+		public HashSet<FileReference> StaticAnalyzerRulesets = new HashSet<FileReference>();
+
 		/// <summary>
 		/// The static analyzer checkers that should be enabled rather than the defaults. This is only supported for Clang.
 		/// </summary>
