Forcing the Gameplay Debugger Menu to be compiled even in Shipping when bBuildDeveloperTools=true

diff --git a/Engine/Source/Programs/UnrealBuildTool/Configuration/ModuleRules.cs b/Engine/Source/Programs/UnrealBuildTool/Configuration/ModuleRules.cs
--- a/Engine/Source/Programs/UnrealBuildTool/Configuration/ModuleRules.cs
+++ b/Engine/Source/Programs/UnrealBuildTool/Configuration/ModuleRules.cs
@@ -1336,7 +1336,7 @@ namespace UnrealBuildTool
 			{
 				PublicDefinitions.Add("WITH_GAMEPLAY_DEBUGGER_CORE=1");
 				PublicDefinitions.Add("WITH_GAMEPLAY_DEBUGGER=" + (Target.bUseGameplayDebugger ? 1 : 0));
-				if (Target.Configuration != UnrealTargetConfiguration.Shipping)
+				if (Target.Configuration != UnrealTargetConfiguration.Shipping || Target.bBuildDeveloperTools)
 				{
 					PublicDefinitions.Add("WITH_GAMEPLAY_DEBUGGER_MENU=1");
 				}
diff --git a/Engine/Source/Runtime/GameplayDebugger/GameplayDebugger.Build.cs b/Engine/Source/Runtime/GameplayDebugger/GameplayDebugger.Build.cs
--- a/Engine/Source/Runtime/GameplayDebugger/GameplayDebugger.Build.cs
+++ b/Engine/Source/Runtime/GameplayDebugger/GameplayDebugger.Build.cs
@@ -55,7 +55,7 @@ namespace UnrealBuildTool.Rules
 			{
 				PublicDefinitions.Add("WITH_GAMEPLAY_DEBUGGER_CORE=1");
 				PublicDefinitions.Add("WITH_GAMEPLAY_DEBUGGER=" + (Target.bUseGameplayDebugger ? 1 : 0));
-				if (Target.Configuration != UnrealTargetConfiguration.Shipping)
+				if (Target.Configuration != UnrealTargetConfiguration.Shipping || Target.bBuildDeveloperTools)
 				{
 					PublicDefinitions.Add("WITH_GAMEPLAY_DEBUGGER_MENU=1");
 				}
