Remove check for supressing PIE and disable PIE unless being explicity enabled

diff --git a/Engine/Source/Programs/UnrealBuildTool/Platform/Linux/LinuxToolChain.cs b/Engine/Source/Programs/UnrealBuildTool/Platform/Linux/LinuxToolChain.cs
--- a/Engine/Source/Programs/UnrealBuildTool/Platform/Linux/LinuxToolChain.cs
+++ b/Engine/Source/Programs/UnrealBuildTool/Platform/Linux/LinuxToolChain.cs
@@ -41,9 +41,6 @@ namespace UnrealBuildTool
 		/** Flavor of the current build (target triplet)*/
 		string Architecture;
 
-		/** Whether the compiler is set up to produce PIE executables by default */
-		bool bSuppressPIE = false;
-
 		/** Pass --gdb-index option to linker to generate .gdb_index section. */
 		protected bool bGdbIndexSection = true;
 
@@ -254,49 +251,6 @@ namespace UnrealBuildTool
 			return Out.ToString();
 		}
 
-		/// <summary>
-		/// Checks default compiler settings
-		/// </summary>
-		private void CheckDefaultCompilerSettings()
-		{
-			using (Process Proc = new Process())
-			{
-				Proc.StartInfo.UseShellExecute = false;
-				Proc.StartInfo.CreateNoWindow = true;
-				Proc.StartInfo.RedirectStandardOutput = true;
-				Proc.StartInfo.RedirectStandardError = true;
-				Proc.StartInfo.RedirectStandardInput = true;
-
-				if (FileReference.Exists(Info.Clang))
-				{
-					Proc.StartInfo.FileName = Info.Clang.FullName;
-					Proc.StartInfo.Arguments = " -E -dM -";
-
-					Proc.Start();
-					Proc.StandardInput.Close();
-
-					for (; ; )
-					{
-						string? CompilerDefine = Proc.StandardOutput.ReadLine();
-						if (string.IsNullOrEmpty(CompilerDefine))
-						{
-							Proc.WaitForExit();
-							break;
-						}
-
-						if (CompilerDefine.Contains("__PIE__") || CompilerDefine.Contains("__pie__"))
-						{
-							bSuppressPIE = true;
-						}
-					}
-				}
-				else
-				{
-					// other compilers aren't implemented atm
-				}
-			}
-		}
-
 		/// <summary>
 		/// Architecture-specific compiler switches
 		/// </summary>
@@ -553,7 +507,7 @@ namespace UnrealBuildTool
 				{
 					Arguments.Add("-fPIE");
 				}
-				else if (bSuppressPIE)
+				else
 				{
 					Arguments.Add("-fno-PIE");
 				}
@@ -708,7 +662,7 @@ namespace UnrealBuildTool
 				{
 					Arguments.Add("-pie");
 				}
-				else if (bSuppressPIE)
+				else
 				{
 					Arguments.Add("-Wl,-no-pie");
 				}
@@ -851,10 +805,6 @@ namespace UnrealBuildTool
 			{
 				Logger.LogInformation("Using position independent executables (PIE)");
 			}
-			else if (bSuppressPIE)
-			{
-				Logger.LogInformation("Compiler is set up to generate position independent executables by default, but we're suppressing it.");
-			}
 
 			if (CompileEnvironment.bUseStackProtection)
 			{
