Fix missing removal of function

diff --git a/Engine/Source/Programs/UnrealBuildTool/Platform/Linux/LinuxToolChain.cs b/Engine/Source/Programs/UnrealBuildTool/Platform/Linux/LinuxToolChain.cs
--- a/Engine/Source/Programs/UnrealBuildTool/Platform/Linux/LinuxToolChain.cs
+++ b/Engine/Source/Programs/UnrealBuildTool/Platform/Linux/LinuxToolChain.cs
@@ -56,8 +56,6 @@ namespace UnrealBuildTool
 		public LinuxToolChain(string InArchitecture, LinuxPlatformSDK InSDK, ClangToolChainOptions InOptions, ILogger InLogger)
 			: this(UnrealTargetPlatform.Linux, InArchitecture, InSDK, InOptions, InLogger)
 		{
-			CheckDefaultCompilerSettings();
-
 			// prevent unknown clangs since the build is likely to fail on too old or too new compilers
 			if (CompilerVersionLessThan(13, 0, 0) || CompilerVersionGreaterOrEqual(14, 0, 0))
 			{
