PR #10020: Add -Project param to command line build params (Contributed by [REDACTED_PERSON])

diff --git a/Engine/Source/Programs/UnrealBuildTool/ProjectFiles/VisualStudio/VCProject.cs b/Engine/Source/Programs/UnrealBuildTool/ProjectFiles/VisualStudio/VCProject.cs
--- a/Engine/Source/Programs/UnrealBuildTool/ProjectFiles/VisualStudio/VCProject.cs
+++ b/Engine/Source/Programs/UnrealBuildTool/ProjectFiles/VisualStudio/VCProject.cs
@@ -1697,7 +1697,7 @@ namespace UnrealBuildTool
 				StringBuilder BuildArguments = new StringBuilder();
 				
 				BuildArguments.AppendFormat("{0} {1} {2}", TargetName, Platform.ToString(), Configuration.ToString());
-				if (bIsForeignProject)
+				if (UProjectPath.Length > 0)
 				{
 					BuildArguments.AppendFormat(" -Project={0}", UProjectPath);
 				}
@@ -1784,10 +1784,11 @@ namespace UnrealBuildTool
 
 			PlatformProjectGenerator? ProjGenerator = Combination.Platform != null ? PlatformProjectGenerators.GetPlatformProjectGenerator(Combination.Platform.Value, true) : null;
 
+			FileReference? UProjectPathNullable = Combination.ProjectTarget?.UnrealProjectFilePath;
 			string UProjectPath = "";
-			if (IsForeignProject)
+			if (UProjectPathNullable != null)
 			{
-				UProjectPath = String.Format("\"{0}\"", InsertPathVariables(Combination.ProjectTarget!.UnrealProjectFilePath!));
+				UProjectPath = String.Format("\"{0}\"", InsertPathVariables(UProjectPathNullable));
 			}
 
 			string ConditionString = "Condition=\"'$(Configuration)|$(Platform)'=='" + Combination.ProjectConfigurationAndPlatformName + "'\"";
