PR #10021: Better Additional Plugin Project Support (Contributed by [REDACTED_PERSON])

diff --git a/Engine/Source/Programs/UnrealBuildTool/ProjectFiles/ProjectFileGenerator.cs b/Engine/Source/Programs/UnrealBuildTool/ProjectFiles/ProjectFileGenerator.cs
--- a/Engine/Source/Programs/UnrealBuildTool/ProjectFiles/ProjectFileGenerator.cs
+++ b/Engine/Source/Programs/UnrealBuildTool/ProjectFiles/ProjectFileGenerator.cs
@@ -907,10 +907,11 @@ namespace UnrealBuildTool
 			List<ProjectFile> ModProjects = new List<ProjectFile>();
 			Dictionary<FileReference, ProjectFile> ProgramProjects = new Dictionary<FileReference, ProjectFile>();
 			Dictionary<RulesAssembly, DirectoryReference> RulesAssemblies = new Dictionary<RulesAssembly, DirectoryReference>();
+			Dictionary<ProjectFile, FileReference> ProjectFileToUProjectFile = new Dictionary<ProjectFile, FileReference>();
 			if(IncludeCppSource)
 			{
 				// Setup buildable projects for all targets
-				AddProjectsForAllTargets(PlatformProjectGenerators, AllGameProjects, AllTargetFiles, Arguments, ref EngineProject, GameProjects, ProgramProjects, RulesAssemblies, Logger);
+				AddProjectsForAllTargets(PlatformProjectGenerators, AllGameProjects, AllTargetFiles, Arguments, ref EngineProject, GameProjects, ProjectFileToUProjectFile, ProgramProjects, RulesAssemblies, Logger);
 
 				// Add projects for mods
 				AddProjectsForMods(GameProjects, ModProjects);
@@ -1078,7 +1079,7 @@ namespace UnrealBuildTool
 			}
 
 			// Setup "stub" projects for all modules
-			AddProjectsForAllModules(AllGameProjects, ProgramProjects, ModProjects, AllModuleFiles, AdditionalSearchPaths, bGatherThirdPartySource, Logger);
+			AddProjectsForAllModules(AllGameProjects, ProjectFileToUProjectFile, ProgramProjects, ModProjects, AllModuleFiles, AdditionalSearchPaths, bGatherThirdPartySource, Logger);
 
 			{
 				if (bIncludeDotNetPrograms)
@@ -2143,13 +2144,14 @@ namespace UnrealBuildTool
 		/// Finds all modules and code files, given a list of games to process
 		/// </summary>
 		/// <param name="AllGames">All game folders</param>
+		/// <param name="ProjectFileToUProjectFile">Map of generated Project File to the .uproject file for a given game.</param>
 		/// <param name="ProgramProjects">All program projects</param>
 		/// <param name="ModProjects">All mod projects</param>
 		/// <param name="AllModuleFiles">List of *.Build.cs files for all engine programs and games</param>
 		/// <param name="AdditionalSearchPaths"></param>
 		/// <param name="bGatherThirdPartySource">True to gather source code from third party projects too</param>
 		/// <param name="Logger">Logger for output</param>
-		protected void AddProjectsForAllModules(List<FileReference> AllGames, Dictionary<FileReference, ProjectFile> ProgramProjects, List<ProjectFile> ModProjects, List<FileReference> AllModuleFiles, Dictionary<FileReference, List<DirectoryReference>> AdditionalSearchPaths, bool bGatherThirdPartySource, ILogger Logger)
+		protected void AddProjectsForAllModules(List<FileReference> AllGames, Dictionary<ProjectFile, FileReference> ProjectFileToUProjectFile, Dictionary<FileReference, ProjectFile> ProgramProjects, List<ProjectFile> ModProjects, List<FileReference> AllModuleFiles, Dictionary<FileReference, List<DirectoryReference>> AdditionalSearchPaths, bool bGatherThirdPartySource, ILogger Logger)
 		{
 			HashSet<ProjectFile> ProjectsWithPlugins = new HashSet<ProjectFile>();
 			foreach (FileReference CurModuleFile in AllModuleFiles)
@@ -2220,6 +2222,22 @@ namespace UnrealBuildTool
 								}
 							}
 						}
+
+						FileReference? GameUProject = ProjectFileToUProjectFile.GetValueOrDefault(ProjectFile);
+						if (GameUProject != null && AdditionalSearchPaths.ContainsKey(GameUProject!))
+						{
+							foreach (DirectoryReference PluginFolder in AdditionalSearchPaths[GameUProject])
+							{
+								foreach (FileReference PluginFileName in PluginsBase.EnumeratePlugins(PluginFolder))
+								{
+									if (!ModProjects.Any(x => x.BaseDir == PluginFileName.Directory))
+									{
+										AddPluginFilesToProject(PluginFileName, BaseFolder, ProjectFile);
+									}
+								}
+							}
+						}
+
 						ProjectsWithPlugins.Add(ProjectFile);
 					}
 				}
@@ -2344,6 +2362,7 @@ namespace UnrealBuildTool
 		/// <param name="Arguments">The commandline arguments used</param>
 		/// <param name="EngineProject">The engine project we created</param>
 		/// <param name="GameProjects">Map of game folder name to all of the game projects we created</param>
+		/// <param name="ProjectFileToUProjectFile">Map of generated Project File to the .uproject file for a given game.</param>
 		/// <param name="ProgramProjects">Map of program names to all of the program projects we created</param>
 		/// <param name="RulesAssemblies">Map of RuleAssemblies to their base folders</param>
 		/// <param name="Logger">Logger for output</param>
@@ -2354,6 +2373,7 @@ namespace UnrealBuildTool
 			String[] Arguments,
 			ref ProjectFile? EngineProject,
 			List<ProjectFile> GameProjects,
+			Dictionary<ProjectFile, FileReference> ProjectFileToUProjectFile,
 			Dictionary<FileReference, ProjectFile> ProgramProjects,
 			Dictionary<RulesAssembly, DirectoryReference> RulesAssemblies,
 			ILogger Logger)
@@ -2552,6 +2572,7 @@ namespace UnrealBuildTool
 							if (FileReference.Exists(UProjectFilePath))
 							{
 								ProjectFile.AddFileToProject(UProjectFilePath, BaseFolder);
+								ProjectFileToUProjectFile.Add(ProjectFile, UProjectFilePath);
 							}
 							else
 							{
