PR #9932: Fix action graph copy action on non-Windows platforms (Contributed by [REDACTED_PERSON])

diff --git a/Engine/Source/Programs/UnrealBuildTool/System/ActionGraphBuilder.cs b/Engine/Source/Programs/UnrealBuildTool/System/ActionGraphBuilder.cs
--- a/Engine/Source/Programs/UnrealBuildTool/System/ActionGraphBuilder.cs
+++ b/Engine/Source/Programs/UnrealBuildTool/System/ActionGraphBuilder.cs
@@ -252,7 +252,7 @@ namespace UnrealBuildTool
 			}
 			else
 			{
-				CopyAction.CommandArguments = String.Format("-c 'cp -f \"\"{0}\"\" \"\"{1}\"'", SourceFile.AbsolutePath, TargetFile.AbsolutePath);
+				CopyAction.CommandArguments = String.Format("-c 'cp -f \"{0}\" \"{1}\"'", SourceFile.AbsolutePath, TargetFile.AbsolutePath);
 			}
 			CopyAction.WorkingDirectory = Unreal.EngineSourceDirectory;
 			CopyAction.PrerequisiteItems.Add(SourceFile);
