UnrealBuildTool: Ensure Module.bHasUObjects is set for IncludePathModules as well, so this property is consistent when building a Target with -AllModules and the generated UHT path is added to the include paths when compiling

#rnx

diff --git a/Engine/Source/Programs/UnrealBuildTool/System/ExternalExecution.cs b/Engine/Source/Programs/UnrealBuildTool/System/ExternalExecution.cs
--- a/Engine/Source/Programs/UnrealBuildTool/System/ExternalExecution.cs
+++ b/Engine/Source/Programs/UnrealBuildTool/System/ExternalExecution.cs
@@ -590,6 +590,38 @@ namespace UnrealBuildTool
 					}
 				}
 			}
+
+			// Set Module.bHasUObjects for any IncludePathModules not already processed.
+			// This is necessary to keep include paths consistent between targets built with -AllModules and without
+			HashSet<UEBuildModuleCPP> IncludePathModules = new HashSet<UEBuildModuleCPP>();
+			IncludePathModules.UnionWith(ModulesToGenerateHeadersFor.SelectMany(x => x.PrivateIncludePathModules ?? new()).OfType<UEBuildModuleCPP>().Where(x => !x.bHasUObjects));
+			IncludePathModules.UnionWith(ModulesToGenerateHeadersFor.SelectMany(x => x.PublicIncludePathModules ?? new()).OfType<UEBuildModuleCPP>().Where(x => !x.bHasUObjects));
+			IncludePathModules.ExceptWith(ModulesToGenerateHeadersFor);
+			if (IncludePathModules.Count > 0)
+			{
+				Dictionary<UEBuildModuleCPP, UHTModuleInfo> IncludePathInfo = new();
+				using (ThreadPoolWorkQueue Queue = new ThreadPoolWorkQueue())
+				{
+					ReadOnlyHashSet<string> ExcludedFolders = UEBuildPlatform.GetBuildPlatform(Platform).GetExcludedFolderNames();
+					foreach (var Module in IncludePathModules)
+					{
+						Queue.Enqueue(() =>
+						{
+							foreach (DirectoryItem ModuleDirectoryItem in Module.ModuleDirectories.Select(x => DirectoryItem.GetItemByDirectoryReference(x)))
+							{
+								List<FileItem> HeaderFiles = new();
+								FindHeaders(ModuleDirectoryItem, ExcludedFolders, HeaderFiles);
+
+								if (HeaderFiles.Any(x => MetadataCache.ContainsReflectionMarkup(x)))
+								{
+									Module.bHasUObjects = true;
+									break;
+								}
+							}
+						});
+					}
+				}
+			}
 		}
 
 		static void SetupUObjectModule(UHTModuleInfo ModuleInfo, ReadOnlyHashSet<string> ExcludedFolders, SourceFileMetadataCache MetadataCache, ThreadPoolWorkQueue Queue)
