Added the parent rule assemblies to the referencedassemblies list when compiling the rule assemblies dlls so the parent code can be used/referenced.

diff --git a/Engine/Source/Programs/UnrealBuildTool/System/RulesAssembly.cs b/Engine/Source/Programs/UnrealBuildTool/System/RulesAssembly.cs
--- a/Engine/Source/Programs/UnrealBuildTool/System/RulesAssembly.cs
+++ b/Engine/Source/Programs/UnrealBuildTool/System/RulesAssembly.cs
@@ -134,8 +134,22 @@ namespace UnrealBuildTool
 			// Compile the assembly
 			if (AssemblySourceFiles.Count > 0)
 			{
+				// Add the parent assemblies as references so they can be used
+				List<string>? ReferencedAssembies = null;
+				var CurrentParent = Parent;
+				while (CurrentParent != null && CurrentParent.CompiledAssembly != null)
+				{
+					if (ReferencedAssembies == null)
+					{
+						ReferencedAssembies = new List<string>();
+					}
+
+					ReferencedAssembies.Add(CurrentParent.CompiledAssembly.Location);
+					CurrentParent = CurrentParent.Parent;
+				}
+
 				PreprocessorDefines = GetPreprocessorDefinitions();
-				CompiledAssembly = DynamicCompilation.CompileAndLoadAssembly(AssemblyFileName, AssemblySourceFiles, Logger, PreprocessorDefines: PreprocessorDefines, DoNotCompile: bSkipCompile, ForceCompile: bForceCompile);
+				CompiledAssembly = DynamicCompilation.CompileAndLoadAssembly(AssemblyFileName, AssemblySourceFiles, Logger, ReferencedAssembies: ReferencedAssembies, PreprocessorDefines: PreprocessorDefines, DoNotCompile: bSkipCompile, ForceCompile: bForceCompile);
 			}
 
 			// Setup the module map
@@ -158,7 +172,7 @@ namespace UnrealBuildTool
 				}
 			}
 
-			// Write any deprecation warnings for methods overriden from a base with the [ObsoleteOverride] attribute. Unlike the [Obsolete] attribute, this ensures the message
+			// Write any deprecation warnings for methods overridden from a base with the [ObsoleteOverride] attribute. Unlike the [Obsolete] attribute, this ensures the message
 			// is given because the method is implemented, not because it's called.
 			if (CompiledAssembly != null)
 			{
