Move to new 15.0.1 LLVM Linux Toolchain

diff --git a/Engine/Build/BatchFiles/Linux/SetupToolchain.sh b/Engine/Build/BatchFiles/Linux/SetupToolchain.sh
--- a/Engine/Build/BatchFiles/Linux/SetupToolchain.sh
+++ b/Engine/Build/BatchFiles/Linux/SetupToolchain.sh
@@ -4,7 +4,7 @@
 # put ourselves into Engine directory (two up from location of this script)
 pushd "`dirname "$0"`/../../.." > /dev/null
 
-TOOLCHAIN_VERSION=v20_clang-13.0.1-centos7
+TOOLCHAIN_VERSION=v21_clang-15.0.1-centos7
 TOOLCHAIN_ARCHIVE=$TOOLCHAIN_VERSION.tar.gz
 
 TOOLCHAIN_URL=http://cdn.unrealengine.com/Toolchain_Linux/native-linux-$TOOLCHAIN_ARCHIVE
diff --git a/Engine/Source/Programs/UnrealBuildTool/Platform/Linux/LinuxPlatformSDK.Versions.cs b/Engine/Source/Programs/UnrealBuildTool/Platform/Linux/LinuxPlatformSDK.Versions.cs
--- a/Engine/Source/Programs/UnrealBuildTool/Platform/Linux/LinuxPlatformSDK.Versions.cs
+++ b/Engine/Source/Programs/UnrealBuildTool/Platform/Linux/LinuxPlatformSDK.Versions.cs
@@ -8,14 +8,14 @@ namespace UnrealBuildTool
 	{
 		public override string GetMainVersion()
 		{
-			return "v20_clang-13.0.1-centos7";
+			return "v21_clang-15.0.1-centos7";
 		}
 
 		protected override void GetValidVersionRange(out string MinVersion, out string MaxVersion)
 		{
 			// all that matters is the number after the v, according to TryConvertVersionToInt()
-			MinVersion = "v20_clang-13.0.1-centos7";
-			MaxVersion = "v20_clang-13.0.1-centos7";
+			MinVersion = "v21_clang-15.0.1-centos7";
+			MaxVersion = "v21_clang-15.0.1-centos7";
 		}
 
 		protected override void GetValidSoftwareVersionRange(out string? MinVersion, out string? MaxVersion)
diff --git a/Engine/Source/Programs/UnrealBuildTool/Platform/Linux/LinuxToolChain.cs b/Engine/Source/Programs/UnrealBuildTool/Platform/Linux/LinuxToolChain.cs
--- a/Engine/Source/Programs/UnrealBuildTool/Platform/Linux/LinuxToolChain.cs
+++ b/Engine/Source/Programs/UnrealBuildTool/Platform/Linux/LinuxToolChain.cs
@@ -57,10 +57,10 @@ namespace UnrealBuildTool
 			: this(UnrealTargetPlatform.Linux, InArchitecture, InSDK, InOptions, InLogger)
 		{
 			// prevent unknown clangs since the build is likely to fail on too old or too new compilers
-			if (CompilerVersionLessThan(13, 0, 0) || CompilerVersionGreaterOrEqual(14, 0, 0))
+			if (CompilerVersionLessThan(15, 0, 0) || CompilerVersionGreaterOrEqual(16, 0, 0))
 			{
 				throw new BuildException(
-					string.Format("This version of the Unreal Engine can only be compiled with clang 13.0. clang {0} may not build it - please use a different version.",
+					string.Format("This version of the Unreal Engine can only be compiled with clang 15.0. clang {0} may not build it - please use a different version.",
 						Info.ClangVersion)
 					);
 			}
