Horde: Handle unset agent name in XGE metadata metrics

diff --git a/Engine/Source/Programs/Horde/Horde.Agent/Execution/JobExecutor.cs b/Engine/Source/Programs/Horde/Horde.Agent/Execution/JobExecutor.cs
--- a/Engine/Source/Programs/Horde/Horde.Agent/Execution/JobExecutor.cs
+++ b/Engine/Source/Programs/Horde/Horde.Agent/Execution/JobExecutor.cs
@@ -1170,10 +1170,12 @@ namespace Horde.Agent.Execution
 			{
 				try
 				{
-					string taskName = logEvent.GetProperty<string>("name");
-					int duration = logEvent.GetProperty<int>("duration");
-					string agent = logEvent.GetProperty<string>("agent");
-					_logger.LogInformation("Ran XGE task '{XgeTaskName}' on agent '{XgeAgent}' for {Duration} secs", taskName, agent, duration);	
+					if (logEvent.TryGetProperty("agent", out string? agentName))
+					{
+						string taskName = logEvent.GetProperty<string>("name");
+						int duration = logEvent.GetProperty<int>("duration");
+						_logger.LogInformation("Executed XGE task {XgeTaskName} on agent {XgeAgent} for {Duration} secs", taskName, agentName, duration);	
+					}
 				}
 				catch (Exception e)
 				{
diff --git a/Engine/Source/Programs/UnrealBuildTool/Matchers/XoreaxEventMatcher.cs b/Engine/Source/Programs/UnrealBuildTool/Matchers/XoreaxEventMatcher.cs
--- a/Engine/Source/Programs/UnrealBuildTool/Matchers/XoreaxEventMatcher.cs
+++ b/Engine/Source/Programs/UnrealBuildTool/Matchers/XoreaxEventMatcher.cs
@@ -127,7 +127,7 @@ namespace UnrealBuildTool.Matchers
 				int minutes = Convert.ToInt32(m.Groups[2].Value);
 				int seconds = Convert.ToInt32(m.Groups[3].Value);
 				TimeSpan duration = new TimeSpan(hours, minutes, seconds);
-				builder.AddProperty(PropertyDuration, duration.TotalSeconds);
+				builder.AddProperty(PropertyDuration, (int)duration.TotalSeconds);
 			}
 		}
 		
@@ -139,7 +139,7 @@ namespace UnrealBuildTool.Matchers
 				int minutes = Convert.ToInt32(m.Groups[1].Value);
 				int seconds = Convert.ToInt32(m.Groups[2].Value);
 				TimeSpan startTime = new TimeSpan(0, minutes, seconds);
-				builder.AddProperty(PropertyStartTime, startTime.TotalSeconds);
+				builder.AddProperty(PropertyStartTime, (int)startTime.TotalSeconds);
 			}
 		}
 	}
