UnrealBuildTool: Don't merge actions that ignore conflicts

#rnx

diff --git a/Engine/Source/Programs/UnrealBuildTool/Configuration/UEBuildTarget.cs b/Engine/Source/Programs/UnrealBuildTool/Configuration/UEBuildTarget.cs
--- a/Engine/Source/Programs/UnrealBuildTool/Configuration/UEBuildTarget.cs
+++ b/Engine/Source/Programs/UnrealBuildTool/Configuration/UEBuildTarget.cs
@@ -2158,7 +2158,8 @@ namespace UnrealBuildTool
 
 					WriteMetadataTargetInfo EngineInfo = new WriteMetadataTargetInfo(null, VersionFile, Receipt.Version, null, null, EngineFileToManifest, null);
 					FileReference EngineInfoFile = FileReference.Combine(ProjectIntermediateDirectory, "EngineMetadata.dat");
-					CreateWriteMetadataAction(MakefileBuilder, VersionFile.GetFileName(), EngineInfoFile, EngineInfo, EnginePrereqItems);
+					string EngineInfoName = TargetName.Equals(VersionFile.GetFileNameWithoutExtension()) ? VersionFile.GetFileName() : $"{VersionFile.GetFileName()} ({TargetName})";
+					CreateWriteMetadataAction(MakefileBuilder, EngineInfoName, EngineInfoFile, EngineInfo, EnginePrereqItems);
 
 					WriteMetadataTargetInfo TargetInfo = new WriteMetadataTargetInfo(ProjectFile, VersionFile, null, ReceiptFileName, Receipt, TargetFileToManifest, FileNameToLoadOrderManifest);
 					FileReference TargetInfoFile = FileReference.Combine(ProjectIntermediateDirectory, "TargetMetadata.dat");
diff --git a/Engine/Source/Programs/UnrealBuildTool/Modes/BuildMode.cs b/Engine/Source/Programs/UnrealBuildTool/Modes/BuildMode.cs
--- a/Engine/Source/Programs/UnrealBuildTool/Modes/BuildMode.cs
+++ b/Engine/Source/Programs/UnrealBuildTool/Modes/BuildMode.cs
@@ -422,7 +422,7 @@ namespace UnrealBuildTool
 					// TODO: Skipping conflicts for WriteMetadata is a hack to maintain parity with the BuildGraph executor which allowed
 					// exporting mulitple conflicing WriteMetadata actions.
 					// This is technically buggy behavior and should be properly fixed in UEBuildTarget.cs
-					IEnumerable<IExternalAction> CheckActions = Makefiles.SelectMany(x => x.Actions).Where(x => x.ActionType != ActionType.WriteMetadata);
+					IEnumerable<IExternalAction> CheckActions = Makefiles.SelectMany(x => x.Actions).Where(x => !x.IgnoreConflicts());
 					ActionGraph.CheckForConflicts(CheckActions, Logger);
 				}
 
@@ -975,11 +975,18 @@ namespace UnrealBuildTool
 		{
 			// Set of all output items. Knowing that there are no conflicts in produced items, we use this to eliminate duplicate actions.
 			Dictionary<FileItem, LinkedAction> OutputItemToProducingAction = new Dictionary<FileItem, LinkedAction>();
+			HashSet<LinkedAction> IgnoreConflictActions = new HashSet<LinkedAction>();
 			for(int TargetIdx = 0; TargetIdx < TargetDescriptors.Count; TargetIdx++)
 			{
 				string GroupPrefix = String.Format("{0}-{1}-{2}", TargetDescriptors[TargetIdx].Name, TargetDescriptors[TargetIdx].Platform, TargetDescriptors[TargetIdx].Configuration);
 				foreach(LinkedAction TargetAction in TargetActions[TargetIdx])
 				{
+					if (TargetAction.IgnoreConflicts())
+					{
+						IgnoreConflictActions.Add(TargetAction);
+						continue;
+					}
+
 					FileItem ProducedItem = TargetAction.ProducedItems.First();
 
 					LinkedAction? ExistingAction;
@@ -991,7 +998,9 @@ namespace UnrealBuildTool
 					ExistingAction.GroupNames.Add(GroupPrefix);
 				}
 			}
-			return new List<LinkedAction>(OutputItemToProducingAction.Values);
+			List<LinkedAction> Results = new List<LinkedAction>(OutputItemToProducingAction.Values);
+			Results.AddRange(IgnoreConflictActions);
+			return Results;
 		}
 
 		void ProcessCoreDumps(DirectoryReference? SaveCrashDumpDirectory, ILogger Logger)
diff --git a/Engine/Source/Programs/UnrealBuildTool/System/Action.cs b/Engine/Source/Programs/UnrealBuildTool/System/Action.cs
--- a/Engine/Source/Programs/UnrealBuildTool/System/Action.cs
+++ b/Engine/Source/Programs/UnrealBuildTool/System/Action.cs
@@ -472,6 +472,16 @@ namespace UnrealBuildTool
 				Writer.WriteValue("DependencyListFile", Action.DependencyListFile.AbsolutePath);
 			}
 		}
+
+		/// <summary>
+		/// Determine if this Action ignores output conficts.
+		/// </summary>
+		/// <param name="Action"> The Action to check</param>
+		/// <returns>True if conflicts are ignored, else false.</returns>
+		public static bool IgnoreConflicts(this IExternalAction Action)
+		{
+			return Action.ActionType == ActionType.WriteMetadata;
+		}
 	}
 
 	/// <summary>
