- Fixed NMakeOutput to not have x64 architecture inlined for Windows builds

diff --git a/Engine/Source/Programs/UnrealBuildTool/ProjectFiles/VisualStudio/VCProject.cs b/Engine/Source/Programs/UnrealBuildTool/ProjectFiles/VisualStudio/VCProject.cs
--- a/Engine/Source/Programs/UnrealBuildTool/ProjectFiles/VisualStudio/VCProject.cs
+++ b/Engine/Source/Programs/UnrealBuildTool/ProjectFiles/VisualStudio/VCProject.cs
@@ -1888,7 +1888,10 @@ namespace UnrealBuildTool
 					{
 						NMakePath += "-" + UBTPlatformName + "-" + UBTConfigurationName;
 					}
-					NMakePath += TargetRulesObject.Architecture.ToString();
+					if (UnrealArchitectureConfig.ForPlatform(Platform).RequiresArchitectureFilenames(TargetRulesObject.Architectures))
+					{
+						NMakePath += TargetRulesObject.Architecture.ToString();
+					}
 					NMakePath += BuildPlatform.GetBinaryExtension(UEBuildBinaryType.Executable);
 
 					VCProjectFileContent.AppendLine("  <PropertyGroup {0}>", ConditionString);
