Fix build failure from new arch changes

diff --git a/Engine/Source/Programs/UnrealBuildTool/Platform/Linux/LinuxToolChain.cs b/Engine/Source/Programs/UnrealBuildTool/Platform/Linux/LinuxToolChain.cs
--- a/Engine/Source/Programs/UnrealBuildTool/Platform/Linux/LinuxToolChain.cs
+++ b/Engine/Source/Programs/UnrealBuildTool/Platform/Linux/LinuxToolChain.cs
@@ -592,15 +592,15 @@ namespace UnrealBuildTool
 				if (CrossCompiling())
 				{
 					// x64 only replaced the linux folder with arch, while on arm64 its still linux
-					if (LinkEnvironment.Architecture.StartsWith("x86_64"))
+					if (LinkEnvironment.Architecture == UnrealArch.Arm64)
 					{
-						Arguments.Add(string.Format("-Wl,-rpath=\"{0}/lib/clang/{1}.{2}.{3}/lib/{4}\"",
-								LinuxInfo.BaseLinuxPath, Info.ClangVersion.Major, Info.ClangVersion.Minor, Info.ClangVersion.Build, LinkEnvironment.Architecture));
+						Arguments.Add(string.Format("-Wl,-rpath=\"{0}/lib/clang/{1}.{2}.{3}/lib/linux\"",
+								LinuxInfo.BaseLinuxPath, Info.ClangVersion.Major, Info.ClangVersion.Minor, Info.ClangVersion.Build));
 					}
 					else
 					{
-						Arguments.Add(string.Format("-Wl,-rpath=\"{0}/lib/clang/{1}.{2}.{3}/lib/linux\"",
-								LinuxInfo.BaseLinuxPath, Info.ClangVersion.Major, Info.ClangVersion.Minor, Info.ClangVersion.Build));
+						Arguments.Add(string.Format("-Wl,-rpath=\"{0}/lib/clang/{1}.{2}.{3}/lib/{4}\"",
+								LinuxInfo.BaseLinuxPath, Info.ClangVersion.Major, Info.ClangVersion.Minor, Info.ClangVersion.Build, LinkEnvironment.Architecture));
 					}
 				}
 			}
