Re-fix arch name change

diff --git a/Engine/Source/Programs/UnrealBuildTool/Platform/Linux/LinuxToolChain.cs b/Engine/Source/Programs/UnrealBuildTool/Platform/Linux/LinuxToolChain.cs
--- a/Engine/Source/Programs/UnrealBuildTool/Platform/Linux/LinuxToolChain.cs
+++ b/Engine/Source/Programs/UnrealBuildTool/Platform/Linux/LinuxToolChain.cs
@@ -599,8 +599,8 @@ namespace UnrealBuildTool
 					}
 					else
 					{
-						Arguments.Add(string.Format("-Wl,-rpath=\"{0}/lib/clang/{1}.{2}.{3}/lib/{4}\"",
-								LinuxInfo.BaseLinuxPath, Info.ClangVersion.Major, Info.ClangVersion.Minor, Info.ClangVersion.Build, LinkEnvironment.Architecture));
+						Arguments.Add(string.Format("-Wl,-rpath=\"{0}/lib/clang/{1}.{2}.{3}/lib/x86_64-unknown-linux-gnu\"",
+								LinuxInfo.BaseLinuxPath, Info.ClangVersion.Major, Info.ClangVersion.Minor, Info.ClangVersion.Build));
 					}
 				}
 			}
