- Don't do the path length check (focused on Portability between users) for files in Intermediate directory, as these files are inherently not portable

diff --git a/Engine/Source/Programs/UnrealBuildTool/System/ActionGraph.cs b/Engine/Source/Programs/UnrealBuildTool/System/ActionGraph.cs
--- a/Engine/Source/Programs/UnrealBuildTool/System/ActionGraph.cs
+++ b/Engine/Source/Programs/UnrealBuildTool/System/ActionGraph.cs
@@ -304,8 +304,11 @@ namespace UnrealBuildTool
 							FailPaths.Add(ProducedItem.Location);
 						}
 
-						if (ProducedItem.Location.FullName.Length > Unreal.RootDirectory.FullName.Length +
-						    BuildConfiguration.MaxNestedPathLength && ProducedItem.Location.IsUnderDirectory(Unreal.RootDirectory))
+						// don't look in Intermediate directories - these aren't portable between machines, so don't need to be cbecked for length underneath the root
+						if (ProducedItem.Location.FullName.Length > Unreal.RootDirectory.FullName.Length + BuildConfiguration.MaxNestedPathLength && 
+							ProducedItem.Location.IsUnderDirectory(Unreal.RootDirectory) &&
+							!ProducedItem.Location.ContainsName("Intermediate", Unreal.RootDirectory)
+							)
 						{
 							WarnPaths.Add(ProducedItem.Location);
 						}
