- Fixing Arch parsing from commandline (only UnrealArch.Parse was allowing for alternate names, not UnrealArch.TryParse)

diff --git a/Engine/Source/Programs/UnrealBuildTool/Configuration/UEBuildTarget.cs b/Engine/Source/Programs/UnrealBuildTool/Configuration/UEBuildTarget.cs
--- a/Engine/Source/Programs/UnrealBuildTool/Configuration/UEBuildTarget.cs
+++ b/Engine/Source/Programs/UnrealBuildTool/Configuration/UEBuildTarget.cs
@@ -710,6 +710,30 @@ namespace UnrealBuildTool
 			return GetUniqueStringRegistry().GetStringForId(Id);
 		}
 
+		private static string FixName(string Name)
+		{
+			// allow for some alternate names
+			switch (Name.ToLower())
+			{
+				case "a8":
+				case "a":
+				case "-arm64":
+					Name = "arm64";
+					break;
+
+				case "x86_64":
+				case "x6":
+				case "x":
+				case "-x64":
+				case "intel":
+					Name = "x64";
+					break;
+			}
+
+			return Name;
+		}
+
+
 		/// <summary>
 		/// 
 		/// </summary>
@@ -718,6 +742,8 @@ namespace UnrealBuildTool
 		/// <returns></returns>
 		static public bool TryParse(string Name, out UnrealArch Arch)
 		{
+			Name = FixName(Name);
+
 			if (GetUniqueStringRegistry().HasString(Name))
 			{
 				Arch.Id = GetUniqueStringRegistry().FindOrAddByName(Name);
@@ -741,23 +767,7 @@ namespace UnrealBuildTool
 		/// <returns></returns>
 		static public UnrealArch Parse(string Name)
 		{
-			// allow for some alternate names
-			switch (Name.ToLower())
-			{
-				case "a8":
-				case "a":
-				case "-arm64":
-					Name = "arm64";
-					break;
-
-				case "x86_64":
-				case "x6":
-				case "x":
-				case "-x64":
-				case "intel":
-					Name = "x64";
-					break;
-			}
+			Name = FixName(Name);
 
 			if (GetUniqueStringRegistry().HasString(Name))
 			{
