UnrealBuildTool: Don't halt if Visual Studio installations can not be enumerated

#rnx

diff --git a/Engine/Source/Programs/UnrealBuildTool/Platform/Windows/MicrosoftPlatformSDK.cs b/Engine/Source/Programs/UnrealBuildTool/Platform/Windows/MicrosoftPlatformSDK.cs
--- a/Engine/Source/Programs/UnrealBuildTool/Platform/Windows/MicrosoftPlatformSDK.cs
+++ b/Engine/Source/Programs/UnrealBuildTool/Platform/Windows/MicrosoftPlatformSDK.cs
@@ -832,7 +832,8 @@ namespace UnrealBuildTool
 				}
 				catch (Exception Ex)
 				{
-					throw new BuildException(Ex, "Error while enumerating Visual Studio toolchains");
+					Logger.LogWarning("Unable to enumerate Visual Studio installations");
+					Logger.LogDebug(Ex.ToString());
 				}
 			}
 			CachedVisualStudioInstallations = Installations;
