Add clang version checks on deprecated arguments

diff --git a/Engine/Source/Programs/UnrealBuildTool/Platform/Windows/VCToolChain.cs b/Engine/Source/Programs/UnrealBuildTool/Platform/Windows/VCToolChain.cs
--- a/Engine/Source/Programs/UnrealBuildTool/Platform/Windows/VCToolChain.cs
+++ b/Engine/Source/Programs/UnrealBuildTool/Platform/Windows/VCToolChain.cs
@@ -287,8 +287,12 @@ namespace UnrealBuildTool
 					// Enable the static analyzer with default checks.
 					Arguments.Add("--analyze");
 
-					// Make sure we check inside nested blocks (e.g. 'if ((foo = getchar()) == 0) {}')
-					Arguments.Add("-Xclang -analyzer-opt-analyze-nested-blocks");
+					// Deprecated in LLVM 15
+					if (EnvVars.CompilerVersion <= new VersionNumber(14))
+					{
+						// Make sure we check inside nested blocks (e.g. 'if ((foo = getchar()) == 0) {}')
+						Arguments.Add("-Xclang -analyzer-opt-analyze-nested-blocks");
+					}
 
 					// Write out a pretty web page with navigation to understand how the analysis was derived if HTML is enabled.
 					Arguments.Add($"-Xclang -analyzer-output={Target.StaticAnalyzerOutputType.ToString().ToLowerInvariant()}");
diff --git a/Engine/Source/Programs/UnrealBuildTool/ToolChain/ClangToolChain.cs b/Engine/Source/Programs/UnrealBuildTool/ToolChain/ClangToolChain.cs
--- a/Engine/Source/Programs/UnrealBuildTool/ToolChain/ClangToolChain.cs
+++ b/Engine/Source/Programs/UnrealBuildTool/ToolChain/ClangToolChain.cs
@@ -775,8 +775,12 @@ namespace UnrealBuildTool
 			// Enable the static analyzer with default checks.
 			Arguments.Add("--analyze");
 
-			// Make sure we check inside nested blocks (e.g. 'if ((foo = getchar()) == 0) {}')
-			Arguments.Add("-Xclang -analyzer-opt-analyze-nested-blocks");
+			// Deprecated in LLVM 15
+			if (CompilerVersionLessThan(15, 0, 0))
+			{
+				// Make sure we check inside nested blocks (e.g. 'if ((foo = getchar()) == 0) {}')
+				Arguments.Add("-Xclang -analyzer-opt-analyze-nested-blocks");
+			}
 
 			// Write out a pretty web page with navigation to understand how the analysis was derived if HTML is enabled.
 			Arguments.Add($"-Xclang -analyzer-output={StaticAnalyzerOutputType.ToString().ToLowerInvariant()}");
