UnrealBuildTool: Fix multiple targets on Mac using wrong bundle directory

#rnx

diff --git a/Engine/Source/Programs/UnrealBuildTool/Platform/Mac/MacToolChain.cs b/Engine/Source/Programs/UnrealBuildTool/Platform/Mac/MacToolChain.cs
--- a/Engine/Source/Programs/UnrealBuildTool/Platform/Mac/MacToolChain.cs
+++ b/Engine/Source/Programs/UnrealBuildTool/Platform/Mac/MacToolChain.cs
@@ -1060,7 +1060,7 @@ namespace UnrealBuildTool
 			}
 		}
 
-		static private DirectoryReference? BundleContentsDirectory;
+		private static Dictionary<ReadOnlyTargetRules, DirectoryReference> BundleContentsDirectories = new();
 
 		public override void ModifyBuildProducts(ReadOnlyTargetRules Target, UEBuildBinary Binary, List<string> Libraries, List<UEBuildBundleResource> BundleResources, Dictionary<FileReference, BuildProductType> BuildProducts)
 		{
@@ -1112,10 +1112,11 @@ namespace UnrealBuildTool
 				return;
 			}
 
-			if (BundleContentsDirectory == null && Binary.Type == UEBuildBinaryType.Executable)
+			if (!BundleContentsDirectories.ContainsKey(Target) && Binary.Type == UEBuildBinaryType.Executable)
 			{
-				BundleContentsDirectory = Binary.OutputFilePath.Directory.ParentDirectory!;
+				BundleContentsDirectories.Add(Target, Binary.OutputFilePath.Directory.ParentDirectory!);
 			}
+			DirectoryReference? BundleContentsDirectory = BundleContentsDirectories.GetValueOrDefault(Target);
 
 			// We need to know what third party dylibs would be copied to the bundle
 			if (Binary.Type != UEBuildBinaryType.StaticLibrary)
