IWYU should use the latest deprecation tag and not have a hardcoded string.
Deduplicating the engine include order strings.

diff --git a/Engine/Source/Programs/UnrealBuildTool/Configuration/TargetRules.cs b/Engine/Source/Programs/UnrealBuildTool/Configuration/TargetRules.cs
--- a/Engine/Source/Programs/UnrealBuildTool/Configuration/TargetRules.cs
+++ b/Engine/Source/Programs/UnrealBuildTool/Configuration/TargetRules.cs
@@ -292,6 +292,25 @@ namespace UnrealBuildTool
 	/// </summary>
 	public class EngineIncludeOrderHelper
 	{
+		private static string GetDeprecationDefine(EngineIncludeOrderVersion InVersion)
+		{
+			switch (InVersion)
+			{
+				default:
+				case EngineIncludeOrderVersion.Unreal5_2:
+					return "UE_ENABLE_INCLUDE_ORDER_DEPRECATED_IN_5_2";
+				case EngineIncludeOrderVersion.Unreal5_1:
+					return "UE_ENABLE_INCLUDE_ORDER_DEPRECATED_IN_5_1";
+				case EngineIncludeOrderVersion.Unreal5_0:
+					return "UE_ENABLE_INCLUDE_ORDER_DEPRECATED_IN_5_0";
+			}
+		}
+
+		private static string GetDeprecationDefine(EngineIncludeOrderVersion InTestVersion, EngineIncludeOrderVersion InVersion)
+		{
+			return GetDeprecationDefine(InTestVersion) + "=" + (InVersion < InTestVersion ? "1" : "0");
+		}
+
 		/// <summary>
 		/// Returns a list of every deprecation define available.
 		/// </summary>
@@ -300,8 +319,8 @@ namespace UnrealBuildTool
 		{
 			return new List<string>()
 			{
-				"UE_ENABLE_INCLUDE_ORDER_DEPRECATED_IN_5_1",
-				"UE_ENABLE_INCLUDE_ORDER_DEPRECATED_IN_5_2",
+				GetDeprecationDefine(EngineIncludeOrderVersion.Unreal5_1),
+				GetDeprecationDefine(EngineIncludeOrderVersion.Unreal5_2),
 			};
 		}
 
@@ -314,10 +333,19 @@ namespace UnrealBuildTool
 		{
 			return new List<string>()
 			{
-				"UE_ENABLE_INCLUDE_ORDER_DEPRECATED_IN_5_1=" + (InVersion < EngineIncludeOrderVersion.Unreal5_1 ? "1" : "0"),
-				"UE_ENABLE_INCLUDE_ORDER_DEPRECATED_IN_5_2=" + (InVersion < EngineIncludeOrderVersion.Unreal5_2 ? "1" : "0"),
+				GetDeprecationDefine(EngineIncludeOrderVersion.Unreal5_1, InVersion),
+				GetDeprecationDefine(EngineIncludeOrderVersion.Unreal5_2, InVersion),
 			};
 		}
+
+		/// <summary>
+		/// Returns the latest deprecation define.
+		/// </summary>
+		/// <returns></returns>
+		public static string GetLatestDeprecationDefine()
+		{
+			return GetDeprecationDefine(EngineIncludeOrderVersion.Latest);
+		}
 	}
 
 	/// <summary>
diff --git a/Engine/Source/Programs/UnrealBuildTool/Modes/IWYUMode.cs b/Engine/Source/Programs/UnrealBuildTool/Modes/IWYUMode.cs
--- a/Engine/Source/Programs/UnrealBuildTool/Modes/IWYUMode.cs
+++ b/Engine/Source/Programs/UnrealBuildTool/Modes/IWYUMode.cs
@@ -257,7 +257,23 @@ namespace UnrealBuildTool
 		/// Applied to non-private headers that are part of the Engine folder.
 		/// </summary>
 		[CommandLine("-DeprecateTag")]
-		public string KeepRemovedIncludesInDeprecated = "UE_ENABLE_INCLUDE_ORDER_DEPRECATED_IN_5_2";
+		private string? HeaderDeprecationTagOverride;
+
+		public string HeaderDeprecationTag
+		{
+			get
+			{
+				if (!String.IsNullOrEmpty(HeaderDeprecationTagOverride))
+				{
+					return HeaderDeprecationTagOverride;
+				}
+				return EngineIncludeOrderHelper.GetLatestDeprecationDefine();
+			}
+			set
+			{
+				HeaderDeprecationTagOverride = value;
+			}
+		}
 
 		private string? GetModuleToUpdateName(TargetDescriptor Descriptor)
 		{
@@ -1216,7 +1232,7 @@ namespace UnrealBuildTool
 				// If file is public, in engine and we have a deprecation tag set we will 
 				// add a deprecated include scope at the end of the file (unless scope already exists, then we'll add it inside that)
 				string EngineDir = Unreal.EngineDirectory.FullName.Replace('\\', '/');
-				if (!IsPrivate && Info.File.StartsWith(EngineDir) && !String.IsNullOrEmpty(KeepRemovedIncludesInDeprecated))
+				if (!IsPrivate && Info.File.StartsWith(EngineDir) && !String.IsNullOrEmpty(HeaderDeprecationTag))
 				{
 					// Remove the includes in LinesRemoved
 					LinesRemoved.RemoveWhere(Line =>
@@ -1239,7 +1255,7 @@ namespace UnrealBuildTool
 					if (LinesRemoved.Count > 0)
 					{
 						int IndexOfDeprecateScope = -1;
-						string Match = "#if " + KeepRemovedIncludesInDeprecated;
+						string Match = "#if " + HeaderDeprecationTag;
 						for (int I = NewLines.Count - 1; I != 0; --I)
 						{
 							if (NewLines[I] == Match)
