UnrealBuildTool: Mac stub dylib paths should be in Engine/Intermediate for Engine outputs

#rnx

diff --git a/Engine/Source/Programs/UnrealBuildTool/Platform/Mac/MacToolChain.cs b/Engine/Source/Programs/UnrealBuildTool/Platform/Mac/MacToolChain.cs
--- a/Engine/Source/Programs/UnrealBuildTool/Platform/Mac/MacToolChain.cs
+++ b/Engine/Source/Programs/UnrealBuildTool/Platform/Mac/MacToolChain.cs
@@ -703,7 +703,7 @@ namespace UnrealBuildTool
 
 		private FileItem MakeStubItem(LinkEnvironment LinkEnvironment, string DylibPath)
 		{
-			string IntermediateDirectory = (ProjectFile == null ? Unreal.EngineDirectory : ProjectFile.Directory) + "/Intermediate/Mac";
+			string IntermediateDirectory = (ProjectFile == null || LinkEnvironment.OutputFilePaths.All(x => x.IsUnderDirectory(Unreal.EngineDirectory)) ? Unreal.EngineDirectory : ProjectFile.Directory) + "/Intermediate/Mac";
 			return FileItem.GetItemByPath(Path.Combine(IntermediateDirectory, "Stubs", LinkEnvironment.Architecture.ToString(), Path.GetFileName(DylibPath)));
 		}
 
