Fix ASan crashing in getaddrinfo. Bug introduced in LLVM 15 compiler-rt

[REDACTED_PROFILE_URL]/llvm-project/issues/59007

diff --git a/Engine/Source/Programs/UnrealBuildTool/Platform/Linux/LinuxToolChain.cs b/Engine/Source/Programs/UnrealBuildTool/Platform/Linux/LinuxToolChain.cs
--- a/Engine/Source/Programs/UnrealBuildTool/Platform/Linux/LinuxToolChain.cs
+++ b/Engine/Source/Programs/UnrealBuildTool/Platform/Linux/LinuxToolChain.cs
@@ -568,6 +568,11 @@ namespace UnrealBuildTool
 				if (Options.HasFlag(ClangToolChainOptions.EnableSharedSanitizer))
 				{
 					Arguments.Add("-shared-libsan");
+
+					// LLVM 15 compiler-rt introduced a new bug causing ASan to crash when going over net code
+					// [REDACTED_PROFILE_URL]/llvm-project/issues/59007
+					// adding this fixes it, may be required now?
+					Arguments.Add("-lresolv");
 				}
 
 				if (Options.HasFlag(ClangToolChainOptions.EnableAddressSanitizer))
