UnrealBuildTool: Prevent multiple target leaking bundle dependencies on Mac

#rnx

diff --git a/Engine/Source/Programs/UnrealBuildTool/Configuration/UEBuildBinary.cs b/Engine/Source/Programs/UnrealBuildTool/Configuration/UEBuildBinary.cs
--- a/Engine/Source/Programs/UnrealBuildTool/Configuration/UEBuildBinary.cs
+++ b/Engine/Source/Programs/UnrealBuildTool/Configuration/UEBuildBinary.cs
@@ -242,13 +242,13 @@ namespace UnrealBuildTool
 
 				foreach (FileItem Executable in ConsoleAppOutputFiles)
 				{
-					OutputFiles.AddRange(ToolChain.PostBuild(Executable, ConsoleAppLinkEvironment, Graph));
+					OutputFiles.AddRange(ToolChain.PostBuild(Target, Executable, ConsoleAppLinkEvironment, Graph));
 				}
 			}
 
 			foreach (FileItem Executable in Executables)
 			{
-				OutputFiles.AddRange(ToolChain.PostBuild(Executable, BinaryLinkEnvironment, Graph));
+				OutputFiles.AddRange(ToolChain.PostBuild(Target, Executable, BinaryLinkEnvironment, Graph));
 			}
 
 			return OutputFiles;
diff --git a/Engine/Source/Programs/UnrealBuildTool/Configuration/UEBuildTarget.cs b/Engine/Source/Programs/UnrealBuildTool/Configuration/UEBuildTarget.cs
--- a/Engine/Source/Programs/UnrealBuildTool/Configuration/UEBuildTarget.cs
+++ b/Engine/Source/Programs/UnrealBuildTool/Configuration/UEBuildTarget.cs
@@ -2332,7 +2332,7 @@ namespace UnrealBuildTool
 			}
 
 			// On Mac and Linux we have actions that should be executed after all the binaries are created
-			TargetToolChain.SetupBundleDependencies(Binaries, TargetName);
+			TargetToolChain.SetupBundleDependencies(Rules, Binaries, TargetName);
 
 			// Gather modules we might want to generate headers for
 			HashSet<UEBuildModuleCPP> ModulesToGenerateHeadersFor = GatherDependencyModules(OriginalBinaries.ToList());
diff --git a/Engine/Source/Programs/UnrealBuildTool/Platform/IOS/IOSToolChain.cs b/Engine/Source/Programs/UnrealBuildTool/Platform/IOS/IOSToolChain.cs
--- a/Engine/Source/Programs/UnrealBuildTool/Platform/IOS/IOSToolChain.cs
+++ b/Engine/Source/Programs/UnrealBuildTool/Platform/IOS/IOSToolChain.cs
@@ -1160,9 +1160,9 @@ namespace UnrealBuildTool
 			return new DirectoryReference(ResourcesDir);
 		}
 
-		public override ICollection<FileItem> PostBuild(FileItem Executable, LinkEnvironment BinaryLinkEnvironment, IActionGraphBuilder Graph)
+		public override ICollection<FileItem> PostBuild(ReadOnlyTargetRules Target, FileItem Executable, LinkEnvironment BinaryLinkEnvironment, IActionGraphBuilder Graph)
 		{
-			List<FileItem> OutputFiles = new List<FileItem>(base.PostBuild(Executable, BinaryLinkEnvironment, Graph));
+			List<FileItem> OutputFiles = new List<FileItem>(base.PostBuild(Target, Executable, BinaryLinkEnvironment, Graph));
 
 			if (BinaryLinkEnvironment.bIsBuildingLibrary)
 			{
@@ -1990,9 +1990,9 @@ namespace UnrealBuildTool
 			return TargetItem;
 		}
 
-		public override void SetupBundleDependencies(List<UEBuildBinary> Binaries, string GameName)
+		public override void SetupBundleDependencies(ReadOnlyTargetRules Target, List<UEBuildBinary> Binaries, string GameName)
 		{
-			base.SetupBundleDependencies(Binaries, GameName);
+			base.SetupBundleDependencies(Target, Binaries, GameName);
 
 			foreach (UEBuildBinary Binary in Binaries)
 			{
diff --git a/Engine/Source/Programs/UnrealBuildTool/Platform/Linux/LinuxToolChain.cs b/Engine/Source/Programs/UnrealBuildTool/Platform/Linux/LinuxToolChain.cs
--- a/Engine/Source/Programs/UnrealBuildTool/Platform/Linux/LinuxToolChain.cs
+++ b/Engine/Source/Programs/UnrealBuildTool/Platform/Linux/LinuxToolChain.cs
@@ -735,7 +735,7 @@ namespace UnrealBuildTool
 		/// <summary>
 		/// Holds all the binaries for a particular target (except maybe the executable itself).
 		/// </summary>
-		private static List<FileItem> AllBinaries = new List<FileItem>();
+		private static Dictionary<ReadOnlyTargetRules, List<FileItem>> AllBinaries = new();
 
 		/// <summary>
 		/// Tracks that information about used C++ library is only printed once
@@ -890,7 +890,7 @@ namespace UnrealBuildTool
 			return OutputFile;
 		}
 
-		public FileItem? FixDependencies(LinkEnvironment LinkEnvironment, FileItem Executable, IActionGraphBuilder Graph, ILogger Logger)
+		public FileItem? FixDependencies(ReadOnlyTargetRules Target, LinkEnvironment LinkEnvironment, FileItem Executable, IActionGraphBuilder Graph, ILogger Logger)
 		{
 			if (bUseFixdeps)
 			{
@@ -924,7 +924,7 @@ namespace UnrealBuildTool
 				// Make sure we don't run this script until the all executables and shared libraries
 				// have been built.
 				PostLinkAction.PrerequisiteItems.Add(Executable);
-				foreach (FileItem Dependency in AllBinaries)
+				foreach (FileItem Dependency in AllBinaries.GetValueOrDefault(Target, new()))
 				{
 					PostLinkAction.PrerequisiteItems.Add(Dependency);
 				}
@@ -1497,20 +1497,24 @@ namespace UnrealBuildTool
 			return OutputFile;
 		}
 
-		public override void SetupBundleDependencies(List<UEBuildBinary> Binaries, string GameName)
+		public override void SetupBundleDependencies(ReadOnlyTargetRules Target, List<UEBuildBinary> Binaries, string GameName)
 		{
 			if (bUseFixdeps)
 			{
+				if (!AllBinaries.ContainsKey(Target))
+				{
+					AllBinaries.Add(Target, new());
+				}
 				foreach (UEBuildBinary Binary in Binaries)
 				{
-					AllBinaries.Add(FileItem.GetItemByFileReference(Binary.OutputFilePath));
+					AllBinaries[Target].Add(FileItem.GetItemByFileReference(Binary.OutputFilePath));
 				}
 			}
 		}
 
-		public override ICollection<FileItem> PostBuild(FileItem Executable, LinkEnvironment BinaryLinkEnvironment, IActionGraphBuilder Graph)
+		public override ICollection<FileItem> PostBuild(ReadOnlyTargetRules Target, FileItem Executable, LinkEnvironment BinaryLinkEnvironment, IActionGraphBuilder Graph)
 		{
-			ICollection<FileItem> OutputFiles = base.PostBuild(Executable, BinaryLinkEnvironment, Graph);
+			ICollection<FileItem> OutputFiles = base.PostBuild(Target, Executable, BinaryLinkEnvironment, Graph);
 
 			if (bUseFixdeps)
 			{
@@ -1519,7 +1523,7 @@ namespace UnrealBuildTool
 					return OutputFiles;
 				}
 
-				FileItem? FixDepsOutputFile = FixDependencies(BinaryLinkEnvironment, Executable, Graph, Logger);
+				FileItem? FixDepsOutputFile = FixDependencies(Target, BinaryLinkEnvironment, Executable, Graph, Logger);
 				if (FixDepsOutputFile != null)
 				{
 					OutputFiles.Add(FixDepsOutputFile);
diff --git a/Engine/Source/Programs/UnrealBuildTool/Platform/Mac/MacToolChain.cs b/Engine/Source/Programs/UnrealBuildTool/Platform/Mac/MacToolChain.cs
--- a/Engine/Source/Programs/UnrealBuildTool/Platform/Mac/MacToolChain.cs
+++ b/Engine/Source/Programs/UnrealBuildTool/Platform/Mac/MacToolChain.cs
@@ -101,7 +101,7 @@ namespace UnrealBuildTool
 		/// </summary>
 		private const string MacArchiver = "libtool";
 
-		private static List<FileItem> BundleDependencies = new List<FileItem>();
+		private static Dictionary<ReadOnlyTargetRules, List<FileItem>> BundleDependencies = new();
 
 		protected override ClangToolChainInfo GetToolChainInfo()
 		{
@@ -996,10 +996,11 @@ namespace UnrealBuildTool
 		/// <summary>
 		/// Creates app bundle for a given executable
 		/// </summary>
+		/// <param name="Target"></param>
 		/// <param name="LinkEnvironment"></param>
 		/// <param name="Executable">FileItem describing the executable to generate app bundle for</param>
 		/// <param name="Graph">List of actions to be executed. Additional actions will be added to this list.</param>
-		FileItem FinalizeAppBundle(LinkEnvironment LinkEnvironment, FileItem Executable, IActionGraphBuilder Graph)
+		FileItem FinalizeAppBundle(ReadOnlyTargetRules Target, LinkEnvironment LinkEnvironment, FileItem Executable, IActionGraphBuilder Graph)
 		{
 			// Make a file item for the source and destination files
 			string FullDestPath = Executable.AbsolutePath.Substring(0, Executable.AbsolutePath.IndexOf(".app") + 4);
@@ -1016,7 +1017,7 @@ namespace UnrealBuildTool
 
 			FinalizeAppBundleAction.CommandArguments = "\"" + BundleScript.AbsolutePath + "\"";
 			FinalizeAppBundleAction.PrerequisiteItems.Add(Executable);
-			foreach (FileItem Dependency in BundleDependencies)
+			foreach (FileItem Dependency in BundleDependencies.GetValueOrDefault(Target, new()))
 			{
 				FinalizeAppBundleAction.PrerequisiteItems.Add(Dependency);
 			}
@@ -1050,13 +1051,18 @@ namespace UnrealBuildTool
 			return TargetItem;
 		}
 
-		public override void SetupBundleDependencies(List<UEBuildBinary> Binaries, string GameName)
+		public override void SetupBundleDependencies(ReadOnlyTargetRules Target, List<UEBuildBinary> Binaries, string GameName)
 		{
-			base.SetupBundleDependencies(Binaries, GameName);
+			base.SetupBundleDependencies(Target, Binaries, GameName);
+
+			if (!BundleDependencies.ContainsKey(Target))
+			{
+				BundleDependencies.Add(Target, new());
+			}
 
 			foreach (UEBuildBinary Binary in Binaries)
 			{
-				BundleDependencies.Add(FileItem.GetItemByFileReference(Binary.OutputFilePath));
+				BundleDependencies[Target].Add(FileItem.GetItemByFileReference(Binary.OutputFilePath));
 			}
 		}
 
@@ -1179,9 +1185,9 @@ namespace UnrealBuildTool
 			}
 		}
 
-		public override ICollection<FileItem> PostBuild(FileItem Executable, LinkEnvironment BinaryLinkEnvironment, IActionGraphBuilder Graph)
+		public override ICollection<FileItem> PostBuild(ReadOnlyTargetRules Target, FileItem Executable, LinkEnvironment BinaryLinkEnvironment, IActionGraphBuilder Graph)
 		{
-			ICollection<FileItem> OutputFiles = base.PostBuild(Executable, BinaryLinkEnvironment, Graph);
+			ICollection<FileItem> OutputFiles = base.PostBuild(Target, Executable, BinaryLinkEnvironment, Graph);
 
 			if (BinaryLinkEnvironment.bIsBuildingLibrary)
 			{
@@ -1216,7 +1222,7 @@ namespace UnrealBuildTool
 				}
 				else
 				{
-					OutputFiles.Add(FinalizeAppBundle(BinaryLinkEnvironment, Executable, Graph));
+					OutputFiles.Add(FinalizeAppBundle(Target, BinaryLinkEnvironment, Executable, Graph));
 				}
 			}
 
diff --git a/Engine/Source/Programs/UnrealBuildTool/ToolChain/UEToolChain.cs b/Engine/Source/Programs/UnrealBuildTool/ToolChain/UEToolChain.cs
--- a/Engine/Source/Programs/UnrealBuildTool/ToolChain/UEToolChain.cs
+++ b/Engine/Source/Programs/UnrealBuildTool/ToolChain/UEToolChain.cs
@@ -161,15 +161,15 @@ namespace UnrealBuildTool
 			return FileReference.Combine(LinkEnvironment.IntermediateDirectory!, OutputFile.Location.GetFileName() + ".response");
 		}
 
-		public virtual ICollection<FileItem> PostBuild(FileItem Executable, LinkEnvironment ExecutableLinkEnvironment, IActionGraphBuilder Graph)
+		public virtual ICollection<FileItem> PostBuild(ReadOnlyTargetRules Target, FileItem Executable, LinkEnvironment ExecutableLinkEnvironment, IActionGraphBuilder Graph)
 		{
 			return new List<FileItem>();
 		}
 
-		public virtual ICollection<FileItem> PostBuild(FileItem[] Executables, LinkEnvironment ExecutableLinkEnvironment, IActionGraphBuilder Graph)
+		public virtual ICollection<FileItem> PostBuild(ReadOnlyTargetRules Target, FileItem[] Executables, LinkEnvironment ExecutableLinkEnvironment, IActionGraphBuilder Graph)
 		{
 			// by default, run PostBuild for exe Exe and merge results
-			return Executables.SelectMany(x => PostBuild(x, ExecutableLinkEnvironment, Graph)).ToList();
+			return Executables.SelectMany(x => PostBuild(Target, x, ExecutableLinkEnvironment, Graph)).ToList();
 		}
 
 		public virtual void SetUpGlobalEnvironment(ReadOnlyTargetRules Target)
@@ -208,7 +208,7 @@ namespace UnrealBuildTool
 			return OutputFile.ChangeExtension(DebugExtension);
 		}
 
-		public virtual void SetupBundleDependencies(List<UEBuildBinary> Binaries, string GameName)
+		public virtual void SetupBundleDependencies(ReadOnlyTargetRules Target, List<UEBuildBinary> Binaries, string GameName)
 		{
 		}
 
