Adding EngineIncludeOrderVersion for UE5.3

diff --git a/Engine/Source/Programs/UnrealBuildTool/Configuration/TargetRules.cs b/Engine/Source/Programs/UnrealBuildTool/Configuration/TargetRules.cs
--- a/Engine/Source/Programs/UnrealBuildTool/Configuration/TargetRules.cs
+++ b/Engine/Source/Programs/UnrealBuildTool/Configuration/TargetRules.cs
@@ -161,12 +161,17 @@ namespace UnrealBuildTool
 		/// </summary>
 		Unreal5_2,
 
+		/// <summary>
+		/// Include order used in Unreal 5.3
+		/// </summary>
+		Unreal5_3,
+
 		// *** When adding new entries here, be sure to update UEBuildModuleCPP.CurrentIncludeOrderDefine to ensure that the correct guard is used. ***
 
 		/// <summary>
 		/// Always use the latest version of include order.
 		/// </summary>
-		Latest = Unreal5_2,
+		Latest = Unreal5_3,
 
 		/// <summary>
 		/// Contains the oldest version of include order that the engine supports.
@@ -297,6 +302,8 @@ namespace UnrealBuildTool
 			switch (InVersion)
 			{
 				default:
+				case EngineIncludeOrderVersion.Unreal5_3:
+					return "UE_ENABLE_INCLUDE_ORDER_DEPRECATED_IN_5_3";
 				case EngineIncludeOrderVersion.Unreal5_2:
 					return "UE_ENABLE_INCLUDE_ORDER_DEPRECATED_IN_5_2";
 				case EngineIncludeOrderVersion.Unreal5_1:
@@ -321,6 +328,7 @@ namespace UnrealBuildTool
 			{
 				GetDeprecationDefine(EngineIncludeOrderVersion.Unreal5_1),
 				GetDeprecationDefine(EngineIncludeOrderVersion.Unreal5_2),
+				GetDeprecationDefine(EngineIncludeOrderVersion.Unreal5_3),
 			};
 		}
 
@@ -335,6 +343,7 @@ namespace UnrealBuildTool
 			{
 				GetDeprecationDefine(EngineIncludeOrderVersion.Unreal5_1, InVersion),
 				GetDeprecationDefine(EngineIncludeOrderVersion.Unreal5_2, InVersion),
+				GetDeprecationDefine(EngineIncludeOrderVersion.Unreal5_3, InVersion),
 			};
 		}
 
