[UnrealBuildTool - IWYUMode]
* IWYUMode - Changed so ThirdParty includes are not modified. Since iwyu mode has no insight into external modules it ends up producing endlessly long lists of includes. Instead we just rely on that user has added the optimal includes to the code file
* Fixed bug related to p4 edit code.
* Added tracking of if #include is inside declaration or not (namespace/type/function/etc).. in those cases we don't want to remove that include and also not add any includes to that used include
* Renamed -NoTransitiveIncludes to -NoTransitive
* Added -KeepRedundant which will keep includes that are in the header and are transitively included by needed includes. Keeping them shouldnt affect any compile times and dramatically reduced number of changed files.
* Added iwyu developer options -WriteToc and -ReadToc which can be used to avoid having to run build all the time when iterating on iwyu mode itself. Just run once with -WriteToc and then use -ReadToc to skip the building part and jump directly to parsing .iwyu files
* Changed so transitively include generation is using full paths instead of printable paths. This fixes problems where the same header is included using different paths
* Added -Compare which currently is WIP.. plan is to try to give user a sense of how much updating all headers will give in terms of number of lines parsed
* Made sure headers that has no includes but has XXX_API adds HAL/Platform include

diff --git a/Engine/Source/Programs/UnrealBuildTool/Modes/IWYUMode.cs b/Engine/Source/Programs/UnrealBuildTool/Modes/IWYUMode.cs
--- a/Engine/Source/Programs/UnrealBuildTool/Modes/IWYUMode.cs
+++ b/Engine/Source/Programs/UnrealBuildTool/Modes/IWYUMode.cs
@@ -67,6 +67,11 @@ namespace UnrealBuildTool
 		/// Decides if printable should be #include quote or angle bracket
 		/// </summary>
 		public bool System { get; set; }
+
+		/// <summary>
+		/// If true #include was found inside a declaration such as a namespace or class/struct/function
+		/// </summary>
+		public bool InsideDecl { get; set; }
 	}
 
 	/// <summary>
@@ -164,6 +169,11 @@ namespace UnrealBuildTool
 		/// Module that this file belongs to
 		/// </summary>
 		public UEBuildModule? Module;
+
+		/// <summary>
+		/// Is true if this file was included inside a declaration such as a namespace or class/struct/function
+		/// </summary>
+		public bool IncludedInsideDecl;
 	}
 
 	/// <summary>
@@ -237,9 +247,15 @@ namespace UnrealBuildTool
 		/// <summary>
 		/// Allow files to not add includes if they are transitively included by other includes
 		/// </summary>
-		[CommandLine("-NoTransitiveIncludes")]
+		[CommandLine("-NoTransitive")]
 		public bool bNoTransitiveIncludes = false;
 
+		/// <summary>
+		/// Allow files to not add includes if they are transitively included by other includes
+		/// </summary>
+		[CommandLine("-KeepRedundant")]
+		public bool bKeepRedundantIncludes = false;
+
 		/// <summary>
 		/// Will skip compiling before updating. Handle with care, this is dangerous since files might not match .iwyu files
 		/// </summary>
@@ -275,6 +291,24 @@ namespace UnrealBuildTool
 			}
 		}
 
+		/// <summary>
+		/// Compare current include structure with how it would look like if iwyu was applied on all files
+		/// </summary>
+		[CommandLine("-Compare")]
+		public bool bCompare = false;
+
+		/// <summary>
+		/// For Development only - Will write a toc referencing all .iwyu files. Toc can be used by -ReadToc
+		/// </summary>
+		[CommandLine("-WriteToc")]
+		public bool bWriteToc = false;
+
+		/// <summary>
+		/// For Development only - Will read a toc to find all .iwyu files instead of building the code.
+		/// </summary>
+		[CommandLine("-ReadToc")]
+		public bool bReadToc = false;
+
 		private string? GetModuleToUpdateName(TargetDescriptor Descriptor)
 		{
 			if (!String.IsNullOrEmpty(ModuleToUpdate))
@@ -361,15 +395,19 @@ namespace UnrealBuildTool
 					PathToUpdateList[I] = PathToUpdateList[I].Replace('\\', '/');
 				}
 			}
-			
+
+			string TargetName = $"{TargetDescriptors[0].Name}_{TargetDescriptors[0].Configuration}_{TargetDescriptors[0].Platform}";
+
 			// Calculate file paths filter to figure out which files that IWYU will run on
-			Logger.LogInformation($"Calculating file filter for IWYU...");
 			HashSet<string> ValidPaths = new();
 			Dictionary<string, UEBuildModule?> PathToModule = new();
 			{
 				// Create target to be able to traverse all existing modules
+				Logger.LogInformation($"Creating BuildTarget for {TargetName}...");
 				UEBuildTarget Target = UEBuildTarget.Create(TargetDescriptors[0], BuildConfiguration.bSkipRulesCompile, BuildConfiguration.bForceRulesCompile, BuildConfiguration.bUsePrecompiled, Logger);
 
+				Logger.LogInformation($"Calculating file filter for IWYU...");
+
 				UEBuildModule? UEModuleToUpdate = null;
 
 				// Turn provided module string into UEBuildModule reference
@@ -382,7 +420,7 @@ namespace UnrealBuildTool
 						return -1;
 					}
 				}
-				
+
 				if (PathToUpdateList.Length == 0 && UEModuleToUpdate == null)
 				{
 					Logger.LogError($"Need to provide -ModuleToUpdate <modulename> or -PathToUpdate <partofpath> to run IWYU.");
@@ -438,11 +476,11 @@ namespace UnrealBuildTool
 							}
 							ShouldUpdate = DependsOnModule;
 						}
-						
+
 						if (PathToUpdateList.Length > 0)
 						{
 							bool Match = false;
-							for (int I=0; I!= PathToUpdateList.Length; ++I)
+							for (int I = 0; I != PathToUpdateList.Length; ++I)
 							{
 								Match = Match || ModuleDir.Contains(PathToUpdateList[I], StringComparison.OrdinalIgnoreCase);
 							}
@@ -510,54 +548,74 @@ namespace UnrealBuildTool
 			IWYUInfo? ObjectMacrosInfo = null;
 			IWYUInfo? ScriptMacrosInfo = null;
 			IWYUInfo? PlatformInfo = null;
-			int ReadSuccess = 1;
 
-			TargetDescriptor Descriptor = TargetDescriptors.First();
 			using (ISourceFileWorkingSet WorkingSet = new EmptySourceFileWorkingSet())
 			{
-				// Create the make file that contains all the actions we will use to find .iwyu files.
-				Logger.LogInformation($"Creating MakeFile for target...");
-				TargetMakefile Makefile = BuildMode.CreateMakefile(BuildConfiguration, Descriptor, WorkingSet, Logger);
+				List<string>? TocContent = null;
+				int ReadSuccess = 1;
 
-				// Use the make file to build unless -NoCompile is set.
-				if (!bNoCompile)
+				if (bReadToc)
 				{
-					BuildMode.Build(new TargetMakefile[] { Makefile }, new List<TargetDescriptor>() { Descriptor }, BuildConfiguration, BuildOptions.None, null, Logger);
-					SourceFileMetadataCache.SaveAll();
-					CppDependencyCache.SaveAll();
+					string TocName = TargetName + ".txt";
+					FileReference TocReference = FileReference.Combine(Unreal.EngineDirectory, "Intermediate", "IWYU", TocName);
+					Logger.LogInformation($"Reading TOC from {TocReference.FullName}...");
+					if (FileReference.Exists(TocReference))
+					{
+						TocContent = FileReference.ReadAllLines(TocReference).ToList();
+					}
 				}
 
-				HashSet<FileItem> OutputItems = new HashSet<FileItem>();
-				if (Descriptor.OnlyModuleNames.Count > 0)
+				if (TocContent == null)
 				{
-					foreach (string OnlyModuleName in Descriptor.OnlyModuleNames)
+					TargetDescriptor Descriptor = TargetDescriptors.First();
+
+					// Create the make file that contains all the actions we will use to find .iwyu files.
+					Logger.LogInformation($"Creating MakeFile for target...");
+					TargetMakefile Makefile = BuildMode.CreateMakefile(BuildConfiguration, Descriptor, WorkingSet, Logger);
+
+					// Use the make file to build unless -NoCompile is set.
+					if (!bNoCompile)
+					{
+						BuildMode.Build(new TargetMakefile[] { Makefile }, new List<TargetDescriptor>() { Descriptor }, BuildConfiguration, BuildOptions.None, null, Logger);
+						SourceFileMetadataCache.SaveAll();
+						CppDependencyCache.SaveAll();
+					}
+
+					HashSet<FileItem> OutputItems = new HashSet<FileItem>();
+					if (Descriptor.OnlyModuleNames.Count > 0)
 					{
-						FileItem[]? OutputItemsForModule;
-						if (!Makefile.ModuleNameToOutputItems.TryGetValue(OnlyModuleName, out OutputItemsForModule))
+						foreach (string OnlyModuleName in Descriptor.OnlyModuleNames)
 						{
-							throw new BuildException("Unable to find output items for module '{0}'", OnlyModuleName);
+							FileItem[]? OutputItemsForModule;
+							if (!Makefile.ModuleNameToOutputItems.TryGetValue(OnlyModuleName, out OutputItemsForModule))
+							{
+								throw new BuildException("Unable to find output items for module '{0}'", OnlyModuleName);
+							}
+							OutputItems.UnionWith(OutputItemsForModule);
+						}
+					}
+					else
+					{
+						// Use all the output items from the target
+						OutputItems.UnionWith(Makefile.OutputItems);
+					}
+
+					TocContent = new();
+
+					foreach (var OutputItem in OutputItems)
+					{
+						if (OutputItem.Name.EndsWith(".iwyu"))
+						{
+							TocContent.Add(OutputItem.AbsolutePath);
 						}
-						OutputItems.UnionWith(OutputItemsForModule);
 					}
-				}
-				else
-				{
-					// Use all the output items from the target
-					OutputItems.UnionWith(Makefile.OutputItems);
 				}
 
 				// Time to parse all the .iwyu files generated from iwyu.
 				// We Do this in parallel since it involves reading a ton of .iwyu files from disk.
-				Logger.LogInformation("Parsing output from IWYU...");
-				Parallel.ForEach(OutputItems, OutputItem =>
+				Logger.LogInformation($"Parsing {TocContent.Count} .iwyu files...");
+				Parallel.ForEach(TocContent, IWYUFilePath =>
 				{
-					if (!OutputItem.Name.EndsWith(".iwyu"))
-					{
-						return;
-					}
-
-					string IWYUFilePath = OutputItem.AbsolutePath;
-
 					string? JsonContent = File.ReadAllText(IWYUFilePath);
 					if (JsonContent == null)
 					{
@@ -577,7 +635,7 @@ namespace UnrealBuildTool
 							Info.Source = IWYUFile;
 
 							// We track .gen.cpp in a special list, they need special treatment later
-							if (Info.File.Contains(".gen.cpp"))
+							if (Info.File.Contains(".gen.cpp", StringComparison.Ordinal))
 							{
 								lock (GeneratedCppInfos)
 									GeneratedCppInfos.Add(Info);
@@ -614,6 +672,16 @@ namespace UnrealBuildTool
 								}
 							}
 
+							// TODO: Fix bad formatting coming from iwyu
+							foreach (IWYUIncludeEntry Entry in Info.IncludesSeenInFile)
+							{
+								if (Entry.Printable[0] == '<')
+								{
+									Entry.Printable = Entry.Printable.Substring(1, Entry.Printable.Length - 2);
+									Entry.System = true;
+								}
+							}
+
 							string FilePath = Info.File;
 							while (true)
 							{
@@ -657,11 +725,21 @@ namespace UnrealBuildTool
 						return;
 					}
 				});
-			}
 
-			// Something went wrong parsing iwyu files.
-			if (ReadSuccess == 0)
-				return -1;
+				// Something went wrong parsing iwyu files.
+				if (ReadSuccess == 0)
+					return -1;
+
+				if (bWriteToc)
+				{
+					Logger.LogInformation($"Writing TOC that references all .iwyu files ");
+					DirectoryReference TocDir = DirectoryReference.Combine(Unreal.EngineDirectory, "Intermediate", "IWYU");
+					DirectoryReference.CreateDirectory(TocDir);
+					string TocName = TargetName + ".txt";
+					FileReference TocReference = FileReference.Combine(TocDir, TocName);
+					FileReference.WriteAllLines(TocReference, TocContent);
+				}
+			}
 
 			// Add all .generated.h files as entries in the lookup and explicitly add the includes they have which will never be removed
 			Logger.LogInformation($"Generating infos for .generated.h files...");
@@ -689,6 +767,13 @@ namespace UnrealBuildTool
 			Logger.LogInformation($"Found {Infos.Count} IWYU entries...");
 
 			Logger.LogInformation($"Resolving Includes...");
+
+			var PlatformSDK = new LinuxPlatformSDK(Logger);
+			DirectoryReference? BaseLinuxPath = PlatformSDK.GetBaseLinuxPathForArchitecture(LinuxPlatform.DefaultHostArchitecture);
+			string SystemPath = BaseLinuxPath!.FullName.Replace('\\', '/');
+
+			HashSet<IWYUInfo> IncludedInsideDecl = new();
+
 			Parallel.ForEach(Infos.Values, Info =>
 			{
 				// If KeepAsIs we transfer all "seen includes" into the include list.
@@ -697,15 +782,45 @@ namespace UnrealBuildTool
 					foreach (IWYUIncludeEntry Entry in Info.IncludesSeenInFile)
 					{
 						// Special hack, we don't want CoreMinimal to be the reason we remove includes transitively
-						if (!Entry.Printable.Contains("CoreMinimal.h"))
+						if (!Entry.Printable.Contains("CoreMinimal.h", StringComparison.Ordinal))
 						{
 							Info.Includes.Add(Entry);
 						}
 					}
 				}
 
+				// We don't want to mess around with third party includes.. since we can't see the hierarchy we just assumes that they are optimally included
+				Info.Includes.RemoveAll(Entry => Entry.Full.Contains("/ThirdParty/", StringComparison.Ordinal) || Entry.Full.StartsWith(SystemPath));
+
+				foreach (IWYUIncludeEntry Entry in Info.IncludesSeenInFile)
+				{
+					if (Entry.Full.Contains("/ThirdParty/", StringComparison.Ordinal) || Entry.Full.StartsWith(SystemPath, StringComparison.Ordinal))
+					{
+						Info.Includes.Add(Entry);
+					}
+
+					if (Entry.InsideDecl && Infos.TryGetValue(Entry.Full, out Entry.Resolved))
+					{
+						lock (IncludedInsideDecl)
+						{
+							IncludedInsideDecl.Add(Entry.Resolved);
+							Info.Includes.Add(Entry);
+						}
+					}
+				}
+				/*
+				if (Info.IncludesSeenInFile.Count == 0)
+				{
+					if (Info.File.EndsWith(".inl"))
+					{
+						Info.Includes.Clear();
+						Info.ForwardDeclarations.Clear();
+					}
+				}
+				*/
 				// Definitions.h is automatically added in the reponse file and iwyu sees it as not included
 				// If there are no includes but we depend on Definitions.h (which is missing) we add Platform.h because it is most likely a <module>_API entry in the file
+				/*
 				if (Info.Includes.Count == 0)
 				{
 					bool IsUsingDefinitionsH = false;
@@ -718,19 +833,33 @@ namespace UnrealBuildTool
 						Info.Includes.Add(new IWYUIncludeEntry() { Full = PlatformInfo!.File, Printable = "HAL/Platform.h", Resolved = PlatformInfo });
 					}
 				}
-
+				*/
 				foreach (var Include in Info.Includes)
 				{
 					if (Include.Resolved == null)
+					{
 						Infos.TryGetValue(Include.Full, out Include.Resolved);
+					}
 				}
 			});
 
+			foreach (IWYUInfo Info in IncludedInsideDecl)
+			{
+				if (Info.IncludesSeenInFile.Count != 0 && !Info.File.EndsWith("ScriptSerialization.h")) // Remove include in ScriptSerialization.h and remove this check
+				{
+					Logger.LogWarning($"{Info.File} - Included inside declaration in other file but has includes itself.");
+				}
+				Info.Includes.Clear();
+				Info.ForwardDeclarations.Clear();
+			}
+
 			Logger.LogInformation($"Generating transitive include lists and forward declaration lists...");
 			Parallel.ForEach(Infos.Values, Info =>
 			{
 				Stack<string> Stack = new();
-				CalculateTransitive(Info, Info, Stack, Info.TransitiveIncludes, Info.TransitiveForwardDeclarations);
+				Info.TransitiveIncludes.EnsureCapacity(300);
+				Info.TransitiveForwardDeclarations.EnsureCapacity(300);
+				CalculateTransitive(Info, Info, Stack, Info.TransitiveIncludes, Info.TransitiveForwardDeclarations, false);
 			});
 
 			// If we have built .gen.cpp it means that it is not inlined in another cpp file
@@ -742,7 +871,7 @@ namespace UnrealBuildTool
 				HashSet<string> SeenTransitiveIncludes = new();
 				foreach (var SeenInclude in GeneratedCpp.IncludesSeenInFile)
 				{
-					SeenTransitiveIncludes.Add(SeenInclude.Printable);
+					SeenTransitiveIncludes.Add(SeenInclude.Full);
 					foreach (var Include in GeneratedCpp.Includes)
 					{
 						IWYUInfo? IncludeInfo;
@@ -785,7 +914,7 @@ namespace UnrealBuildTool
 
 				foreach (var Include in GeneratedCpp.Includes)
 				{
-					if (SeenTransitiveIncludes.Contains(Include.Printable))
+					if (SeenTransitiveIncludes.Contains(Include.Full))
 					{
 						continue;
 					}
@@ -797,13 +926,57 @@ namespace UnrealBuildTool
 						continue;
 					}
 
-					if (!IncluderInfo!.TransitiveIncludes.ContainsKey(Include.Printable))
+					if (!IncluderInfo!.TransitiveIncludes.ContainsKey(Include.Full))
 					{
-						IncluderInfo.Includes.Add(Include);
+						if (!Include.Full.Contains("/ThirdParty/") && !Include.Full.StartsWith(SystemPath))
+						{
+							IncluderInfo.Includes.Add(Include);
+						}
 					}
 				}
 			});
 
+			if (bCompare)
+			{
+				return CompareFiles(Infos, ValidPaths, Logger);
+			}
+			else
+			{
+				return UpdateFiles(Infos, ValidPaths, Logger);
+			}
+		}
+
+		bool IsValidForUpdate(IWYUInfo Info, HashSet<string> ValidPaths, bool ObeyModuleRules)
+		{
+			if (Info.Source == null) // .generated.h is also in this list, ignore them
+			{
+				return false;
+			}
+
+			if (ObeyModuleRules && Info.Module?.Rules.IWYUSupport != IWYUSupport.Full)
+			{
+				return false;
+			}
+
+			// There are some codegen files with this name
+			if (Info.File.Contains(".gen.h"))
+			{
+				return false;
+			}
+
+			// Filter out files
+			foreach (string ValidPath in ValidPaths)
+			{
+				if (Info.File.Contains(ValidPath))
+				{
+					return true;
+				}
+			}
+			return false;
+		}
+
+		int UpdateFiles(Dictionary<string, IWYUInfo> Infos, HashSet<string> ValidPaths, ILogger Logger)
+		{
 			object? ShouldLog = bWrite ? null : new();
 
 			List<ValueTuple<IWYUInfo, List<string>>> UpdatedFiles = new(); // <FilePath>, <NewLines>
@@ -818,19 +991,8 @@ namespace UnrealBuildTool
 
 			Parallel.ForEach(Infos.Values, Info =>
 			{
-				if (Info.Source == null) // .generated.h is also in this list, ignore them
-					return;
-
-				// Filter out files
-				if (!ValidPaths.Any(ValidPath => Info.File.Contains(ValidPath, StringComparison.OrdinalIgnoreCase)))
-				{
-					return;
-				}
-
-				if (Info.Module?.Rules.IWYUSupport != IWYUSupport.Full)
-				{
+				if (!IsValidForUpdate(Info, ValidPaths, true))
 					return;
-				}
 
 				bool IsCpp = Info.File.EndsWith(".cpp");
 				bool IsPrivate = IsCpp || Info.File.Contains("/Private/");
@@ -875,7 +1037,7 @@ namespace UnrealBuildTool
 						{
 							if (Include2.Resolved != null && Include != Include2)
 							{
-								string Key = Include.Printable; //Full;
+								string Key = Include.Full;
 								string? TransitivePath;
 								if (Include2.Resolved.TransitiveIncludes!.TryGetValue(Key, out TransitivePath))
 								{
@@ -893,7 +1055,22 @@ namespace UnrealBuildTool
 					}
 				}
 
-				if (!IsCpp) // Ignore forward declarations for cpp files. They should never need them afaik
+				// If keep redundant includes is set we don't remove the includes that are included through other includes
+				if (bKeepRedundantIncludes)
+				{
+					List<string> ToReadd = new();
+					foreach (var Seen in Info.IncludesSeenInFile)
+					{
+						string QuotedPrintable = Seen.System ? $"<{Seen.Printable}>" : $"\"{Seen.Printable}\"";
+						if (!CleanedupIncludes.Contains(QuotedPrintable) && Info.TransitiveIncludes.ContainsKey(Seen.Full))
+						{
+							CleanedupIncludes.Add(QuotedPrintable);
+						}
+					}
+				}
+
+				// Ignore forward declarations for cpp files. They are very rarely needed and we let the user add them manually instead
+				if (!IsCpp)
 				{
 					foreach (var ForwardDeclaration in Info.ForwardDeclarations)
 					{
@@ -920,8 +1097,7 @@ namespace UnrealBuildTool
 				// Read all lines of the header/source file
 				string[] ExistingLines = File.ReadAllLines(Info.File);
 
-				SortedSet<string> LinesToRemove = new();
-				SortedSet<string> IncludesToRemove = new();
+				SortedDictionary<string, string> LinesToRemove = new();
 				SortedSet<string> IncludesToAdd = new(CleanedupIncludes);
 
 				bool HasIncludes = false;
@@ -930,16 +1106,30 @@ namespace UnrealBuildTool
 				HashSet<string> SeenIncludes = new();
 				foreach (IWYUIncludeEntry SeenInclude in Info.IncludesSeenInFile)
 				{
-					if (SeenInclude.Printable[0] != '<')
-						SeenIncludes.Add($"\"{SeenInclude.Printable}\"");
+					if (SeenInclude.System)
+						SeenIncludes.Add($"<{SeenInclude.Printable}>");
 					else
-						SeenIncludes.Add(SeenInclude.Printable);
+						SeenIncludes.Add($"\"{SeenInclude.Printable}\"");
 				}
 
 
 				bool ForceKeepScope = false;
 				bool ErrorOnMoreIncludes = false;
 				int LineIndex = -1;
+
+
+				// This makes sure that we have at least HAL/Platform.h included if the file contains XXX_API
+				bool Contains_API = false;
+				bool LookFor_API = false;
+				if (Info.Includes.Count == 0)
+				{
+					foreach (var Missing in Info.MissingIncludes)
+					{
+						LookFor_API = LookFor_API || Missing.Full.Contains("Definitions.h");
+					}
+				}
+
+				// Traverse all lines in file and figure out which includes that should be added or removed
 				foreach (var L in ExistingLines)
 				{
 					++LineIndex;
@@ -970,6 +1160,9 @@ namespace UnrealBuildTool
 							return;
 						}
 
+						if (LookFor_API && !Contains_API && Line.Contains("_API", StringComparison.Ordinal))
+							Contains_API = true;
+
 						continue;
 					}
 
@@ -1023,19 +1216,19 @@ namespace UnrealBuildTool
 							IncludeSpan = IncludeSpan.Slice(0, Index + 2);
 						}
 
-						if (Line.Contains("IWYU pragma: "))
+						if (Line.Contains("IWYU pragma: ", StringComparison.Ordinal))
 							ForceKeep = true;
 					}
 
 					string Include = IncludeSpan.ToString();
 
+					// If Include is not seen it means that it is probably inside a #if/#endif with condition false. These includes we can't touch
 					if (!SeenIncludes.Contains(Include))
 						continue;
 
 					if (!ForceKeep && !ForceKeepScope && !CleanedupIncludes.Contains(Include))
 					{
-						LinesToRemove.Add(L);
-						IncludesToRemove.Add(Include);
+						LinesToRemove.TryAdd(L, Include);
 					}
 					else
 					{
@@ -1043,6 +1236,13 @@ namespace UnrealBuildTool
 					}
 				}
 
+				if (Contains_API && LookFor_API)
+				{
+					//IncludesToAdd.Add(new IWYUIncludeEntry() { Full = PlatformInfo!.File, Printable = "HAL/Platform.h", Resolved = PlatformInfo });
+					if (!LinesToRemove.Remove("#include \"HAL/Platform.h\""))
+						IncludesToAdd.Add("\"HAL/Platform.h\"");
+				}
+
 				// Nothing has changed! early out of this file
 				if (IncludesToAdd.Count == 0 && LinesToRemove.Count == 0 && ForwardDeclarationsToAdd.Count == 0)
 				{
@@ -1078,11 +1278,11 @@ namespace UnrealBuildTool
 						{
 							System.Console.WriteLine("  +" + I);
 						}
-						foreach (var I in IncludesToRemove)
+						foreach (var Pair in LinesToRemove)
 						{
-							System.Console.Write("  -" + I);
+							System.Console.Write("  -" + Pair.Key);
 							string? Reason;
-							if (TransitivelyIncluded.TryGetValue(I, out Reason))
+							if (TransitivelyIncluded.TryGetValue(Pair.Value, out Reason))
 							{
 								System.Console.Write("  (Transitively included from " + Reason + ")");
 							}
@@ -1164,7 +1364,7 @@ namespace UnrealBuildTool
 
 						ContiguousNonIncludeLineCount = 0;
 
-						if (!LinesToRemove.Contains(OldLine))
+						if (!LinesToRemove.ContainsKey(OldLine))
 						{
 							// If we find #include SOME_DEFINE we assume that should be last and "end" the include block with that
 							if (!OldLineTrimmedStart.Contains('\"') && !OldLineTrimmedStart.Contains('<'))
@@ -1197,7 +1397,7 @@ namespace UnrealBuildTool
 					if (LinesToAdd.Count > 0)
 					{
 						int InsertPos = LastSeenIncludeBeforeCode + 1;
-						if (NewLines[LastSeenIncludeBeforeCode].Contains(".generated.h"))
+						if (NewLines[LastSeenIncludeBeforeCode].Contains(".generated.h", StringComparison.Ordinal))
 						{
 							--InsertPos;
 						}
@@ -1234,6 +1434,10 @@ namespace UnrealBuildTool
 				string EngineDir = Unreal.EngineDirectory.FullName.Replace('\\', '/');
 				if (!IsPrivate && Info.File.StartsWith(EngineDir) && !String.IsNullOrEmpty(HeaderDeprecationTag))
 				{
+					Dictionary<string, string> PrintableToFull = new();
+					foreach (var Seen in Info.IncludesSeenInFile)
+						PrintableToFull.TryAdd(Seen.Printable, Seen.Full);
+
 					// Remove the includes in LinesRemoved
 					LinesRemoved.RemoveWhere(Line =>
 					{
@@ -1249,7 +1453,10 @@ namespace UnrealBuildTool
 							return false;
 						}
 						IncludeSpan = IncludeSpan.Slice(1, Index);
-						return Info.TransitiveIncludes.ContainsKey(IncludeSpan.ToString());
+						string? Full;
+						if (!PrintableToFull.TryGetValue(IncludeSpan.ToString(), out Full))
+							return false;
+						return Info.TransitiveIncludes.ContainsKey(Full);
 					});
 
 					if (LinesRemoved.Count > 0)
@@ -1345,7 +1552,7 @@ namespace UnrealBuildTool
 				{
 					// Check out files in batches. This can go quite crazy if there are lots of files.
 					// Should probably revisit this code to prevent 100s of p4 processes to start at once
-					Logger.LogInformation($"Opening {UpdatedFiles.Count} files for edit in P4... ({SkippedCount} files skipped)");
+					Logger.LogInformation($"Opening {ReadOnlyFileInfos.Count} files for edit in P4... ({SkippedCount} files skipped)");
 					int ShowCount = 8;
 					foreach (var Info in ReadOnlyFileInfos)
 					{
@@ -1353,8 +1560,8 @@ namespace UnrealBuildTool
 						if (--ShowCount == 0)
 							break;
 					}
-					if (UpdatedFiles.Count > 5)
-						Logger.LogInformation($"   ... and {UpdatedFiles.Count - 5} more.");
+					if (ReadOnlyFileInfos.Count > 5)
+						Logger.LogInformation($"   ... and {ReadOnlyFileInfos.Count - 5} more.");
 
 					StringBuilder P4Arguments = new();
 					int BatchSize = 10;
@@ -1368,7 +1575,7 @@ namespace UnrealBuildTool
 							++BatchCount;
 						}
 						++Index;
-						if (BatchCount == BatchSize || Index == UpdatedFiles.Count)
+						if (BatchCount == BatchSize || Index == ReadOnlyFileInfos.Count)
 						{
 							AddP4Process($"edit{P4Arguments}");
 							P4Arguments.Clear();
@@ -1409,7 +1616,7 @@ namespace UnrealBuildTool
 			if (!WaitForP4())
 			{
 				return -1;
-			}	
+			}
 
 			if (!WriteSuccess)
 				return -1;
@@ -1421,11 +1628,13 @@ namespace UnrealBuildTool
 		/// <summary>
 		/// Calculate all indirect transitive includes for a file. This list contains does not contain itself and will handle circular dependencies
 		/// </summary>
-		void CalculateTransitive(IWYUInfo Root, IWYUInfo Info, Stack<string> Stack, Dictionary<string, string> TransitiveIncludes, Dictionary<string, string> TransitiveForwardDeclarations)
+		void CalculateTransitive(IWYUInfo Root, IWYUInfo Info, Stack<string> Stack, Dictionary<string, string> TransitiveIncludes, Dictionary<string, string> TransitiveForwardDeclarations, bool UseSeenIncludes)
 		{
-			foreach (var Include in Info.Includes)
+			List<IWYUIncludeEntry> Includes = UseSeenIncludes ? Info.IncludesSeenInFile : Info.Includes;
+			foreach (var Include in Includes)
 			{
-				string Key = Include.Printable; // Full
+				string Key = Include.Full;
+
 				if (TransitiveIncludes.ContainsKey(Key) || Include.Resolved == Root)
 				{
 					continue;
@@ -1436,7 +1645,7 @@ namespace UnrealBuildTool
 				TransitiveIncludes.Add(Key, TransitivePath);
 				if (Include.Resolved != null)
 				{
-					CalculateTransitive(Root, Include.Resolved, Stack, TransitiveIncludes, TransitiveForwardDeclarations);
+					CalculateTransitive(Root, Include.Resolved, Stack, TransitiveIncludes, TransitiveForwardDeclarations, UseSeenIncludes);
 				}
 				Stack.Pop();
 			}
@@ -1444,9 +1653,70 @@ namespace UnrealBuildTool
 			foreach (var ForwardDeclaration in Info.ForwardDeclarations)
 			{
 				string Key = ForwardDeclaration.Printable;
-				if (!TransitiveForwardDeclarations.ContainsKey(Key))
-					TransitiveForwardDeclarations.Add(Key, Info.File);
+				TransitiveForwardDeclarations.TryAdd(Key, Info.File);
 			}
 		}
+
+		int CompareFiles(Dictionary<string, IWYUInfo> Infos, HashSet<string> ValidPaths, ILogger Logger)
+		{
+			Logger.LogInformation($"Comparing code files...");
+
+			SortedSet<string> UnresolvableFiles = new();
+
+			Parallel.ForEach(Infos.Values, Info =>
+			{
+				foreach (var Include in Info.IncludesSeenInFile)
+				{
+					if (Include.Resolved == null)
+					{
+						if (!Infos.TryGetValue(Include.Full, out Include.Resolved))
+						{
+							if (!Include.Full.Contains(".gen.cpp") && !Include.Full.Contains("/ThirdParty/") && !Include.Full.Contains("/AutoSDK/") && !Include.Full.Contains("/VNI/"))
+							{
+								lock (UnresolvableFiles)
+									UnresolvableFiles.Add(Include.Full);
+							}
+						}
+					}
+				}
+			});
+
+			//foreach (string Unresolved in UnresolvableFiles)
+			//	Console.WriteLine("Failed to resolve " + Unresolved);
+
+			List<ValueTuple<float, IWYUInfo>> RatioList = new();
+
+			Parallel.ForEach(Infos.Values, Info =>
+			{
+				if (!IsValidForUpdate(Info, ValidPaths, false))
+					return;
+
+				Dictionary<string, string> SeenTransitiveIncludes = new();
+				Stack<string> Stack = new();
+				CalculateTransitive(Info, Info, Stack, SeenTransitiveIncludes, new Dictionary<string, string>(), true);
+
+				// Sometimes seen are lower and that is because there are a bunch of files that seen includes but is not visible (UnrealNames.inl, UnrealPlatformMathSSE4.h etc
+				// So actually Seen should be 3-4 higher than what it says (probably more)
+				if (Info.TransitiveIncludes.Count < SeenTransitiveIncludes.Count)
+				{
+					float Ratio = (float)Info.TransitiveIncludes.Count / SeenTransitiveIncludes.Count;
+					lock (RatioList)
+						RatioList.Add(new(Ratio, Info));
+				}
+				//Console.WriteLine($"{Info.File} - Now: {SeenTransitiveIncludes.Count}  Optimized: {Info.TransitiveIncludes.Count}");
+			});
+
+			RatioList.Sort((a, b) =>
+			{
+				if (a.Item1 == b.Item1)
+					return 0;
+				if (a.Item1 < b.Item1)
+					return -1;
+				return 1;
+			});
+			//var List = RatioList.OrderBy(I => I.Item1);
+
+			return 0;
+		}
 	}
 }
