UnrealBuildTool: Mac legacy app bundles don't need to have dylib as prerequisites.

This removes an unnecessary dependency which causes conflicts when building multiple editors in the same call to UBT

#rnx

diff --git a/Engine/Source/Programs/UnrealBuildTool/Platform/Mac/MacToolChain.cs b/Engine/Source/Programs/UnrealBuildTool/Platform/Mac/MacToolChain.cs
--- a/Engine/Source/Programs/UnrealBuildTool/Platform/Mac/MacToolChain.cs
+++ b/Engine/Source/Programs/UnrealBuildTool/Platform/Mac/MacToolChain.cs
@@ -101,8 +101,6 @@ namespace UnrealBuildTool
 		/// </summary>
 		private const string MacArchiver = "libtool";
 
-		private static Dictionary<ReadOnlyTargetRules, List<FileItem>> BundleDependencies = new();
-
 		protected override ClangToolChainInfo GetToolChainInfo()
 		{
 			FileReference CompilerPath = new FileReference(Settings.ToolchainDir + MacCompiler);
@@ -1017,10 +1015,6 @@ namespace UnrealBuildTool
 
 			FinalizeAppBundleAction.CommandArguments = "\"" + BundleScript.AbsolutePath + "\"";
 			FinalizeAppBundleAction.PrerequisiteItems.Add(Executable);
-			foreach (FileItem Dependency in BundleDependencies.GetValueOrDefault(Target, new()))
-			{
-				FinalizeAppBundleAction.PrerequisiteItems.Add(Dependency);
-			}
 			FinalizeAppBundleAction.ProducedItems.Add(DestFile);
 			FinalizeAppBundleAction.StatusDescription = string.Format("Finalizing app bundle: {0}.app", Path.GetFileName(Executable.AbsolutePath));
 			FinalizeAppBundleAction.bCanExecuteRemotely = false;
@@ -1051,21 +1045,6 @@ namespace UnrealBuildTool
 			return TargetItem;
 		}
 
-		public override void SetupBundleDependencies(ReadOnlyTargetRules Target, List<UEBuildBinary> Binaries, string GameName)
-		{
-			base.SetupBundleDependencies(Target, Binaries, GameName);
-
-			if (!BundleDependencies.ContainsKey(Target))
-			{
-				BundleDependencies.Add(Target, new());
-			}
-
-			foreach (UEBuildBinary Binary in Binaries)
-			{
-				BundleDependencies[Target].Add(FileItem.GetItemByFileReference(Binary.OutputFilePath));
-			}
-		}
-
 		private static Dictionary<ReadOnlyTargetRules, DirectoryReference> BundleContentsDirectories = new();
 
 		public override void ModifyBuildProducts(ReadOnlyTargetRules Target, UEBuildBinary Binary, List<string> Libraries, List<UEBuildBundleResource> BundleResources, Dictionary<FileReference, BuildProductType> BuildProducts)
