Added a missing circular dependency with AssetTools and UnrealEd.
Updated AssetTypeActivationOpenedMethod.h to include the generated.h.

diff --git a/Engine/Source/Developer/AssetTools/AssetTools.Build.cs b/Engine/Source/Developer/AssetTools/AssetTools.Build.cs
--- a/Engine/Source/Developer/AssetTools/AssetTools.Build.cs
+++ b/Engine/Source/Developer/AssetTools/AssetTools.Build.cs
@@ -109,5 +109,11 @@ public class AssetTools : ModuleRules
                 "VirtualTexturingEditor"
             }
 		);
+
+		CircularlyReferencedDependentModules.AddRange(
+			new string[] {
+				"UnrealEd",
+			}
+		);
 	}
 }
diff --git a/Engine/Source/Developer/AssetTools/Public/AssetTypeActivationOpenedMethod.h b/Engine/Source/Developer/AssetTools/Public/AssetTypeActivationOpenedMethod.h
--- a/Engine/Source/Developer/AssetTools/Public/AssetTypeActivationOpenedMethod.h
+++ b/Engine/Source/Developer/AssetTools/Public/AssetTypeActivationOpenedMethod.h
@@ -3,6 +3,7 @@
 #pragma once
 
 #include "UObject/ObjectMacros.h"
+#include "AssetTypeActivationOpenedMethod.generated.h"
 
 /**
  *	(jcotton) This enum has been extracted into a separate header as it would ideally live in IAssetTypeActions.h or AssetTypeActions_Base.h, 
diff --git a/Engine/Source/Programs/UnrealBuildTool/Configuration/UEBuildModuleCPP.cs b/Engine/Source/Programs/UnrealBuildTool/Configuration/UEBuildModuleCPP.cs
--- a/Engine/Source/Programs/UnrealBuildTool/Configuration/UEBuildModuleCPP.cs
+++ b/Engine/Source/Programs/UnrealBuildTool/Configuration/UEBuildModuleCPP.cs
@@ -112,28 +112,29 @@ namespace UnrealBuildTool
 		static readonly KeyValuePair<string, string>[] CircularDependenciesAllowList =
 		{
 			new KeyValuePair<string, string>("AIModule", "AITestSuite"),
-			new KeyValuePair<string, string>("AnimGraph", "UnrealEd"),
 			new KeyValuePair<string, string>("AnimGraph", "GraphEditor"),
+			new KeyValuePair<string, string>("AnimGraph", "UnrealEd"),
+			new KeyValuePair<string, string>("AudioEditor", "DetailCustomizations"),
 			new KeyValuePair<string, string>("AudioMixer", "NonRealtimeAudioRenderer"),
 			new KeyValuePair<string, string>("AudioMixer", "SoundFieldRendering"),
-			new KeyValuePair<string, string>("AudioEditor", "DetailCustomizations"),
-			new KeyValuePair<string, string>("BlueprintGraph", "KismetCompiler"),
-			new KeyValuePair<string, string>("BlueprintGraph", "UnrealEd"),
+			new KeyValuePair<string, string>("AssetTools", "UnrealEd"),
+			new KeyValuePair<string, string>("BlueprintGraph", "CinematicCamera"),
 			new KeyValuePair<string, string>("BlueprintGraph", "GraphEditor"),
 			new KeyValuePair<string, string>("BlueprintGraph", "Kismet"),
-			new KeyValuePair<string, string>("BlueprintGraph", "CinematicCamera"),
+			new KeyValuePair<string, string>("BlueprintGraph", "KismetCompiler"),
+			new KeyValuePair<string, string>("BlueprintGraph", "UnrealEd"),
 			new KeyValuePair<string, string>("ConfigEditor", "PropertyEditor"),
 			new KeyValuePair<string, string>("Documentation", "SourceControl"),
-			new KeyValuePair<string, string>("Engine", "Landscape"),
-			new KeyValuePair<string, string>("Engine", "UMG"),
-			new KeyValuePair<string, string>("Engine", "GameplayTags"),
-			new KeyValuePair<string, string>("Engine", "MaterialShaderQualitySettings"),
-			new KeyValuePair<string, string>("Engine", "UnrealEd"),
 			new KeyValuePair<string, string>("Engine", "AudioMixer"),
 			new KeyValuePair<string, string>("Engine", "CinematicCamera"),
 			new KeyValuePair<string, string>("Engine", "CollisionAnalyzer"),
-			new KeyValuePair<string, string>("Engine", "LogVisualizer"),
+			new KeyValuePair<string, string>("Engine", "GameplayTags"),
 			new KeyValuePair<string, string>("Engine", "Kismet"),
+			new KeyValuePair<string, string>("Engine", "Landscape"),
+			new KeyValuePair<string, string>("Engine", "LogVisualizer"),
+			new KeyValuePair<string, string>("Engine", "MaterialShaderQualitySettings"),
+			new KeyValuePair<string, string>("Engine", "UMG"),
+			new KeyValuePair<string, string>("Engine", "UnrealEd"),
 			new KeyValuePair<string, string>("FoliageEdit", "ViewportInteraction"),
 			new KeyValuePair<string, string>("FoliageEdit", "VREditor"),
 			new KeyValuePair<string, string>("FunctionalTesting", "UnrealEd"),
@@ -144,11 +145,11 @@ namespace UnrealBuildTool
 			new KeyValuePair<string, string>("GraphEditor", "Kismet"),
 			new KeyValuePair<string, string>("HierarchicalLODOutliner", "UnrealEd"),
 			new KeyValuePair<string, string>("Kismet", "BlueprintGraph"),
-			new KeyValuePair<string, string>("Kismet", "UMGEditor"),
 			new KeyValuePair<string, string>("Kismet", "Merge"),
+			new KeyValuePair<string, string>("Kismet", "UMGEditor"),
 			new KeyValuePair<string, string>("KismetWidgets", "BlueprintGraph"),
-			new KeyValuePair<string, string>("Landscape", "UnrealEd"),
 			new KeyValuePair<string, string>("Landscape", "MaterialUtilities"),
+			new KeyValuePair<string, string>("Landscape", "UnrealEd"),
 			new KeyValuePair<string, string>("LandscapeEditor", "ViewportInteraction"),
 			new KeyValuePair<string, string>("LandscapeEditor", "VREditor"),
 			new KeyValuePair<string, string>("LocalizationDashboard", "LocalizationService"),
@@ -170,36 +171,14 @@ namespace UnrealBuildTool
 			new KeyValuePair<string, string>("UnrealEd", "GraphEditor"),
 			new KeyValuePair<string, string>("UnrealEd", "InputBindingEditor"),
 			new KeyValuePair<string, string>("UnrealEd", "Kismet"),
-			new KeyValuePair<string, string>("UnrealEd", "AudioEditor"),
-			new KeyValuePair<string, string>("BlueprintGraph", "KismetCompiler"),
-			new KeyValuePair<string, string>("BlueprintGraph", "UnrealEd"),
-			new KeyValuePair<string, string>("BlueprintGraph", "GraphEditor"),
-			new KeyValuePair<string, string>("BlueprintGraph", "Kismet"),
-			new KeyValuePair<string, string>("BlueprintGraph", "CinematicCamera"),
-			new KeyValuePair<string, string>("ConfigEditor", "PropertyEditor"),
-			new KeyValuePair<string, string>("SourceControl", "UnrealEd"),
-			new KeyValuePair<string, string>("Kismet", "BlueprintGraph"),
-			new KeyValuePair<string, string>("Kismet", "UMGEditor"),
-			new KeyValuePair<string, string>("MovieSceneTools", "Sequencer"),
-			new KeyValuePair<string, string>("Sequencer", "MovieSceneTools"),
-			new KeyValuePair<string, string>("AIModule", "AITestSuite"),
-			new KeyValuePair<string, string>("GameplayTasks", "UnrealEd"),
-			new KeyValuePair<string, string>("AnimGraph", "UnrealEd"),
-			new KeyValuePair<string, string>("AnimGraph", "GraphEditor"),
-			new KeyValuePair<string, string>("MaterialUtilities", "Landscape"),
-			new KeyValuePair<string, string>("HierarchicalLODOutliner", "UnrealEd"),
-			new KeyValuePair<string, string>("PixelInspectorModule", "UnrealEd"),
-			new KeyValuePair<string, string>("GameplayAbilitiesEditor", "BlueprintGraph"),
-            new KeyValuePair<string, string>("UnrealEd", "ViewportInteraction"),
-            new KeyValuePair<string, string>("UnrealEd", "VREditor"),
-            new KeyValuePair<string, string>("Sequencer", "ViewportInteraction"),
-            new KeyValuePair<string, string>("NavigationSystem", "UnrealEd"),
+			new KeyValuePair<string, string>("UnrealEd", "MeshPaint"),
 			new KeyValuePair<string, string>("UnrealEd", "PluginWarden"),
 			new KeyValuePair<string, string>("UnrealEd", "PropertyEditor"),
 			new KeyValuePair<string, string>("UnrealEd", "ToolMenusEditor"),
+			new KeyValuePair<string, string>("UnrealEd", "ViewportInteraction"),
+			new KeyValuePair<string, string>("UnrealEd", "VREditor"),
 			new KeyValuePair<string, string>("WebBrowser", "WebBrowserTexture"),
-            new KeyValuePair<string, string>("UnrealEd", "MeshPaint"),
-        };
+		};
 
 		public UEBuildModuleCPP(ModuleRules Rules, DirectoryReference IntermediateDirectory, DirectoryReference? GeneratedCodeDirectory, ILogger Logger)
 			: base(Rules, IntermediateDirectory, Logger)
