[UBT]
* IWYUMode - Changed so IWYU always keep includes that are needed but redundant (they are transitively included by other needed includes). To remove redundant includes there is now an option -RemoveRedundant that brings back old behavior

diff --git a/Engine/Source/Programs/UnrealBuildTool/Modes/IWYUMode.cs b/Engine/Source/Programs/UnrealBuildTool/Modes/IWYUMode.cs
--- a/Engine/Source/Programs/UnrealBuildTool/Modes/IWYUMode.cs
+++ b/Engine/Source/Programs/UnrealBuildTool/Modes/IWYUMode.cs
@@ -251,10 +251,10 @@ namespace UnrealBuildTool
 		public bool bNoTransitiveIncludes = false;
 
 		/// <summary>
-		/// Allow files to not add includes if they are transitively included by other includes
+		/// Will remove headers that are needed but redundant because they are included through other needed includes
 		/// </summary>
-		[CommandLine("-KeepRedundant")]
-		public bool bKeepRedundantIncludes = false;
+		[CommandLine("-RemoveRedundant")]
+		public bool bRemoveRedundantIncludes = false;
 
 		/// <summary>
 		/// Will skip compiling before updating. Handle with care, this is dangerous since files might not match .iwyu files
@@ -1055,8 +1055,8 @@ namespace UnrealBuildTool
 					}
 				}
 
-				// If keep redundant includes is set we don't remove the includes that are included through other includes
-				if (bKeepRedundantIncludes)
+				// We don't remove seen includes that are redundant because they are included through other includes that are needed.
+				if (!bRemoveRedundantIncludes)
 				{
 					List<string> ToReadd = new();
 					foreach (var Seen in Info.IncludesSeenInFile)
