UnrealBuildTool: Rename .response to .rsp

#rnx

diff --git a/Engine/Source/Programs/UnrealBuildTool/Executors/SNDBS.cs b/Engine/Source/Programs/UnrealBuildTool/Executors/SNDBS.cs
--- a/Engine/Source/Programs/UnrealBuildTool/Executors/SNDBS.cs
+++ b/Engine/Source/Programs/UnrealBuildTool/Executors/SNDBS.cs
@@ -235,7 +235,7 @@ namespace UnrealBuildTool
 
 					if (a.PrerequisiteItems.Count() > 0)
 					{
-						Job["explicit_input_files"] = a.PrerequisiteItems.Where(i => !i.AbsolutePath.EndsWith(".response")).Select(i => new Dictionary<string, object>()
+						Job["explicit_input_files"] = a.PrerequisiteItems.Where(i => !(i.AbsolutePath.EndsWith(".rsp") || i.AbsolutePath.EndsWith(".response"))).Select(i => new Dictionary<string, object>()
 						{
 							["filename"] = i.AbsolutePath
 						}).ToList();
diff --git a/Engine/Source/Programs/UnrealBuildTool/Platform/Windows/UEBuildWindows.cs b/Engine/Source/Programs/UnrealBuildTool/Platform/Windows/UEBuildWindows.cs
--- a/Engine/Source/Programs/UnrealBuildTool/Platform/Windows/UEBuildWindows.cs
+++ b/Engine/Source/Programs/UnrealBuildTool/Platform/Windows/UEBuildWindows.cs
@@ -1319,6 +1319,7 @@ namespace UnrealBuildTool
 			return IsBuildProductName(FileName, NamePrefixes, NameSuffixes, ".exe")
 				|| IsBuildProductName(FileName, NamePrefixes, NameSuffixes, ".dll")
 				|| IsBuildProductName(FileName, NamePrefixes, NameSuffixes, ".dll.response")
+				|| IsBuildProductName(FileName, NamePrefixes, NameSuffixes, ".dll.rsp")
 				|| IsBuildProductName(FileName, NamePrefixes, NameSuffixes, ".lib")
 				|| IsBuildProductName(FileName, NamePrefixes, NameSuffixes, ".pdb")
 				|| IsBuildProductName(FileName, NamePrefixes, NameSuffixes, ".exp")
diff --git a/Engine/Source/Programs/UnrealBuildTool/Platform/Windows/VCToolChain.cs b/Engine/Source/Programs/UnrealBuildTool/Platform/Windows/VCToolChain.cs
--- a/Engine/Source/Programs/UnrealBuildTool/Platform/Windows/VCToolChain.cs
+++ b/Engine/Source/Programs/UnrealBuildTool/Platform/Windows/VCToolChain.cs
@@ -1440,7 +1440,7 @@ namespace UnrealBuildTool
 				if (CompileEnvironment.bPreprocessOnly)
 				{
 					CompileAction.PreprocessedFile = FileItem.GetItemByFileReference(FileReference.Combine(OutputDir, SourceFile.Location.GetFileName() + ".i"));
-					CompileAction.ResponseFile = FileItem.GetItemByPath(CompileAction.PreprocessedFile.FullName + ".response");
+					CompileAction.ResponseFile = FileItem.GetItemByFileReference(GetResponseFileName(CompileEnvironment, CompileAction.PreprocessedFile));
 				}
 				else
 				{
@@ -1449,7 +1449,7 @@ namespace UnrealBuildTool
 					FileItem ObjectFile = FileItem.GetItemByFileReference(FileReference.Combine(OutputDir, ObjectLeafFilename));
 
 					CompileAction.ObjectFile = ObjectFile;
-					CompileAction.ResponseFile = FileItem.GetItemByPath(ObjectFile.FullName + ".response");
+					CompileAction.ResponseFile = FileItem.GetItemByFileReference(GetResponseFileName(CompileEnvironment, ObjectFile)); ;
 
 					if (Target.WindowsPlatform.ObjSrcMapFile != null)
 					{
@@ -1546,7 +1546,7 @@ namespace UnrealBuildTool
 
 					VCCompileAction CompileDepsAction = new VCCompileAction(CompileAction);
 					CompileDepsAction.ActionType = ActionType.GatherModuleDependencies;
-					CompileDepsAction.ResponseFile = FileItem.GetItemByPath(IfcDepsFile + ".response");
+					CompileDepsAction.ResponseFile = FileItem.GetItemByFileReference(GetResponseFileName(CompileEnvironment, IfcDepsFile));
 					CompileDepsAction.ObjectFile = null;
 					CompileDepsAction.DependencyListFile = IfcDepsFile;
 					CompileDepsAction.Arguments.Add($"/sourceDependencies:directives \"{IfcDepsFile.Location}\"");
@@ -1929,7 +1929,7 @@ namespace UnrealBuildTool
 				Arguments.Add($"\"{NormalizeCommandLinePath(RCFile)}\"");
 
 				// Create a response file for the resource compilier
-				FileItem ResponseFile = FileItem.GetItemByPath(CompiledResourceFile.FullName + ".response");
+				FileItem ResponseFile = FileItem.GetItemByFileReference(GetResponseFileName(CompileEnvironment, CompiledResourceFile));
 				Graph.CreateIntermediateTextFile(ResponseFile, Arguments);
 				CompileAction.PrerequisiteItems.Add(ResponseFile);
 
diff --git a/Engine/Source/Programs/UnrealBuildTool/System/HotReload.cs b/Engine/Source/Programs/UnrealBuildTool/System/HotReload.cs
--- a/Engine/Source/Programs/UnrealBuildTool/System/HotReload.cs
+++ b/Engine/Source/Programs/UnrealBuildTool/System/HotReload.cs
@@ -1006,7 +1006,7 @@ namespace UnrealBuildTool
 				string NewFileNameWithoutExtension = Utils.GetFilenameWithoutAnyExtensions(HotReloadFile.FullName);
 
 				// Find the response file in the command line.  We'll need to make a copy of it with our new file name.
-				string ResponseFileExtension = ".response";
+				string ResponseFileExtension = UEToolChain.ResponseExt;
 				int ResponseExtensionIndex = Action.CommandArguments.IndexOf(ResponseFileExtension, StringComparison.InvariantCultureIgnoreCase);
 				if (ResponseExtensionIndex != -1)
 				{
diff --git a/Engine/Source/Programs/UnrealBuildTool/ToolChain/ClangToolChain.cs b/Engine/Source/Programs/UnrealBuildTool/ToolChain/ClangToolChain.cs
--- a/Engine/Source/Programs/UnrealBuildTool/ToolChain/ClangToolChain.cs
+++ b/Engine/Source/Programs/UnrealBuildTool/ToolChain/ClangToolChain.cs
@@ -1009,7 +1009,7 @@ namespace UnrealBuildTool
 			FileItem TargetFile = CompileAction.ProducedItems.First();
 
 			// Creates the path to the response file using the name of the output file and creates its contents.
-			FileReference ResponseFileName = new FileReference(TargetFile.AbsolutePath + ".response");
+			FileReference ResponseFileName = GetResponseFileName(CompileEnvironment, TargetFile);
 			List<string> ResponseFileContents = ExpandResponseFileContents(FileArguments);
 
 			if (RuntimePlatform.IsWindows)
@@ -1085,7 +1085,7 @@ namespace UnrealBuildTool
 				FileItem PreprocessTargetFile = PrepassAction.ProducedItems.First();
 
 				// Creates the path to the response file using the name of the output file and creates its contents.
-				FileReference PreprocessResponseFileName = new FileReference(PreprocessTargetFile.AbsolutePath + ".response");
+				FileReference PreprocessResponseFileName = GetResponseFileName(CompileEnvironment, PreprocessTargetFile);
 				List<string> PreprocessResponseFileContents = new();
 				PreprocessResponseFileContents.AddRange(PreprocessGlobalArguments);
 				PreprocessResponseFileContents.AddRange(PreprocessFileArguments);
diff --git a/Engine/Source/Programs/UnrealBuildTool/ToolChain/ISPCToolChain.cs b/Engine/Source/Programs/UnrealBuildTool/ToolChain/ISPCToolChain.cs
--- a/Engine/Source/Programs/UnrealBuildTool/ToolChain/ISPCToolChain.cs
+++ b/Engine/Source/Programs/UnrealBuildTool/ToolChain/ISPCToolChain.cs
@@ -491,7 +491,7 @@ namespace UnrealBuildTool
 
 				CompileAction.ProducedItems.Add(ISPCIncludeHeaderFile);
 
-				FileReference ResponseFileName = new FileReference(ISPCIncludeHeaderFile.AbsolutePath + ".response");
+				FileReference ResponseFileName = GetResponseFileName(CompileEnvironment, ISPCIncludeHeaderFile);
 				FileItem ResponseFileItem = Graph.CreateIntermediateTextFile(ResponseFileName, Arguments.Select(x => Utils.ExpandVariables(x)));
 				CompileAction.CommandArguments = $"@\"{NormalizeCommandLinePath(ResponseFileName)}\"";
 				CompileAction.PrerequisiteItems.Add(ResponseFileItem);
@@ -697,7 +697,7 @@ namespace UnrealBuildTool
 
 				CompileAction.ProducedItems.UnionWith(CompiledISPCObjFiles);
 
-				FileReference ResponseFileName = new FileReference(CompiledISPCObjFileNoISA.AbsolutePath + ".response");
+				FileReference ResponseFileName = GetResponseFileName(CompileEnvironment, CompiledISPCObjFileNoISA);
 				FileItem ResponseFileItem = Graph.CreateIntermediateTextFile(ResponseFileName, Arguments.Select(x => Utils.ExpandVariables(x)));
 				CompileAction.CommandArguments = $"@\"{ResponseFileName}\"";
 				CompileAction.PrerequisiteItems.Add(ResponseFileItem);
@@ -732,7 +732,7 @@ namespace UnrealBuildTool
 							PostCompileArgs.Add($"-o \"{NormalizeCommandLinePath(FinalCompiledISPCObjFile)}\"");
 
 							// Write the args to a response file
-							FileReference PostCompileResponseFileName = new FileReference(FinalCompiledISPCObjFile.AbsolutePath + ".response");
+							FileReference PostCompileResponseFileName = GetResponseFileName(CompileEnvironment, FinalCompiledISPCObjFile);
 							FileItem PostCompileResponseFileItem = Graph.CreateIntermediateTextFile(PostCompileResponseFileName, PostCompileArgs.Select(x => Utils.ExpandVariables(x)));
 							PostCompileAction.CommandArguments = $"@\"{PostCompileResponseFileName}\"";
 							PostCompileAction.PrerequisiteItems.Add(PostCompileResponseFileItem);
diff --git a/Engine/Source/Programs/UnrealBuildTool/ToolChain/UEToolChain.cs b/Engine/Source/Programs/UnrealBuildTool/ToolChain/UEToolChain.cs
--- a/Engine/Source/Programs/UnrealBuildTool/ToolChain/UEToolChain.cs
+++ b/Engine/Source/Programs/UnrealBuildTool/ToolChain/UEToolChain.cs
@@ -13,6 +13,9 @@ namespace UnrealBuildTool
 	abstract class UEToolChain
 	{
 		protected readonly ILogger Logger;
+		
+		// Return the extension for response files
+		public static string ResponseExt => ".rsp";
 
 		public UEToolChain(ILogger InLogger)
 		{
@@ -149,6 +152,18 @@ namespace UnrealBuildTool
 			return Result.ToArray();
 		}
 
+		/// <summary>
+		/// Get the name of the response file for the current compile environment and output file
+		/// </summary>
+		/// <param name="CompileEnvironment"></param>
+		/// <param name="OutputFile"></param>
+		/// <returns></returns>
+		public static FileReference GetResponseFileName(CppCompileEnvironment CompileEnvironment, FileItem OutputFile)
+		{
+			// Construct a relative path for the intermediate response file
+			return OutputFile.Location.ChangeExtension(OutputFile.Location.GetExtension() + ResponseExt);
+		}
+
 		/// <summary>
 		/// Get the name of the response file for the current linker environment and output file
 		/// </summary>
@@ -158,7 +173,7 @@ namespace UnrealBuildTool
 		public static FileReference GetResponseFileName(LinkEnvironment LinkEnvironment, FileItem OutputFile)
 		{
 			// Construct a relative path for the intermediate response file
-			return FileReference.Combine(LinkEnvironment.IntermediateDirectory!, OutputFile.Location.GetFileName() + ".response");
+			return FileReference.Combine(LinkEnvironment.IntermediateDirectory!, OutputFile.Location.GetFileName() + ResponseExt);
 		}
 
 		public virtual ICollection<FileItem> PostBuild(ReadOnlyTargetRules Target, FileItem Executable, LinkEnvironment ExecutableLinkEnvironment, IActionGraphBuilder Graph)
