[UBT - IWYU]
* Fixed so source file metadata cache is saved after make file is created. Added try finally to always save
* Fixed so Cpp dependency cache is in try/finally statement since we want the dependency cache to be updated for the build steps that finished

diff --git a/Engine/Source/Programs/UnrealBuildTool/Modes/IWYUMode.cs b/Engine/Source/Programs/UnrealBuildTool/Modes/IWYUMode.cs
--- a/Engine/Source/Programs/UnrealBuildTool/Modes/IWYUMode.cs
+++ b/Engine/Source/Programs/UnrealBuildTool/Modes/IWYUMode.cs
@@ -571,14 +571,27 @@ namespace UnrealBuildTool
 
 					// Create the make file that contains all the actions we will use to find .iwyu files.
 					Logger.LogInformation($"Creating MakeFile for target...");
-					TargetMakefile Makefile = BuildMode.CreateMakefile(BuildConfiguration, Descriptor, WorkingSet, Logger);
+					TargetMakefile Makefile;
+					try
+					{
+						Makefile = BuildMode.CreateMakefile(BuildConfiguration, Descriptor, WorkingSet, Logger);
+					}
+					finally
+					{
+						SourceFileMetadataCache.SaveAll();
+					}
 
 					// Use the make file to build unless -NoCompile is set.
 					if (!bNoCompile)
 					{
-						BuildMode.Build(new TargetMakefile[] { Makefile }, new List<TargetDescriptor>() { Descriptor }, BuildConfiguration, BuildOptions.None, null, Logger);
-						SourceFileMetadataCache.SaveAll();
-						CppDependencyCache.SaveAll();
+						try
+						{
+							BuildMode.Build(new TargetMakefile[] { Makefile }, new List<TargetDescriptor>() { Descriptor }, BuildConfiguration, BuildOptions.None, null, Logger);
+						}
+						finally
+						{
+							CppDependencyCache.SaveAll();
+						}
 					}
 
 					HashSet<FileItem> OutputItems = new HashSet<FileItem>();
