If running through Wine use a different method for touching files as the current way does not work

diff --git a/Engine/Source/Programs/UnrealBuildTool/Platform/Windows/VCToolChain.cs b/Engine/Source/Programs/UnrealBuildTool/Platform/Windows/VCToolChain.cs
--- a/Engine/Source/Programs/UnrealBuildTool/Platform/Windows/VCToolChain.cs
+++ b/Engine/Source/Programs/UnrealBuildTool/Platform/Windows/VCToolChain.cs
@@ -2012,12 +2012,21 @@ namespace UnrealBuildTool
 			Action TouchAction = Graph.CreateAction(ActionType.BuildProject);
 			TouchAction.CommandDescription = "Touch";
 			TouchAction.CommandPath = BuildHostPlatform.Current.Shell;
-			TouchAction.CommandArguments = $"/C \"copy /b \"{OutputFile.FullName}\"+,, \"{OutputFile.FullName}\" 1>nul:\"";
 			TouchAction.WorkingDirectory = Unreal.EngineSourceDirectory;
 			TouchAction.PrerequisiteItems.Add(ObjectFile);
 			TouchAction.ProducedItems.Add(FileItem.GetItemByFileReference(OutputFile));
 			TouchAction.StatusDescription = OutputFile.GetFileName();
 			TouchAction.bCanExecuteRemotely = false;
+
+			// Wine has an issue with copy /b <file>+,, <file> and does not correctly touch the file, and instead makes a copy of the file at the CWD
+			if (BuildHostPlatform.Current.IsRunningOnWine())
+			{
+				TouchAction.CommandArguments = $"/C \"copy /Y \"{OutputFile.FullName}\" \"%TMP%/{OutputFile.GetFileName()}\" && type \"%TMP%/{OutputFile.GetFileName()}\" > \"{OutputFile.FullName}\"\"";
+			}
+			else
+			{
+				TouchAction.CommandArguments = $"/C \"copy /b \"{OutputFile.FullName}\"+,, \"{OutputFile.FullName}\"\"";
+			}
 		}
 
 		/// <summary>
diff --git a/Engine/Source/Programs/UnrealBuildTool/System/BuildHostPlatform.cs b/Engine/Source/Programs/UnrealBuildTool/System/BuildHostPlatform.cs
--- a/Engine/Source/Programs/UnrealBuildTool/System/BuildHostPlatform.cs
+++ b/Engine/Source/Programs/UnrealBuildTool/System/BuildHostPlatform.cs
@@ -224,6 +224,15 @@ namespace UnrealBuildTool
 			return Modules.ToArray();
 		}
 
+		/// <summary>
+		/// Determines if the UBT process is running through WINE
+		/// </summary>
+		/// <returns>Sequence of project file formats</returns>
+		public virtual bool IsRunningOnWine()
+		{
+			return false;
+		}
+
 		/// <summary>
 		/// Determines the default project file formats for this platform
 		/// </summary>
@@ -241,6 +250,24 @@ namespace UnrealBuildTool
 
 		public override string BinarySuffix => ".exe";
 
+		[DllImport("ntdll.dll", EntryPoint="wine_get_version", CallingConvention=CallingConvention.Cdecl, CharSet=CharSet.Ansi)]
+		private static extern string GetWineVersion();
+
+		public override bool IsRunningOnWine()
+		{
+			try
+			{
+				// if we dont throw an exception we are running in Wine
+				GetWineVersion();
+			}
+			catch
+			{
+				return false;
+			};
+
+			return true;
+		}
+
 		internal override IEnumerable<ProjectFileFormat> GetDefaultProjectFileFormats()
 		{
 			yield return ProjectFileFormat.VisualStudio;
