Fixed issue when building with 2022, but having a 2019 Arm64 toolchain when making HoloLens project

diff --git a/Engine/Platforms/Hololens/Source/Programs/UnrealBuildTool/UEBuildHoloLens.cs b/Engine/Platforms/Hololens/Source/Programs/UnrealBuildTool/UEBuildHoloLens.cs
--- a/Engine/Platforms/Hololens/Source/Programs/UnrealBuildTool/UEBuildHoloLens.cs
+++ b/Engine/Platforms/Hololens/Source/Programs/UnrealBuildTool/UEBuildHoloLens.cs
@@ -205,6 +205,36 @@ namespace UnrealBuildTool
 
 	class HoloLensArchitectureConfig : UnrealArchitectureConfig
 	{
+		// we need to get the compiler the Target will use, before the target is created, when making projectfiles especially,
+		// so we duplicate some logic elsewhere in this file
+		[ConfigFile(ConfigHierarchyType.Engine, "/Script/HoloLensPlatformEditor.HoloLensTargetSettings", "CompilerVersion")]
+		[XmlConfigFile(Category = "HoloLensPlatform")]
+		[CommandLine("-2019", Value = nameof(WindowsCompiler.VisualStudio2019))]
+		[CommandLine("-2022", Value = nameof(WindowsCompiler.VisualStudio2022))]
+		[CommandLine("-Compiler=")]
+		private WindowsCompiler Compiler = WindowsCompiler.Default;
+
+		private WindowsCompiler GetCompiler(FileReference? ProjectFile, UnrealArch Architecture)
+		{
+			HoloLensTargetRules HoloLensPlatform = new();
+			ConfigCache.ReadSettings(ProjectFile?.Directory, UnrealTargetPlatform.HoloLens, HoloLensPlatform);
+
+			// HoloLensPlatform has logic that the value in ReadSettings will override whatever the commandline specified version is,
+			// which is hard to say if it's desired, but need to stay consistent (could move this into a function)
+			if (HoloLensPlatform.Compiler != WindowsCompiler.Default)
+			{
+				return HoloLensPlatform.Compiler;
+			}
+
+			if (Compiler != WindowsCompiler.Default)
+			{
+				return Compiler;
+			}
+
+			// get the default compiler, but don't spit out a nasty warning if it doesn't exist
+			return WindowsPlatform.GetDefaultCompiler(ProjectFile, Architecture, Logger, bSkipWarning: true);
+		}
+
 		private ILogger Logger;
 		public HoloLensArchitectureConfig(ILogger Logger)
 			: base(UnrealArchitectureMode.OneTargetPerArchitecture, new[] { UnrealArch.X64, UnrealArch.Arm64 })
@@ -219,8 +249,8 @@ namespace UnrealBuildTool
 			bool bBuildForEmulation;
 			bool bBuildForDevice;
 
-			bool bHasArmToolchain = MicrosoftPlatformSDK.HasValidCompiler(WindowsCompiler.VisualStudio2022, UnrealArch.Arm64, Logger) || MicrosoftPlatformSDK.HasValidCompiler(WindowsCompiler.VisualStudio2019, UnrealArch.Arm64, Logger);
-			bool bHasX64Toolchain = MicrosoftPlatformSDK.HasValidCompiler(WindowsCompiler.VisualStudio2022, UnrealArch.X64, Logger) || MicrosoftPlatformSDK.HasValidCompiler(WindowsCompiler.VisualStudio2019, UnrealArch.X64, Logger);
+			bool bHasArmToolchain = MicrosoftPlatformSDK.HasValidCompiler(GetCompiler(ProjectFile, UnrealArch.Arm64), UnrealArch.Arm64, Logger);
+			bool bHasX64Toolchain = MicrosoftPlatformSDK.HasValidCompiler(GetCompiler(ProjectFile, UnrealArch.X64), UnrealArch.X64, Logger);
 	
 			ConfigHierarchy Ini = ConfigCache.ReadHierarchy(ConfigHierarchyType.Engine, ProjectFile?.Directory, UnrealTargetPlatform.HoloLens);
 			if (Ini.GetBool("/Script/HoloLensPlatformEditor.HoloLensTargetSettings", "bBuildForEmulation", out bBuildForEmulation) && bBuildForEmulation)
@@ -249,7 +279,7 @@ namespace UnrealBuildTool
 				{
 					if (!bShownLog)
 					{
-						Logger.LogInformation("Arm64 was requested for HoloLens, but the Arm64 toolchain is not installed, defaulting to X64 build");
+						Logger.LogInformation("Arm64 was requested for HoloLens, but the {0} Arm64 toolchain is not installed, defaulting to X64 build", GetCompiler(ProjectFile, UnrealArch.Arm64));
 						bShownLog = true;
 					}
 					ArchList.Add(UnrealArch.X64);
diff --git a/Engine/Source/Programs/UnrealBuildTool/Platform/Windows/UEBuildWindows.cs b/Engine/Source/Programs/UnrealBuildTool/Platform/Windows/UEBuildWindows.cs
--- a/Engine/Source/Programs/UnrealBuildTool/Platform/Windows/UEBuildWindows.cs
+++ b/Engine/Source/Programs/UnrealBuildTool/Platform/Windows/UEBuildWindows.cs
@@ -1122,7 +1122,7 @@ namespace UnrealBuildTool
 		/// Gets the default compiler which should be used, if it's not set explicitly by the target, command line, or config file.
 		/// </summary>
 		/// <returns>The default compiler version</returns>
-		internal static WindowsCompiler GetDefaultCompiler(FileReference? ProjectFile, UnrealArch Architecture, ILogger Logger)
+		internal static WindowsCompiler GetDefaultCompiler(FileReference? ProjectFile, UnrealArch Architecture, ILogger Logger, bool bSkipWarning=false)
 		{
 			// If there's no specific compiler set, try to pick the matching compiler for the selected IDE
 			if (ProjectFileGeneratorSettings.Format != null)
@@ -1179,24 +1179,27 @@ namespace UnrealBuildTool
 				return WindowsCompiler.VisualStudio2019;
 			}
 
-			// If we do have a Visual Studio installation, but we're missing just the C++ parts, warn about that.
-			if (TryGetVSInstallDirs(WindowsCompiler.VisualStudio2022, Logger) != null)
+			if (!bSkipWarning)
 			{
-				string ToolSetWarning = Architecture == UnrealArch.X64 ?
-					"MSVC v143 - VS 2022 C++ x64/x86 build tools (Latest)" :
-					"MSVC v143 - VS 2022 C++ ARM64 build tools (Latest)";
-				Logger.LogWarning("Visual Studio 2022 is installed, but is missing the C++ toolchain. Please verify that the \"{Component}\" component is selected in the Visual Studio 2022 installation options.", ToolSetWarning);
-			}
-			else if (TryGetVSInstallDirs(WindowsCompiler.VisualStudio2019, Logger) != null)
-			{
-				string ToolSetWarning = Architecture == UnrealArch.X64 ?
-					"MSVC v142 - VS 2019 C++ x64/x86 build tools (Latest)" :
-					"MSVC v142 - VS 2019 C++ ARM64 build tools (Latest)";
-				Logger.LogWarning("Visual Studio 2019 is installed, but is missing the C++ toolchain. Please verify that the \"{Component}\" component is selected in the Visual Studio 2019 installation options.", ToolSetWarning);
-			}
-			else
-			{
-				Logger.LogWarning("No Visual C++ installation was found. Please download and install Visual Studio 2022 or 2019 with C++ components.");
+				// If we do have a Visual Studio installation, but we're missing just the C++ parts, warn about that.
+				if (TryGetVSInstallDirs(WindowsCompiler.VisualStudio2022, Logger) != null)
+				{
+					string ToolSetWarning = Architecture == UnrealArch.X64 ?
+						"MSVC v143 - VS 2022 C++ x64/x86 build tools (Latest)" :
+						"MSVC v143 - VS 2022 C++ ARM64 build tools (Latest)";
+					Logger.LogWarning("Visual Studio 2022 is installed, but is missing the C++ toolchain. Please verify that the \"{Component}\" component is selected in the Visual Studio 2022 installation options.", ToolSetWarning);
+				}
+				else if (TryGetVSInstallDirs(WindowsCompiler.VisualStudio2019, Logger) != null)
+				{
+					string ToolSetWarning = Architecture == UnrealArch.X64 ?
+						"MSVC v142 - VS 2019 C++ x64/x86 build tools (Latest)" :
+						"MSVC v142 - VS 2019 C++ ARM64 build tools (Latest)";
+					Logger.LogWarning("Visual Studio 2019 is installed, but is missing the C++ toolchain. Please verify that the \"{Component}\" component is selected in the Visual Studio 2019 installation options.", ToolSetWarning);
+				}
+				else
+				{
+					Logger.LogWarning("No Visual C++ installation was found. Please download and install Visual Studio 2022 or 2019 with C++ components.");
+				}
 			}
 
 			// Finally, default to VS2022 anyway
