UBT: Remove XGE task metadata extraction

Current regex matchers potentially caused a slowdown in cooks. Also, this approach of grabbing metrics cannot be applied for cooker-based XGE tasks, so removing this completely in favor of using *.ib_mon files is for the better.

diff --git a/Engine/Source/Programs/UnrealBuildTool/Matchers/XoreaxEventMatcher.cs b/Engine/Source/Programs/UnrealBuildTool/Matchers/XoreaxEventMatcher.cs
--- a/Engine/Source/Programs/UnrealBuildTool/Matchers/XoreaxEventMatcher.cs
+++ b/Engine/Source/Programs/UnrealBuildTool/Matchers/XoreaxEventMatcher.cs
@@ -28,11 +28,6 @@ namespace UnrealBuildTool.Matchers
 		static readonly Regex s_cacheWarning = new Regex(@"^\s*WARNING: \d+ items \([^\)]*\) removed from the cache due to reaching the cache size limit");
 		static readonly Regex s_cacheWarning2 = new Regex(@"^\s*WARNING: Several items removed from the cache due to reaching the cache size limit");
 		static readonly Regex s_cacheWarning3 = new Regex(@"^\s*WARNING: The Build Cache is close to full");
-		
-		static readonly Regex s_taskFinishedOnHelper = new Regex(@"(.*?) \(Agent '(.+)', (.*) at \+(.*)\)");
-		static readonly Regex s_taskFinishedLocally = new Regex(@"(.*?) \((.*) at \+(.*)\)");
-		static readonly Regex s_duration = new Regex(@"(\d+):(\d+)\.(\d+)");
-		static readonly Regex s_startTime = new Regex(@"\+?(\d+):(\d+)");
 
 		public LogEventMatch? Match(ILogCursor cursor)
 		{
@@ -96,51 +91,7 @@ namespace UnrealBuildTool.Matchers
 				}
 			}
 			
-			if (cursor.TryMatch(s_taskFinishedOnHelper, out Match? taskHelperMatch))
-			{
-				LogEventBuilder builder = new LogEventBuilder(cursor);
-				builder.AddProperty(PropertyName, taskHelperMatch.Groups[1].Value);
-				builder.AddProperty(PropertyAgent, taskHelperMatch.Groups[2].Value);
-				SetDuration(builder, taskHelperMatch.Groups[3].Value);
-				SetStartTime(builder, taskHelperMatch.Groups[4].Value);
-				return builder.ToMatch(LogEventPriority.High, LogLevel.Information, KnownLogEvents.Systemic_Xge_TaskMetadata);
-			}
-			
-			if (cursor.TryMatch(s_taskFinishedLocally, out Match? taskLocalMatch))
-			{
-				LogEventBuilder builder = new LogEventBuilder(cursor);
-				builder.AddProperty(PropertyName, taskLocalMatch.Groups[1].Value);
-				SetDuration(builder, taskLocalMatch.Groups[2].Value);
-				SetStartTime(builder, taskLocalMatch.Groups[3].Value);
-				return builder.ToMatch(LogEventPriority.High, LogLevel.Information, KnownLogEvents.Systemic_Xge_TaskMetadata);
-			}
-			
 			return null;
 		}
-
-		private static void SetDuration(LogEventBuilder builder, string durationText)
-		{
-			Match m = s_duration.Match(durationText);
-			if (m.Success)
-			{
-				int minutes = Convert.ToInt32(m.Groups[1].Value);
-				int seconds = Convert.ToInt32(m.Groups[2].Value);
-				int milliseconds = Convert.ToInt32(m.Groups[3].Value) * 10;
-				TimeSpan duration = new TimeSpan(0, 0, minutes, seconds, milliseconds);
-				builder.AddProperty(PropertyDuration, (int)duration.TotalMilliseconds);
-			}
-		}
-		
-		private static void SetStartTime(LogEventBuilder builder, string startTimeText)
-		{
-			Match m = s_startTime.Match(startTimeText);
-			if (m.Success)
-			{
-				int minutes = Convert.ToInt32(m.Groups[1].Value);
-				int seconds = Convert.ToInt32(m.Groups[2].Value);
-				TimeSpan startTime = new TimeSpan(0, minutes, seconds);
-				builder.AddProperty(PropertyStartTime, (int)startTime.TotalMilliseconds);
-			}
-		}
 	}
 }
