UnrealBuildTool: Reference any Target.cs files as rules dependencies so modifying them will update the Makefile, as these TargetRules can be referenced

#rnx

diff --git a/Engine/Source/Programs/AutomationTool/AutomationUtils/ProjectUtils.cs b/Engine/Source/Programs/AutomationTool/AutomationUtils/ProjectUtils.cs
--- a/Engine/Source/Programs/AutomationTool/AutomationUtils/ProjectUtils.cs
+++ b/Engine/Source/Programs/AutomationTool/AutomationUtils/ProjectUtils.cs
@@ -857,7 +857,7 @@ namespace AutomationTool
 
 					// Create an instance of this type
 					CommandUtils.LogVerbose("Creating target rules object: {0}", TargetType.Name);
-					TargetRules Rules = TargetRules.Create(TargetType, DummyTargetInfo, null, null, null, Log.Logger);
+					TargetRules Rules = TargetRules.Create(TargetType, DummyTargetInfo, null, null, null, null, Log.Logger);
 					CommandUtils.LogVerbose("Adding target: {0} ({1})", TargetType.Name, Rules.Type);
 
 					SingleTargetProperties TargetData = new SingleTargetProperties();
diff --git a/Engine/Source/Programs/UnrealBuildTool/Configuration/TargetRules.cs b/Engine/Source/Programs/UnrealBuildTool/Configuration/TargetRules.cs
--- a/Engine/Source/Programs/UnrealBuildTool/Configuration/TargetRules.cs
+++ b/Engine/Source/Programs/UnrealBuildTool/Configuration/TargetRules.cs
@@ -497,6 +497,11 @@ namespace UnrealBuildTool
 		/// </summary>
 		internal FileReference? TargetSourceFile { get; set; }
 
+		/// <summary>
+		/// All target files that could be referenced by this target and will require updating the Makefile if changed
+		/// </summary>
+		internal HashSet<FileReference>? TargetFiles { get; set; }
+
 		/// <summary>
 		/// Logger for output relating to this target. Set before the constructor is run from <see cref="RulesCompiler"/>
 		/// </summary>
@@ -2248,10 +2253,11 @@ namespace UnrealBuildTool
 		/// <param name="TargetInfo">Target info</param>
 		/// <param name="BaseFile">Path to the file for the rules assembly</param>
 		/// <param name="PlatformFile">Path to the platform specific rules file</param>
+		/// <param name="TargetFiles">Path to all possible rules file</param>
 		/// <param name="DefaultBuildSettings"></param>
 		/// <param name="Logger">Logger for the new target rules</param>
 		/// <returns>Target instance</returns>
-		public static TargetRules Create(Type RulesType, TargetInfo TargetInfo, FileReference? BaseFile, FileReference? PlatformFile, BuildSettingsVersion? DefaultBuildSettings, ILogger Logger)
+		public static TargetRules Create(Type RulesType, TargetInfo TargetInfo, FileReference? BaseFile, FileReference? PlatformFile, IEnumerable<FileReference>? TargetFiles, BuildSettingsVersion? DefaultBuildSettings, ILogger Logger)
 		{
 			TargetRules Rules = (TargetRules)FormatterServices.GetUninitializedObject(RulesType);
 			if (DefaultBuildSettings.HasValue)
@@ -2265,6 +2271,9 @@ namespace UnrealBuildTool
 			// The platform/group-specific target file name
 			Rules.TargetSourceFile = PlatformFile;
 
+			// All target files for this target that could cause invalidation
+			Rules.TargetFiles = TargetFiles?.ToHashSet() ?? new();
+
 			// Initialize the logger
 			Rules.Logger = Logger;
 
@@ -2690,6 +2699,12 @@ namespace UnrealBuildTool
 			get { return Inner.TargetSourceFile!; }
 		}
 
+
+		internal ReadOnlyHashSet<FileReference> TargetFiles
+		{
+			get { return Inner.TargetFiles!; }
+		}
+
 		public UnrealTargetPlatform Platform
 		{
 			get { return Inner.Platform; }
diff --git a/Engine/Source/Programs/UnrealBuildTool/Configuration/UEBuildTarget.cs b/Engine/Source/Programs/UnrealBuildTool/Configuration/UEBuildTarget.cs
--- a/Engine/Source/Programs/UnrealBuildTool/Configuration/UEBuildTarget.cs
+++ b/Engine/Source/Programs/UnrealBuildTool/Configuration/UEBuildTarget.cs
@@ -2309,6 +2309,9 @@ namespace UnrealBuildTool
 			TargetToolChain.GetVersionInfo(Makefile.Diagnostics);
 			Rules.GetBuildSettingsInfo(Makefile.Diagnostics);
 
+			// Mark dependencies to other target rules files, as they can be referenced by this makefile
+			Makefile.InternalDependencies.UnionWith(Rules.TargetFiles.Select(x => FileItem.GetItemByFileReference(x)));
+
 			// Get any pre-build targets.
 			Makefile.PreBuildTargets = Rules.PreBuildTargets.ToArray();
 
diff --git a/Engine/Source/Programs/UnrealBuildTool/System/RulesAssembly.cs b/Engine/Source/Programs/UnrealBuildTool/System/RulesAssembly.cs
--- a/Engine/Source/Programs/UnrealBuildTool/System/RulesAssembly.cs
+++ b/Engine/Source/Programs/UnrealBuildTool/System/RulesAssembly.cs
@@ -610,7 +610,7 @@ namespace UnrealBuildTool
 			Type RulesType = PlatformRulesType ?? BaseRulesType;
 			FileReference BaseFile = TargetNameToTargetFile[TargetInfo.Name];
 			FileReference PlatformFile = TargetNameToTargetFile.TryGetValue(PlatformRulesName, out FileReference? PlatformTargetFile) ? PlatformTargetFile : BaseFile;
-			TargetRules Rules = TargetRules.Create(RulesType, TargetInfo, BaseFile, PlatformFile, DefaultBuildSettings, Logger);
+			TargetRules Rules = TargetRules.Create(RulesType, TargetInfo, BaseFile, PlatformFile, TargetNameToTargetFile.Values, DefaultBuildSettings, Logger);
 
 			// Set the default overriddes for the configured target type
 			Rules.SetOverridesForTargetType();
