- Fixed IOS/TVOS codesigning in InstalledBuild (caused by an issue when forcing InstalledBuilds to be Intel only on Mac - it was making all platforms Intel!)

diff --git a/Engine/Source/Programs/UnrealBuildTool/ProjectFiles/Xcode/XcodeProjectLegacy.cs b/Engine/Source/Programs/UnrealBuildTool/ProjectFiles/Xcode/XcodeProjectLegacy.cs
--- a/Engine/Source/Programs/UnrealBuildTool/ProjectFiles/Xcode/XcodeProjectLegacy.cs
+++ b/Engine/Source/Programs/UnrealBuildTool/ProjectFiles/Xcode/XcodeProjectLegacy.cs
@@ -915,7 +915,7 @@ namespace UnrealBuildTool.XcodeProjectLegacy
 			Content.Append("\t\t\t\tALWAYS_SEARCH_USER_PATHS = NO;" + ProjectFileGenerator.NewLine);
 			if (Unreal.IsEngineInstalled())
 			{
-				Content.Append("\t\t\t\tARCHS = x86_64;" + ProjectFileGenerator.NewLine);
+				Content.Append("\t\t\t\tARCHS[sdk=macosx*] = x86_64;" + ProjectFileGenerator.NewLine);
 			}
 			Content.Append("\t\t\t\tCLANG_CXX_LANGUAGE_STANDARD = \"c++14\";" + ProjectFileGenerator.NewLine);
 			Content.Append("\t\t\t\tGCC_ENABLE_CPP_RTTI = NO;" + ProjectFileGenerator.NewLine);
