[UBT]
* Run VNI at the same time as UHT to save wall time

diff --git a/Engine/Source/Programs/UnrealBuildTool/Modes/BuildMode.cs b/Engine/Source/Programs/UnrealBuildTool/Modes/BuildMode.cs
--- a/Engine/Source/Programs/UnrealBuildTool/Modes/BuildMode.cs
+++ b/Engine/Source/Programs/UnrealBuildTool/Modes/BuildMode.cs
@@ -10,6 +10,7 @@ using OpenTracing.Util;
 using UnrealBuildBase;
 using Microsoft.Extensions.Logging;
 using System.Collections.Concurrent;
+using System.Threading.Tasks;
 
 namespace UnrealBuildTool
 {
@@ -910,6 +911,14 @@ namespace UnrealBuildTool
 					Environment.SetEnvironmentVariable(EnvironmentVariable.Item1, EnvironmentVariable.Item2);
 				}
 
+#if __VPROJECT_AVAILABLE__
+				Task VNITask = Task.CompletedTask;
+				// Same for VNI
+				if (Makefile.VNIModules.Count > 0)
+				{
+					VNITask = Task.Run(() => VNIExecution.ExecuteVNITool(Makefile, TargetDescriptor, Logger));
+				}
+#endif
 				// If the target needs UHT to be run, we'll go ahead and do that now
 				if (Makefile.UObjectModules.Count > 0)
 				{
@@ -917,12 +926,7 @@ namespace UnrealBuildTool
 				}
 
 #if __VPROJECT_AVAILABLE__
-				// Same for VNI
-				if (Makefile.VNIModules.Count > 0)
-				{
-
-					VNIExecution.ExecuteVNITool(Makefile, TargetDescriptor, Logger);
-				}
+				VNITask.Wait();
 #endif
 			}
 			return Makefile;
