[EpicGames.Build]
* DirectoryItem - Added CompatRead/Write to BinaryArchive

diff --git a/Engine/Source/Programs/Shared/EpicGames.Build/System/DirectoryItem.cs b/Engine/Source/Programs/Shared/EpicGames.Build/System/DirectoryItem.cs
--- a/Engine/Source/Programs/Shared/EpicGames.Build/System/DirectoryItem.cs
+++ b/Engine/Source/Programs/Shared/EpicGames.Build/System/DirectoryItem.cs
@@ -448,5 +448,27 @@ namespace UnrealBuildBase
 			// Use lambda that doesn't require anything to be captured thus eliminating an allocation.
 			Writer.WriteObjectReference<DirectoryItem>(DirectoryItem, (BinaryArchiveWriter Writer, DirectoryItem DirectoryItem) => Writer.WriteDirectoryReference(DirectoryItem.Location));
 		}
+
+		/// <summary>
+		/// Writes a directory reference  to a binary archive
+		/// </summary>
+		/// <param name="writer">The writer to output data to</param>
+		/// <param name="directory">The item to write</param>
+		public static void WriteCompactDirectoryReference(this BinaryArchiveWriter Writer, DirectoryReference Directory)
+		{
+			DirectoryItem Item = DirectoryItem.GetItemByDirectoryReference(Directory);
+			Writer.WriteDirectoryItem(Item);
+		}
+
+		/// <summary>
+		/// Reads a directory reference from a binary archive
+		/// </summary>
+		/// <param name="reader">Reader to serialize data from</param>
+		/// <returns>New directory reference instance</returns>
+		public static DirectoryReference ReadCompactDirectoryReference(this BinaryArchiveReader Reader)
+		{
+			return Reader.ReadDirectoryItem()!.Location;
+		}
+
 	}
 }
diff --git a/Engine/Source/Programs/UnrealBuildTool/Platform/Windows/VCCompileAction.cs b/Engine/Source/Programs/UnrealBuildTool/Platform/Windows/VCCompileAction.cs
--- a/Engine/Source/Programs/UnrealBuildTool/Platform/Windows/VCCompileAction.cs
+++ b/Engine/Source/Programs/UnrealBuildTool/Platform/Windows/VCCompileAction.cs
@@ -344,8 +344,8 @@ namespace UnrealBuildTool
 			CreatePchFile = Reader.ReadFileItem();
 			UsingPchFile = Reader.ReadFileItem();
 			PchThroughHeaderFile = Reader.ReadFileItem();
-			IncludePaths = Reader.ReadList(() => Reader.ReadDirectoryReference())!;
-			SystemIncludePaths = Reader.ReadList(() => Reader.ReadDirectoryReference())!;
+			IncludePaths = Reader.ReadList(() => Reader.ReadCompactDirectoryReference())!;
+			SystemIncludePaths = Reader.ReadList(() => Reader.ReadCompactDirectoryReference())!;
 			Definitions = Reader.ReadList(() => Reader.ReadString())!;
 			ForceIncludeFiles = Reader.ReadList(() => Reader.ReadFileItem())!;
 			Arguments = Reader.ReadList(() => Reader.ReadString())!;
@@ -375,8 +375,8 @@ namespace UnrealBuildTool
 			Writer.WriteFileItem(CreatePchFile);
 			Writer.WriteFileItem(UsingPchFile);
 			Writer.WriteFileItem(PchThroughHeaderFile);
-			Writer.WriteList(IncludePaths, Item => Writer.WriteDirectoryReference(Item));
-			Writer.WriteList(SystemIncludePaths, Item => Writer.WriteDirectoryReference(Item));
+			Writer.WriteList(IncludePaths, Item => Writer.WriteCompactDirectoryReference(Item));
+			Writer.WriteList(SystemIncludePaths, Item => Writer.WriteCompactDirectoryReference(Item));
 			Writer.WriteList(Definitions, Item => Writer.WriteString(Item));
 			Writer.WriteList(ForceIncludeFiles, Item => Writer.WriteFileItem(Item));
 			Writer.WriteList(Arguments, Item => Writer.WriteString(Item));
