[UBT]
* Deleted code that should be deleted after 5.2

diff --git a/Engine/Source/Programs/UnrealBuildTool/System/ExternalExecution.cs b/Engine/Source/Programs/UnrealBuildTool/System/ExternalExecution.cs
--- a/Engine/Source/Programs/UnrealBuildTool/System/ExternalExecution.cs
+++ b/Engine/Source/Programs/UnrealBuildTool/System/ExternalExecution.cs
@@ -534,17 +534,6 @@ namespace UnrealBuildTool
 				UHTModuleInfo Info = ModuleInfoArray[Idx];
 				Info.PublicDefines.AddRange(Module.PublicDefinitions);
 
-				// Delete any obsolete files at the root of the generated code directory
-				// Remove this code after engine version 5.2 (SOL-2551)
-				if (Module.GeneratedCodeDirectory != null && DirectoryReference.Exists(Module.GeneratedCodeDirectory))
-				{
-					DirectoryItem GeneratedCodeDirectory = DirectoryItem.GetItemByDirectoryReference(Module.GeneratedCodeDirectory);
-					foreach (FileItem File in GeneratedCodeDirectory.EnumerateFiles())
-					{
-						File.Delete(Logger);
-					}
-				}
-
 				if (Info.PublicUObjectClassesHeaders.Count > 0 || Info.PrivateUObjectHeaders.Count > 0 || Info.PublicUObjectHeaders.Count > 0 || Info.InternalUObjectHeaders.Count > 0)
 				{
 					Module.bHasUObjects = true;
