[UBT]
* Optimization - WriteFileIfChanged taking enumerable of strings now join them all up to one string and use the WriteFileIfChanged path for one string. This saves some time seen when profiling
* Optimization - ExpandVariables use OrdinalCompare to find "$("

diff --git a/Engine/Source/Programs/UnrealBuildTool/System/Utils.cs b/Engine/Source/Programs/UnrealBuildTool/System/Utils.cs
--- a/Engine/Source/Programs/UnrealBuildTool/System/Utils.cs
+++ b/Engine/Source/Programs/UnrealBuildTool/System/Utils.cs
@@ -93,7 +93,7 @@ namespace UnrealBuildTool
 		public static string ExpandVariables(string InputString, Dictionary<string, string>? AdditionalVariables = null, bool bUseAdditionalVariablesOnly = false)
 		{
 			string Result = InputString;
-			for (int Idx = Result.IndexOf("$("); Idx != -1; Idx = Result.IndexOf("$(", Idx))
+			for (int Idx = Result.IndexOf("$(", StringComparison.Ordinal); Idx != -1; Idx = Result.IndexOf("$(", Idx, StringComparison.Ordinal))
 			{
 				// Find the end of the variable name
 				int EndIdx = Result.IndexOf(')', Idx + 2);
@@ -1699,44 +1699,8 @@ namespace UnrealBuildTool
 		/// <param name="Logger">Logger for output</param>
 		internal static void WriteFileIfChanged(FileItem FileItem, IEnumerable<string> ContentLines, StringComparison Comparison, ILogger Logger)
 		{
-			// Only write the file if its contents have changed.
-			FileReference Location = FileItem.Location;
-			
-			if (!FileItem.Exists)
-			{
-				DirectoryReference.CreateDirectory(Location.Directory);
-				FileReference.WriteAllLines(Location, ContentLines, GetEncodingForStrings(ContentLines));
-				FileItem.ResetCachedInfo();
-
-				RecordWriteFileIfChanged(FileItem.Location, bNew: true, bChanged: true, Logger);
-			}
-			else
-			{
-				string[] CurrentContents = File.ReadAllLines(FileItem.FullName);
-				if (!CurrentContents.SequenceEqual(ContentLines, StringComparer.FromComparison(Comparison)))
-				{
-					FileReference BackupFile = new FileReference($"{FileItem.FullName}.old");
-					try
-					{
-						Logger.LogDebug("Updating {FileItem}: contents have changed. Saving previous version to {BackupFile}.", FileItem.Location, BackupFile);
-						FileReference.Delete(BackupFile);
-						FileReference.Move(Location, BackupFile);
-					}
-					catch (Exception Ex)
-					{
-						Logger.LogWarning("Unable to rename {FileItem} to {BackupFile}", FileItem, BackupFile);
-						Logger.LogDebug(Ex, "{Ex}", ExceptionUtils.FormatExceptionDetails(Ex));
-					}
-					FileReference.WriteAllLines(Location, ContentLines, GetEncodingForStrings(ContentLines));
-					FileItem.ResetCachedInfo();
-
-					RecordWriteFileIfChanged(FileItem.Location, bNew: false, bChanged: true, Logger);
-				}
-				else
-				{ 
-					RecordWriteFileIfChanged(FileItem.Location, bNew: false, bChanged: false, Logger);
-				}
-			}
+			var Contents = string.Join(Environment.NewLine, ContentLines);
+			WriteFileIfChanged(FileItem, Contents, Comparison, Logger);
 		}
 
 		/// <summary>
