[UBT]
* Minor optimizations based on superluminal profiling information

diff --git a/Engine/Source/Programs/UnrealBuildTool/Configuration/UEBuildBinary.cs b/Engine/Source/Programs/UnrealBuildTool/Configuration/UEBuildBinary.cs
--- a/Engine/Source/Programs/UnrealBuildTool/Configuration/UEBuildBinary.cs
+++ b/Engine/Source/Programs/UnrealBuildTool/Configuration/UEBuildBinary.cs
@@ -380,10 +380,8 @@ namespace UnrealBuildTool
 
 			foreach (UEBuildModule Module in Modules)
 			{
-				if (!IgnoreReferencedModules.Contains(Module))
+				if (IgnoreReferencedModules.Add(Module))
 				{
-					IgnoreReferencedModules.Add(Module);
-
 					Module.GetAllDependencyModules(ReferencedModules, IgnoreReferencedModules, bIncludeDynamicallyLoaded, bForceCircular, bOnlyDirectDependencies: false);
 
 					ReferencedModules.Add(Module);
@@ -727,18 +725,22 @@ namespace UnrealBuildTool
 					if (Target.LinkType == TargetLinkType.Monolithic)
 					{
 						Graph.SetOutputItemsForModule(Module.Name, LinkInputFiles.ToArray());
-					}
 
-					// NOTE: Because of 'Shared PCHs', in monolithic builds the same PCH file may appear as a link input
-					// multiple times for a single binary.  We'll check for that here, and only add it once.  This avoids
-					// a linker warning about redundant .obj files. 
-					foreach (FileItem LinkInputFile in LinkInputFiles)
-					{
-						if (!BinaryLinkEnvironment.InputFiles.Contains(LinkInputFile))
+						// NOTE: Because of 'Shared PCHs', in monolithic builds the same PCH file may appear as a link input
+						// multiple times for a single binary.  We'll check for that here, and only add it once.  This avoids
+						// a linker warning about redundant .obj files. 
+						foreach (FileItem LinkInputFile in LinkInputFiles)
 						{
-							BinaryLinkEnvironment.InputFiles.Add(LinkInputFile);
+							if (!BinaryLinkEnvironment.InputFiles.Contains(LinkInputFile))
+							{
+								BinaryLinkEnvironment.InputFiles.Add(LinkInputFile);
+							}
 						}
 					}
+					else
+					{
+						BinaryLinkEnvironment.InputFiles.AddRange(LinkInputFiles);
+					}
 
 					// Force a reference to initialize module for this binary
 					if(Module.Rules.bRequiresImplementModule ?? true)
diff --git a/Engine/Source/Programs/UnrealBuildTool/Configuration/UEBuildModule.cs b/Engine/Source/Programs/UnrealBuildTool/Configuration/UEBuildModule.cs
--- a/Engine/Source/Programs/UnrealBuildTool/Configuration/UEBuildModule.cs
+++ b/Engine/Source/Programs/UnrealBuildTool/Configuration/UEBuildModule.cs
@@ -555,10 +555,7 @@ namespace UnrealBuildTool
 		private void AddIncludePaths(HashSet<DirectoryReference> IncludePaths, HashSet<DirectoryReference> IncludePathsToAdd)
 		{
 			// Need to check whether directories exist to avoid bloating compiler command line with generated code directories
-			foreach(DirectoryReference IncludePathToAdd in IncludePathsToAdd)
-			{
-				IncludePaths.Add(IncludePathToAdd);
-			}
+			IncludePaths.UnionWith(IncludePathsToAdd);
 		}
 
 		/// <summary>
@@ -958,7 +955,7 @@ namespace UnrealBuildTool
 		/// <param name="bIncludeDynamicallyLoaded">True if dynamically loaded modules (and all of their dependent modules) should be included.</param>
 		/// <param name="bForceCircular">True if circular dependencies should be processed</param>
 		/// <param name="bOnlyDirectDependencies">True to return only this module's direct dependencies</param>
-		public virtual void GetAllDependencyModules(List<UEBuildModule> ReferencedModules, HashSet<UEBuildModule> IgnoreReferencedModules, bool bIncludeDynamicallyLoaded, bool bForceCircular, bool bOnlyDirectDependencies)
+		public void GetAllDependencyModules(List<UEBuildModule> ReferencedModules, HashSet<UEBuildModule> IgnoreReferencedModules, bool bIncludeDynamicallyLoaded, bool bForceCircular, bool bOnlyDirectDependencies)
 		{
 			List<UEBuildModule> AllDependencyModules = new List<UEBuildModule>(PrivateDependencyModules!.Count + PublicDependencyModules!.Count + (bIncludeDynamicallyLoaded ? DynamicallyLoadedModules!.Count : 0));
 			AllDependencyModules.AddRange(PrivateDependencyModules!);
@@ -970,13 +967,11 @@ namespace UnrealBuildTool
 
 			foreach (UEBuildModule DependencyModule in AllDependencyModules)
 			{
-				if (!IgnoreReferencedModules.Contains(DependencyModule))
+				// Don't follow circular back-references!
+				if (bForceCircular || !HasCircularDependencyOn(DependencyModule.Name))
 				{
-					// Don't follow circular back-references!
-					if (bForceCircular || !HasCircularDependencyOn(DependencyModule.Name))
+					if (IgnoreReferencedModules.Add(DependencyModule))
 					{
-						IgnoreReferencedModules.Add(DependencyModule);
-
 						if (!bOnlyDirectDependencies)
 						{
 							// Recurse into dependent modules first
