[UBT]
* ActionGraphBuilder - Added AllowAsync parameter to CreateIntermediateTextFile which defaults to true.
* UEBuildModuleCPP - Added AllowAsync false to one place where the code touches the file after CreateIntermediateTextFile has been called
* UEBuildModuleCPP - Changed so code writing LiveCodingInfo.json is using same path as others when it comes to writing files

diff --git a/Engine/Source/Programs/UnrealBuildTool/Configuration/UEBuildModuleCPP.cs b/Engine/Source/Programs/UnrealBuildTool/Configuration/UEBuildModuleCPP.cs
--- a/Engine/Source/Programs/UnrealBuildTool/Configuration/UEBuildModuleCPP.cs
+++ b/Engine/Source/Programs/UnrealBuildTool/Configuration/UEBuildModuleCPP.cs
@@ -745,8 +745,8 @@ namespace UnrealBuildTool
 			// Write a mapping of unity object file to standalone object file for live coding
 			if(Rules.Target.bWithLiveCoding)
 			{
-				FileReference UnityManifestFile = FileReference.Combine(IntermediateDirectory, "LiveCodingInfo.json");
-				using (JsonWriter Writer = new JsonWriter(UnityManifestFile))
+				StringWriter StringWriter = new();
+				using (JsonWriter Writer = new JsonWriter(StringWriter))
 				{
 					Writer.WriteObjectStart();
 					Writer.WriteObjectStart("RemapUnityFiles");
@@ -762,6 +762,9 @@ namespace UnrealBuildTool
 					Writer.WriteObjectEnd();
 					Writer.WriteObjectEnd();
 				}
+
+				FileReference UnityManifestFile = FileReference.Combine(IntermediateDirectory, "LiveCodingInfo.json");
+				Graph.CreateIntermediateTextFile(UnityManifestFile, StringWriter.ToString());
 			}
 
 			// IWYU needs to build all headers separate from cpp files to produce proper recommendations for includes
@@ -1433,7 +1436,7 @@ namespace UnrealBuildTool
 			}
 
 			// Create the item
-			FileItem WrapperFile = Graph.CreateIntermediateTextFile(OutputFile, WrapperContents.ToString());
+			FileItem WrapperFile = Graph.CreateIntermediateTextFile(OutputFile, WrapperContents.ToString(), AllowAsync: false);
 
 			// Touch it if the included file is newer, to make sure our timestamp dependency checking is accurate.
 			if (IncludedFile.LastWriteTimeUtc > WrapperFile.LastWriteTimeUtc)
diff --git a/Engine/Source/Programs/UnrealBuildTool/Configuration/UEBuildTarget.cs b/Engine/Source/Programs/UnrealBuildTool/Configuration/UEBuildTarget.cs
--- a/Engine/Source/Programs/UnrealBuildTool/Configuration/UEBuildTarget.cs
+++ b/Engine/Source/Programs/UnrealBuildTool/Configuration/UEBuildTarget.cs
@@ -2734,6 +2734,9 @@ namespace UnrealBuildTool
 				WritePublicHeader(Binary, Header, GlobalCompileEnvironment, Logger);
 			}
 
+			// Wait for makefile tasks to finish (tasks writing files to disk etc)
+			MakefileBuilder.WaitOnWriteTasks();
+
 			// Clean any stale modules which exist in multiple output directories. This can lead to the wrong DLL being loaded on Windows.
 			CleanStaleModules(Logger);
 
diff --git a/Engine/Source/Programs/UnrealBuildTool/Modes/GenerateClangDatabase.cs b/Engine/Source/Programs/UnrealBuildTool/Modes/GenerateClangDatabase.cs
--- a/Engine/Source/Programs/UnrealBuildTool/Modes/GenerateClangDatabase.cs
+++ b/Engine/Source/Programs/UnrealBuildTool/Modes/GenerateClangDatabase.cs
@@ -56,13 +56,13 @@ namespace UnrealBuildTool
 			}
 
 			/// <inheritdoc/>
-			public void CreateIntermediateTextFile(FileItem FileItem, string Contents)
+			public void CreateIntermediateTextFile(FileItem FileItem, string Contents, bool AllowAsync)
 			{
 				Utils.WriteFileIfChanged(FileItem, Contents, Logger);
 			}
 
 			/// <inheritdoc/>
-			public void CreateIntermediateTextFile(FileItem FileItem, IEnumerable<string> ContentLines)
+			public void CreateIntermediateTextFile(FileItem FileItem, IEnumerable<string> ContentLines, bool AllowAsync = true)
 			{
 				Utils.WriteFileIfChanged(FileItem, ContentLines, Logger);
 				CapturedTextFiles.Add(new Tuple<FileItem, IEnumerable<string>>(FileItem, ContentLines));
diff --git a/Engine/Source/Programs/UnrealBuildTool/System/ActionGraphBuilder.cs b/Engine/Source/Programs/UnrealBuildTool/System/ActionGraphBuilder.cs
--- a/Engine/Source/Programs/UnrealBuildTool/System/ActionGraphBuilder.cs
+++ b/Engine/Source/Programs/UnrealBuildTool/System/ActionGraphBuilder.cs
@@ -28,16 +28,18 @@ namespace UnrealBuildTool
 		/// </summary>
 		/// <param name="Location">Location of the response file</param>
 		/// <param name="Contents">Contents of the file</param>
+		/// <param name="AllowAsync">Allows the backend to write the file in a separate task.</param>
 		/// <returns>New file item</returns>
-		void CreateIntermediateTextFile(FileItem Location, string Contents);
+		void CreateIntermediateTextFile(FileItem Location, string Contents, bool AllowAsync = true);
 
 		/// <summary>
 		/// Creates a response file for use in the action graph, with a newline between each string in ContentLines
 		/// </summary>
 		/// <param name="Location">Location of the response file</param>
 		/// <param name="ContentLines">Contents of the file</param>
+		/// <param name="AllowAsync">Allows the backend to write the file in a separate task.</param>
 		/// <returns>New file item</returns>
-		void CreateIntermediateTextFile(FileItem Location, IEnumerable<string> ContentLines);
+		void CreateIntermediateTextFile(FileItem Location, IEnumerable<string> ContentLines, bool AllowAsync = true);
 
 		/// <summary>
 		/// Adds a file which is in the non-unity working set
@@ -100,13 +102,13 @@ namespace UnrealBuildTool
 		}
 
 		/// <inheritdoc/>
-		public void CreateIntermediateTextFile(FileItem FileItem, string Contents)
+		public void CreateIntermediateTextFile(FileItem FileItem, string Contents, bool AllowAsync = true)
 		{
 			Utils.WriteFileIfChanged(FileItem, Contents, Logger);
 		}
 
 		/// <inheritdoc/>
-		public void CreateIntermediateTextFile(FileItem FileItem, IEnumerable<string> ContentLines)
+		public void CreateIntermediateTextFile(FileItem FileItem, IEnumerable<string> ContentLines, bool AllowAsync = true)
 		{
 			Utils.WriteFileIfChanged(FileItem, ContentLines, Logger);
 		}
@@ -168,15 +170,15 @@ namespace UnrealBuildTool
 		}
 
 		/// <inheritdoc/>
-		public virtual void CreateIntermediateTextFile(FileItem FileItem, string Contents)
+		public virtual void CreateIntermediateTextFile(FileItem FileItem, string Contents, bool AllowAsync = true)
 		{
-			Inner.CreateIntermediateTextFile(FileItem, Contents);
+			Inner.CreateIntermediateTextFile(FileItem, Contents, AllowAsync);
 		}
 
 		/// <inheritdoc/>
-		public virtual void CreateIntermediateTextFile(FileItem FileItem, IEnumerable<string> ContentLines)
+		public virtual void CreateIntermediateTextFile(FileItem FileItem, IEnumerable<string> ContentLines, bool AllowAsync = true)
 		{
-			Inner.CreateIntermediateTextFile(FileItem, ContentLines);
+			Inner.CreateIntermediateTextFile(FileItem, ContentLines, AllowAsync);
 		}
 
 		/// <inheritdoc/>
@@ -309,11 +311,12 @@ namespace UnrealBuildTool
 		/// <param name="Graph">The action graph</param>
 		/// <param name="AbsolutePath">Path to the intermediate file to create</param>
 		/// <param name="Contents">Contents of the new file</param>
+		/// <param name="AllowAsync">Allows the backend to write the file in a separate task.</param>
 		/// <returns>File item for the newly created file</returns>
-		public static FileItem CreateIntermediateTextFile(this IActionGraphBuilder Graph, FileReference AbsolutePath, string Contents)
+		public static FileItem CreateIntermediateTextFile(this IActionGraphBuilder Graph, FileReference AbsolutePath, string Contents, bool AllowAsync = true)
 		{
 			FileItem FileItem = FileItem.GetItemByFileReference(AbsolutePath);
-			Graph.CreateIntermediateTextFile(FileItem, Contents);
+			Graph.CreateIntermediateTextFile(FileItem, Contents, AllowAsync);
 			return FileItem;
 		}
 
@@ -324,11 +327,12 @@ namespace UnrealBuildTool
 		/// <param name="Graph">The action graph</param>
 		/// <param name="AbsolutePath">Path to the intermediate file to create</param>
 		/// <param name="ContentLines">Contents of the new file</param>
+		/// <param name="AllowAsync">Allows the backend to write the file in a separate task.</param>
 		/// <returns>File item for the newly created file</returns>
-		public static FileItem CreateIntermediateTextFile(this IActionGraphBuilder Graph, FileReference AbsolutePath, IEnumerable<string> ContentLines)
+		public static FileItem CreateIntermediateTextFile(this IActionGraphBuilder Graph, FileReference AbsolutePath, IEnumerable<string> ContentLines, bool AllowAsync = true)
 		{
 			FileItem FileItem = UnrealBuildBase.FileItem.GetItemByFileReference(AbsolutePath);
-			Graph.CreateIntermediateTextFile(FileItem, ContentLines);
+			Graph.CreateIntermediateTextFile(FileItem, ContentLines, AllowAsync);
 			return FileItem;
 		}
 	}
diff --git a/Engine/Source/Programs/UnrealBuildTool/System/TargetMakefile.cs b/Engine/Source/Programs/UnrealBuildTool/System/TargetMakefile.cs
--- a/Engine/Source/Programs/UnrealBuildTool/System/TargetMakefile.cs
+++ b/Engine/Source/Programs/UnrealBuildTool/System/TargetMakefile.cs
@@ -5,6 +5,7 @@ using System.Collections.Generic;
 using System.IO;
 using System.Linq;
 using System.Reflection;
+using System.Threading.Tasks;
 using EpicGames.Core;
 using EpicGames.UHT.Utils;
 using Microsoft.Extensions.Logging;
@@ -68,13 +69,14 @@ namespace UnrealBuildTool
 		private readonly ILogger Logger;
 		public TargetMakefile Makefile { get; }
 
+		private List<Task> WriteTasks = new();
+
 		public TargetMakefileBuilder(TargetMakefile InMakefile, ILogger InLogger)
 		{
 			Logger = InLogger;
 			Makefile = InMakefile;
 		}
 
-
 		/// <inheritdoc/>
 		public void AddAction(IExternalAction Action)
 		{
@@ -82,25 +84,39 @@ namespace UnrealBuildTool
 		}
 
 		/// <inheritdoc/>
-		public void CreateIntermediateTextFile(FileItem FileItem, string Contents)
+		public void CreateIntermediateTextFile(FileItem FileItem, string Contents, bool AllowAsync = true)
 		{
-			// Write the file
-			Utils.WriteFileIfChanged(FileItem.Location, Contents, Logger);
+			if (AllowAsync)
+			{
+				WriteTasks.Add(Task.Run(() =>
+				{
+					Utils.WriteFileIfChanged(FileItem.Location, Contents, Logger);
+				}));
+			}
+			else
+			{
+				Utils.WriteFileIfChanged(FileItem.Location, Contents, Logger);
+			}
 
-			// Reset the file info, in case it already knows about the old file
 			Makefile.InternalDependencies.Add(FileItem);
-			FileItem.ResetCachedInfo();
 		}
 
 		/// <inheritdoc/>
-		public void CreateIntermediateTextFile(FileItem FileItem, IEnumerable<string> ContentLines)
+		public void CreateIntermediateTextFile(FileItem FileItem, IEnumerable<string> ContentLines, bool AllowAsync = true)
 		{
-			// Write the file
-			Utils.WriteFileIfChanged(FileItem, ContentLines, Logger);
-
+			if (AllowAsync)
+			{
+				WriteTasks.Add(Task.Run(() =>
+				{
+					Utils.WriteFileIfChanged(FileItem, ContentLines, Logger);
+				}));
+			}
+			else
+			{
+				Utils.WriteFileIfChanged(FileItem.Location, ContentLines, Logger);
+			}
 			// Reset the file info, in case it already knows about the old file
 			Makefile.InternalDependencies.Add(FileItem);
-			FileItem.ResetCachedInfo();
 		}
 
 		/// <inheritdoc/>
@@ -141,6 +157,11 @@ namespace UnrealBuildTool
 		{
 			Makefile.ModuleNameToOutputItems[ModuleName] = OutputItems;
 		}
+
+		internal void WaitOnWriteTasks()
+		{
+			Task.WaitAll(WriteTasks.ToArray());
+		}
 	}
 
 	/// <summary>
@@ -566,6 +587,12 @@ namespace UnrealBuildTool
 				}
 			}
 
+			// Run task that parse file system and find plugin files etc. (this can run in parallel with make file creation)
+			Task<FileReference[]> EnumPluginsTask = Task.Run(() =>
+			{
+				return PluginsBase.EnumerateUbtPlugins(ProjectFile).ToArray();
+			});
+
 			TargetMakefile Makefile;
 			using (GlobalTracer.Instance.BuildSpan("Loading makefile").StartActive())
 			{
@@ -631,7 +658,7 @@ namespace UnrealBuildTool
 
 				// Check to see if the number of plugins changed
 				{
-					FileReference[]? Plugins = PluginsBase.EnumerateUbtPlugins(ProjectFile)?.ToArray();
+					FileReference[] Plugins = EnumPluginsTask.Result;
 
 					if ((Plugins != null) != (Makefile.UbtPlugins != null) ||
 						(Plugins != null && !Enumerable.SequenceEqual(Plugins, Makefile.UbtPlugins!)))
