[UBT]
* Fixed so only colliding headers gets hash suffix for files in immediate folder

diff --git a/Engine/Source/Programs/UnrealBuildTool/Configuration/UEBuildModuleCPP.cs b/Engine/Source/Programs/UnrealBuildTool/Configuration/UEBuildModuleCPP.cs
--- a/Engine/Source/Programs/UnrealBuildTool/Configuration/UEBuildModuleCPP.cs
+++ b/Engine/Source/Programs/UnrealBuildTool/Configuration/UEBuildModuleCPP.cs
@@ -770,33 +770,8 @@ namespace UnrealBuildTool
 			// IWYU needs to build all headers separate from cpp files to produce proper recommendations for includes
 			if (Target.bIWYU)
 			{
-				// Find FileItems for module's pch files
-				FileItem? PrivatePchFileItem = null;
-				if (Rules.PrivatePCHHeaderFile != null)
-					PrivatePchFileItem = FileItem.GetItemByFileReference(FileReference.Combine(ModuleDirectory, Rules.PrivatePCHHeaderFile));
-				FileItem? SharedPchFileItem = null;
-				if (Rules.SharedPCHHeaderFile != null)
-					SharedPchFileItem = FileItem.GetItemByFileReference(FileReference.Combine(ModuleDirectory, Rules.SharedPCHHeaderFile));
-
 				// Collect the headers that should be built
-				List<FileItem> HeaderFileItems = new();
-				foreach (FileItem HeaderFileItem in InputFiles.HeaderFiles)
-				{
-					// We don't want to build pch files in iwyu, skip those.
-					if (HeaderFileItem == PrivatePchFileItem || HeaderFileItem == SharedPchFileItem)
-					{
-						continue;
-					}
-
-					// If file is skipped by header units it means they can't be compiled by themselves and we must skip them too
-					if (CompileEnvironment.MetadataCache.GetHeaderUnitType(HeaderFileItem) != HeaderUnitType.Valid)
-					{
-						continue;
-					}
-
-					HeaderFileItems.Add(HeaderFileItem);
-				}
-
+				List<FileItem> HeaderFileItems = GetCompilableHeaders(InputFiles, CompileEnvironment);
 				if (HeaderFileItems.Count > 0)
 				{
 					if (Target.bIWYUHeadersOnly)
@@ -812,6 +787,51 @@ namespace UnrealBuildTool
 			return LinkInputFiles;
 		}
 
+		List<FileItem> GetCompilableHeaders(InputFileCollection InputFiles, CppCompileEnvironment CompileEnvironment)
+		{
+			// Find FileItems for module's pch files
+			FileItem? PrivatePchFileItem = null;
+			if (Rules.PrivatePCHHeaderFile != null)
+				PrivatePchFileItem = FileItem.GetItemByFileReference(FileReference.Combine(ModuleDirectory, Rules.PrivatePCHHeaderFile));
+			FileItem? SharedPchFileItem = null;
+			if (Rules.SharedPCHHeaderFile != null)
+				SharedPchFileItem = FileItem.GetItemByFileReference(FileReference.Combine(ModuleDirectory, Rules.SharedPCHHeaderFile));
+
+			Dictionary<string, FileItem> NameToFileItem = new();
+
+			HashSet<FileItem> CollidingFiles = new();
+
+			// Collect the headers that should be built
+			List<FileItem> HeaderFileItems = new();
+			foreach (FileItem HeaderFileItem in InputFiles.HeaderFiles)
+			{
+				// We don't want to build pch files in iwyu, skip those.
+				if (HeaderFileItem == PrivatePchFileItem || HeaderFileItem == SharedPchFileItem)
+				{
+					continue;
+				}
+
+				// If file is skipped by header units it means they can't be compiled by themselves and we must skip them too
+				if (CompileEnvironment.MetadataCache.GetHeaderUnitType(HeaderFileItem) != HeaderUnitType.Valid)
+				{
+					continue;
+				}
+
+				if (!NameToFileItem.TryAdd(HeaderFileItem.Name, HeaderFileItem))
+				{
+					CollidingFiles.Add(NameToFileItem[HeaderFileItem.Name]);
+					CollidingFiles.Add(HeaderFileItem);
+				}
+
+				HeaderFileItems.Add(HeaderFileItem);
+			}
+
+			if (CollidingFiles.Count != 0)
+				CompileEnvironment.CollidingNames = CollidingFiles;
+
+			return HeaderFileItems;
+		}
+
 		/// <summary>
 		/// Create a shared PCH template for this module, which allows constructing shared PCH instances in the future
 		/// </summary>
diff --git a/Engine/Source/Programs/UnrealBuildTool/System/CppCompileEnvironment.cs b/Engine/Source/Programs/UnrealBuildTool/System/CppCompileEnvironment.cs
--- a/Engine/Source/Programs/UnrealBuildTool/System/CppCompileEnvironment.cs
+++ b/Engine/Source/Programs/UnrealBuildTool/System/CppCompileEnvironment.cs
@@ -467,6 +467,11 @@ namespace UnrealBuildTool
 		/// </summary>
 		public Dictionary<FileItem, List<FileItem>> FileInlineGenCPPMap = new();
 
+		/// <summary>
+		/// FileItems with colliding names. (Which means they would overwrite each other in intermediate folder
+		/// </summary>
+		public HashSet<FileItem>? CollidingNames;
+
 		/// <summary>
 		/// The C++ preprocessor definitions to use.
 		/// </summary>
@@ -628,6 +633,7 @@ namespace UnrealBuildTool
 			bCheckSystemHeadersForModification = Other.bCheckSystemHeadersForModification;
 			ForceIncludeFiles.AddRange(Other.ForceIncludeFiles);
 			AdditionalPrerequisites.AddRange(Other.AdditionalPrerequisites);
+			CollidingNames = Other.CollidingNames;
 			FileInlineGenCPPMap = new Dictionary<FileItem, List<FileItem>>(Other.FileInlineGenCPPMap);
 			Definitions.AddRange(Other.Definitions);
 			AdditionalArguments = Other.AdditionalArguments;
diff --git a/Engine/Source/Programs/UnrealBuildTool/ToolChain/ClangToolChain.cs b/Engine/Source/Programs/UnrealBuildTool/ToolChain/ClangToolChain.cs
--- a/Engine/Source/Programs/UnrealBuildTool/ToolChain/ClangToolChain.cs
+++ b/Engine/Source/Programs/UnrealBuildTool/ToolChain/ClangToolChain.cs
@@ -931,15 +931,22 @@ namespace UnrealBuildTool
 				CompileAction.PrerequisiteItems.Add(FileItem.GetItemByFileReference(CompileEnvironment.PrecompiledHeaderIncludeFilename!));
 			}
 
+			string FileName = SourceFile.Name;
+			if (CompileEnvironment.CollidingNames != null && CompileEnvironment.CollidingNames.Contains(SourceFile))
+			{
+				string HashString = ContentHash.MD5(SourceFile.AbsolutePath.Substring(Unreal.RootDirectory.FullName.Length)).GetHashCode().ToString("X4");
+				FileName = Path.GetFileNameWithoutExtension(FileName) + "_" + HashString + Path.GetExtension(FileName);
+			}
+
 			FileItem OutputFile;
 			if (CompileEnvironment.PrecompiledHeaderAction == PrecompiledHeaderAction.Create)
 			{
-				OutputFile = FileItem.GetItemByFileReference(FileReference.Combine(OutputDir, GetFileNameFromExtension(SourceFile.AbsolutePath, ".gch")));
+				OutputFile = FileItem.GetItemByFileReference(FileReference.Combine(OutputDir, GetFileNameFromExtension(FileName, ".gch")));
 				CompileResult.PrecompiledHeaderFile = OutputFile;
 			}
 			else if (CompileEnvironment.bPreprocessOnly)
 			{
-				OutputFile = FileItem.GetItemByFileReference(FileReference.Combine(OutputDir, GetFileNameFromExtension(SourceFile.AbsolutePath, ".i")));
+				OutputFile = FileItem.GetItemByFileReference(FileReference.Combine(OutputDir, GetFileNameFromExtension(FileName, ".i")));
 				CompileResult.ObjectFiles.Add(OutputFile);
 
 				// Clang does EITHER pre-process or object file.
@@ -952,7 +959,7 @@ namespace UnrealBuildTool
 			else
 			{
 				string ObjectFileExtension = (CompileEnvironment.AdditionalArguments != null && CompileEnvironment.AdditionalArguments.Contains("-emit-llvm")) ? ".bc" : ".o";
-				OutputFile = FileItem.GetItemByFileReference(FileReference.Combine(OutputDir, GetFileNameFromExtension(SourceFile.AbsolutePath, ObjectFileExtension)));
+				OutputFile = FileItem.GetItemByFileReference(FileReference.Combine(OutputDir, GetFileNameFromExtension(FileName, ObjectFileExtension)));
 				CompileResult.ObjectFiles.Add(OutputFile);
 			}
 
@@ -967,7 +974,7 @@ namespace UnrealBuildTool
 				// Generate the timing info
 				if (CompileEnvironment.bPrintTimingInfo)
 				{
-					FileItem TraceFile = FileItem.GetItemByFileReference(FileReference.Combine(OutputDir, GetFileNameFromExtension(SourceFile.AbsolutePath, ".json")));
+					FileItem TraceFile = FileItem.GetItemByFileReference(FileReference.Combine(OutputDir, GetFileNameFromExtension(FileName, ".json")));
 					Arguments.Add("-ftime-trace");
 					CompileAction.ProducedItems.Add(TraceFile);
 				}
@@ -975,7 +982,7 @@ namespace UnrealBuildTool
 				// Generate the included header dependency list
 				if (!PreprocessDepends)
 				{
-					FileItem DependencyListFile = FileItem.GetItemByFileReference(FileReference.Combine(OutputDir, GetFileNameFromExtension(SourceFile.AbsolutePath, ".d")));
+					FileItem DependencyListFile = FileItem.GetItemByFileReference(FileReference.Combine(OutputDir, GetFileNameFromExtension(FileName, ".d")));
 					Arguments.Add(GetDepencenciesListFileArgument(DependencyListFile));
 					CompileAction.DependencyListFile = DependencyListFile;
 					CompileAction.ProducedItems.Add(DependencyListFile);
@@ -983,7 +990,7 @@ namespace UnrealBuildTool
 			}
 			else
 			{
-				OutputFile = FileItem.GetItemByFileReference(FileReference.Combine(OutputDir, GetFileNameFromExtension(SourceFile.AbsolutePath, ".analysis")));
+				OutputFile = FileItem.GetItemByFileReference(FileReference.Combine(OutputDir, GetFileNameFromExtension(FileName, ".analysis")));
 			}
 
 			// Add the parameters needed to compile the output file to the command-line.
@@ -992,7 +999,27 @@ namespace UnrealBuildTool
 
 		protected virtual List<string> ExpandResponseFileContents(List<string> ResponseFileContents)
         {
-			return ResponseFileContents.Select(x => Utils.ExpandVariables(x)).ToList();
+			// This code is optimized for the scenario where ResponseFile has no variables to expand
+			List<string> NewList = ResponseFileContents;
+			for (int I=0; I!=NewList.Count; ++I)
+			{
+				string Line = ResponseFileContents[I];
+				string NewLine = Utils.ExpandVariables(Line);
+				if (NewLine == Line)
+				{
+					continue;
+				}
+
+				lock (ResponseFileContents)
+				{
+					if (NewList == ResponseFileContents)
+					{
+						NewList = new List<string>(ResponseFileContents);
+					}
+				}
+				NewList[I] = NewLine;
+			}
+			return NewList;
 		}
 
 		protected virtual Action CompileCPPFile(CppCompileEnvironment CompileEnvironment, FileItem SourceFile, DirectoryReference OutputDir, string ModuleName, IActionGraphBuilder Graph, IReadOnlyCollection<string> GlobalArguments, CPPOutput Result)
@@ -1021,7 +1048,9 @@ namespace UnrealBuildTool
 				// >>>> xxx-clang.exe': The filename or extension is too long.  (0xCE)
 				// Clang processes and modifies the response file contents and this makes the final cmd line size hard for us to predict.
 				// To be conservative we add a dummy define to inflate the response file size and force clang to use the response file mode when we are close to the limit.
-				int CmdLineLength = Info.Clang.ToString().Length + string.Join(' ', ResponseFileContents).Length;
+				int CmdLineLength = Info.Clang.ToString().Length;
+				foreach (string Line in ResponseFileContents)
+					CmdLineLength += 1 + Line.Length;
 				bool bIsInDangerZone = CmdLineLength >= ClangCmdlineDangerZone && CmdLineLength <= ClangCmdLineMaxSize;
 				if (bIsInDangerZone)
 				{
diff --git a/Engine/Source/Programs/UnrealBuildTool/ToolChain/IWYUToolChain.cs b/Engine/Source/Programs/UnrealBuildTool/ToolChain/IWYUToolChain.cs
--- a/Engine/Source/Programs/UnrealBuildTool/ToolChain/IWYUToolChain.cs
+++ b/Engine/Source/Programs/UnrealBuildTool/ToolChain/IWYUToolChain.cs
@@ -60,16 +60,6 @@ namespace UnrealBuildTool
 		protected override string GetFileNameFromExtension(string AbsolutePath, string Extension)
 		{
 			string FileName = Path.GetFileName(AbsolutePath);
-			if (AbsolutePath.EndsWith(".h"))
-			{
-				string HashString = ContentHash.MD5(AbsolutePath).GetHashCode().ToString("X4");
-				FileName = Path.GetFileNameWithoutExtension(FileName) + "_" + HashString + ".h";
-
-				if (Extension != ".d" && Extension != ".o")
-				{
-					throw new NotImplementedException($"Files with extension {Extension} not handled by IWYUToolChain");
-				}
-			}
 
 			if (Extension == ".o")
 			{
@@ -126,6 +116,11 @@ namespace UnrealBuildTool
 						}
 					}
 				}
+				else
+				{
+					// We want to build header files first
+					Action.bIsHighPriority = true;
+				}
 
 				Action.CommandArguments = CommandLineArgs + Action.CommandArguments;
 				IwyuFiles.Add(Action.ProducedItems.First());
