[UBT]
* Run VNI at the same time as UHT to save wall time
* Changed from linq code to normal parallel.foreach as optimization

diff --git a/Engine/Source/Programs/UnrealBuildTool/Configuration/UEBuildTarget.cs b/Engine/Source/Programs/UnrealBuildTool/Configuration/UEBuildTarget.cs
--- a/Engine/Source/Programs/UnrealBuildTool/Configuration/UEBuildTarget.cs
+++ b/Engine/Source/Programs/UnrealBuildTool/Configuration/UEBuildTarget.cs
@@ -19,6 +19,7 @@ using UnrealBuildBase;
 using EpicGames.UHT.Utils;
 using Microsoft.Extensions.Logging;
 using System.Runtime.InteropServices;
+using System.Threading.Tasks;
 using UnrealBuildTool;
 
 
@@ -2365,6 +2366,27 @@ namespace UnrealBuildTool
 				return Makefile;
 			}
 
+#if __VPROJECT_AVAILABLE__
+			Task VNITask = Task.Run(() =>
+			{
+				// Prepare cached data for VNI header generation
+				if (Rules.bUseVerse)
+				{
+					using (Timeline.ScopeEvent("ExternalExecution.SetupVNIModules()"))
+					{
+						VNIExecution.SetupVNIModules(ModulesToGenerateHeadersFor, RulesAssembly, out Makefile.VNIModules);
+					}
+				}
+
+				// NOTE: Even in Gather mode, we need to run VNI to make sure the files exist for the static action graph to be setup correctly.  This is because VNI generates .cpp
+				// files that are injected as top level prerequisites.  If VNI only emitted included header files, we wouldn't need to run it during the Gather phase at all.
+				if (Makefile.VNIModules.Count > 0)
+				{
+					VNIExecution.ExecuteVNITool(Makefile, TargetDescriptor, Logger);
+				}
+			});
+#endif
+
 			// NOTE: Even in Gather mode, we need to run UHT to make sure the files exist for the static action graph to be setup correctly.  This is because UHT generates .cpp
 			// files that are injected as top level prerequisites.  If UHT only emitted included header files, we wouldn't need to run it during the Gather phase at all.
 			if (Makefile.UObjectModules.Count > 0)
@@ -2372,30 +2394,24 @@ namespace UnrealBuildTool
 				ExternalExecution.ExecuteHeaderToolIfNecessary(BuildConfiguration, ProjectFile, Makefile, TargetName, WorkingSet, Logger);
 			}
 
-#if __VPROJECT_AVAILABLE__
-			// Prepare cached data for VNI header generation
-			if (Rules.bUseVerse)
+			// Prefetch directory items for UHT folders since they are going to be used later
+			foreach (var Module in ModulesToGenerateHeadersFor)
 			{
-				using (Timeline.ScopeEvent("ExternalExecution.SetupVNIModules()"))
-				{
-					VNIExecution.SetupVNIModules(ModulesToGenerateHeadersFor, RulesAssembly, out Makefile.VNIModules);
-				}
+				if (Module.GeneratedCodeDirectoryUHT != null)
+					FileMetadataPrefetch.QueueDirectoryTree(Module.GeneratedCodeDirectoryUHT);
 			}
 
-			// NOTE: Even in Gather mode, we need to run VNI to make sure the files exist for the static action graph to be setup correctly.  This is because VNI generates .cpp
-			// files that are injected as top level prerequisites.  If VNI only emitted included header files, we wouldn't need to run it during the Gather phase at all.
-			if (Makefile.VNIModules.Count > 0)
-			{
-				VNIExecution.ExecuteVNITool(Makefile, TargetDescriptor, Logger);
-			}
-#endif
-
 			// Find all the shared PCHs.
 			if (Rules.bUseSharedPCHs)
 			{
 				FindSharedPCHs(OriginalBinaries, GlobalCompileEnvironment, Logger);
 			}
 
+			// Can probably be moved further down?
+#if __VPROJECT_AVAILABLE__
+			VNITask.Wait();
+#endif
+
 			// Compile the resource files common to all DLLs on Windows
 			if (!ShouldCompileMonolithic())
 			{
@@ -2444,14 +2460,14 @@ namespace UnrealBuildTool
 			// Cache inline gen cpp data
 			using (GlobalTracer.Instance.BuildSpan("CacheInlineGenCppData").StartActive())
 			{
-				Makefile.DirectoryToSourceFiles.AsParallel()
-					.ForAll((kvp) =>
+				// We don't want to run this in parallel since the tasks most likely is choked on writing response files
+				foreach (TargetMakefileSourceFileInfo[] Values in Makefile.DirectoryToSourceFiles.Values)
+				{
+					foreach (TargetMakefileSourceFileInfo SourceFileInfo in Values)
 					{
-						foreach (TargetMakefileSourceFileInfo SourceFileInfo in kvp.Value)
-						{
-							SourceFileInfo.InlineGenCppHash = TargetMakefileSourceFileInfo.CalculateInlineGenCppHash(MetadataCache.GetListOfInlinedGeneratedCppFiles(SourceFileInfo.SourceFileItem!));
-						}
-					});
+						SourceFileInfo.InlineGenCppHash = TargetMakefileSourceFileInfo.CalculateInlineGenCppHash(MetadataCache.GetListOfInlinedGeneratedCppFiles(SourceFileInfo.SourceFileItem!));
+					}
+				}
 			}
 
 			// Prepare all the runtime dependencies, copying them from their source folders if necessary
