[UBT]
* Changed so code related to inlined .gen.cpp is using FileItems etc to reduce number of sys calls

diff --git a/Engine/Source/Programs/UnrealBuildTool/Configuration/UEBuildModuleCPP.cs b/Engine/Source/Programs/UnrealBuildTool/Configuration/UEBuildModuleCPP.cs
--- a/Engine/Source/Programs/UnrealBuildTool/Configuration/UEBuildModuleCPP.cs
+++ b/Engine/Source/Programs/UnrealBuildTool/Configuration/UEBuildModuleCPP.cs
@@ -478,7 +478,7 @@ namespace UnrealBuildTool
 			// Compile all the generated CPP files
 			if (GeneratedCppDirectories != null && !CompileEnvironment.bHackHeaderGenerator)
 			{
-				var GeneratedFiles = new Dictionary<string, string>();
+				var GeneratedFiles = new Dictionary<string, FileItem>();
 				if (SpecificFilesToCompile.Count == 0)
 				{
 					foreach (string GeneratedDir in GeneratedCppDirectories)
@@ -490,14 +490,15 @@ namespace UnrealBuildTool
 
 						string Prefix = Path.GetFileName(GeneratedDir) + '/'; // "UHT/" or "VNI/"
 
-						string[] Files = Directory.GetFiles(GeneratedDir, "*.gen.cpp");
-						GeneratedFiles.EnsureCapacity(GeneratedFiles.Count + Files.Length);
-						foreach (var File in Files)
+						DirectoryItem DirItem = DirectoryItem.GetItemByPath(GeneratedDir);
+						foreach (FileItem File in DirItem.EnumerateFiles())
 						{
-							// Can't use GetFileNameWithoutAnyExtensions because of the .init.gen.cpp files
-							string FileName = Path.GetFileName(File);
-							string Key = Prefix + FileName.Substring(0, FileName.Length - ".gen.cpp".Length); 
-							GeneratedFiles.Add(Key, File);
+							string FileName = File.Name;
+							if (FileName.EndsWith(".gen.cpp"))
+							{
+								string Key = Prefix + FileName.Substring(0, FileName.Length - ".gen.cpp".Length);
+								GeneratedFiles.Add(Key, File);
+							}
 						}
 					}
 				}
@@ -507,7 +508,7 @@ namespace UnrealBuildTool
 					{
 						if (GeneratedCppDirectories.Any(x => FileToCompile.IsUnderDirectory(new DirectoryReference(x))))
 						{
-							GeneratedFiles.Add(FileToCompile.GetFileNameWithoutAnyExtensions(), FileToCompile.FullName);
+							GeneratedFiles.Add(FileToCompile.GetFileNameWithoutAnyExtensions(), FileItem.GetItemByFileReference(FileToCompile));
 						}
 					}
 				}
@@ -524,13 +525,13 @@ namespace UnrealBuildTool
 							{
 								string Prefix = "UHT/";
 								string Key = Prefix + ListOfInlinedGenCppsItem;
-								if (GeneratedFiles.Remove(Key, out string? FoundGenCppFile))
+								if (GeneratedFiles.Remove(Key, out FileItem? FoundGenCppFile))
 								{
 									if (!CompileEnvironment.FileInlineGenCPPMap.ContainsKey(CPPFileItem))
 									{
 										CompileEnvironment.FileInlineGenCPPMap[CPPFileItem] = new List<FileItem>();
 									}
-									CompileEnvironment.FileInlineGenCPPMap[CPPFileItem].Add(FileItem.GetItemByPath(FoundGenCppFile));
+									CompileEnvironment.FileInlineGenCPPMap[CPPFileItem].Add(FoundGenCppFile);
 								}
 								else
 								{
@@ -584,9 +585,8 @@ namespace UnrealBuildTool
 					}
 
 					// Compile all the generated files
-					foreach (string GeneratedFilename in GeneratedFiles.Values)
+					foreach (FileItem GeneratedCppFileItem in GeneratedFiles.Values)
 					{
-						FileItem GeneratedCppFileItem = FileItem.GetItemByPath(GeneratedFilename);
 						if (SpecificFilesToCompile.Count == 0 || SpecificFilesToCompile.Contains(GeneratedCppFileItem.Location))
 						{
 							GeneratedFileItems.Add(GeneratedCppFileItem);
