[UBT]
* Changed a bunch of code to use StringBuilder instead of TextWriter. Based on superluminal profiling data

diff --git a/Engine/Source/Programs/UnrealBuildTool/Configuration/UEBuildModuleCPP.cs b/Engine/Source/Programs/UnrealBuildTool/Configuration/UEBuildModuleCPP.cs
--- a/Engine/Source/Programs/UnrealBuildTool/Configuration/UEBuildModuleCPP.cs
+++ b/Engine/Source/Programs/UnrealBuildTool/Configuration/UEBuildModuleCPP.cs
@@ -1251,7 +1251,6 @@ namespace UnrealBuildTool
 			{
 				// Build the contents of the wrapper file
 				StringBuilder WrapperContents = new StringBuilder();
-				using (StringWriter Writer = new StringWriter(WrapperContents))
 				{
 					string FileString = File.AbsolutePath;
 					if (File.Location.IsUnderDirectory(Unreal.RootDirectory))
@@ -1259,13 +1258,13 @@ namespace UnrealBuildTool
 						FileString = File.Location.MakeRelativeTo(Unreal.EngineSourceDirectory);
 					}
 					FileString = FileString.Replace('\\', '/');
-					Writer.WriteLine("// Dedicated PCH for {0}", FileString);
-					Writer.WriteLine();
-					WriteDefinitions(CompileEnvironment.Definitions, Writer);
-					Writer.WriteLine();
+					WrapperContents.AppendLine("// Dedicated PCH for {0}", FileString);
+					WrapperContents.AppendLine();
+					WriteDefinitions(CompileEnvironment.Definitions, WrapperContents);
+					WrapperContents.AppendLine();
 					using(StreamReader Reader = new StreamReader(File.Location.FullName))
 					{
-						CppIncludeParser.CopyIncludeDirectives(Reader, Writer);
+						CppIncludeParser.CopyIncludeDirectives(Reader, WrapperContents);
 					}
 				}
 
@@ -1342,22 +1341,22 @@ namespace UnrealBuildTool
 						FileReference PrivateDefinitionsFile = FileReference.Combine(IntermediateDirectory, String.Format("Definitions.{0}.h", Name));
 
 						FileItem PrivateDefinitionsFileItem;
-						using (StringWriter Writer = new StringWriter())
 						{
+							StringBuilder Writer = new();
 							// Remove the module _API definition for cases where there are circular dependencies between the shared PCH module and modules using it
-							Writer.WriteLine("#undef {0}", ModuleApiDefine);
+							Writer.AppendLine("#undef {0}", ModuleApiDefine);
 
 							// Games may choose to use shared PCHs from the engine, so allow them to change the value of these macros
 							if(!Rules.bTreatAsEngineModule)
 							{
-								Writer.WriteLine("#undef UE_IS_ENGINE_MODULE");
-								Writer.WriteLine("#undef DEPRECATED_FORGAME");
-								Writer.WriteLine("#define DEPRECATED_FORGAME DEPRECATED");
-								Writer.WriteLine("#undef UE_DEPRECATED_FORGAME");
-								Writer.WriteLine("#define UE_DEPRECATED_FORGAME UE_DEPRECATED");
+								Writer.AppendLine("#undef UE_IS_ENGINE_MODULE");
+								Writer.AppendLine("#undef DEPRECATED_FORGAME");
+								Writer.AppendLine("#define DEPRECATED_FORGAME DEPRECATED");
+								Writer.AppendLine("#undef UE_DEPRECATED_FORGAME");
+								Writer.AppendLine("#define UE_DEPRECATED_FORGAME UE_DEPRECATED");
 								foreach (string DeprecationDefine in EngineIncludeOrderHelper.GetAllDeprecationDefines())
 								{
-									Writer.WriteLine("#undef " + DeprecationDefine);
+									Writer.AppendLine("#undef " + DeprecationDefine);
 								}
 							}
 
@@ -1401,8 +1400,9 @@ namespace UnrealBuildTool
 				}
 
 				FileReference PrivateDefinitionsFile = FileReference.Combine(IntermediateDirectory, PrivateDefinitionsName);
-				using (StringWriter Writer = new StringWriter())
 				{
+					StringBuilder Writer = new();
+					Writer.AppendLine("#pragma once");
 					WriteDefinitions(CompileEnvironment.Definitions, Writer);
 					CompileEnvironment.Definitions.Clear();
 
@@ -1442,7 +1442,6 @@ namespace UnrealBuildTool
 		{
 			// Build the contents of the wrapper file
 			StringBuilder WrapperContents = new StringBuilder();
-			using (StringWriter Writer = new StringWriter(WrapperContents))
 			{
 				string IncludeFileString = IncludedFile.AbsolutePath;
 				if (IncludedFile.Location.IsUnderDirectory(Unreal.RootDirectory))
@@ -1450,9 +1449,9 @@ namespace UnrealBuildTool
 					IncludeFileString = IncludedFile.Location.MakeRelativeTo(Unreal.EngineSourceDirectory);
 				}
 				IncludeFileString = IncludeFileString.Replace('\\', '/');
-				Writer.WriteLine("// PCH for {0}", IncludeFileString);
-				WriteDefinitions(Definitions, Writer);
-				Writer.WriteLine("#include \"{0}\"", IncludeFileString);
+				WrapperContents.AppendLine("// PCH for {0}", IncludeFileString);
+				WriteDefinitions(Definitions, WrapperContents);
+				WrapperContents.AppendLine("#include \"{0}\"", IncludeFileString);
 			}
 
 			// Create the item
@@ -1472,19 +1471,19 @@ namespace UnrealBuildTool
 		/// </summary>
 		/// <param name="Definitions">List of definitions</param>
 		/// <param name="Writer">Writer to receive output</param>
-		static void WriteDefinitions(IEnumerable<string> Definitions, TextWriter Writer)
+		static void WriteDefinitions(IEnumerable<string> Definitions, StringBuilder Writer)
 		{
-			Writer.WriteLine("// Generated by UnrealBuildTool (UEBuildModuleCPP.cs)");
+			Writer.AppendLine("// Generated by UnrealBuildTool (UEBuildModuleCPP.cs)");
 			foreach(string Definition in Definitions)
 			{
 				int EqualsIdx = Definition.IndexOf('=');
 				if(EqualsIdx == -1)
 				{
-					Writer.WriteLine("#define {0} 1", Definition);
+					Writer.Append("#define ").Append(Definition).AppendLine(" 1");
 				}
 				else
 				{
-					Writer.WriteLine("#define {0} {1}", Definition.Substring(0, EqualsIdx), Definition.Substring(EqualsIdx + 1));
+					Writer.Append("#define ").Append(Definition.AsSpan(0, EqualsIdx)).Append(' ').Append(Definition.AsSpan(EqualsIdx + 1)).AppendLine();
 				}
 			}
 		}
diff --git a/Engine/Source/Programs/UnrealBuildTool/System/CppIncludeParser.cs b/Engine/Source/Programs/UnrealBuildTool/System/CppIncludeParser.cs
--- a/Engine/Source/Programs/UnrealBuildTool/System/CppIncludeParser.cs
+++ b/Engine/Source/Programs/UnrealBuildTool/System/CppIncludeParser.cs
@@ -19,7 +19,7 @@ namespace UnrealBuildTool
 		/// </summary>
 		/// <param name="Reader">Stream to read from</param>
 		/// <param name="Writer">Stream to write directives to</param>
-		public static void CopyIncludeDirectives(StreamReader Reader, StringWriter Writer)
+		public static void CopyIncludeDirectives(StreamReader Reader, StringBuilder Writer)
 		{
 			StringBuilder Token = new StringBuilder();
 			while(TryReadToken(Reader, Token))
@@ -69,7 +69,7 @@ namespace UnrealBuildTool
 							break;
 						}
 
-						Writer.WriteLine("#include {0}", IncludeFile);
+						Writer.Append("#include ").AppendLine(IncludeFile);
 					}
 					else
 					{
