[UBT]
* Fixed mac linker error where same link input files were added multiple times. Removed confusing comment saying this could only happen in monolithic builds. Also changed the code to be fronted by a hashset to reduce performance cost (which was the reason the change was made first place)

This fixes Mac linker dup errors

diff --git a/Engine/Source/Programs/UnrealBuildTool/Configuration/UEBuildBinary.cs b/Engine/Source/Programs/UnrealBuildTool/Configuration/UEBuildBinary.cs
--- a/Engine/Source/Programs/UnrealBuildTool/Configuration/UEBuildBinary.cs
+++ b/Engine/Source/Programs/UnrealBuildTool/Configuration/UEBuildBinary.cs
@@ -712,35 +712,30 @@ namespace UnrealBuildTool
 				BinaryCompileEnvironment.bUseSharedBuildEnvironment = false;
 			}
 
+			HashSet<FileItem> InputFilesLookup = new();
+
 			foreach (UEBuildModule Module in Modules)
 			{
-				List<FileItem> LinkInputFiles;
 				if (Module.Binary == null || Module.Binary == this)
 				{
 					// Compile each module.
 					Logger.LogDebug("Compile module: {ModuleName}", Module.Name);
-					LinkInputFiles = Module.Compile(Target, ToolChain, BinaryCompileEnvironment, SpecificFilesToCompile, WorkingSet, Graph, Logger);
+					List<FileItem> LinkInputFiles = Module.Compile(Target, ToolChain, BinaryCompileEnvironment, SpecificFilesToCompile, WorkingSet, Graph, Logger);
 
 					// Save the module outputs. In monolithic builds, this is just the object files.
 					if (Target.LinkType == TargetLinkType.Monolithic)
 					{
 						Graph.SetOutputItemsForModule(Module.Name, LinkInputFiles.ToArray());
+					}
 
-						// NOTE: Because of 'Shared PCHs', in monolithic builds the same PCH file may appear as a link input
-						// multiple times for a single binary.  We'll check for that here, and only add it once.  This avoids
-						// a linker warning about redundant .obj files. 
-						foreach (FileItem LinkInputFile in LinkInputFiles)
+					// Sometimes the same link input file is added multiple times, so make sure to filter out the dups
+					foreach (FileItem LinkInputFile in LinkInputFiles)
+					{
+						if (InputFilesLookup.Add(LinkInputFile))
 						{
-							if (!BinaryLinkEnvironment.InputFiles.Contains(LinkInputFile))
-							{
-								BinaryLinkEnvironment.InputFiles.Add(LinkInputFile);
-							}
+							BinaryLinkEnvironment.InputFiles.Add(LinkInputFile);
 						}
 					}
-					else
-					{
-						BinaryLinkEnvironment.InputFiles.AddRange(LinkInputFiles);
-					}
 
 					// Force a reference to initialize module for this binary
 					if(Module.Rules.bRequiresImplementModule ?? true)
