Add support for overriding an individual module's optimization level via a new  OptimizationLevel property. This requires either a private PCH, or PCHUsageMode=NoPCHs (not recommended) because Clang requires that a PCH's optimization flags match the compile unit.

diff --git a/Engine/Source/Programs/UnrealBuildTool/Configuration/ModuleRules.cs b/Engine/Source/Programs/UnrealBuildTool/Configuration/ModuleRules.cs
--- a/Engine/Source/Programs/UnrealBuildTool/Configuration/ModuleRules.cs
+++ b/Engine/Source/Programs/UnrealBuildTool/Configuration/ModuleRules.cs
@@ -633,6 +633,27 @@ namespace UnrealBuildTool
 			set { OptimizeCodeOverride = value; }
 		}
 
+		private OptimizationMode? OptimizationLevelOverride = null;
+
+		/// <summary>
+		/// Allows fine tuning optimization level for speed and\or code size. This requires a private PCH (or NoPCHs, which is not recommended)
+		/// </summary>
+		public OptimizationMode OptimizationLevel
+		{
+			get
+			{
+				if (OptimizationLevelOverride.HasValue)
+				{
+					return OptimizationLevelOverride.Value;
+				}
+				else
+				{
+					return Target.OptimizationLevel;
+				}
+			}
+			set { OptimizationLevelOverride = value; }
+		}
+
 		/// <summary>
 		/// Explicit private PCH for this module. Implies that this module will not use a shared PCH.
 		/// </summary>
diff --git a/Engine/Source/Programs/UnrealBuildTool/Configuration/UEBuildModuleCPP.cs b/Engine/Source/Programs/UnrealBuildTool/Configuration/UEBuildModuleCPP.cs
--- a/Engine/Source/Programs/UnrealBuildTool/Configuration/UEBuildModuleCPP.cs
+++ b/Engine/Source/Programs/UnrealBuildTool/Configuration/UEBuildModuleCPP.cs
@@ -347,7 +347,7 @@ namespace UnrealBuildTool
 
 			List<FileItem> LinkInputFiles = base.Compile(Target, ToolChain, BinaryCompileEnvironment, SpecificFilesToCompile, WorkingSet, Graph, Logger);
 
-			CppCompileEnvironment ModuleCompileEnvironment = CreateModuleCompileEnvironment(Target, BinaryCompileEnvironment);
+			CppCompileEnvironment ModuleCompileEnvironment = CreateModuleCompileEnvironment(Target, BinaryCompileEnvironment, Logger);
 
 			// If the module is precompiled, read the object files from the manifest
 			if(Rules.bUsePrecompiled && Target.LinkType == TargetLinkType.Monolithic)
@@ -1573,7 +1573,7 @@ namespace UnrealBuildTool
 
 		public CppCompileEnvironment CreateCompileEnvironmentForIntellisense(ReadOnlyTargetRules Target, CppCompileEnvironment BaseCompileEnvironment, ILogger Logger)
 		{
-			CppCompileEnvironment CompileEnvironment = CreateModuleCompileEnvironment(Target, BaseCompileEnvironment);
+			CppCompileEnvironment CompileEnvironment = CreateModuleCompileEnvironment(Target, BaseCompileEnvironment, Logger);
 			CompileEnvironment = SetupPrecompiledHeaders(Target, null, CompileEnvironment, new List<FileItem>(), new NullActionGraphBuilder(Logger));
 			CreateHeaderForDefinitions(CompileEnvironment, IntermediateDirectory, null, new NullActionGraphBuilder(Logger));
 			return CompileEnvironment;
@@ -1584,8 +1584,9 @@ namespace UnrealBuildTool
 		/// </summary>
 		/// <param name="Target">Rules for the target being built</param>
 		/// <param name="BaseCompileEnvironment">An existing environment to base the module compile environment on.</param>
+		/// <param name="Logger">Logger for output</param>
 		/// <returns>The new module compile environment.</returns>
-		public CppCompileEnvironment CreateModuleCompileEnvironment(ReadOnlyTargetRules Target, CppCompileEnvironment BaseCompileEnvironment)
+		public CppCompileEnvironment CreateModuleCompileEnvironment(ReadOnlyTargetRules Target, CppCompileEnvironment BaseCompileEnvironment, ILogger Logger)
 		{
 			CppCompileEnvironment Result = new CppCompileEnvironment(BaseCompileEnvironment);
 
@@ -1613,6 +1614,17 @@ namespace UnrealBuildTool
 			Result.IncludeOrderVersion = Rules.IncludeOrderVersion;
 			Result.bDeterministic |= Rules.bDeterministic;
 
+			if (Result.OptimizationLevel != Rules.OptimizationLevel)
+			{
+				if ( Rules.PCHUsage != ModuleRules.PCHUsageMode.NoPCHs && Rules.PrivatePCHHeaderFile == null )
+				{
+					throw new BuildException("Module {0} - Overriding OptimizationLevel requires a private PCH", Name );
+				}
+				Logger.LogInformation("Optimization level changed for module {0} due to override. Old: {1} New: {2}", Name, Result.OptimizationLevel, Rules.OptimizationLevel);
+				Result.OptimizationLevel = Rules.OptimizationLevel;
+			}
+
+
 			// If the module overrides the C++ language version, override it on the compile environment
 			if (Rules.CppStandard != CppStandardVersion.Default)
 			{
diff --git a/Engine/Source/Programs/UnrealBuildTool/Modes/ServerMode.cs b/Engine/Source/Programs/UnrealBuildTool/Modes/ServerMode.cs
--- a/Engine/Source/Programs/UnrealBuildTool/Modes/ServerMode.cs
+++ b/Engine/Source/Programs/UnrealBuildTool/Modes/ServerMode.cs
@@ -397,7 +397,7 @@ namespace UnrealBuildTool
 							continue;
 						}
 
-						CppCompileEnvironment ModuleCompileEnvironment = Module.CreateModuleCompileEnvironment(CurrentTarget.Rules, BinaryCompileEnvironment);
+						CppCompileEnvironment ModuleCompileEnvironment = Module.CreateModuleCompileEnvironment(CurrentTarget.Rules, BinaryCompileEnvironment, Logger);
 						foreach (DirectoryReference Dir in Module.ModuleDirectories)
 						{
 							BrowseConfigurationFolders.Add(Dir.ToString());
