[UBT]
* IWYUMode - Added so skipped files are manually parsed for headers. This code is pessimistic and doesnt include any includes that are inside #if/#endif

diff --git a/Engine/Source/Programs/UnrealBuildTool/Modes/IWYUMode.cs b/Engine/Source/Programs/UnrealBuildTool/Modes/IWYUMode.cs
--- a/Engine/Source/Programs/UnrealBuildTool/Modes/IWYUMode.cs
+++ b/Engine/Source/Programs/UnrealBuildTool/Modes/IWYUMode.cs
@@ -12,6 +12,7 @@ using System.Threading.Tasks;
 using System.Text;
 using System.Diagnostics.CodeAnalysis;
 using System.Threading;
+using System.Collections.Concurrent;
 
 
 // IWYUMode is a mode that can be used to clean up includes in source code. It uses the clang based tool include-what-you-use (IWYU) to figure out what is needed in each .h/.cpp file
@@ -400,7 +401,8 @@ namespace UnrealBuildTool
 
 			// Calculate file paths filter to figure out which files that IWYU will run on
 			HashSet<string> ValidPaths = new();
-			Dictionary<string, UEBuildModule?> PathToModule = new();
+			Dictionary<string, UEBuildModule> PathToModule = new();
+			Dictionary<string, UEBuildModule> NameToModule = new();
 			{
 				// Create target to be able to traverse all existing modules
 				Logger.LogInformation($"Creating BuildTarget for {TargetName}...");
@@ -431,99 +433,119 @@ namespace UnrealBuildTool
 				int ModulesSkippedCount = 0;
 
 				// Traverse all modules to figure out which ones should be updated. This is based on -ModuleToUpdate and -PathToUpdate
+				List<UEBuildModule> ReferencedModules = new();
+				HashSet<UEBuildModule> IgnoreReferencedModules = new();
 				foreach (var Binary in Target.Binaries)
 				{
 					foreach (var Module in Binary.Modules)
 					{
-						if (Module.Rules.Type != ModuleRules.ModuleType.CPlusPlus)
+						if (IgnoreReferencedModules.Add(Module))
 						{
-							continue;
+							ReferencedModules.Add(Module);
+							Module.GetAllDependencyModules(ReferencedModules, IgnoreReferencedModules, true, false, false);
 						}
+					}
+				}
+				foreach (UEBuildModule Module in ReferencedModules)
+				{
+					NameToModule.TryAdd(Module.Name, Module);
+					foreach (var Module2 in Module.GetDependencies(true, true))
+					{
+						NameToModule.TryAdd(Module2.Name, Module2);
+					}
+				}
 
-						bool ShouldUpdate = false;
-						string ModuleDir = Module.ModuleDirectory.FullName.Replace('\\', '/');
+				foreach (var Module in NameToModule.Values)
+				{
+					bool ShouldUpdate = false;
+					string ModuleDir = Module.ModuleDirectory.FullName.Replace('\\', '/');
 
-						if (!PathToModule.TryAdd(ModuleDir, Module))
-						{
-							if (PathToModule[ModuleDir] != Module)
-							{
-								Logger.LogError("FOOO");
-								return -1;
-							}
-							continue;
-						}
-						foreach (DirectoryReference AdditionalDir in Module.ModuleDirectories)
+					if (!PathToModule.TryAdd(ModuleDir, Module))
+					{
+						continue;
+					}
+					foreach (DirectoryReference AdditionalDir in Module.ModuleDirectories)
+					{
+						if (AdditionalDir != Module.ModuleDirectory)
 						{
-							if (AdditionalDir != Module.ModuleDirectory)
-							{
-								PathToModule.TryAdd(AdditionalDir.FullName.Replace('\\', '/'), Module);
-							}
+							PathToModule.TryAdd(AdditionalDir.FullName.Replace('\\', '/'), Module);
 						}
+					}
+
+					if (Module.Rules.Type != ModuleRules.ModuleType.CPlusPlus)
+					{
+						continue;
+					}
 
-						if (UEModuleToUpdate != null)
+					if (UEModuleToUpdate != null)
+					{
+						bool DependsOnModule = Module == UEModuleToUpdate;
+						if (bUpdateDependents)
 						{
-							bool DependsOnModule = Module == UEModuleToUpdate;
-							if (bUpdateDependents)
+							if (Module.PublicDependencyModules != null)
 							{
-								foreach (var Dependency in Module.PublicDependencyModules!)
+								foreach (var Dependency in Module.PublicDependencyModules)
 								{
 									DependsOnModule = DependsOnModule || Dependency == UEModuleToUpdate;
 								}
+							}
+							if (Module.PrivateDependencyModules != null)
+							{
 								foreach (var Dependency in Module.PrivateDependencyModules!)
 								{
 									DependsOnModule = DependsOnModule || Dependency == UEModuleToUpdate;
 								}
 							}
-							ShouldUpdate = DependsOnModule;
 						}
+						ShouldUpdate = DependsOnModule;
+					}
 
-						if (PathToUpdateList.Length > 0)
+					if (PathToUpdateList.Length > 0)
+					{
+						bool Match = false;
+						for (int I = 0; I != PathToUpdateList.Length; ++I)
 						{
-							bool Match = false;
-							for (int I = 0; I != PathToUpdateList.Length; ++I)
-							{
-								Match = Match || ModuleDir.Contains(PathToUpdateList[I], StringComparison.OrdinalIgnoreCase);
-							}
+							Match = Match || ModuleDir.Contains(PathToUpdateList[I], StringComparison.OrdinalIgnoreCase);
+						}
 
-							if (Match)
-							{
-								if (UEModuleToUpdate == null)
-								{
-									ShouldUpdate = true;
-								}
-							}
-							else
+						if (Match)
+						{
+							if (UEModuleToUpdate == null)
 							{
-								ShouldUpdate = false;
+								ShouldUpdate = true;
 							}
 						}
-						else if (UEModuleToUpdate == null)
+						else
 						{
-							ShouldUpdate = true;
+							ShouldUpdate = false;
 						}
+					}
+					else if (UEModuleToUpdate == null)
+					{
+						ShouldUpdate = true;
+					}
 
-						if (!ShouldUpdate)
-						{
-							continue;
-						}
+					if (!ShouldUpdate)
+					{
+						continue;
+					}
 
-						if (Module.Rules.IWYUSupport == IWYUSupport.None)
-						{
-							++ModulesSkippedCount;
-							continue;
-						}
+					if (Module.Rules.IWYUSupport == IWYUSupport.None)
+					{
+						++ModulesSkippedCount;
+						continue;
+					}
 
-						++ModulesToUpdateCount;
+					++ModulesToUpdateCount;
 
-						// When adding to ValidPaths, make sure the Module directory ends in a / so we match the exact folder later
-						ValidPaths.Add(AdjustModulePathForMatching(ModuleDir));
+					// When adding to ValidPaths, make sure the Module directory ends in a / so we match the exact folder later
+					ValidPaths.Add(AdjustModulePathForMatching(ModuleDir));
 
-						foreach (DirectoryReference AdditionalDir in Module.ModuleDirectories)
+					foreach (DirectoryReference AdditionalDir in Module.ModuleDirectories)
+					{
+						if (AdditionalDir != Module.ModuleDirectory)
 						{
-							if (AdditionalDir != Module.ModuleDirectory)
-							{
-								ValidPaths.Add(AdjustModulePathForMatching(AdditionalDir.FullName));
-							}
+							ValidPaths.Add(AdjustModulePathForMatching(AdditionalDir.FullName));
 						}
 					}
 				}
@@ -547,6 +569,7 @@ namespace UnrealBuildTool
 			List<IWYUInfo> GeneratedCppInfos = new();
 			IWYUInfo? ObjectMacrosInfo = null;
 			IWYUInfo? ScriptMacrosInfo = null;
+			IWYUInfo? VerseInteropUtilsInfo = null;
 			IWYUInfo? PlatformInfo = null;
 
 			using (ISourceFileWorkingSet WorkingSet = new EmptySourceFileWorkingSet())
@@ -695,16 +718,7 @@ namespace UnrealBuildTool
 								}
 							}
 
-							string FilePath = Info.File;
-							while (true)
-							{
-								int LastIndexOfSlash = FilePath.LastIndexOf('/');
-								if (LastIndexOfSlash == -1)
-									break;
-								FilePath = FilePath.Substring(0, LastIndexOfSlash);
-								if (PathToModule.TryGetValue(FilePath, out Info.Module))
-									break;
-							}
+							Info.Module = GetModule(Info.File, PathToModule);
 
 							// Some special logic for headers.
 							if (Info.File.EndsWith(".h"))
@@ -714,18 +728,25 @@ namespace UnrealBuildTool
 									ObjectMacrosInfo = Info;
 								else if (ScriptMacrosInfo == null && Info.File.EndsWith("UObject/ScriptMacros.h"))
 									ScriptMacrosInfo = Info;
+								else if (VerseInteropUtilsInfo == null && Info.File.EndsWith("VerseInteropUtils.h"))
+									VerseInteropUtilsInfo = Info;
 								else if (PlatformInfo == null && Info.File.EndsWith("HAL/Platform.h"))
 									PlatformInfo = Info;
 
 								// We need to add entries for .generated.h. They are not in the iwyu files
 								foreach (var Include in Info.Includes)
 								{
-									if (Include.Printable.EndsWith(".generated.h"))
+									if (Include.Full.Contains("/UHT/"))
 									{
 										lock (GeneratedHeaderInfos)
 											GeneratedHeaderInfos.Add(Include.Full);
 										break;
 									}
+									else if (Include.Full.Contains("/VNI/"))
+									{
+										lock (GeneratedHeaderInfos)
+											GeneratedHeaderInfos.Add(Include.Full);
+									}
 								}
 							}
 						}
@@ -760,19 +781,33 @@ namespace UnrealBuildTool
 			{
 				IWYUIncludeEntry ObjectMacrosInclude = new();
 				IWYUIncludeEntry ScriptMacrosInclude = new();
+				IWYUIncludeEntry VerseInteropUtilsInclude = new();
 				ObjectMacrosInclude.Printable = "UObject/ObjectMacros.h";
 				ObjectMacrosInclude.Full = ObjectMacrosInfo!.File;
 				ObjectMacrosInclude.Resolved = ObjectMacrosInfo;
 				ScriptMacrosInclude.Printable = "UObject/ScriptMacros.h";
 				ScriptMacrosInclude.Full = ScriptMacrosInfo!.File;
 				ScriptMacrosInclude.Resolved = ScriptMacrosInfo;
+				VerseInteropUtilsInclude.Printable = "VerseInteropUtils.h";
+				VerseInteropUtilsInclude.Full = VerseInteropUtilsInfo!.File;
+				VerseInteropUtilsInclude.Resolved = VerseInteropUtilsInfo;
 
 				foreach (var Gen in GeneratedHeaderInfos)
 				{
 					IWYUInfo GenInfo = new();
 					GenInfo.File = Gen;
-					GenInfo.Includes.Add(ObjectMacrosInclude);
-					GenInfo.Includes.Add(ScriptMacrosInclude);
+					if (Gen.EndsWith(".generated.h"))
+					{
+						GenInfo.IncludesSeenInFile.Add(ObjectMacrosInclude);
+						GenInfo.IncludesSeenInFile.Add(ScriptMacrosInclude);
+						GenInfo.Includes.Add(ObjectMacrosInclude);
+						GenInfo.Includes.Add(ScriptMacrosInclude);
+					}
+					else
+					{
+						GenInfo.IncludesSeenInFile.Add(VerseInteropUtilsInclude);
+						GenInfo.Includes.Add(VerseInteropUtilsInclude);
+					}
 					Infos.Add(Gen, GenInfo);
 				}
 			}
@@ -786,6 +821,7 @@ namespace UnrealBuildTool
 			string SystemPath = BaseLinuxPath!.FullName.Replace('\\', '/');
 
 			HashSet<IWYUInfo> IncludedInsideDecl = new();
+			Dictionary<string, IWYUInfo> SkippedHeaders = new();
 
 			Parallel.ForEach(Infos.Values, Info =>
 			{
@@ -851,11 +887,31 @@ namespace UnrealBuildTool
 				{
 					if (Include.Resolved == null)
 					{
-						Infos.TryGetValue(Include.Full, out Include.Resolved);
+						if (!Infos.TryGetValue(Include.Full, out Include.Resolved))
+						{
+							if (!Include.Full.Contains(".gen.cpp") && !Include.Full.Contains("/ThirdParty/") && !Include.Full.Contains("/AutoSDK/") && !Include.Full.Contains("/VNI/"))
+							{
+								lock (SkippedHeaders)
+								{
+									if (!SkippedHeaders.TryGetValue(Include.Full, out Include.Resolved))
+									{
+										IWYUInfo NewInfo = new();
+										Include.Resolved = NewInfo;
+										NewInfo.File = Include.Full;
+										SkippedHeaders.Add(Include.Full, NewInfo);
+									}
+								}
+							}
+						}
 					}
 				}
 			});
 
+			int InfosCount = Infos.Count;
+			Logger.LogInformation($"Parsing included headers not supporting being compiled...");
+			ParseSkippedHeaders(SkippedHeaders, Infos, PathToModule, NameToModule);
+			Logger.LogInformation($"Added {Infos.Count - InfosCount} more read-only IWYU entries...");
+
 			foreach (IWYUInfo Info in IncludedInsideDecl)
 			{
 				if (Info.IncludesSeenInFile.Count != 0 && !Info.File.EndsWith("ScriptSerialization.h")) // Remove include in ScriptSerialization.h and remove this check
@@ -951,7 +1007,7 @@ namespace UnrealBuildTool
 
 			if (bCompare)
 			{
-				return CompareFiles(Infos, ValidPaths, Logger);
+				return CompareFiles(Infos, ValidPaths, PathToModule, NameToModule, Logger);
 			}
 			else
 			{
@@ -959,7 +1015,162 @@ namespace UnrealBuildTool
 			}
 		}
 
-		bool IsValidForUpdate(IWYUInfo Info, HashSet<string> ValidPaths, bool ObeyModuleRules)
+		static UEBuildModule? GetModule(string File, Dictionary<string, UEBuildModule> PathToModule)
+		{
+			string FilePath = File;
+			while (true)
+			{
+				int LastIndexOfSlash = FilePath.LastIndexOf('/');
+				if (LastIndexOfSlash == -1)
+					break;
+				FilePath = FilePath.Substring(0, LastIndexOfSlash);
+				UEBuildModule? Module;
+				if (PathToModule.TryGetValue(FilePath, out Module))
+					return Module;
+			}
+			return null;
+		}
+
+		static string? GetFullName(string Include, string From, UEBuildModule? Module, Dictionary<string, UEBuildModule> NameToModule, HashSet<UEBuildModule>? Visited = null)
+		{
+			if (Module == null)
+			{
+				return null;
+			}
+			foreach (var Dir in Module.PublicIncludePaths.Union(Module.PrivateIncludePaths).Union(Module.PublicSystemIncludePaths))
+			{
+				FileReference FileReference = FileReference.Combine(Dir, Include);
+				FileItem FileItem = FileItem.GetItemByFileReference(FileReference);
+				if (FileItem.Exists)
+					return FileItem.Location.FullName.Replace('\\', '/');
+			}
+
+			bool Nested = Visited != null;
+
+			if (!Nested)
+			{
+				FileReference FileReference = FileReference.Combine(FileReference.FromString(From).Directory, Include);
+				FileItem FileItem = FileItem.GetItemByFileReference(FileReference);
+				if (FileItem.Exists)
+					return FileItem.Location.FullName.Replace('\\', '/');
+			}
+
+			foreach (var PublicModule in Module.Rules.PublicDependencyModuleNames.Union(Module.Rules.PublicIncludePathModuleNames).Union(Module.Rules.PrivateIncludePathModuleNames).Union(Module.Rules.PrivateDependencyModuleNames))
+			{
+				UEBuildModule? DependencyModule;
+				if (NameToModule.TryGetValue(PublicModule, out DependencyModule))
+				{
+					if (Visited == null)
+					{
+						Visited = new();
+						Visited.Add(Module);
+						Visited.Add(DependencyModule);
+					}
+					else if (!Visited.Add(DependencyModule))
+						continue;
+
+					string? Str = GetFullName(Include, From, DependencyModule, NameToModule, Visited);
+					if (Str != null)
+						return Str;
+				}
+			}
+
+			if (!Nested)
+			{
+				// This should only show system includes..
+				//Console.WriteLine($"Can't resolve {Include} from module {Module.Name}");
+			}
+
+			return null;
+		}
+
+		static void ParseSkippedHeaders(Dictionary<string, IWYUInfo> SkippedHeaders, Dictionary<string, IWYUInfo> Infos, Dictionary<string, UEBuildModule> PathToModule, Dictionary<string, UEBuildModule> NameToModule)
+		{
+			foreach (var Kvp in SkippedHeaders)
+			{
+				Infos.Add(Kvp.Key, Kvp.Value);
+			}
+			while (true)
+			{
+				Dictionary<string, IWYUInfo> NewSkippedHeaders = new();
+				Parallel.ForEach(SkippedHeaders.Values, Info =>
+				{
+					Info.Module = GetModule(Info.File, PathToModule);
+
+					bool HasIncludeGuard = false;
+					int IfCount = 0;
+					string[] Lines = File.ReadAllLines(Info.File);
+					foreach (string Line in Lines)
+					{
+						ReadOnlySpan<char> LineSpan = Line.AsSpan().TrimStart();
+						if (Line.StartsWith("#if"))
+						{
+							// Include guards are special
+							if (!HasIncludeGuard)
+							{
+								if (Line.StartsWith("#ifndef ") && Line.EndsWith("_H"))
+								{
+									HasIncludeGuard = true;
+									continue;
+								}
+							}
+							++IfCount;
+						}
+						else if (Line.StartsWith("#endif"))
+						{
+							--IfCount;
+						}
+						else if (Line.StartsWith("#include"))
+						{
+							var IncludeSpan = LineSpan.Slice("#include".Length).TrimStart();
+							char LeadingIncludeChar = IncludeSpan[0];
+							if (LeadingIncludeChar == '"' || LeadingIncludeChar == '<')
+							{
+								var Start = IncludeSpan.Slice(1);
+								var EndIndex = Start.IndexOf(LeadingIncludeChar == '"' ? '"' : '>');
+								var FileSpan = Start.Slice(0, EndIndex);
+								var Include = FileSpan.ToString();
+								string? File = GetFullName(Include, Info.File, Info.Module, NameToModule);
+								if (File != null)
+								{
+									IWYUInfo? IncludeInfo;
+									lock (Infos)
+									{
+										if (!Infos.TryGetValue(File, out IncludeInfo))
+										{
+											if (!SkippedHeaders.TryGetValue(File, out IncludeInfo))
+											{
+												IncludeInfo = new();
+												IncludeInfo.File = File;
+												NewSkippedHeaders.Add(File, IncludeInfo);
+												Infos.Add(File, IncludeInfo);
+											}
+										}
+									}
+									if (IncludeInfo != null)
+									{
+										IWYUIncludeEntry Entry = new();
+										Entry.Full = File;
+										Entry.Printable = Include;
+										Entry.System = LeadingIncludeChar == '<';
+										Entry.Resolved = IncludeInfo;
+										Info.IncludesSeenInFile.Add(Entry);
+										Info.Includes.Add(Entry);
+									}
+								}
+							}
+						}
+					}
+				});
+				if (NewSkippedHeaders.Count == 0)
+				{
+					return;
+				}
+				SkippedHeaders = NewSkippedHeaders;
+			}
+		}
+
+		static bool IsValidForUpdate(IWYUInfo Info, HashSet<string> ValidPaths, bool ObeyModuleRules)
 		{
 			if (Info.Source == null) // .generated.h is also in this list, ignore them
 			{
@@ -1641,7 +1852,7 @@ namespace UnrealBuildTool
 		/// <summary>
 		/// Calculate all indirect transitive includes for a file. This list contains does not contain itself and will handle circular dependencies
 		/// </summary>
-		void CalculateTransitive(IWYUInfo Root, IWYUInfo Info, Stack<string> Stack, Dictionary<string, string> TransitiveIncludes, Dictionary<string, string> TransitiveForwardDeclarations, bool UseSeenIncludes)
+		static void CalculateTransitive(IWYUInfo Root, IWYUInfo Info, Stack<string> Stack, Dictionary<string, string> TransitiveIncludes, Dictionary<string, string> TransitiveForwardDeclarations, bool UseSeenIncludes)
 		{
 			List<IWYUIncludeEntry> Includes = UseSeenIncludes ? Info.IncludesSeenInFile : Info.Includes;
 			foreach (var Include in Includes)
@@ -1670,12 +1881,13 @@ namespace UnrealBuildTool
 			}
 		}
 
-		int CompareFiles(Dictionary<string, IWYUInfo> Infos, HashSet<string> ValidPaths, ILogger Logger)
+		int CompareFiles(Dictionary<string, IWYUInfo> Infos, HashSet<string> ValidPaths, Dictionary<string, UEBuildModule> PathToModule, Dictionary<string, UEBuildModule> NameToModule, ILogger Logger)
 		{
 			Logger.LogInformation($"Comparing code files...");
 
-			SortedSet<string> UnresolvableFiles = new();
+			Dictionary<string, IWYUInfo> SkippedHeaders = new();
 
+			Logger.LogInformation($"Parsing seen headers not supporting being compiled...");
 			Parallel.ForEach(Infos.Values, Info =>
 			{
 				foreach (var Include in Info.IncludesSeenInFile)
@@ -1684,41 +1896,115 @@ namespace UnrealBuildTool
 					{
 						if (!Infos.TryGetValue(Include.Full, out Include.Resolved))
 						{
-							if (!Include.Full.Contains(".gen.cpp") && !Include.Full.Contains("/ThirdParty/") && !Include.Full.Contains("/AutoSDK/") && !Include.Full.Contains("/VNI/"))
+							lock (SkippedHeaders)
 							{
-								lock (UnresolvableFiles)
-									UnresolvableFiles.Add(Include.Full);
+								if (!SkippedHeaders.TryGetValue(Include.Full, out Include.Resolved))
+								{
+									IWYUInfo NewInfo = new();
+									Include.Resolved = NewInfo;
+									NewInfo.File = Include.Full;
+									SkippedHeaders.Add(Include.Full, NewInfo);
+								}
 							}
 						}
 					}
 				}
 			});
+			ParseSkippedHeaders(SkippedHeaders, Infos, PathToModule, NameToModule);
 
-			//foreach (string Unresolved in UnresolvableFiles)
-			//	Console.WriteLine("Failed to resolve " + Unresolved);
+			//List<ValueTuple<float, IWYUInfo>> RatioList = new();
 
-			List<ValueTuple<float, IWYUInfo>> RatioList = new();
+			ConcurrentDictionary<string, long> AllSeenIncludes = new();
 
+			List<ValueTuple<IWYUInfo, Dictionary<string, string> >> InfosToCompare = new();
+
+			Logger.LogInformation($"Generating transitive seen include lists...");
 			Parallel.ForEach(Infos.Values, Info =>
 			{
-				if (!IsValidForUpdate(Info, ValidPaths, false))
+				if (!IsValidForUpdate(Info, ValidPaths, false) || Info.File.EndsWith(".h"))
 					return;
 
 				Dictionary<string, string> SeenTransitiveIncludes = new();
 				Stack<string> Stack = new();
 				CalculateTransitive(Info, Info, Stack, SeenTransitiveIncludes, new Dictionary<string, string>(), true);
 
-				// Sometimes seen are lower and that is because there are a bunch of files that seen includes but is not visible (UnrealNames.inl, UnrealPlatformMathSSE4.h etc
-				// So actually Seen should be 3-4 higher than what it says (probably more)
+				lock (InfosToCompare)
+					InfosToCompare.Add(new(Info, SeenTransitiveIncludes));
+
+				foreach (var Include in SeenTransitiveIncludes.Union(Info.TransitiveIncludes))
+				{
+					AllSeenIncludes.TryAdd(Include.Key, 0);
+				}
+
+				/*
 				if (Info.TransitiveIncludes.Count < SeenTransitiveIncludes.Count)
 				{
 					float Ratio = (float)Info.TransitiveIncludes.Count / SeenTransitiveIncludes.Count;
 					lock (RatioList)
 						RatioList.Add(new(Ratio, Info));
 				}
-				//Console.WriteLine($"{Info.File} - Now: {SeenTransitiveIncludes.Count}  Optimized: {Info.TransitiveIncludes.Count}");
+				if (SeenTransitiveIncludes.Count < Info.TransitiveIncludes.Count)
+					Console.WriteLine($"{Info.File} - Now: {SeenTransitiveIncludes.Count}  Optimized: {Info.TransitiveIncludes.Count}");
+				*/
+			});
+
+			var FileToSize = new Dictionary<string, long>(AllSeenIncludes);
+
+			Logger.LogInformation($"Reading file sizes...");
+			Parallel.ForEach(AllSeenIncludes, Info =>
+			{
+				FileToSize[Info.Key] = (new FileInfo(Info.Key)).Length;
+			});
+
+			Logger.LogInformation($"Calculate total amount of bytes included per file...");
+			List<ValueTuple<long, IWYUInfo>> ByteSizeDiffList = new();
+			Parallel.ForEach(InfosToCompare, Kvp =>
+			{
+				long OptimizedSize = 0;
+				long SeenSize = 0;
+				foreach (var Include in Kvp.Item1.TransitiveIncludes.Keys)
+					OptimizedSize += FileToSize[Include];
+				foreach (var Include in Kvp.Item2.Keys)
+					SeenSize += FileToSize[Include];
+				long Diff = SeenSize - OptimizedSize;
+				lock (ByteSizeDiffList)
+					ByteSizeDiffList.Add(new(Diff, Kvp.Item1));
+			});
+
+			ByteSizeDiffList.Sort((a, b) =>
+			{
+				if (a.Item1 == b.Item1)
+					return 0;
+				if (a.Item1 < b.Item1)
+					return 1;
+				return -1;
 			});
 
+			Func<long, string> PrettySize = (Size) =>
+			{
+				if (Size > 1024 * 1024)
+					return (((float)Size) / (1024 * 1024)).ToString("0.00") + "mb";
+				if (Size > 1024)
+					return (((float)Size) / (1024)).ToString("0.00") + "kb";
+				return Size + "b";
+			};
+
+			Logger.LogInformation("");
+			Logger.LogInformation("Top 20 most reduction in bytes parsed by compiler frontend");
+			Logger.LogInformation("----------------------------------------------------------");
+			for (int I = 0; I != Math.Min(20, ByteSizeDiffList.Count); ++I)
+			{
+				long Saved = ByteSizeDiffList[I].Item1;
+				IWYUInfo File = ByteSizeDiffList[I].Item2;
+				long OptimizedSize = 0;
+				foreach (var Include in File.TransitiveIncludes.Keys)
+					OptimizedSize += FileToSize[Include];
+				float Percent = 100.0f - (((float)OptimizedSize)/(OptimizedSize + Saved) * 100.0f);
+				Logger.LogInformation($"{Path.GetFileName(File.File)}   {PrettySize(OptimizedSize + Saved)} -> {PrettySize(OptimizedSize)}  (Saved {PrettySize(Saved)} or {Percent.ToString("0.0")}%)");
+			}
+			Logger.LogInformation("");
+
+			/*
 			RatioList.Sort((a, b) =>
 			{
 				if (a.Item1 == b.Item1)
@@ -1727,7 +2013,7 @@ namespace UnrealBuildTool
 					return -1;
 				return 1;
 			});
-			//var List = RatioList.OrderBy(I => I.Item1);
+			*/
 
 			return 0;
 		}
