[UBT]
* Enabled /Zc:preprocessor for c++20 path (it is needed for header units and probably doesn't add any large amounts of problems since we already build with clang)

diff --git a/Engine/Source/Programs/UnrealBuildTool/Platform/Windows/VCToolChain.cs b/Engine/Source/Programs/UnrealBuildTool/Platform/Windows/VCToolChain.cs
--- a/Engine/Source/Programs/UnrealBuildTool/Platform/Windows/VCToolChain.cs
+++ b/Engine/Source/Programs/UnrealBuildTool/Platform/Windows/VCToolChain.cs
@@ -903,7 +903,8 @@ namespace UnrealBuildTool
 					break;
 				case CppStandardVersion.Cpp20:
 					Arguments.Add("/std:c++20");
-					
+					Arguments.Add("/Zc:preprocessor");
+
 					// warning C5054: operator ___: deprecated between enumerations of different types
 					// re: http://www.open-std.org/jtc1/sc22/wg21/docs/papers/2018/p1120r0.html
 						
