- Working on improving IOS Plist generation, with UBT making a template plist project that Xcode will finalize, or project can use a fully premade plist
- Mac/MacEditor can also have template/premade plist files
- Legacy xcode mode will still write the entire plist file

diff --git a/Engine/Source/Programs/UnrealBuildTool/Platform/IOS/UEDeployIOS.cs b/Engine/Source/Programs/UnrealBuildTool/Platform/IOS/UEDeployIOS.cs
--- a/Engine/Source/Programs/UnrealBuildTool/Platform/IOS/UEDeployIOS.cs
+++ b/Engine/Source/Programs/UnrealBuildTool/Platform/IOS/UEDeployIOS.cs
@@ -217,7 +217,295 @@ namespace UnrealBuildTool
 			return MinVersionToReturn;
 		}
 
+
+		public static void WritePlistFile(FileReference PlistFile, DirectoryReference? ProjectLocation, UnrealPluginLanguage? UPL, string GameName, bool bIsUnrealGame, ILogger Logger)
+		{
+			ConfigHierarchy Ini = ConfigCache.ReadHierarchy(ConfigHierarchyType.Engine, ProjectLocation, UnrealTargetPlatform.IOS);
+			// required capabilities
+			List <string> RequiredCaps = new() { "arm64", "metal" };
+
+			// orientations
+			string InterfaceOrientation = "";
+			string PreferredLandscapeOrientation = "";
+			Ini.GetString("/Script/IOSRuntimeSettings.IOSRuntimeSettings", "PreferredLandscapeOrientation", out PreferredLandscapeOrientation);
+
+			string SupportedOrientations = "";
+			bool bSupported = true;
+			Ini.GetBool("/Script/IOSRuntimeSettings.IOSRuntimeSettings", "bSupportsPortraitOrientation", out bSupported);
+			SupportedOrientations += bSupported ? "\t\t<string>UIInterfaceOrientationPortrait</string>\n" : "";
+			bool bSupportsPortrait = bSupported;
+
+			Ini.GetBool("/Script/IOSRuntimeSettings.IOSRuntimeSettings", "bSupportsUpsideDownOrientation", out bSupported);
+			SupportedOrientations += bSupported ? "\t\t<string>UIInterfaceOrientationPortraitUpsideDown</string>\n" : "";
+			bSupportsPortrait |= bSupported;
+
+			bool bSupportsLandscapeLeft = false;
+			Ini.GetBool("/Script/IOSRuntimeSettings.IOSRuntimeSettings", "bSupportsLandscapeLeftOrientation", out bSupportsLandscapeLeft);
+			bool bSupportsLandscapeRight = false;
+			Ini.GetBool("/Script/IOSRuntimeSettings.IOSRuntimeSettings", "bSupportsLandscapeRightOrientation", out bSupportsLandscapeRight);
+			bool bSupportsLandscape = bSupportsLandscapeLeft || bSupportsLandscapeRight;
+
+			if (bSupportsLandscapeLeft && bSupportsLandscapeRight)
+			{
+				// if both landscape orientations are present, set the UIInterfaceOrientation key
+				// in the orientation list, the preferred orientation should be first
+				if (PreferredLandscapeOrientation == "LandscapeLeft")
+				{
+					InterfaceOrientation = "\t<key>UIInterfaceOrientation</key>\n\t<string>UIInterfaceOrientationLandscapeLeft</string>\n";
+					SupportedOrientations += "\t\t<string>UIInterfaceOrientationLandscapeLeft</string>\n\t\t<string>UIInterfaceOrientationLandscapeRight</string>\n";
+				}
+				else
+				{
+					// by default, landscape right is the preferred orientation - Apple's UI guidlines
+					InterfaceOrientation = "\t<key>UIInterfaceOrientation</key>\n\t<string>UIInterfaceOrientationLandscapeRight</string>\n";
+					SupportedOrientations += "\t\t<string>UIInterfaceOrientationLandscapeRight</string>\n\t\t<string>UIInterfaceOrientationLandscapeLeft</string>\n";
+				}
+			}
+			else
+			{
+				// max one landscape orientation is supported
+				SupportedOrientations += bSupportsLandscapeRight ? "\t\t<string>UIInterfaceOrientationLandscapeRight</string>\n" : "";
+				SupportedOrientations += bSupportsLandscapeLeft ? "\t\t<string>UIInterfaceOrientationLandscapeLeft</string>\n" : "";
+			}
+
+			// ITunes file sharing
+			bool bSupportsITunesFileSharing = false;
+			Ini.GetBool("/Script/IOSRuntimeSettings.IOSRuntimeSettings", "bSupportsITunesFileSharing", out bSupportsITunesFileSharing);
+			bool bSupportsFilesApp = false;
+			Ini.GetBool("/Script/IOSRuntimeSettings.IOSRuntimeSettings", "bSupportsFilesApp", out bSupportsFilesApp);
+
+			// disable https requirement
+			bool bDisableHTTPS;
+			Ini.GetBool("/Script/IOSRuntimeSettings.IOSRuntimeSettings", "bDisableHTTPS", out bDisableHTTPS);
+
+			// short version string
+			string BundleShortVersion;
+			Ini.GetString("/Script/IOSRuntimeSettings.IOSRuntimeSettings", "VersionInfo", out BundleShortVersion);
+
+
+			// Get Google Support details
+			bool bEnableGoogleSupport = true;
+			Ini.GetBool("/Script/IOSRuntimeSettings.IOSRuntimeSettings", "bEnableGoogleSupport", out bEnableGoogleSupport);
+
+			// Write the Google iOS URL Scheme if we need it.
+			string GoogleReversedClientId = "";
+			Ini.GetString("/Script/IOSRuntimeSettings.IOSRuntimeSettings", "GoogleReversedClientId", out GoogleReversedClientId);
+			bEnableGoogleSupport = bEnableGoogleSupport && !string.IsNullOrWhiteSpace(GoogleReversedClientId);
+
+			// Add remote-notifications as background mode
+			bool bRemoteNotificationsSupported = false;
+			Ini.GetBool("/Script/IOSRuntimeSettings.IOSRuntimeSettings", "bEnableRemoteNotificationsSupport", out bRemoteNotificationsSupported);
+
+			// Add audio as background mode
+			bool bBackgroundAudioSupported = false;
+			Ini.GetBool("/Script/IOSRuntimeSettings.IOSRuntimeSettings", "bSupportsBackgroundAudio", out bBackgroundAudioSupported);
+
+			// Add background fetch as background mode
+			bool bBackgroundFetch = false;
+			Ini.GetBool("/Script/IOSRuntimeSettings.IOSRuntimeSettings", "bEnableBackgroundFetch", out bBackgroundFetch);
+
+			// Get any Location Services permission descriptions added
+			string LocationAlwaysUsageDescription = "";
+			string LocationWhenInUseDescription = "";
+			Ini.GetString("/Script/LocationServicesIOSEditor.LocationServicesIOSSettings", "LocationAlwaysUsageDescription", out LocationAlwaysUsageDescription);
+			Ini.GetString("/Script/LocationServicesIOSEditor.LocationServicesIOSSettings", "LocationWhenInUseDescription", out LocationWhenInUseDescription);
+
+			// extra plist data
+			string ExtraData = "";
+			Ini.GetString("/Script/IOSRuntimeSettings.IOSRuntimeSettings", "AdditionalPlistData", out ExtraData);
+
+			Ini.GetBool("/Script/IOSRuntimeSettings.IOSRuntimeSettings", "bCustomLaunchscreenStoryboard", out VersionUtilities.bCustomLaunchscreenStoryboard);
+
+			// generate the plist file
+			StringBuilder Text = new StringBuilder();
+			Text.AppendLine("<?xml version=\"1.0\" encoding=\"UTF-8\"?>");
+			Text.AppendLine("<!DOCTYPE plist PUBLIC \"-//Apple//DTD PLIST 1.0//EN\" \"http://www.apple.com/DTDs/PropertyList-1.0.dtd\">");
+			Text.AppendLine("<plist version=\"1.0\">");
+			Text.AppendLine("<dict>");
+			Text.AppendLine("\t<key>CFBundleURLTypes</key>");
+			Text.AppendLine("\t<array>");
+			Text.AppendLine("\t\t<dict>");
+
+			Text.AppendLine("\t\t\t<key>CFBundleURLName</key>");
+			Text.AppendLine("\t\t\t<string>com.Epic.Unreal</string>");
+			Text.AppendLine("\t\t\t<key>CFBundleURLSchemes</key>");
+			Text.AppendLine("\t\t\t<array>");
+			Text.AppendLine(string.Format("\t\t\t\t<string>{0}</string>", bIsUnrealGame ? "UnrealGame" : GameName));
+			if (bEnableGoogleSupport)
+			{
+				Text.AppendLine(string.Format("\t\t\t\t<string>{0}</string>", GoogleReversedClientId));
+			}
+			Text.AppendLine("\t\t\t</array>");
+			Text.AppendLine("\t\t</dict>");
+			Text.AppendLine("\t</array>");
+			Text.AppendLine("\t<key>UIStatusBarHidden</key>");
+			Text.AppendLine("\t<true/>");
+			Text.AppendLine("\t<key>UIFileSharingEnabled</key>");
+			Text.AppendLine(string.Format("\t<{0}/>", bSupportsITunesFileSharing ? "true" : "false"));
+			if (bSupportsFilesApp)
+			{
+				Text.AppendLine("\t<key>LSSupportsOpeningDocumentsInPlace</key>");
+				Text.AppendLine("\t<true/>");
+			}
+			Text.AppendLine("\t<key>UIRequiresFullScreen</key>");
+			Text.AppendLine("\t<true/>");
+			Text.AppendLine("\t<key>UIViewControllerBasedStatusBarAppearance</key>");
+			Text.AppendLine("\t<false/>");
+			if (InterfaceOrientation != "")
+			{
+				Text.AppendLine(InterfaceOrientation);
+			}
+			Text.AppendLine("\t<key>UISupportedInterfaceOrientations</key>");
+			Text.AppendLine("\t<array>");
+			foreach (string Line in SupportedOrientations.Split("\r\n".ToCharArray()))
+			{
+				if (!string.IsNullOrWhiteSpace(Line))
+				{
+					Text.AppendLine(Line);
+				}
+			}
+			Text.AppendLine("\t</array>");
+			Text.AppendLine("\t<key>UIRequiredDeviceCapabilities</key>");
+			Text.AppendLine("\t<array>");
+			foreach (string Cap in RequiredCaps)
+				{
+				Text.AppendLine($"\t\t<string>{Cap}</string>\n");
+			}
+			Text.AppendLine("\t</array>");
+
+			Text.AppendLine("\t<key>UILaunchStoryboardName</key>");
+			Text.AppendLine("\t<string>LaunchScreen</string>");
+
+			// disable exempt encryption
+			Text.AppendLine("\t<key>ITSAppUsesNonExemptEncryption</key>");
+			Text.AppendLine("\t<false/>");
+			// add location services descriptions if used
+			if (!string.IsNullOrWhiteSpace(LocationAlwaysUsageDescription))
+			{
+				Text.AppendLine("\t<key>NSLocationAlwaysAndWhenInUseUsageDescription</key>");
+				Text.AppendLine(string.Format("\t<string>{0}</string>", LocationAlwaysUsageDescription));
+			}
+			if (!string.IsNullOrWhiteSpace(LocationWhenInUseDescription))
+			{
+				Text.AppendLine("\t<key>NSLocationWhenInUseUsageDescription</key>");
+				Text.AppendLine(string.Format("\t<string>{0}</string>", LocationWhenInUseDescription));
+			}
+			// disable HTTPS requirement
+			if (bDisableHTTPS)
+			{
+				Text.AppendLine("\t<key>NSAppTransportSecurity</key>");
+				Text.AppendLine("\t\t<dict>");
+				Text.AppendLine("\t\t\t<key>NSAllowsArbitraryLoads</key><true/>");
+				Text.AppendLine("\t\t</dict>");
+			}
+
+			if (!string.IsNullOrEmpty(ExtraData))
+			{
+				ExtraData = ExtraData.Replace("\\n", "\n");
+				foreach(string Line in ExtraData.Split("\r\n".ToCharArray()))
+				{
+					if (!string.IsNullOrWhiteSpace(Line))
+					{
+						Text.AppendLine("\t" + Line);
+					}
+				}
+			}
+
+			// Add remote-notifications as background mode
+			if (bRemoteNotificationsSupported || bBackgroundFetch || bBackgroundAudioSupported)
+			{
+				Text.AppendLine("\t<key>UIBackgroundModes</key>");
+				Text.AppendLine("\t<array>");
+				if (bBackgroundAudioSupported)
+				{
+					Text.AppendLine("\t\t<string>audio</string>");
+				}
+				if (bRemoteNotificationsSupported)
+				{
+					Text.AppendLine("\t\t<string>remote-notification</string>");
+				}
+				if (bBackgroundFetch)
+				{
+					Text.AppendLine("\t\t<string>fetch</string>");
+				}
+				Text.AppendLine("\t</array>");
+			}
+			Text.AppendLine("</dict>");
+			Text.AppendLine("</plist>");
+
+
+			DirectoryReference.CreateDirectory(PlistFile.Directory);
+
+			if (UPL != null)
+			{
+				// Allow UPL to modify the plist here
+				XDocument XDoc;
+				try
+				{
+					XDoc = XDocument.Parse(Text.ToString());
+				}
+				catch (Exception e)
+				{
+					throw new BuildException("plist is invalid {0}\n{1}", e, Text.ToString());
+				}
+
+				XDoc.DocumentType!.InternalSubset = "";
+				UPL.ProcessPluginNode("None", "iosPListUpdates", "", ref XDoc);
+				string result = XDoc.Declaration?.ToString() + "\n" + XDoc.ToString().Replace("<!DOCTYPE plist PUBLIC \"-//Apple//DTD PLIST 1.0//EN\" \"http://www.apple.com/DTDs/PropertyList-1.0.dtd\"[]>", "<!DOCTYPE plist PUBLIC \"-//Apple//DTD PLIST 1.0//EN\" \"http://www.apple.com/DTDs/PropertyList-1.0.dtd\">");
+				File.WriteAllText(PlistFile.FullName, result);
+
+				Text = new StringBuilder(result);
+			}
+
+			File.WriteAllText(PlistFile.FullName, Text.ToString());
+		}
+
 		public static bool GenerateIOSPList(FileReference? ProjectFile, UnrealTargetConfiguration Config, string ProjectDirectory, bool bIsUnrealGame, string GameName, bool bIsClient, string ProjectName, string InEngineDir, string AppDirectory, UnrealPluginLanguage? UPL, string? BundleID, bool bBuildAsFramework, ILogger Logger)
+		{
+			// get the settings from the ini file
+			// plist replacements
+			DirectoryReference? DirRef = bIsUnrealGame ? (!string.IsNullOrEmpty(UnrealBuildTool.GetRemoteIniPath()) ? new DirectoryReference(UnrealBuildTool.GetRemoteIniPath()!) : null) : new DirectoryReference(ProjectDirectory);
+			ConfigHierarchy Ini = ConfigCache.ReadHierarchy(ConfigHierarchyType.Engine, DirRef, UnrealTargetPlatform.Mac);
+			bool bUseModernXcode;
+			Ini.TryGetValue("XcodeConfiguration", "bUseModernXcode", out bUseModernXcode);
+
+			if (!bUseModernXcode)
+			{
+				return GenerateLegacyIOSPList(ProjectFile, Config, ProjectDirectory, bIsUnrealGame, GameName, bIsClient, ProjectName, InEngineDir, AppDirectory, UPL, BundleID, bBuildAsFramework, Logger);
+			}
+
+			// generate the Info.plist for future use
+			string BuildDirectory = ProjectDirectory + "/Build/IOS";
+			string IntermediateDirectory = (bIsUnrealGame ? InEngineDir : ProjectDirectory) + "/Intermediate/IOS";
+			string PListFile = IntermediateDirectory + "/" + GameName + "-Info.plist";;
+			ProjectName = !String.IsNullOrEmpty(ProjectName) ? ProjectName : GameName;
+			VersionUtilities.BuildDirectory = BuildDirectory;
+			VersionUtilities.GameName = GameName;
+
+			WritePlistFile(new FileReference(PListFile), DirRef, UPL, GameName, bIsUnrealGame, Logger);
+
+            if (BuildHostPlatform.Current.Platform == UnrealTargetPlatform.Mac && !bBuildAsFramework)
+			{
+				FileReference FinalPlistFile;
+				if (bUseModernXcode)
+				{
+					// @todo: do we need one per Target for IOS? Client? I dont think so for Modern
+					FinalPlistFile = new FileReference($"{ProjectDirectory}/Build/IOS/UBTGenerated/Info.plist");
+				}
+				else
+				{
+					FinalPlistFile = new FileReference(AppDirectory + "/Info.plist");
+				}
+				DirectoryReference.CreateDirectory(FinalPlistFile.Directory);
+				// @todo: writeifdifferent is better
+				FileReference.Delete(FinalPlistFile);
+				File.Copy(PListFile, FinalPlistFile.FullName);
+			}
+
+			return true;
+		}
+
+		public static bool GenerateLegacyIOSPList(FileReference? ProjectFile, UnrealTargetConfiguration Config, string ProjectDirectory, bool bIsUnrealGame, string GameName, bool bIsClient, string ProjectName, string InEngineDir, string AppDirectory, UnrealPluginLanguage? UPL, string? BundleID, bool bBuildAsFramework, ILogger Logger)
 		{
 			// generate the Info.plist for future use
 			string BuildDirectory = ProjectDirectory + "/Build/IOS";
@@ -318,7 +606,7 @@ namespace UnrealBuildTool
 			// minimum iOS version
 			string MinVersionSetting = "";
 			Ini.GetString("/Script/IOSRuntimeSettings.IOSRuntimeSettings", "MinimumiOSVersion", out MinVersionSetting);
-			string MinVersion  = GetMinimumOSVersion(MinVersionSetting, Logger);
+			string MinVersion = GetMinimumOSVersion(MinVersionSetting, Logger);
 
 			// Get Google Support details
 			bool bEnableGoogleSupport = true;
@@ -413,12 +701,12 @@ namespace UnrealBuildTool
 			Text.AppendLine("\t<key>UIViewControllerBasedStatusBarAppearance</key>");
 			Text.AppendLine("\t<false/>");
 			if (InterfaceOrientation != "")
-			{ 
+			{
 				Text.AppendLine(InterfaceOrientation);
 			}
 			Text.AppendLine("\t<key>UISupportedInterfaceOrientations</key>");
 			Text.AppendLine("\t<array>");
-			foreach(string Line in SupportedOrientations.Split("\r\n".ToCharArray()))
+			foreach (string Line in SupportedOrientations.Split("\r\n".ToCharArray()))
 			{
 				if (!string.IsNullOrWhiteSpace(Line))
 				{
@@ -428,7 +716,7 @@ namespace UnrealBuildTool
 			Text.AppendLine("\t</array>");
 			Text.AppendLine("\t<key>UIRequiredDeviceCapabilities</key>");
 			Text.AppendLine("\t<array>");
-			foreach(string Line in RequiredCaps.Split("\r\n".ToCharArray()))
+			foreach (string Line in RequiredCaps.Split("\r\n".ToCharArray()))
 			{
 				if (!string.IsNullOrWhiteSpace(Line))
 				{
@@ -519,7 +807,7 @@ namespace UnrealBuildTool
 			if (!string.IsNullOrEmpty(ExtraData))
 			{
 				ExtraData = ExtraData.Replace("\\n", "\n");
-				foreach(string Line in ExtraData.Split("\r\n".ToCharArray()))
+				foreach (string Line in ExtraData.Split("\r\n".ToCharArray()))
 				{
 					if (!string.IsNullOrWhiteSpace(Line))
 					{
@@ -594,7 +882,7 @@ namespace UnrealBuildTool
 
 			File.WriteAllText(PListFile, Text.ToString());
 
-            if (BuildHostPlatform.Current.Platform == UnrealTargetPlatform.Mac && !bBuildAsFramework)
+			if (BuildHostPlatform.Current.Platform == UnrealTargetPlatform.Mac && !bBuildAsFramework)
 			{
 				if (!Directory.Exists(AppDirectory))
 				{
@@ -606,7 +894,7 @@ namespace UnrealBuildTool
 			return true;
 		}
 
-        public static VersionNumber? GetSdkVersion(TargetReceipt Receipt)
+		public static VersionNumber? GetSdkVersion(TargetReceipt Receipt)
         {
 			VersionNumber? SdkVersion = null;
             if (Receipt != null)
@@ -800,6 +1088,15 @@ namespace UnrealBuildTool
 
 		public bool PrepForUATPackageOrDeploy(UnrealTargetConfiguration Config, FileReference? ProjectFile, string InProjectName, string InProjectDirectory, FileReference Executable, string InEngineDir, bool bForDistribution, string CookFlavor, bool bIsDataDeploy, bool bCreateStubIPA, List<string> UPLScripts, string? BundleID, bool bBuildAsFramework)
 		{
+			ConfigHierarchy Ini = ConfigCache.ReadHierarchy(ConfigHierarchyType.Engine, ProjectFile?.Directory, UnrealTargetPlatform.IOS);
+			bool bUseModernXcode;
+			if (Ini.TryGetValue("XcodeConfiguration", "bUseModernXcode", out bUseModernXcode) && bUseModernXcode)
+			{
+				Logger.LogInformation("Modern Xcode doens't need to Deploy");
+				// none of this is needed with modern xcode
+				return false;
+			}
+
 			if (BuildHostPlatform.Current.Platform != UnrealTargetPlatform.Mac)
 			{
 				throw new BuildException("UEDeployIOS.PrepForUATPackageOrDeploy only supports running on the Mac");
@@ -988,8 +1285,6 @@ namespace UnrealBuildTool
 			string ProjectDirectory = (DirectoryReference.FromFile(ProjectFile) ?? Unreal.EngineDirectory).FullName;
 			bool bIsUnrealGame = GameName.Contains("UnrealGame");
 
-			Console.WriteLine("1 GameName: {0}, ProjectName: {1}, Executable: {2}", GameName, (ProjectFile == null) ? "" : Path.GetFileNameWithoutExtension(ProjectFile.FullName), Executable);
-
 			if (BuildHostPlatform.Current.Platform == UnrealTargetPlatform.Mac && Environment.GetEnvironmentVariable("UBT_NO_POST_DEPLOY") != "true")
 			{
 				return PrepForUATPackageOrDeploy(Configuration, ProjectFile, GameName, ProjectDirectory, Executable, "../../Engine", bForDistribution, "", false, bCreateStubIPA, UPLScripts, BundleID, bBuildAsFramework);
diff --git a/Engine/Source/Programs/UnrealBuildTool/ProjectFiles/Xcode/XcodeProject.cs b/Engine/Source/Programs/UnrealBuildTool/ProjectFiles/Xcode/XcodeProject.cs
--- a/Engine/Source/Programs/UnrealBuildTool/ProjectFiles/Xcode/XcodeProject.cs
+++ b/Engine/Source/Programs/UnrealBuildTool/ProjectFiles/Xcode/XcodeProject.cs
@@ -13,6 +13,7 @@ using UnrealBuildBase;
 using Microsoft.Extensions.Logging;
 using System.Text.RegularExpressions;
 using System.Diagnostics.CodeAnalysis;
+using System.Security.Cryptography;
 
 
 /*****
@@ -157,11 +158,11 @@ namespace UnrealBuildTool.XcodeProjectXcconfig
 		UpdateTemplate,
 	}
 
-	class Metadata
+	class MetadataItem
 	{
 		public MetadataMode Mode;
-		public FileReference File;
-		public string XcodeProjectRelative;
+		public FileReference? File = null;
+		public string? XcodeProjectRelative = null;
 
 		//public Metadata(MetadataMode Mode, FileReference File, DirectoryReference XcodeProjectFile)
 		//{
@@ -170,42 +171,90 @@ namespace UnrealBuildTool.XcodeProjectXcconfig
 		//	XcodeProjectRelative = File.MakeRelativeTo(XcodeProjectFile.ParentDirectory!);
 		//}
 
-		public Metadata(DirectoryReference ProductDirectory, DirectoryReference XcodeProject, ConfigHierarchy Ini, string IniKey, string StandardLocationSubdir, string? ProjectTemplateSubdir = null)
+		public MetadataItem(DirectoryReference ProductDirectory, DirectoryReference XcodeProject, ConfigHierarchy Ini, params string[] Locations)
 		{
-			string? PremadeLocation;
-			Ini.TryGetValue("XcodeConfiguration", IniKey, out PremadeLocation);
-
-			if (!string.IsNullOrEmpty(PremadeLocation))
-			{
-				Mode = MetadataMode.UsePremade;
-				File = FileReference.Combine(ProductDirectory, PremadeLocation);
-			}
-			else
+			foreach (string Loc in Locations)
 			{
-				if (File == null)
+				string[] Tokens = Loc.Split(':');
+				Mode = (Tokens[0] == "premade") ? MetadataMode.UsePremade : MetadataMode.UpdateTemplate;
+
+				// no extension means it's a .ini entry
+				if (Path.GetExtension(Tokens[1]) == "")
 				{
-					FileReference StandardProjectLocation = FileReference.Combine(ProductDirectory, StandardLocationSubdir);
-					if (FileReference.Exists(StandardProjectLocation))
+					string? PremadeLocation;
+					if (Ini.TryGetValue("XcodeConfiguration", Tokens[1], out PremadeLocation) && PremadeLocation != "")
 					{
-						Mode = MetadataMode.UsePremade;
-						File = StandardProjectLocation;
-					}
-					else
-					{
-						Mode = MetadataMode.UpdateTemplate;
-						if (ProjectTemplateSubdir != null)
-						{
-							File = FileReference.Combine(ProductDirectory, ProjectTemplateSubdir);
-						}
-						else
+						// if the file exists (or it's a 'default' type which we use even if it doesn't exist), we are done!
+						FileReference TestFile = FileReference.Combine(ProductDirectory, PremadeLocation);
+						if (Tokens[0] == "default" || FileReference.Exists(TestFile))
 						{
-							File = FileReference.Combine(Unreal.EngineDirectory, StandardLocationSubdir);
+							File = TestFile;
+							break;
 						}
+						Log.TraceWarningOnce($"Metadata file {Tokens[1]} was specified in project settings, but does not exist. Ignoring...");
+					}
+				}
+				// otherwise it's a sub-path
+				else
+				{
+					FileReference EngineTestFile = FileReference.Combine(Unreal.EngineDirectory, Tokens[1]);
+					FileReference ProductTestFile = FileReference.Combine(ProductDirectory, Tokens[1]);
+					// if the file exists (or it's a 'default' type which we use even if it doesn't exist), we are done!
+					if (Tokens[0] == "default" || FileReference.Exists(ProductTestFile))
+					{
+						File = ProductTestFile;
+						break;
+					}
+					if (FileReference.Exists(EngineTestFile))
+					{
+						File = EngineTestFile;
+						break;
 					}
 				}
 			}
 
-			XcodeProjectRelative = File.MakeRelativeTo(XcodeProject.ParentDirectory!);
+			if (File == null)
+			{
+				// if null, we assume UpdateTemplate, since xcode can update "nothing" to something useful
+				Mode = MetadataMode.UpdateTemplate;
+			}
+			else
+			{
+				// make it relative if it exists
+				XcodeProjectRelative = File?.MakeRelativeTo(XcodeProject.ParentDirectory!);
+			}
+			
+			// Log.TraceInformation($"   Metadata {Locations[0]} found {Mode} File {File}");
+		}
+	}
+
+	class Metadata
+	{
+		public Dictionary<MetadataPlatform, MetadataItem> PlistFiles = new();
+		public Dictionary<MetadataPlatform, MetadataItem> EntitlementsFiles = new();
+
+
+		public Metadata(DirectoryReference ProductDirectory, DirectoryReference XcodeProject, ConfigHierarchy Ini, bool bSupportsMac, bool bSupportsIOSOrTVOS)
+		{
+			Log.TraceInformation($"making metadata for {ProductDirectory} / {XcodeProject}");
+			if (bSupportsMac)
+			{
+				PlistFiles[MetadataPlatform.MacEditor] = new MetadataItem(ProductDirectory, XcodeProject, Ini, "premade:PremadeMacEditorPlist", "template:TemplateMacEditorPlist", 
+					"premade:Build/Mac/Resources/Info-Editor.plist", "template:Build/Mac/Resources/Info-Editor.Template.plist");
+
+				PlistFiles[MetadataPlatform.Mac] = new MetadataItem(ProductDirectory, XcodeProject, Ini, "premade:PremadeMacPlist", "template:TemplateMacPlist", 
+					"premade:Build/Mac/Resources/Info.plist", "template:Build/Mac/Resources/Info.Template.plist");
+
+				EntitlementsFiles[MetadataPlatform.MacEditor] = new MetadataItem(ProductDirectory, XcodeProject, Ini, "premade:PremadeMacEditorEntitlements", "premade:Build/Mac/Resources/MacEditor.entitlements");
+				EntitlementsFiles[MetadataPlatform.Mac] = new MetadataItem(ProductDirectory, XcodeProject, Ini, "premade:PremadeMacEntitlements", "premade:Build/Mac/Resources/Mac.entitlements");
+			}
+			if (bSupportsIOSOrTVOS)
+			{
+				PlistFiles[MetadataPlatform.IOS] = new MetadataItem(ProductDirectory, XcodeProject, Ini, 
+					"premade:PremadeIOSPlist", "template:TemplateIOSPlist", "premade:Build/IOS/Resources/Info.plist", "template:Build/IOS/Resources/Info.Template.plist", "default:Build/IOS/UBTGenerated/Info.plist");
+
+				EntitlementsFiles[MetadataPlatform.IOS] = new MetadataItem(ProductDirectory, XcodeProject, Ini, "premade:PremadeIOSEntitlements", "premade:Build/IOS/Resources/IOS.entitlements");
+			}
 		}
 	}
 
@@ -218,8 +267,7 @@ namespace UnrealBuildTool.XcodeProjectXcconfig
 		public bool bIsMergingProjects = false;
 		public bool bWriteCodeSigningSettings = true;
 
-		public Dictionary<MetadataPlatform, Metadata> PlistFiles = new();
-		public Dictionary<MetadataPlatform, Metadata> EntitlementsFiles = new();
+		public Metadata? Metadata;
 
 		public List<UnrealBuildConfig> AllConfigs = new();
 
@@ -239,6 +287,9 @@ namespace UnrealBuildTool.XcodeProjectXcconfig
 		// Name of the product (usually the project name, but UE5.xcodeproj is actually UnrealGame product)
 		public string ProductName;
 
+		// Name of the xcode project
+		public string XcodeProjectName;
+
 		// Display name, can be overridden from commandline
 		public string DisplayName;
 
@@ -285,7 +336,7 @@ namespace UnrealBuildTool.XcodeProjectXcconfig
 			this.XcodeProjectFileLocation = new DirectoryReference(XcodeProjectFileLocation.FullName);
 			// default to engine director, will be fixed in Initialize if needed
 			ProductDirectory = Unreal.EngineDirectory;
-			ProductName = XcodeProjectFileLocation.GetFileNameWithoutAnyExtensions();
+			XcodeProjectName = ProductName = XcodeProjectFileLocation.GetFileNameWithoutAnyExtensions();
 			if (ProductName == "UE5")
 			{
 				ProductName = "UnrealGame";
@@ -293,7 +344,7 @@ namespace UnrealBuildTool.XcodeProjectXcconfig
 
 			this.bForDistribution = bIsForDistribution;
 			this.BundleIdentifier = string.IsNullOrEmpty(BundleID) ? "$(UE_SIGNING_PREFIX).$(UE_PRODUCT_NAME_STRIPPED)" : BundleID;
-			this.DisplayName = string.IsNullOrEmpty(AppName) ? "$(UE_PRODUCT_NAME)" : BundleID; ;
+			this.DisplayName = string.IsNullOrEmpty(AppName) ? "$(UE_PRODUCT_NAME)" : AppName;
 		}
 
 		public FileReference? FindUProjectFileLocation(XcodeProjectFile ProjectFile)
@@ -390,18 +441,7 @@ namespace UnrealBuildTool.XcodeProjectXcconfig
 			ConfigHierarchy SharedPlatformIni = ConfigCache.ReadHierarchy(ConfigHierarchyType.Engine, UProjectFileLocation?.Directory, UnrealTargetPlatform.Mac);
 			SharedPlatformIni.TryGetValue("XcodeConfiguration", "bUseModernCodeSigning", out bUseAutomaticSigning);
 
-			if (bSupportsMac)
-			{
-				PlistFiles[MetadataPlatform.MacEditor] = new Metadata(ProductDirectory, XcodeProjectFileLocation, SharedPlatformIni, "PremadeMacEditorPlist", "Build/Mac/Resources/Info-Editor.plist");
-				PlistFiles[MetadataPlatform.Mac] = new Metadata(ProductDirectory, XcodeProjectFileLocation, SharedPlatformIni, "PremadeMacPlist", "Build/Mac/Resources/Info.plist");
-				EntitlementsFiles[MetadataPlatform.MacEditor] = new Metadata(ProductDirectory, XcodeProjectFileLocation, SharedPlatformIni, "PremadeMacEditorPlist", "Build/Mac/Resources/Info-Editor.plist");
-				EntitlementsFiles[MetadataPlatform.Mac] = new Metadata(ProductDirectory, XcodeProjectFileLocation, SharedPlatformIni, "PremadeMacPlist", "Build/Mac/Resources/.plist");
-			}
-			if (bSupportsIOS || bSupportsTVOS)
-			{
-				PlistFiles[MetadataPlatform.IOS] = new Metadata(ProductDirectory, XcodeProjectFileLocation, SharedPlatformIni, "PremadeIOSPlist", "Build/IOS/Resources/Info.plist", "Build/IOS/UBTGenerated/Info.plist");
-				EntitlementsFiles[MetadataPlatform.IOS] = new Metadata(ProductDirectory, XcodeProjectFileLocation, SharedPlatformIni, "PremadeIOSPlist", "Build/IOS/Resources/Info.plist");
-			}
+			Metadata = new Metadata(ProductDirectory, XcodeProjectFileLocation, SharedPlatformIni, bSupportsMac, bSupportsIOS || bSupportsTVOS);
 		}
 
 		public string? FindFile(List<string> Paths, UnrealTargetPlatform Platform, bool bMakeRelative)
@@ -1214,7 +1254,7 @@ namespace UnrealBuildTool.XcodeProjectXcconfig
 			}
 
 			// set up names
-			Name = UnrealData.ProductName + (OverrideName ?? ConfigName);
+			Name = UnrealData.XcodeProjectName + (OverrideName ?? ConfigName);
 		}
 
 		public void AddDependency(XcodeTarget Target, XcodeProject Project)
@@ -1587,15 +1627,26 @@ namespace UnrealBuildTool.XcodeProjectXcconfig
 			return Orientations;
 		}
 
-		string FindBundleID(FileReference PlistFile)
+		void FindBundleID(MetadataItem PlistItem, ref string BundleId)
 		{
-			string Identifier = XcodeProjectFile.Plist("Print :CFBundleIdentifier", PlistFile.FullName);
+			if (PlistItem.File == null)
+			{
+				return;
+			}
+
+			string Identifier = XcodeProjectFile.Plist("Print :CFBundleIdentifier", PlistItem.File.FullName);
 			// handle error
-			if (Identifier.StartsWith("Print:"))
+			if (Identifier == "" || Identifier.StartsWith("Print:"))
 			{
-				return "";
+				if (PlistItem.Mode == MetadataMode.UsePremade)
+				{
+					Log.TraceErrorOnce($"Premade .plist file '{PlistItem.File}' was found, but it did not contain a BundleIdentifier (CFBundleIdentifier key is missing or empty)");
+				}
+			}
+			else
+			{
+				BundleId = Identifier;
 			}
-			return Identifier;
 		}
 
 		public override void WriteXcconfigFile()
@@ -1651,7 +1702,7 @@ namespace UnrealBuildTool.XcodeProjectXcconfig
 				Xcconfig.AppendLine($"CONFIGURATION_BUILD_DIR{SettingFilter} = {UnrealData.AllConfigs[0].IOSExecutablePath!.Directory}/Payload");
 				Xcconfig.AppendLine($"IPHONEOS_DEPLOYMENT_TARGET = {UnrealData.IOSProjectSettings.RuntimeVersion}");
 
-				if (UnrealData.PlistFiles[MetadataPlatform.IOS].Mode == MetadataMode.UpdateTemplate)
+				if (UnrealData.Metadata!.PlistFiles[MetadataPlatform.IOS].Mode == MetadataMode.UpdateTemplate)
 				{
 					// allow Xcode to generate the final plist file from our input, some INFOPLIST settings and other settings 
 					Xcconfig.AppendLine("GENERATE_INFOPLIST_FILE = YES");
@@ -1714,7 +1765,7 @@ namespace UnrealBuildTool.XcodeProjectXcconfig
 				Xcconfig.AppendLine($"CONFIGURATION_BUILD_DIR{SettingFilter} = {UnrealData.AllConfigs[0].TVOSExecutablePath!.Directory}/Payload");
 				Xcconfig.AppendLine($"TVOS_DEPLOYMENT_TARGET = {UnrealData.TVOSProjectSettings.RuntimeVersion}");
 
-				if (UnrealData.PlistFiles[MetadataPlatform.IOS].Mode == MetadataMode.UpdateTemplate)
+				if (UnrealData.Metadata!.PlistFiles[MetadataPlatform.IOS].Mode == MetadataMode.UpdateTemplate)
 				{
 					// allow Xcode to generate the final plist file from our input, some INFOPLIST settings and other settings 
 					Xcconfig.AppendLine("GENERATE_INFOPLIST_FILE = YES");
@@ -1800,18 +1851,12 @@ namespace UnrealBuildTool.XcodeProjectXcconfig
 			string IOSBundleID = UnrealData.BundleIdentifier;
 			if (UnrealData.bSupportsMac)
 			{
-				if (UnrealData.PlistFiles[MetadataPlatform.MacEditor].Mode == MetadataMode.UsePremade)
-				{
-					MacEditorBundleID = FindBundleID(UnrealData.PlistFiles[MetadataPlatform.MacEditor].File);
-				}
-				if (UnrealData.PlistFiles[MetadataPlatform.Mac].Mode == MetadataMode.UsePremade)
-				{
-					MacGameBundleID = FindBundleID(UnrealData.PlistFiles[MetadataPlatform.Mac].File);
-				}
+				FindBundleID(UnrealData.Metadata!.PlistFiles[MetadataPlatform.MacEditor], ref MacEditorBundleID);
+				FindBundleID(UnrealData.Metadata.PlistFiles[MetadataPlatform.Mac], ref MacGameBundleID);
 			}
-			if ((UnrealData.bSupportsIOS || UnrealData.bSupportsTVOS) && UnrealData.PlistFiles[MetadataPlatform.IOS].Mode == MetadataMode.UsePremade)
+			if (UnrealData.bSupportsIOS || UnrealData.bSupportsTVOS)
 			{
-				IOSBundleID = FindBundleID(UnrealData.PlistFiles[MetadataPlatform.IOS].File);
+				FindBundleID(UnrealData.Metadata!.PlistFiles[MetadataPlatform.IOS], ref IOSBundleID);
 			}
 
 			// Get Mac architectures supported by this project
@@ -1842,15 +1887,15 @@ namespace UnrealBuildTool.XcodeProjectXcconfig
 				string SupportedPlatforms = "";
 				if (Config.ProjectTarget!.TargetRules != null)
 				{
-					if (Config.bSupportsIOS)
+					if (ProjectBuildConfigs.bSupportsIOS)
 					{
 						SupportedPlatforms += "iphoneos iphonesimulator ";
 					}
-					if (Config.bSupportsTVOS)
+					if (ProjectBuildConfigs.bSupportsTVOS)
 					{
 						SupportedPlatforms += "appletvos ";
 					}
-					if (Config.bSupportsMac)
+					if (ProjectBuildConfigs.bSupportsMac)
 					{
 						SupportedPlatforms += "macosx ";
 					}
@@ -1864,11 +1909,11 @@ namespace UnrealBuildTool.XcodeProjectXcconfig
 
 
 				ConfigXcconfig.AppendLine("// Unreal per-config variables");
-				if (Config.bSupportsMac)
+				if (ProjectBuildConfigs.bSupportsMac)
 				{
 					ConfigXcconfig.AppendLine($"UE_MAC_EXECUTABLE_NAME = {MacExecutableFileName}");
 				}
-				if (Config.bSupportsIOS || Config.bSupportsTVOS)
+				if (ProjectBuildConfigs.bSupportsIOS || ProjectBuildConfigs.bSupportsTVOS)
 				{
 					ConfigXcconfig.AppendLine($"UE_IOS_EXECUTABLE_NAME = {Config.BuildTarget}");
 				}
@@ -1890,8 +1935,8 @@ namespace UnrealBuildTool.XcodeProjectXcconfig
 				if (ProjectBuildConfigs.bSupportsMac)
 				{
 					bool bIsEditor = (Config.ProjectTarget!.TargetRules!.Type == TargetType.Editor);
-					Metadata PlistData = UnrealData.PlistFiles[bIsEditor ? MetadataPlatform.MacEditor : MetadataPlatform.Mac];
-					Metadata EntitlementsData = UnrealData.EntitlementsFiles[bIsEditor ? MetadataPlatform.MacEditor : MetadataPlatform.Mac];
+					MetadataItem PlistData = UnrealData.Metadata!.PlistFiles[bIsEditor ? MetadataPlatform.MacEditor : MetadataPlatform.Mac];
+					MetadataItem EntitlementsData = UnrealData.Metadata.EntitlementsFiles[bIsEditor ? MetadataPlatform.MacEditor : MetadataPlatform.Mac];
 
 					string SettingFilter = Platform != null ? "" : "[sdk=macosx*]";
 
@@ -1925,15 +1970,15 @@ namespace UnrealBuildTool.XcodeProjectXcconfig
 
 					ConfigXcconfig.AppendLine("");
 					ConfigXcconfig.AppendLine("// IOS setup");
-					ConfigXcconfig.AppendLine($"INFOPLIST_FILE{SettingFilter} = {UnrealData.PlistFiles[MetadataPlatform.IOS].XcodeProjectRelative}");
+					ConfigXcconfig.AppendLine($"INFOPLIST_FILE{SettingFilter} = {UnrealData.Metadata!.PlistFiles[MetadataPlatform.IOS].XcodeProjectRelative}");
 					ConfigXcconfig.AppendLine($"PRODUCT_BUNDLE_IDENTIFIER{SettingFilter} = {IOSBundleID}");
-					if (UnrealData.PlistFiles[MetadataPlatform.IOS].Mode == MetadataMode.UpdateTemplate)
+					if (UnrealData.Metadata!.PlistFiles[MetadataPlatform.IOS].Mode == MetadataMode.UpdateTemplate)
 					{
 						ConfigXcconfig.AppendLine($"PRODUCT_NAME{SettingFilter} = {Config.BuildTarget}"); // @todo: change to Path.GetFileName(Config.IOSExecutablePath) when we stop using payload
 					}
-					if (UnrealData.EntitlementsFiles[MetadataPlatform.IOS].Mode == MetadataMode.UsePremade)
+					if (UnrealData.Metadata.EntitlementsFiles[MetadataPlatform.IOS].Mode == MetadataMode.UsePremade)
 					{
-						ConfigXcconfig.AppendLine($"CODE_SIGN_ENTITLEMENTS{SettingFilter} = {UnrealData.EntitlementsFiles[MetadataPlatform.IOS].XcodeProjectRelative}");
+						ConfigXcconfig.AppendLine($"CODE_SIGN_ENTITLEMENTS{SettingFilter} = {UnrealData.Metadata.EntitlementsFiles[MetadataPlatform.IOS].XcodeProjectRelative}");
 					}
 				}
 				if (ProjectBuildConfigs.bSupportsTVOS)
@@ -1942,11 +1987,11 @@ namespace UnrealBuildTool.XcodeProjectXcconfig
 
 					ConfigXcconfig.AppendLine("");
 					ConfigXcconfig.AppendLine("// TVOS setup");
-					ConfigXcconfig.AppendLine($"INFOPLIST_FILE{SettingFilter} = {UnrealData.PlistFiles[MetadataPlatform.IOS].XcodeProjectRelative}");
+					ConfigXcconfig.AppendLine($"INFOPLIST_FILE{SettingFilter} = {UnrealData.Metadata!.PlistFiles[MetadataPlatform.IOS].XcodeProjectRelative}");
 					ConfigXcconfig.AppendLine($"PRODUCT_BUNDLE_IDENTIFIER{SettingFilter} = {IOSBundleID}");
-					if (UnrealData.EntitlementsFiles[MetadataPlatform.IOS].Mode == MetadataMode.UsePremade)
+					if (UnrealData.Metadata.EntitlementsFiles[MetadataPlatform.IOS].Mode == MetadataMode.UsePremade)
 					{
-						ConfigXcconfig.AppendLine($"CODE_SIGN_ENTITLEMENTS{SettingFilter} = {UnrealData.EntitlementsFiles[MetadataPlatform.IOS].XcodeProjectRelative}");
+						ConfigXcconfig.AppendLine($"CODE_SIGN_ENTITLEMENTS{SettingFilter} = {UnrealData.Metadata.EntitlementsFiles[MetadataPlatform.IOS].XcodeProjectRelative}");
 					}
 				}
 
@@ -2690,7 +2735,7 @@ namespace UnrealBuildTool.XcodeProjectXcconfig
 						string ProjectName = ProjectFilePathForPlatform(Platform).GetFileNameWithoutAnyExtensions();
 						string? BuildTargetGuid = XcodeProjectNode.GetNodesOfType<XcodeBuildTarget>(RootProject).FirstOrDefault()?.Guid;
 						string? IndexTargetGuid = XcodeProjectNode.GetNodesOfType<XcodeIndexTarget>(RootProject).FirstOrDefault()?.Guid;
-						WriteSchemeFile(Platform, ProjectName, RootProject.RunTargets, BuildTargetGuid, IndexTargetGuid, UnrealData.bHasEditorConfiguration, UnrealData.UProjectFileLocation != null ? UnrealData.UProjectFileLocation.FullName : "");
+						WriteSchemeFile(Platform, ProjectName, UnrealData.ProductName, RootProject.RunTargets, BuildTargetGuid, IndexTargetGuid, UnrealData.bHasEditorConfiguration, UnrealData.UProjectFileLocation != null ? UnrealData.UProjectFileLocation.FullName : "");
 					}
 				}
 				else
@@ -2713,13 +2758,16 @@ namespace UnrealBuildTool.XcodeProjectXcconfig
 			CurrentPlistPlatform = Platform;
 
 			// we now use templates or premade, no writing out here
-			foreach (Metadata Data in UnrealData.PlistFiles.Values)
+			foreach (MetadataItem Data in UnrealData.Metadata!.PlistFiles.Values)
 			{
-				FileCollection.AddFileReference(XcodeProjectFileGenerator.MakeXcodeGuid(), Data.XcodeProjectRelative, "explicitFileType", "text.plist", "\"<group>\"", "Metadata");
+				if (Data.XcodeProjectRelative != null)
+				{
+					FileCollection.AddFileReference(XcodeProjectFileGenerator.MakeXcodeGuid(), Data.XcodeProjectRelative, "explicitFileType", "text.plist", "\"<group>\"", "Metadata");
+				}
 			}
-			foreach (Metadata Data in UnrealData.EntitlementsFiles.Values)
+			foreach (MetadataItem Data in UnrealData.Metadata.EntitlementsFiles.Values)
 			{
-				if (Data.Mode == MetadataMode.UsePremade)
+				if (Data.XcodeProjectRelative != null && Data.Mode == MetadataMode.UsePremade)
 				{
 					FileCollection.AddFileReference(XcodeProjectFileGenerator.MakeXcodeGuid(), Data.XcodeProjectRelative, "explicitFileType", "text.plist", "\"<group>\"", "Metadata");
 				}
@@ -3340,7 +3388,7 @@ namespace UnrealBuildTool.XcodeProjectXcconfig
 			return FileReference.Combine(GetProjectSchemeDirectory(Platform), TargetName + ".xcscheme");
 		}
 
-		private void WriteSchemeFile(UnrealTargetPlatform? Platform, string ProjectName, List<XcodeRunTarget> RunTargets, string? BuildTargetGuid, string? IndexTargetGuid, bool bHasEditorConfiguration, string GameProjectPath)
+		private void WriteSchemeFile(UnrealTargetPlatform? Platform, string ProjectName, string ProductName, List<XcodeRunTarget> RunTargets, string? BuildTargetGuid, string? IndexTargetGuid, bool bHasEditorConfiguration, string GameProjectPath)
 		{
 			StringBuilder Content = new StringBuilder();
 
@@ -3385,6 +3433,8 @@ namespace UnrealBuildTool.XcodeProjectXcconfig
 				Content.WriteLine("            <BuildableReference");
 				Content.WriteLine("               BuildableIdentifier = \"primary\"");
 				Content.WriteLine("               BlueprintIdentifier = \"" + TargetGuid + "\"");
+				Content.WriteLine("               BuildableName = \"" + ProductName + ".app\"");
+				Content.WriteLine("               BlueprintName = \"" + TargetName + "\"");
 				Content.WriteLine("               ReferencedContainer = \"container:" + ProjectName + ".xcodeproj\">");
 				Content.WriteLine("            </BuildableReference>");
 				Content.WriteLine("         </BuildActionEntry>");
@@ -3395,17 +3445,17 @@ namespace UnrealBuildTool.XcodeProjectXcconfig
 				Content.WriteLine("      selectedDebuggerIdentifier = \"Xcode.DebuggerFoundation.Debugger.LLDB\"");
 				Content.WriteLine("      selectedLauncherIdentifier = \"Xcode.DebuggerFoundation.Launcher.LLDB\"");
 				Content.WriteLine("      shouldUseLaunchSchemeArgsEnv = \"YES\">");
-				Content.WriteLine("      <Testables>");
-				Content.WriteLine("      </Testables>");
 				Content.WriteLine("      <MacroExpansion>");
-				Content.WriteLine("            <BuildableReference");
-				Content.WriteLine("               BuildableIdentifier = \"primary\"");
-				Content.WriteLine("               BlueprintIdentifier = \"" + TargetGuid + "\"");
-				Content.WriteLine("               ReferencedContainer = \"container:" + ProjectName + ".xcodeproj\">");
-				Content.WriteLine("            </BuildableReference>");
+				Content.WriteLine("         <BuildableReference");
+				Content.WriteLine("            BuildableIdentifier = \"primary\"");
+				Content.WriteLine("            BlueprintIdentifier = \"" + TargetGuid + "\"");
+				Content.WriteLine("            BuildableName = \"" + ProductName + ".app\"");
+				Content.WriteLine("            BlueprintName = \"" + TargetName + "\"");
+				Content.WriteLine("            ReferencedContainer = \"container:" + ProjectName + ".xcodeproj\">");
+				Content.WriteLine("         </BuildableReference>");
 				Content.WriteLine("      </MacroExpansion>");
-				Content.WriteLine("      <AdditionalOptions>");
-				Content.WriteLine("      </AdditionalOptions>");
+				Content.WriteLine("      <Testables>");
+				Content.WriteLine("      </Testables>");
 				Content.WriteLine("   </TestAction>");
 				Content.WriteLine("   <LaunchAction");
 				Content.WriteLine("      buildConfiguration = \"" + DefaultConfiguration + "\"");
@@ -3419,11 +3469,13 @@ namespace UnrealBuildTool.XcodeProjectXcconfig
 				Content.WriteLine("      allowLocationSimulation = \"YES\">");
 				Content.WriteLine("      <BuildableProductRunnable");
 				Content.WriteLine("         runnableDebuggingMode = \"0\">");
-				Content.WriteLine("            <BuildableReference");
-				Content.WriteLine("               BuildableIdentifier = \"primary\"");
-				Content.WriteLine("               BlueprintIdentifier = \"" + TargetGuid + "\"");
-				Content.WriteLine("               ReferencedContainer = \"container:" + ProjectName + ".xcodeproj\">");
-				Content.WriteLine("            </BuildableReference>");
+				Content.WriteLine("         <BuildableReference");
+				Content.WriteLine("            BuildableIdentifier = \"primary\"");
+				Content.WriteLine("            BlueprintIdentifier = \"" + TargetGuid + "\"");
+				Content.WriteLine("            BuildableName = \"" + ProductName + ".app\"");
+				Content.WriteLine("            BlueprintName = \"" + TargetName + "\"");
+				Content.WriteLine("            ReferencedContainer = \"container:" + ProjectName + ".xcodeproj\">");
+				Content.WriteLine("         </BuildableReference>");
 				Content.WriteLine("      </BuildableProductRunnable>");
 				if (string.IsNullOrEmpty(OldCommandLineArguments))
 				{
@@ -3456,8 +3508,6 @@ namespace UnrealBuildTool.XcodeProjectXcconfig
 				{
 					Content.WriteLine("      <CommandLineArguments>" + OldCommandLineArguments + "</CommandLineArguments>");
 				}
-				Content.WriteLine("      <AdditionalOptions>");
-				Content.WriteLine("      </AdditionalOptions>");
 				Content.WriteLine("   </LaunchAction>");
 				Content.WriteLine("   <ProfileAction");
 				Content.WriteLine("      buildConfiguration = \"" + DefaultConfiguration + "\"");
@@ -3467,13 +3517,13 @@ namespace UnrealBuildTool.XcodeProjectXcconfig
 				Content.WriteLine("      debugDocumentVersioning = \"NO\">");
 				Content.WriteLine("      <BuildableProductRunnable");
 				Content.WriteLine("         runnableDebuggingMode = \"0\">");
-				Content.WriteLine("            <BuildableReference");
-				Content.WriteLine("               BuildableIdentifier = \"primary\"");
-				Content.WriteLine("               BlueprintIdentifier = \"" + TargetGuid + "\"");
-				Content.WriteLine("               BuildableName = \"" + ProjectName + ".app\"");
-				Content.WriteLine("               BlueprintName = \"" + TargetName + "\"");
-				Content.WriteLine("               ReferencedContainer = \"container:" + ProjectName + ".xcodeproj\">");
-				Content.WriteLine("            </BuildableReference>");
+				Content.WriteLine("         <BuildableReference");
+				Content.WriteLine("            BuildableIdentifier = \"primary\"");
+				Content.WriteLine("            BlueprintIdentifier = \"" + TargetGuid + "\"");
+				Content.WriteLine("            BuildableName = \"" + ProductName + ".app\"");
+				Content.WriteLine("            BlueprintName = \"" + TargetName + "\"");
+				Content.WriteLine("            ReferencedContainer = \"container:" + ProjectName + ".xcodeproj\">");
+				Content.WriteLine("         </BuildableReference>");
 				Content.WriteLine("      </BuildableProductRunnable>");
 				Content.WriteLine("   </ProfileAction>");
 				Content.WriteLine("   <AnalyzeAction");
