[UBT]
* Added so SourceFileMetadataCache can skip tracking of includes if the include has markup HEADER_UNIT_IGNORE. This is needed for header units path. All includes that are circular needs to be marked up (there are ~10 of them in the engine)

diff --git a/Engine/Source/Programs/UnrealBuildTool/System/SourceFileMetadataCache.cs b/Engine/Source/Programs/UnrealBuildTool/System/SourceFileMetadataCache.cs
--- a/Engine/Source/Programs/UnrealBuildTool/System/SourceFileMetadataCache.cs
+++ b/Engine/Source/Programs/UnrealBuildTool/System/SourceFileMetadataCache.cs
@@ -302,18 +302,46 @@ namespace UnrealBuildTool
 					bContainsMarkup = ReflectionMarkupRegex.IsMatch(Line);
 				}
 
-				if (Line.AsSpan().TrimStart().StartsWith("#include"))
+				ReadOnlySpan<char> IncludeSpan = Line.AsSpan().TrimStart();
+				if (IncludeSpan.StartsWith("#include"))
 				{
-					int FirstQuotation = Line.IndexOf('"');
-					int FirstAngleBracket = Line.IndexOf('<');
-					if (FirstQuotation != -1 && (FirstAngleBracket == -1 || FirstQuotation < FirstAngleBracket)) // Handle #include <foo.h> // Some text with "
+					IncludeSpan = IncludeSpan.Slice("#include".Length).TrimStart();
+					if (IncludeSpan.IsEmpty)
+						continue;
+					char EndChar;
+					bool TrimQuotation = true;
+					if (IncludeSpan[0] == '"')
 					{
-						int SecondQuotation = Line.IndexOf('"', FirstQuotation + 1);
-						if (SecondQuotation != -1)
-						{
-							Includes.Add(Line.Substring(FirstQuotation + 1, SecondQuotation - FirstQuotation - 1));
-						}
+						EndChar = '"';
+					}
+					else if (IncludeSpan[0] == '<')
+					{
+						EndChar = '>';
+					}
+					else
+					{
+						EndChar = ')';
+						TrimQuotation = false;
+					}
+					if (TrimQuotation)
+					{
+						IncludeSpan = IncludeSpan.Slice(1);
 					}
+
+					if (IncludeSpan.Contains("HEADER_UNIT_IGNORE", StringComparison.OrdinalIgnoreCase))
+					{
+						continue;
+					}
+
+					int EndIndex = IncludeSpan.IndexOf(EndChar);
+
+					if (EndIndex == -1)
+					{
+						continue;
+					}
+
+					IncludeSpan = IncludeSpan.Slice(0, EndIndex);
+					Includes.Add(IncludeSpan.ToString());
 				}
 
 				int HeaderUnitIndex = Line.IndexOf("HEADER_UNIT_");
