UnrealBuildTool: When checking timestamps of produced files, also consider creation time as in some cases copied outputs (such as from a cache) can have a newer creation time than last write time

#rnx

diff --git a/Engine/Source/Programs/Shared/EpicGames.Build/System/DirectoryItem.cs b/Engine/Source/Programs/Shared/EpicGames.Build/System/DirectoryItem.cs
--- a/Engine/Source/Programs/Shared/EpicGames.Build/System/DirectoryItem.cs
+++ b/Engine/Source/Programs/Shared/EpicGames.Build/System/DirectoryItem.cs
@@ -82,6 +82,11 @@ namespace UnrealBuildBase
 		/// </summary>
 		public DateTime LastWriteTimeUtc => Info.Value.LastWriteTimeUtc;
 
+		/// <summary>
+		/// The creation time of the file.
+		/// </summary>
+		public DateTime CreationTimeUtc => Info.Value.CreationTimeUtc;
+
 		/// <summary>
 		/// Gets the parent directory item
 		/// </summary>
diff --git a/Engine/Source/Programs/Shared/EpicGames.Build/System/FileItem.cs b/Engine/Source/Programs/Shared/EpicGames.Build/System/FileItem.cs
--- a/Engine/Source/Programs/Shared/EpicGames.Build/System/FileItem.cs
+++ b/Engine/Source/Programs/Shared/EpicGames.Build/System/FileItem.cs
@@ -109,6 +109,11 @@ namespace UnrealBuildBase
 		/// </summary>
 		public DateTime LastWriteTimeUtc => Info.Value.LastWriteTimeUtc;
 
+		/// <summary>
+		/// The creation time of the file.
+		/// </summary>
+		public DateTime CreationTimeUtc => Info.Value.CreationTimeUtc;
+
 		/// <summary>
 		/// Determines if the file has the given extension
 		/// </summary>
diff --git a/Engine/Source/Programs/UnrealBuildTool/System/ActionGraph.cs b/Engine/Source/Programs/UnrealBuildTool/System/ActionGraph.cs
--- a/Engine/Source/Programs/UnrealBuildTool/System/ActionGraph.cs
+++ b/Engine/Source/Programs/UnrealBuildTool/System/ActionGraph.cs
@@ -685,10 +685,12 @@ namespace UnrealBuildTool
 				if (ProducedItem.Exists && (RootAction.ActionType != ActionType.Compile || ProducedItem.Length > 0 ||
 				                            (!ProducedItem.Location.HasExtension(".obj") && !ProducedItem.Location.HasExtension(".o"))))
 				{
+					// Find the newer of LastWriteTime and CreationTime, as copied files (such as from a build cache) can have a newer creation time in some cases
+					DateTime ExecutionTimeUtc = ProducedItem.LastWriteTimeUtc > ProducedItem.CreationTimeUtc ? ProducedItem.LastWriteTimeUtc : ProducedItem.CreationTimeUtc;
 					// Use the oldest produced item's time as the last execution time.
-					if (ProducedItem.LastWriteTimeUtc < LastExecutionTimeUtc)
+					if (ExecutionTimeUtc < LastExecutionTimeUtc)
 					{
-						LastExecutionTimeUtc = ProducedItem.LastWriteTimeUtc;
+						LastExecutionTimeUtc = ExecutionTimeUtc;
 					}
 				}
 				else
