UnrealBuildTool: Normalize msvc pdb path in response file

#rnx

diff --git a/Engine/Source/Programs/UnrealBuildTool/Platform/Windows/VCToolChain.cs b/Engine/Source/Programs/UnrealBuildTool/Platform/Windows/VCToolChain.cs
--- a/Engine/Source/Programs/UnrealBuildTool/Platform/Windows/VCToolChain.cs
+++ b/Engine/Source/Programs/UnrealBuildTool/Platform/Windows/VCToolChain.cs
@@ -1506,7 +1506,7 @@ namespace UnrealBuildTool
 					}
 
 					// Specify the PDB file that the compiler should write to.
-					CompileAction.Arguments.Add($"/Fd\"{PDBLocation}\"");
+					CompileAction.Arguments.Add($"/Fd\"{NormalizeCommandLinePath(PDBLocation)}\"");
 
 					// Don't allow remote execution when PDB files are enabled; we need to modify the same files. XGE works around this by generating separate
 					// PDB files per agent, but this functionality is only available with the Visual C++ extension package (via the VCCompiler=true tool option).
