[UBT]
* Added support for additional response files (will be used by header units)
* Added support for shared response files and changed so normal cpp/h build path uses it if toolchain supports it.
* Enabled shared response files on all platforms but android (which doesnt seem to support reponse files including response files)
* Some cleanup in the iwyu toolchain related to shared rsp files

diff --git a/Engine/Source/Programs/UnrealBuildTool/Configuration/UEBuildModuleCPP.cs b/Engine/Source/Programs/UnrealBuildTool/Configuration/UEBuildModuleCPP.cs
--- a/Engine/Source/Programs/UnrealBuildTool/Configuration/UEBuildModuleCPP.cs
+++ b/Engine/Source/Programs/UnrealBuildTool/Configuration/UEBuildModuleCPP.cs
@@ -671,12 +671,16 @@ namespace UnrealBuildTool
 				LinkInputFiles.AddRange(ToolChain.CompileAllCPPFiles(GeneratedCPPCompileEnvironment, GeneratedFileItems, IntermediateDirectory, Name, Graph).ObjectFiles);
 			}
 
+			// Create shared rsp for the normal cpp files
+			FileReference SharedResponseFile = FileReference.Combine(IntermediateDirectory, $"{Name}.Shared{UEToolChain.ResponseExt}");
+			CppCompileEnvironment CodeCompileEnvironment = ToolChain.CreateSharedResponseFile(CompileEnvironment, SharedResponseFile, Graph);
+
 			// Compile CPP files
 			if (bModuleUsesUnityBuild)
 			{
 				Unity.GenerateUnityCPPs(Target, CPPFiles, InputFiles.HeaderFiles, CompileEnvironment, WorkingSet, Rules.ShortName ?? Name, IntermediateDirectory, Graph, SourceFileToUnityFile, 
 					out List<FileItem> NormalFiles, out List<FileItem> AdaptiveFiles, NumIncludedBytesPerUnityCPP);
-				LinkInputFiles.AddRange(CompileFilesWithToolChain(Target, ToolChain, CompileEnvironment, ModuleCompileEnvironment, NormalFiles, AdaptiveFiles, Graph, Logger).ObjectFiles);
+				LinkInputFiles.AddRange(CompileFilesWithToolChain(Target, ToolChain, CodeCompileEnvironment, ModuleCompileEnvironment, NormalFiles, AdaptiveFiles, Graph, Logger).ObjectFiles);
 			}
 			else if (SpecificFilesToCompile.Count == 0)
 			{
@@ -687,7 +691,7 @@ namespace UnrealBuildTool
 					NormalFiles = CPPFiles;
 					AdaptiveFiles.RemoveAll(new HashSet<FileItem>(NormalFiles).Contains);
 				}
-				LinkInputFiles.AddRange(CompileFilesWithToolChain(Target, ToolChain, CompileEnvironment, ModuleCompileEnvironment, NormalFiles, AdaptiveFiles, Graph, Logger).ObjectFiles);
+				LinkInputFiles.AddRange(CompileFilesWithToolChain(Target, ToolChain, CodeCompileEnvironment, ModuleCompileEnvironment, NormalFiles, AdaptiveFiles, Graph, Logger).ObjectFiles);
 			}
 			else
 			{
@@ -771,7 +775,7 @@ namespace UnrealBuildTool
 			if (Target.bIWYU)
 			{
 				// Collect the headers that should be built
-				List<FileItem> HeaderFileItems = GetCompilableHeaders(InputFiles, CompileEnvironment);
+				List<FileItem> HeaderFileItems = GetCompilableHeaders(InputFiles, CodeCompileEnvironment);
 				if (HeaderFileItems.Count > 0)
 				{
 					if (Target.bIWYUHeadersOnly)
@@ -780,7 +784,7 @@ namespace UnrealBuildTool
 					}
 
 					// Add the compile actions
-					LinkInputFiles.AddRange(ToolChain.CompileAllCPPFiles(CompileEnvironment, HeaderFileItems, IntermediateDirectory, Name, Graph).ObjectFiles);
+					LinkInputFiles.AddRange(ToolChain.CompileAllCPPFiles(CodeCompileEnvironment, HeaderFileItems, IntermediateDirectory, Name, Graph).ObjectFiles);
 				}
 			}
 
diff --git a/Engine/Source/Programs/UnrealBuildTool/Platform/Android/AndroidToolChain.cs b/Engine/Source/Programs/UnrealBuildTool/Platform/Android/AndroidToolChain.cs
--- a/Engine/Source/Programs/UnrealBuildTool/Platform/Android/AndroidToolChain.cs
+++ b/Engine/Source/Programs/UnrealBuildTool/Platform/Android/AndroidToolChain.cs
@@ -957,6 +957,12 @@ namespace UnrealBuildTool
 			}
 		}
 
+		public override CppCompileEnvironment CreateSharedResponseFile(CppCompileEnvironment CompileEnvironment, FileReference OutResponseFile, IActionGraphBuilder Graph)
+		{
+			// Seems like Android clang toolchain does not handle response files including response files
+			return CompileEnvironment;
+		}
+
 		void GenerateEmptyLinkFunctionsForRemovedModules(List<FileItem> SourceFiles, UnrealArch Arch, string ModuleName, DirectoryReference OutputDirectory, IActionGraphBuilder Graph, ILogger Logger)
 		{
 			// Only add to Launch module
diff --git a/Engine/Source/Programs/UnrealBuildTool/Platform/Windows/VCToolChain.cs b/Engine/Source/Programs/UnrealBuildTool/Platform/Windows/VCToolChain.cs
--- a/Engine/Source/Programs/UnrealBuildTool/Platform/Windows/VCToolChain.cs
+++ b/Engine/Source/Programs/UnrealBuildTool/Platform/Windows/VCToolChain.cs
@@ -1364,12 +1364,24 @@ namespace UnrealBuildTool
 		protected override CPPOutput CompileCPPFiles(CppCompileEnvironment CompileEnvironment, List<FileItem> InputFiles, DirectoryReference OutputDir, string ModuleName, IActionGraphBuilder Graph)
 		{
 			VCCompileAction BaseCompileAction = new VCCompileAction(EnvVars);
+
+			// Add additional response files
+			foreach (FileItem AdditionalRsp in CompileEnvironment.AdditionalResponseFiles)
+			{
+				BaseCompileAction.Arguments.Add($"@{NormalizeCommandLinePath(AdditionalRsp.Location)}");
+			}
+
 			AppendCLArguments_Global(CompileEnvironment, BaseCompileAction.Arguments);
 
 			// Add include paths to the argument list.
 			BaseCompileAction.IncludePaths.AddRange(CompileEnvironment.UserIncludePaths);
 			BaseCompileAction.SystemIncludePaths.AddRange(CompileEnvironment.SystemIncludePaths);
-			BaseCompileAction.SystemIncludePaths.AddRange(EnvVars.IncludePaths);
+			
+			// Maybe not the nicest way of checking but if we don't have any system include paths it means we are not using a shared rsp file
+			if (CompileEnvironment.SystemIncludePaths.Count != 0)
+			{
+				BaseCompileAction.SystemIncludePaths.AddRange(EnvVars.IncludePaths);
+			}
 			
 			// Remember the architecture
 			BaseCompileAction.Architecture = CompileEnvironment.Architecture;
@@ -2030,6 +2042,38 @@ namespace UnrealBuildTool
 			}
 		}
 
+		public override CppCompileEnvironment CreateSharedResponseFile(CppCompileEnvironment CompileEnvironment, FileReference OutResponseFile, IActionGraphBuilder Graph)
+		{
+			CppCompileEnvironment NewCompileEnvironment = new CppCompileEnvironment(CompileEnvironment);
+			List<string> Arguments = new List<string>();
+
+			foreach (DirectoryReference IncludePath in NewCompileEnvironment.UserIncludePaths)
+			{
+				AddIncludePath(Arguments, IncludePath, Target.WindowsPlatform.Compiler);
+			}
+
+			foreach (DirectoryReference IncludePath in NewCompileEnvironment.SystemIncludePaths)
+			{
+				AddSystemIncludePath(Arguments, IncludePath, Target.WindowsPlatform.Compiler);
+			}
+
+			foreach (DirectoryReference IncludePath in EnvVars.IncludePaths)
+			{
+				AddSystemIncludePath(Arguments, IncludePath, Target.WindowsPlatform.Compiler);
+			}
+
+			NewCompileEnvironment.UserIncludePaths.Clear();
+			NewCompileEnvironment.SystemIncludePaths.Clear();
+
+			FileItem FileItem = FileItem.GetItemByFileReference(OutResponseFile);
+			Graph.CreateIntermediateTextFile(FileItem, Arguments);
+
+			NewCompileEnvironment.AdditionalPrerequisites.Add(FileItem);
+			NewCompileEnvironment.AdditionalResponseFiles.Add(FileItem);
+
+			return NewCompileEnvironment;
+		}
+
 		/// <summary>
 		/// Macros passed via the command line have their quotes stripped, and are tokenized before being re-stringized by the compiler. This conversion
 		/// back and forth is normally ok, but certain characters such as single quotes must always be paired. Remove any such characters here.
diff --git a/Engine/Source/Programs/UnrealBuildTool/System/CppCompileEnvironment.cs b/Engine/Source/Programs/UnrealBuildTool/System/CppCompileEnvironment.cs
--- a/Engine/Source/Programs/UnrealBuildTool/System/CppCompileEnvironment.cs
+++ b/Engine/Source/Programs/UnrealBuildTool/System/CppCompileEnvironment.cs
@@ -477,6 +477,11 @@ namespace UnrealBuildTool
 		/// </summary>
 		public List<string> Definitions = new List<string>();
 
+		/// <summary>
+		/// Additional response files that will be used by main response file
+		/// </summary>
+		public List<FileItem> AdditionalResponseFiles = new();
+
 		/// <summary>
 		/// Additional arguments to pass to the compiler.
 		/// </summary>
@@ -636,6 +641,7 @@ namespace UnrealBuildTool
 			CollidingNames = Other.CollidingNames;
 			FileInlineGenCPPMap = new Dictionary<FileItem, List<FileItem>>(Other.FileInlineGenCPPMap);
 			Definitions.AddRange(Other.Definitions);
+			AdditionalResponseFiles.AddRange(Other.AdditionalResponseFiles);
 			AdditionalArguments = Other.AdditionalArguments;
 			AdditionalFrameworks.AddRange(Other.AdditionalFrameworks);
 			PrecompiledHeaderFile = Other.PrecompiledHeaderFile;
diff --git a/Engine/Source/Programs/UnrealBuildTool/ToolChain/ClangToolChain.cs b/Engine/Source/Programs/UnrealBuildTool/ToolChain/ClangToolChain.cs
--- a/Engine/Source/Programs/UnrealBuildTool/ToolChain/ClangToolChain.cs
+++ b/Engine/Source/Programs/UnrealBuildTool/ToolChain/ClangToolChain.cs
@@ -888,6 +888,12 @@ namespace UnrealBuildTool
 		/// <param name="CompileResult"></param>
 		protected virtual void GetCompileArguments_FileType(CppCompileEnvironment CompileEnvironment, FileItem SourceFile, DirectoryReference OutputDir, List<string> Arguments, Action CompileAction, CPPOutput CompileResult)
 		{
+			// Add the additional response files
+			foreach (FileItem AdditionalResponseFile in CompileEnvironment.AdditionalResponseFiles)
+			{
+				Arguments.Add($"@{NormalizeCommandLinePath(AdditionalResponseFile.Location)}");
+			}
+
 			// Add the C++ source file and its included files to the prerequisite item list.
 			CompileAction.PrerequisiteItems.UnionWith(CompileEnvironment.ForceIncludeFiles);
 			CompileAction.PrerequisiteItems.UnionWith(CompileEnvironment.AdditionalPrerequisites);
@@ -1022,6 +1028,27 @@ namespace UnrealBuildTool
 			return NewList;
 		}
 
+		public override CppCompileEnvironment CreateSharedResponseFile(CppCompileEnvironment CompileEnvironment, FileReference OutResponseFile, IActionGraphBuilder Graph)
+		{
+			CppCompileEnvironment NewCompileEnvironment = new CppCompileEnvironment(CompileEnvironment);
+			List<string> Arguments = new List<string>();
+
+			GetCompileArguments_Global(CompileEnvironment, Arguments);
+			NewCompileEnvironment.UserIncludePaths.Clear();
+			NewCompileEnvironment.SystemIncludePaths.Clear();
+
+			Arguments = ExpandResponseFileContents(Arguments);
+
+			FileItem FileItem = FileItem.GetItemByFileReference(OutResponseFile);
+			Graph.CreateIntermediateTextFile(FileItem, Arguments);
+
+			NewCompileEnvironment.AdditionalPrerequisites.Add(FileItem);
+			NewCompileEnvironment.AdditionalResponseFiles.Add(FileItem);
+
+
+			return NewCompileEnvironment;
+		}
+
 		protected virtual Action CompileCPPFile(CppCompileEnvironment CompileEnvironment, FileItem SourceFile, DirectoryReference OutputDir, string ModuleName, IActionGraphBuilder Graph, IReadOnlyCollection<string> GlobalArguments, CPPOutput Result)
 		{
 			Action CompileAction = Graph.CreateAction(ActionType.Compile);
@@ -1144,7 +1171,12 @@ namespace UnrealBuildTool
 		protected override CPPOutput CompileCPPFiles(CppCompileEnvironment CompileEnvironment, List<FileItem> InputFiles, DirectoryReference OutputDir, string ModuleName, IActionGraphBuilder Graph)
 		{
 			List<string> GlobalArguments = new();
-			GetCompileArguments_Global(CompileEnvironment, GlobalArguments);
+			
+			// Maybe not the nicest way of checking but if we don't have any system include paths it means we are not using a shared rsp file
+			if (CompileEnvironment.SystemIncludePaths.Count != 0)
+			{
+				GetCompileArguments_Global(CompileEnvironment, GlobalArguments);
+			}
 
 			// Create a compile action for each source file.
 			CPPOutput Result = new CPPOutput();
diff --git a/Engine/Source/Programs/UnrealBuildTool/ToolChain/IWYUToolChain.cs b/Engine/Source/Programs/UnrealBuildTool/ToolChain/IWYUToolChain.cs
--- a/Engine/Source/Programs/UnrealBuildTool/ToolChain/IWYUToolChain.cs
+++ b/Engine/Source/Programs/UnrealBuildTool/ToolChain/IWYUToolChain.cs
@@ -78,15 +78,11 @@ namespace UnrealBuildTool
 
 			Dictionary<FileItem, List<FileItem>> GenCppLookup = CompileEnvironment.FileInlineGenCPPMap;
 
-			CompileEnvironment = new(CompileEnvironment);
-			CompileEnvironment.Definitions.Add("SUPPRESS_MONOLITHIC_HEADER_WARNINGS=1");
-			CompileEnvironment.Definitions.Add("PLATFORM_COMPILER_IWYU=1");
-
-			// Remove this once c++20 linux compile errors are fixed
-			CompileEnvironment.CppStandard = CppStandardVersion.Cpp17;
-
 			List<string> GlobalArguments = new();
-			GetCompileArguments_Global(CompileEnvironment, GlobalArguments);
+			if (CompileEnvironment.UserIncludePaths.Count != 0)
+			{
+				GetCompileArguments_Global(CompileEnvironment, GlobalArguments);
+			}
 
 			// Need to add these on cmd line.. IWYU is not parsing commands after response file expansion.
 			string CommonCommandLineArgs = " -Xiwyu --mapping_file=" + IWYUMappingFile.FullName +
@@ -159,6 +155,9 @@ namespace UnrealBuildTool
 
 		protected override void GetCompileArguments_Global(CppCompileEnvironment CompileEnvironment, List<string> Arguments)
 		{
+			Arguments.Add(GetPreprocessorDefinitionArgument("SUPPRESS_MONOLITHIC_HEADER_WARNINGS=1"));
+			Arguments.Add(GetPreprocessorDefinitionArgument("PLATFORM_COMPILER_IWYU=1"));
+
 			base.GetCompileArguments_Global(CompileEnvironment, Arguments);
 
 			if (ShouldUseLibcxx(CompileEnvironment.Architecture))
@@ -185,11 +184,6 @@ namespace UnrealBuildTool
 			}
 		}
 
-		protected override void GetCompileArguments_FileType(CppCompileEnvironment CompileEnvironment, FileItem SourceFile, DirectoryReference OutputDir, List<string> Arguments, Action CompileAction, CPPOutput CompileResult)
-		{
-			base.GetCompileArguments_FileType(CompileEnvironment, SourceFile, OutputDir, Arguments, CompileAction, CompileResult);
-		}
-
 		protected override void GetCompileArguments_WarningsAndErrors(CppCompileEnvironment CompileEnvironment, List<string> Arguments)
 		{
 			CompileEnvironment.ShadowVariableWarningLevel = WarningLevel.Off;
@@ -200,11 +194,18 @@ namespace UnrealBuildTool
 			Arguments.Add("-Wno-pragma-once-outside-header");
 		}
 
+		// TODO: REMOVE THIS ONCE c++20 compile fixes have been done
+		protected override void GetCppStandardCompileArgument(CppCompileEnvironment CompileEnvironment, List<string> Arguments)
+		{
+			Arguments.Add("-std=c++17");
+		}
+
 		// Skip ISPC headers
 		public override CPPOutput GenerateISPCHeaders(CppCompileEnvironment CompileEnvironment, List<FileItem> InputFiles, DirectoryReference OutputDir, IActionGraphBuilder Graph)
 		{
 			return new CPPOutput();
 		}
+
 		public override CPPOutput CompileISPCFiles(CppCompileEnvironment CompileEnvironment, List<FileItem> InputFiles, DirectoryReference OutputDir, IActionGraphBuilder Graph)
 		{
 			return new CPPOutput();
diff --git a/Engine/Source/Programs/UnrealBuildTool/ToolChain/UEToolChain.cs b/Engine/Source/Programs/UnrealBuildTool/ToolChain/UEToolChain.cs
--- a/Engine/Source/Programs/UnrealBuildTool/ToolChain/UEToolChain.cs
+++ b/Engine/Source/Programs/UnrealBuildTool/ToolChain/UEToolChain.cs
@@ -152,6 +152,11 @@ namespace UnrealBuildTool
 			return Result.ToArray();
 		}
 
+		public virtual CppCompileEnvironment CreateSharedResponseFile(CppCompileEnvironment CompileEnvironment, FileReference OutResponseFile, IActionGraphBuilder Graph)
+		{
+			return CompileEnvironment;
+		}
+
 		/// <summary>
 		/// Get the name of the response file for the current compile environment and output file
 		/// </summary>
