[UBT]
* Turned off shared response file for apple platforms to fix compile errors. Need to be revisited

diff --git a/Engine/Source/Programs/UnrealBuildTool/ToolChain/AppleToolChain.cs b/Engine/Source/Programs/UnrealBuildTool/ToolChain/AppleToolChain.cs
--- a/Engine/Source/Programs/UnrealBuildTool/ToolChain/AppleToolChain.cs
+++ b/Engine/Source/Programs/UnrealBuildTool/ToolChain/AppleToolChain.cs
@@ -398,6 +398,12 @@ namespace UnrealBuildTool
 			Arguments.Add("-fpascal-strings");
 		}
 
+		public override CppCompileEnvironment CreateSharedResponseFile(CppCompileEnvironment CompileEnvironment, FileReference OutResponseFile, IActionGraphBuilder Graph)
+		{
+			// Temporarily turn of shared response files for apple toolchains
+			return CompileEnvironment;
+		}
+
 		protected string GetDsymutilPath(ILogger Logger, out string ExtraOptions, bool bIsForLTOBuild=false)
 		{
 			FileReference DsymutilLocation = new FileReference("/usr/bin/dsymutil");
