[UBT]
* Disabled shared rsp for pvs toolchain
* Fixed VcToolchain check for shared rsp to fix hololens compile error and c++modules path

diff --git a/Engine/Source/Programs/UnrealBuildTool/Platform/Windows/PVSToolChain.cs b/Engine/Source/Programs/UnrealBuildTool/Platform/Windows/PVSToolChain.cs
--- a/Engine/Source/Programs/UnrealBuildTool/Platform/Windows/PVSToolChain.cs
+++ b/Engine/Source/Programs/UnrealBuildTool/Platform/Windows/PVSToolChain.cs
@@ -667,6 +667,12 @@ namespace UnrealBuildTool
 			return compileAsCCode;
 		}
 
+		public override CppCompileEnvironment CreateSharedResponseFile(CppCompileEnvironment CompileEnvironment, FileReference OutResponseFile, IActionGraphBuilder Graph)
+		{
+			// Temporarily turn of shared response files since current hololens toolchain does not support it
+			return CompileEnvironment;
+		}
+
 		protected override CPPOutput CompileCPPFiles(CppCompileEnvironment CompileEnvironment, List<FileItem> InputFiles, DirectoryReference OutputDir, string ModuleName, IActionGraphBuilder Graph)
 		{
 			if (CompileEnvironment.bDisableStaticAnalysis)
diff --git a/Engine/Source/Programs/UnrealBuildTool/Platform/Windows/VCToolChain.cs b/Engine/Source/Programs/UnrealBuildTool/Platform/Windows/VCToolChain.cs
--- a/Engine/Source/Programs/UnrealBuildTool/Platform/Windows/VCToolChain.cs
+++ b/Engine/Source/Programs/UnrealBuildTool/Platform/Windows/VCToolChain.cs
@@ -1377,8 +1377,8 @@ namespace UnrealBuildTool
 			BaseCompileAction.IncludePaths.AddRange(CompileEnvironment.UserIncludePaths);
 			BaseCompileAction.SystemIncludePaths.AddRange(CompileEnvironment.SystemIncludePaths);
 			
-			// Maybe not the nicest way of checking but if we don't have any system include paths it means we are not using a shared rsp file
-			if (CompileEnvironment.SystemIncludePaths.Count != 0)
+			// Maybe not the nicest way of checking but if we don't have any user include paths it means we have added them to a shared rsp instead
+			if (CompileEnvironment.UserIncludePaths.Count != 0)
 			{
 				BaseCompileAction.SystemIncludePaths.AddRange(EnvVars.IncludePaths);
 			}
