UBT: VCToolChain: MSVC AddressSanitizer workarounds.

Currently necessary to prevent errors like the following:

    C:\Program Files\Microsoft Visual Studio\2022\Professional\VC\Tools\MSVC\14.34.31933\INCLUDE\vector(451): error C2382: '__sanitizer_annotate_contiguous_container': redefinition; different exception specifications
    C:\Program Files\Microsoft Visual Studio\2022\Professional\VC\Tools\MSVC\14.34.31933\INCLUDE\sanitizer\common_interface_defs.h(154): note: see declaration of '__sanitizer_annotate_contiguous_container'
    C:\Program Files\Microsoft Visual Studio\2022\Professional\VC\Tools\MSVC\14.34.31933\INCLUDE\vector(451): error C2733: '__sanitizer_annotate_contiguous_container': you cannot overload a function with 'extern "C"' linkage
    C:\Program Files\Microsoft Visual Studio\2022\Professional\VC\Tools\MSVC\14.34.31933\INCLUDE\sanitizer\common_interface_defs.h(154): note: see declaration of '__sanitizer_annotate_contiguous_container'

diff --git a/Engine/Source/Programs/UnrealBuildTool/Platform/Windows/VCToolChain.cs b/Engine/Source/Programs/UnrealBuildTool/Platform/Windows/VCToolChain.cs
--- a/Engine/Source/Programs/UnrealBuildTool/Platform/Windows/VCToolChain.cs
+++ b/Engine/Source/Programs/UnrealBuildTool/Platform/Windows/VCToolChain.cs
@@ -561,10 +561,17 @@ namespace UnrealBuildTool
 				// detection.
 				AddDefinition(Arguments, "FORCE_ANSI_ALLOCATOR=1");
 
-				// MSVC has no support for __has_feature(address_sanitizer)
 				if (Target.WindowsPlatform.Compiler.IsMSVC())
 				{
+					// MSVC has no support for __has_feature(address_sanitizer)
 					AddDefinition(Arguments, "USING_ADDRESS_SANITISER=1");
+
+					// Re-evalulate the necessity of these workarounds later
+					if (EnvVars.CompilerVersion < new VersionNumber(14, 35))
+					{
+						AddDefinition(Arguments, "_DISABLE_STRING_ANNOTATION=1");
+						AddDefinition(Arguments, "_DISABLE_VECTOR_ANNOTATION=1");
+					}
 				}
 
 				// Currently the ASan headers are not default around. They can be found at this location so lets use this until this is resolved in the toolchain
