UnrealBuildTool: When requesting -SingleFile builds, outdated PrerequisiteActions still need to be built as one of the single files could require their outputs such as generated ispc headers

#rnx

diff --git a/Engine/Source/Programs/UnrealBuildTool/Modes/BuildMode.cs b/Engine/Source/Programs/UnrealBuildTool/Modes/BuildMode.cs
--- a/Engine/Source/Programs/UnrealBuildTool/Modes/BuildMode.cs
+++ b/Engine/Source/Programs/UnrealBuildTool/Modes/BuildMode.cs
@@ -609,11 +609,8 @@ namespace UnrealBuildTool
 					}
 
 					// Plan the actions to execute for the build. For single file compiles, always rebuild the source file regardless of whether it's out of date.
-					if (TargetDescriptor.SpecificFilesToCompile.Count == 0)
-					{
-						ActionGraph.GatherAllOutdatedActions(PrerequisiteActions, History, ActionToOutdatedFlag, CppDependencies, BuildConfiguration.bIgnoreOutdatedImportLibraries, Logger);
-					}
-					else
+					ActionGraph.GatherAllOutdatedActions(PrerequisiteActions, History, ActionToOutdatedFlag, CppDependencies, BuildConfiguration.bIgnoreOutdatedImportLibraries, Logger);
+					if (TargetDescriptor.SpecificFilesToCompile.Count != 0)
 					{
 						HashSet<FileReference> ForceBuildFiles = new HashSet<FileReference>();
 						ForceBuildFiles.UnionWith(TargetDescriptor.SpecificFilesToCompile);
