UnrealBuildTool: Ensure .cpp files found when scanning -SingleFile headers are added to the build list

This can be a lot of files if you request a very common header! For example -SingleFile=Runtime\Core\Public\Containers\Array.h currently would build over 22k UnrealEditor actions

#rnx

diff --git a/Engine/Source/Programs/UnrealBuildTool/Configuration/UEBuildTarget.cs b/Engine/Source/Programs/UnrealBuildTool/Configuration/UEBuildTarget.cs
--- a/Engine/Source/Programs/UnrealBuildTool/Configuration/UEBuildTarget.cs
+++ b/Engine/Source/Programs/UnrealBuildTool/Configuration/UEBuildTarget.cs
@@ -2434,6 +2434,9 @@ namespace UnrealBuildTool
 			}
 
 			// Find the set of binaries to build. If we're compiling only specific files, filter the list of binaries to only include the files we're interested in.
+			HashSet<FileReference> SpecificFilesToCompile = new();
+			SpecificFilesToCompile.UnionWith(TargetDescriptor.SpecificFilesToCompile);
+			SpecificFilesToCompile.UnionWith(TargetDescriptor.OptionalFilesToCompile);
 			if (TargetDescriptor.SpecificFilesToCompile.Count > 0 && !BuildConfiguration.bIgnoreInvalidFiles)
 			{
 				foreach (FileReference SpecificFile in TargetDescriptor.SpecificFilesToCompile)
@@ -2452,7 +2455,7 @@ namespace UnrealBuildTool
 			{
 				foreach (UEBuildBinary Binary in Binaries)
 				{
-					List<FileItem> BinaryOutputItems = Binary.Build(Rules, TargetToolChain, GlobalCompileEnvironment, GlobalLinkEnvironment, TargetDescriptor.SpecificFilesToCompile, WorkingSet, ExeDir, MakefileBuilder, Logger);
+					List<FileItem> BinaryOutputItems = Binary.Build(Rules, TargetToolChain, GlobalCompileEnvironment, GlobalLinkEnvironment, SpecificFilesToCompile.ToList(), WorkingSet, ExeDir, MakefileBuilder, Logger);
 					Makefile.OutputItems.AddRange(BinaryOutputItems);
 				}
 			}
