UnrealBuildTool: Fix additional response files not being quoted

#rnx

diff --git a/Engine/Source/Programs/UnrealBuildTool/Platform/Windows/VCToolChain.cs b/Engine/Source/Programs/UnrealBuildTool/Platform/Windows/VCToolChain.cs
--- a/Engine/Source/Programs/UnrealBuildTool/Platform/Windows/VCToolChain.cs
+++ b/Engine/Source/Programs/UnrealBuildTool/Platform/Windows/VCToolChain.cs
@@ -1375,7 +1375,7 @@ namespace UnrealBuildTool
 			// Add additional response files
 			foreach (FileItem AdditionalRsp in CompileEnvironment.AdditionalResponseFiles)
 			{
-				BaseCompileAction.Arguments.Add($"@{NormalizeCommandLinePath(AdditionalRsp.Location)}");
+				BaseCompileAction.Arguments.Add($"@\"{NormalizeCommandLinePath(AdditionalRsp.Location)}\"");
 			}
 
 			AppendCLArguments_Global(CompileEnvironment, BaseCompileAction.Arguments);
diff --git a/Engine/Source/Programs/UnrealBuildTool/ToolChain/ClangToolChain.cs b/Engine/Source/Programs/UnrealBuildTool/ToolChain/ClangToolChain.cs
--- a/Engine/Source/Programs/UnrealBuildTool/ToolChain/ClangToolChain.cs
+++ b/Engine/Source/Programs/UnrealBuildTool/ToolChain/ClangToolChain.cs
@@ -891,7 +891,7 @@ namespace UnrealBuildTool
 			// Add the additional response files
 			foreach (FileItem AdditionalResponseFile in CompileEnvironment.AdditionalResponseFiles)
 			{
-				Arguments.Add($"@{NormalizeCommandLinePath(AdditionalResponseFile.Location)}");
+				Arguments.Add(GetResponseFileArgument(AdditionalResponseFile));
 			}
 
 			// Add the C++ source file and its included files to the prerequisite item list.
