changed message for circular plugin dependencies to print the cycle per UDN suggestion from Simon Harvey
https://jira.it.epicgames.com/browse/UE-165111

https://p4-swarm.epicgames.net/reviews/23962465

diff --git a/Engine/Source/Programs/UnrealBuildTool/Configuration/UEBuildTarget.cs b/Engine/Source/Programs/UnrealBuildTool/Configuration/UEBuildTarget.cs
--- a/Engine/Source/Programs/UnrealBuildTool/Configuration/UEBuildTarget.cs
+++ b/Engine/Source/Programs/UnrealBuildTool/Configuration/UEBuildTarget.cs
@@ -4162,7 +4162,7 @@ namespace UnrealBuildTool
 							Dependencies.Add(NextInstance);
 							if (NextInstance.Dependencies == null)
 							{
-								throw new BuildException("Found circular dependency from plugin '{0}' onto itself.", NextReference.Name);
+								throw new BuildException($"Found circular plugin dependency. {PluginReferenceChain} -> {NextReference.Name}");
 							}
 							Dependencies.UnionWith(NextInstance.Dependencies);
 						}
