Fix a few edge cases related to UBT project file detection for programs (especially those in the NotForLicensees directory)

* Default.uprojectdir's now contains NotForLicensees as a valid uproject search path
* Since programs uprojects are not found in the same directory as their .target.cs files, add a special exception for uprojects found in Programs dir that assumes any uproject for a program will have a .target.cs with the same name.

diff --git a/Engine/Source/Programs/UnrealBuildTool/System/NativeProjects.cs b/Engine/Source/Programs/UnrealBuildTool/System/NativeProjects.cs
--- a/Engine/Source/Programs/UnrealBuildTool/System/NativeProjects.cs
+++ b/Engine/Source/Programs/UnrealBuildTool/System/NativeProjects.cs
@@ -64,6 +64,13 @@ namespace UnrealBuildTool
 									FindTargetFiles(IntermediateSourceDirectory, TargetNameToProjectFile, ProjectFile);
 								}
 							}
+
+							// Programs are a special case where the .uproject files are separated from the main project source code- in this case,
+							// we guarantee that a project under the Programs dir will always have an associated target file with the same name.
+							if (!TargetNameToProjectFile.ContainsKey(ProjectFile.GetFileNameWithoutAnyExtensions()) && ProjectFile.ContainsName("Programs", 0))
+							{
+								TargetNameToProjectFile.Add(ProjectFile.GetFileNameWithoutAnyExtensions(), ProjectFile);
+							}
 						}
 						CachedTargetNameToProjectFile = TargetNameToProjectFile;
 					}
