Restore the setting of UE_USE_VERSE_PATHS
#rnx

diff --git a/Engine/Source/Programs/UnrealBuildTool/Configuration/UEBuildTarget.cs b/Engine/Source/Programs/UnrealBuildTool/Configuration/UEBuildTarget.cs
--- a/Engine/Source/Programs/UnrealBuildTool/Configuration/UEBuildTarget.cs
+++ b/Engine/Source/Programs/UnrealBuildTool/Configuration/UEBuildTarget.cs
@@ -4512,10 +4512,12 @@ namespace UnrealBuildTool
 			if (Rules.bUseVerse)
 			{
 				GlobalCompileEnvironment.Definitions.Add("WITH_VERSE=1");
+				GlobalCompileEnvironment.Definitions.Add("UE_USE_VERSE_PATHS=1");
 			}
 			else
 			{
 				GlobalCompileEnvironment.Definitions.Add("WITH_VERSE=0");
+				GlobalCompileEnvironment.Definitions.Add("UE_USE_VERSE_PATHS=0");
 			}
 
 			if (Rules.bCompileWithStatsWithoutEngine)
