LLT cleanup and improvements/updates:
- remove dependency on LaunchEngineLoop.cpp which also allows compilation against engine to succeed
- LowLevelTestsRunner no dependency on engine modules required: Launch, Core, Project
- EXPLICIT_TESTS_TARGET definition for self-contained tests: test modules and targets that derive from TestModuleRules/TestTargetRules respectively
- Cleanup Launch dependencies from existing explicit tests
- Cleanup redundant flags from existing explicit tests
- rename Self -> Foundation for in Horde
- additional platform fixes
Default #preflight 63e14d37244dc45a20e29337
All platform/tests LLTs #preflight 63e022f91b44ee7cb1c11d60
#rnx

diff --git a/Engine/Build/LowLevelTests_GenProps.xml b/Engine/Build/LowLevelTests_GenProps.xml
--- a/Engine/Build/LowLevelTests_GenProps.xml
+++ b/Engine/Build/LowLevelTests_GenProps.xml
@@ -1,18 +1,18 @@
 ﻿<?xml version="1.0" encoding="utf-8"?>
 <BuildGraph xmlns:xsi="http://www.w3.org/2001/XMLSchema-instance" xsi:schemaLocation="http://www.epicgames.com/BuildGraph ../../Build/Graph/Schema.xsd" xmlns="http://www.epicgames.com/BuildGraph">
-  <Property Name="TestNames" Value="HeadlessChaos" />
-  <Property Name="SelfDisabled" Value="False" />
+  <Property Name="TestNames" Value="HeadlessChaos;Foundation" />
+  <Property Name="FoundationDisabled" Value="False" />
   <Property Name="HeadlessChaosDisabled" Value="False" />
-  <Property Name="SelfShort" Value="Self" />
+  <Property Name="FoundationShort" Value="Foundation" />
   <Property Name="HeadlessChaosShort" Value="Headless Chaos" />
-  <Property Name="SelfTarget" Value="LowLevelTests" />
+  <Property Name="FoundationTarget" Value="LowLevelTests" />
   <Property Name="HeadlessChaosTarget" Value="HeadlessChaos" />
-  <Property Name="SelfBinariesRelative" Value="Engine\Binaries" />
+  <Property Name="FoundationBinariesRelative" Value="Engine\Binaries" />
   <Property Name="HeadlessChaosBinariesRelative" Value="Engine\Binaries" />
-  <Property Name="SelfReportType" Value="console" />
+  <Property Name="FoundationReportType" Value="console" />
   <Property Name="HeadlessChaosReportType" Value="console" />
-  <Option Name="RunSelfTests" DefaultValue="False" Description="Run Self Tests" />
+  <Option Name="RunFoundationTests" DefaultValue="False" Description="Run Foundation Tests" />
   <Option Name="RunHeadlessChaosTests" DefaultValue="False" Description="Run Headless Chaos Tests" />
-  <Property Name="SelfSupportedPlatforms" Value="Win64;Mac;Linux;Android" />
+  <Property Name="FoundationSupportedPlatforms" Value="Win64;Mac;Linux;Android" />
   <Property Name="HeadlessChaosSupportedPlatforms" Value="Win64" />
 </BuildGraph>
\ No newline at end of file
diff --git a/Engine/Source/Developer/LowLevelTestsRunner/LowLevelTestsRunner.Build.cs b/Engine/Source/Developer/LowLevelTestsRunner/LowLevelTestsRunner.Build.cs
--- a/Engine/Source/Developer/LowLevelTestsRunner/LowLevelTestsRunner.Build.cs
+++ b/Engine/Source/Developer/LowLevelTestsRunner/LowLevelTestsRunner.Build.cs
@@ -16,21 +16,13 @@ namespace UnrealBuildTool.Rules
 			PublicIncludePaths.AddRange(
 				new string[]
 				{
-				"Runtime/Launch/Public",
-				Path.Combine(Target.UEThirdPartySourceDirectory, "Catch2", "v3.0.1", "src")
-				}
-			);
-
-			PrivateIncludePaths.AddRange(
-				new string[] {
-				"Runtime/Launch/Private"
+					Path.Combine(Target.UEThirdPartySourceDirectory, "Catch2", "v3.0.1", "src")
 				}
 			);
 
 			PrivateDependencyModuleNames.AddRange(
 				new string[] {
-				"Core",
-				"Projects"
+					/* Add any extra deps here, should not require any UE modules */
 				}
 			);
 		}
diff --git a/Engine/Source/Developer/LowLevelTestsRunner/Private/TestCommon/CoreUObjectUtilities.cpp b/Engine/Source/Developer/LowLevelTestsRunner/Private/TestCommon/CoreUObjectUtilities.cpp
--- a/Engine/Source/Developer/LowLevelTestsRunner/Private/TestCommon/CoreUObjectUtilities.cpp
+++ b/Engine/Source/Developer/LowLevelTestsRunner/Private/TestCommon/CoreUObjectUtilities.cpp
@@ -11,10 +11,12 @@
 #include "UObject/Package.h"
 #include "UObject/PackageResourceManager.h"
 
+#if WITH_ENGINE
+#include "Framework/Application/SlateApplication.h"
+#endif
+
 void InitCoreUObject()
 {
-	//Initialize the PackageResourceManager, which is needed to load any (non-script) Packages. It is first used in ProcessNewlyLoadedObjects (due to the loading of asset references in Class Default Objects)
-	// It has to be intialized after the AssetRegistryModule; the editor implementations of PackageResourceManager relies on it
 	IPackageResourceManager::Initialize();
 
 	if (!GetTransientPackage())
@@ -25,6 +27,11 @@ void InitCoreUObject()
 
 		FModuleManager::Get().LoadModule(TEXT("CoreUObject"));
 		FCoreDelegates::OnInit.Broadcast();
+
+#if WITH_ENGINE
+		FSlateApplication::InitializeCoreStyle();
+#endif
+
 		ProcessNewlyLoadedUObjects();
 	}
 }
diff --git a/Engine/Source/Developer/LowLevelTestsRunner/Private/TestCommon/CoreUtilities.cpp b/Engine/Source/Developer/LowLevelTestsRunner/Private/TestCommon/CoreUtilities.cpp
--- a/Engine/Source/Developer/LowLevelTestsRunner/Private/TestCommon/CoreUtilities.cpp
+++ b/Engine/Source/Developer/LowLevelTestsRunner/Private/TestCommon/CoreUtilities.cpp
@@ -10,6 +10,10 @@
 
 #include "TestCommon/CoreUtilities.h"
 
+#if WITH_APPLICATION_CORE
+#include "HAL/PlatformApplicationMisc.h"
+#endif
+
 void InitCommandLine()
 {
 }
@@ -18,6 +22,18 @@ void CleanupCommandLine()
 {
 }
 
+void InitGWarn()
+{
+	if (GWarn == nullptr)
+	{
+#if WITH_APPLICATION_CORE
+		GWarn = FPlatformApplicationMisc::GetFeedbackContext();
+#else
+		GWarn = FPlatformOutputDevices::GetFeedbackContext();
+#endif
+	}
+}
+
 void InitIOThreadPool(bool MultiThreaded, int32 StackSize)
 {
 	GIOThreadPool = FQueuedThreadPool::Allocate();
diff --git a/Engine/Source/Developer/LowLevelTestsRunner/Private/TestCommon/EngineUtilities.cpp b/Engine/Source/Developer/LowLevelTestsRunner/Private/TestCommon/EngineUtilities.cpp
--- a/Engine/Source/Developer/LowLevelTestsRunner/Private/TestCommon/EngineUtilities.cpp
+++ b/Engine/Source/Developer/LowLevelTestsRunner/Private/TestCommon/EngineUtilities.cpp
@@ -7,17 +7,28 @@
 #include "CoreMinimal.h"
 #include "DistanceFieldAtlas.h"
 #include "MeshCardRepresentation.h"
+#include "Modules/ModuleManager.h"
 #include "RenderUtils.h"
 #include "ShaderParameterMetadata.h"
 
+#if WITH_COREUOBJECT
+	#include "Internationalization/EnginePackageLocalizationCache.h"
+	#include "Internationalization/PackageLocalizationManager.h"
+	#include "Templates/SharedPointer.h"
+#endif
+
+struct GlobalEngineInitialization
+{
+	GlobalEngineInitialization()
+	{
+		// Any initialization to be triggered before other global objects
+	}
+} GEngineInitializationLowLevelTests;
 
 void InitAsyncQueues()
 {
 	check(!GDistanceFieldAsyncQueue);
 	GDistanceFieldAsyncQueue = new FDistanceFieldAsyncQueue();
-
-	check(!GCardRepresentationAsyncQueue);
-	GCardRepresentationAsyncQueue = new FCardRepresentationAsyncQueue();
 }
 
 void InitRendering()
@@ -35,4 +46,22 @@ void InitRendering()
 		RenderUtilsInit();
 	}
 }
+
+void InitEngine()
+{
+	FModuleManager::Get().LoadModule(TEXT("Engine"));
+	FModuleManager::Get().LoadModule(TEXT("RenderCore"));
+
+#if WITH_COREUOBJECT
+	FPackageLocalizationManager::Get().InitializeFromLazyCallback([](FPackageLocalizationManager& InPackageLocalizationManager)
+	{
+		InPackageLocalizationManager.InitializeFromCache(MakeShareable(new FEnginePackageLocalizationCache()));
+	});
+#endif
+}
+
+void CleanupEngine()
+{
+}
+
 #endif // WITH_ENGINE
diff --git a/Engine/Source/Developer/LowLevelTestsRunner/Private/TestCommon/Initialization.cpp b/Engine/Source/Developer/LowLevelTestsRunner/Private/TestCommon/Initialization.cpp
--- a/Engine/Source/Developer/LowLevelTestsRunner/Private/TestCommon/Initialization.cpp
+++ b/Engine/Source/Developer/LowLevelTestsRunner/Private/TestCommon/Initialization.cpp
@@ -6,7 +6,6 @@
 #include "TestCommon/EngineUtilities.h"
 
 
-
 void InitAllThreadPoolsEditorEx(bool MultiThreaded)
 {
 #if WITH_EDITOR
@@ -30,7 +29,8 @@ void InitAll(bool bAllowLogging, bool bMultithreaded)
 #endif // WITH_ENGINE
 	InitTaskGraph();
 #if WITH_ENGINE
-	InitRendering();
+	InitGWarn();
+	InitEngine();
 #endif // WITH_ENGINE
 #if WITH_EDITOR
 	InitDerivedDataCache();
@@ -45,6 +45,9 @@ void InitAll(bool bAllowLogging, bool bMultithreaded)
 
 void CleanupAll()
 {
+#if WITH_ENGINE
+	CleanupEngine();
+#endif
 #if WITH_COREUOBJECT
 	CleanupCoreUObject();
 #endif
diff --git a/Engine/Source/Developer/LowLevelTestsRunner/Private/TestExterns.cpp b/Engine/Source/Developer/LowLevelTestsRunner/Private/TestExterns.cpp
new file mode 100644
--- /dev/null
+++ b/Engine/Source/Developer/LowLevelTestsRunner/Private/TestExterns.cpp
@@ -0,0 +1,24 @@
+// Copyright Epic Games, Inc. All Rights Reserved.
+
+#include "CoreMinimal.h"
+#include "CoreTypes.h"
+#include "Modules/ModuleManager.h"
+
+TCHAR GInternalProjectName[64] = TEXT("LowLevelTests");
+const TCHAR* GForeignEngineDir = nullptr;
+
+#if !EXPLICIT_TESTS_TARGET
+bool GIsGameAgnosticExe = false;
+#endif
+
+// Typical defined by TargetRules but LowLevelTestRunner is not setup correctly
+// Should revist this in the future
+#ifndef IMPLEMENT_ENCRYPTION_KEY_REGISTRATION
+	#define IMPLEMENT_ENCRYPTION_KEY_REGISTRATION()
+#endif
+#ifndef IMPLEMENT_SIGNING_KEY_REGISTRATION
+	#define IMPLEMENT_SIGNING_KEY_REGISTRATION()
+#endif
+
+// Debug visualizers and new operator overloads
+PER_MODULE_BOILERPLATE
\ No newline at end of file
diff --git a/Engine/Source/Developer/LowLevelTestsRunner/Private/TestMain.cpp b/Engine/Source/Developer/LowLevelTestsRunner/Private/TestMain.cpp
deleted file mode 100644
diff --git a/Engine/Source/Developer/LowLevelTestsRunner/Private/TestRunner.cpp b/Engine/Source/Developer/LowLevelTestsRunner/Private/TestRunner.cpp
--- a/Engine/Source/Developer/LowLevelTestsRunner/Private/TestRunner.cpp
+++ b/Engine/Source/Developer/LowLevelTestsRunner/Private/TestRunner.cpp
@@ -105,7 +105,7 @@ void FTestRunner::ParseCommandLine(TConstArrayView<const ANSICHAR*> Args)
 				ExtraArgs.Append(Arg).AppendChar(' ');
 			}
 		}
-		else if (Arg == ANSITEXTVIEW("--"))
+		else if (Arg == ANSITEXTVIEW("--extra-args"))
 		{
 			bExtraArg = true;
 		}
diff --git a/Engine/Source/Developer/LowLevelTestsRunner/Public/TestCommon/CoreUtilities.h b/Engine/Source/Developer/LowLevelTestsRunner/Public/TestCommon/CoreUtilities.h
--- a/Engine/Source/Developer/LowLevelTestsRunner/Public/TestCommon/CoreUtilities.h
+++ b/Engine/Source/Developer/LowLevelTestsRunner/Public/TestCommon/CoreUtilities.h
@@ -5,6 +5,8 @@
 #include "CoreMinimal.h"
 #include "Async/TaskGraphInterfaces.h"
 
+void InitGWarn();
+
 void InitIOThreadPool(bool MultiThreaded = false, int32 StackSize = 128 * 1024);
 
 void InitThreadPool(bool MultiThreaded = false, int32 StackSize = 128 * 1024);
diff --git a/Engine/Source/Developer/LowLevelTestsRunner/Public/TestCommon/EngineUtilities.h b/Engine/Source/Developer/LowLevelTestsRunner/Public/TestCommon/EngineUtilities.h
--- a/Engine/Source/Developer/LowLevelTestsRunner/Public/TestCommon/EngineUtilities.h
+++ b/Engine/Source/Developer/LowLevelTestsRunner/Public/TestCommon/EngineUtilities.h
@@ -6,5 +6,7 @@
 
 void InitAsyncQueues();
 void InitRendering();
+void InitEngine();
+void CleanupEngine();
 
 #endif // WITH_ENGINE
diff --git a/Engine/Source/Programs/AutomationTool/LowLevelTests/RunLowLevelTests.cs b/Engine/Source/Programs/AutomationTool/LowLevelTests/RunLowLevelTests.cs
--- a/Engine/Source/Programs/AutomationTool/LowLevelTests/RunLowLevelTests.cs
+++ b/Engine/Source/Programs/AutomationTool/LowLevelTests/RunLowLevelTests.cs
@@ -419,6 +419,22 @@ namespace LowLevelTests
 		}
 	}
 
+	/// <summary>
+	/// Platform and test specific extension that provides extra command line arguments.
+	/// </summary>
+	public interface ILowLevelTestsExtension
+	{
+		/// <summary>
+		/// Use this to implement platform and test specific specializations of ILowLevelTestsExtension.
+		/// </summary>
+		bool IsSupported(UnrealTargetPlatform InPlatform, string InTestApp);
+
+		/// <summary>
+		/// Return extra command line arguments specific to a platform and/or test.
+		/// </summary>
+		string ExtraCommandLine(UnrealTargetPlatform InPlatform, string InTestApp, string InBuildPath);
+	}
+
 	/// <summary>
 	/// Platform-specific reporting utility for defining Catch2 report path and means to copy it to a saved storage either from its local directory or from a target device.
 	/// </summary>
@@ -555,6 +571,9 @@ namespace LowLevelTests
 
 		private ILowLevelTestsBuildFactory LowLevelTestsBuildFactory;
 		private ILowLevelTestsReporting LowLevelTestsReporting;
+#nullable enable
+		private ILowLevelTestsExtension? LowLevelTestsExtension;
+#nullable disable
 
 		public UnrealTargetPlatform Platform { get; protected set; }
 		public UnrealTargetConfiguration Configuration { get; protected set; }
@@ -583,6 +602,10 @@ namespace LowLevelTests
 			LowLevelTestsReporting = Gauntlet.Utils.InterfaceHelpers.FindImplementations<ILowLevelTestsReporting>(true)
 				.Where(B => B.CanSupportPlatform(InTargetPlatform))
 				.First();
+
+			LowLevelTestsExtension = Gauntlet.Utils.InterfaceHelpers.FindImplementations<ILowLevelTestsExtension>(true)
+				.Where(B => B.IsSupported(InTargetPlatform, InTestApp))
+				.FirstOrDefault();
 		}
 
 		public UnrealAppConfig GetUnrealAppConfig(string InTags, int InSleep, bool InAttachToDebugger, string InReportType)
@@ -617,6 +640,14 @@ namespace LowLevelTests
 				{
 					CachedConfig.CommandLineParams.AddRawCommandline("--attach-to-debugger");
 				}
+				if (LowLevelTestsExtension != null)
+				{
+					string ExtraCmd = LowLevelTestsExtension.ExtraCommandLine(Platform, TestApp, BuildPath);
+					if (!string.IsNullOrEmpty(ExtraCmd))
+					{
+						CachedConfig.CommandLineParams.AddRawCommandline(string.Format("--extra-args {0}", ExtraCmd));
+					}
+				}
 			}
 			return CachedConfig;
 		}
diff --git a/Engine/Source/Programs/ChaosUserDataPTTests/ChaosUserDataPTTests.Target.cs b/Engine/Source/Programs/ChaosUserDataPTTests/ChaosUserDataPTTests.Target.cs
--- a/Engine/Source/Programs/ChaosUserDataPTTests/ChaosUserDataPTTests.Target.cs
+++ b/Engine/Source/Programs/ChaosUserDataPTTests/ChaosUserDataPTTests.Target.cs
@@ -7,8 +7,5 @@ public class ChaosUserDataPTTestsTarget : TestTargetRules
 {
 	public ChaosUserDataPTTestsTarget(TargetInfo Target) : base(Target)
 	{
-		IncludeOrderVersion = EngineIncludeOrderVersion.Latest;
-		bCompileAgainstCoreUObject = true;
-		bCompileAgainstApplicationCore = true;
 	}
 }
diff --git a/Engine/Source/Programs/ChaosUserDataPTTests/Private/ChaosUserDataPTTests.cpp b/Engine/Source/Programs/ChaosUserDataPTTests/Private/ChaosUserDataPTTests.cpp
--- a/Engine/Source/Programs/ChaosUserDataPTTests/Private/ChaosUserDataPTTests.cpp
+++ b/Engine/Source/Programs/ChaosUserDataPTTests/Private/ChaosUserDataPTTests.cpp
@@ -6,7 +6,6 @@
 #include "PBDRigidsSolver.h"
 #include "ChaosSolversModule.h"
 #include "Chaos/Box.h"
-#include "TestCommon/Initialization.h"
 
 class FTestUserData : public Chaos::TUserDataManagerPT<FString> { };
 
diff --git a/Engine/Source/Programs/LowLevelTests/LowLevelTests.Build.cs b/Engine/Source/Programs/LowLevelTests/LowLevelTests.Build.cs
--- a/Engine/Source/Programs/LowLevelTests/LowLevelTests.Build.cs
+++ b/Engine/Source/Programs/LowLevelTests/LowLevelTests.Build.cs
@@ -16,10 +16,11 @@ public class LowLevelTests : TestModuleRules
 		if (Target.bBuildWithEditorOnlyData)
 		{
 			PrivateDependencyModuleNames.AddRange(
-				new string[] { "DesktopPlatform" }
-				);
+				new string[] {
+					"DesktopPlatform"
+				});
 		}
 
-		UpdateBuildGraphPropertiesFile(new Metadata("Self"), false);
+		UpdateBuildGraphPropertiesFile(new Metadata("Foundation"), false);
 	}
 }
\ No newline at end of file
diff --git a/Engine/Source/Programs/LowLevelTests/LowLevelTests.Target.cs b/Engine/Source/Programs/LowLevelTests/LowLevelTests.Target.cs
--- a/Engine/Source/Programs/LowLevelTests/LowLevelTests.Target.cs
+++ b/Engine/Source/Programs/LowLevelTests/LowLevelTests.Target.cs
@@ -8,8 +8,8 @@ public class LowLevelTestsTarget : TestTargetRules
 {
 	public LowLevelTestsTarget(TargetInfo Target) : base(Target)
 	{
+		// Collects all tests decorated with #if WITH_LOW_LEVELTESTS from dependencies
 		bWithLowLevelTestsOverride = true;
-		IncludeOrderVersion = EngineIncludeOrderVersion.Latest;
 
 		if (Target.Platform == UnrealTargetPlatform.Android)
 		{
diff --git a/Engine/Source/Programs/UnrealBuildTool/Configuration/ModuleRules.cs b/Engine/Source/Programs/UnrealBuildTool/Configuration/ModuleRules.cs
--- a/Engine/Source/Programs/UnrealBuildTool/Configuration/ModuleRules.cs
+++ b/Engine/Source/Programs/UnrealBuildTool/Configuration/ModuleRules.cs
@@ -1284,6 +1284,10 @@ namespace UnrealBuildTool
 		{
 			get
 			{
+				if (IsTestModule)
+				{
+					return Directory.FullName;
+				}
 				return Path.Combine(Directory.FullName, "Tests");
 			}
 		}
@@ -1519,22 +1523,19 @@ namespace UnrealBuildTool
 				TestTargetRules.LowLevelTestsRunnerModule = this;
 			}
 
-			// If one of these modules is in the dependency graph, we must enable their appropriate global definitions
-			if (Name == "UnrealEd")
+			if (Name == "Engine" && !TestTargetRules.bNeverCompileAgainstEngine)
 			{
-				// TODO: more support code required...
-				TestTargetRules.bTestsRequireEditor = false;
+				TestTargetRules.bTestsRequireEngine = true;
 			}
-			if (Name == "Engine")
+			if (Name == "UnrealEd" && !TestTargetRules.bNeverCompileAgainstEditor)
 			{
-				// TODO: more support code required...
-				TestTargetRules.bTestsRequireEngine = false;
+				TestTargetRules.bTestsRequireEditor = true;
 			}
-			if (Name == "ApplicationCore")
+			if (Name == "ApplicationCore" && !TestTargetRules.bNeverCompileAgainstApplicationCore)
 			{
 				TestTargetRules.bTestsRequireApplicationCore = true;
 			}
-			if (Name == "CoreUObject")
+			if (Name == "CoreUObject" && !TestTargetRules.bNeverCompileAgainstCoreUObject)
 			{
 				TestTargetRules.bTestsRequireCoreUObject = true;
 			}
diff --git a/Engine/Source/Programs/UnrealBuildTool/Configuration/TargetRules.cs b/Engine/Source/Programs/UnrealBuildTool/Configuration/TargetRules.cs
--- a/Engine/Source/Programs/UnrealBuildTool/Configuration/TargetRules.cs
+++ b/Engine/Source/Programs/UnrealBuildTool/Configuration/TargetRules.cs
@@ -465,12 +465,8 @@ namespace UnrealBuildTool
 		/// </summary>
 		public bool ExplicitTestsTarget
 		{
-			get { return bExplicitTestsTargetOverride; }
+			get { return GetType().IsSubclassOf(typeof(TestTargetRules)); }
 		}
-		/// <summary>
-		/// This flag is automatically set when classes inherit from TestTargetRules.
-		/// </summary>
-		protected bool bExplicitTestsTargetOverride;
 
 		/// <summary>
 		/// Controls the value of WITH_LOW_LEVEL_TESTS that dictates whether module-specific low level tests are compiled in or not.
@@ -3857,16 +3853,27 @@ namespace UnrealBuildTool
 			get { return "../Binaries/ThirdParty/"; }
 		}
 
+		/// <summary>
+		/// Whether this is a low level tests target.
+		/// </summary>
 		public bool IsTestTarget
 		{
 			get { return Inner.IsTestTarget; }
 		}
 
+		/// <summary>
+		/// Whether this is a test target explicitly defined.
+		/// Explicitley defined test targets always inherit from TestTargetRules and define their own tests.
+		/// Implicit test targets are created from existing targets when building with -Mode=Test and they include tests from all dependencies.
+		/// </summary>
 		public bool ExplicitTestsTarget
 		{
 			get { return Inner.ExplicitTestsTarget; }
 		}
 
+		/// <summary>
+		/// Controls the value of WITH_LOW_LEVEL_TESTS that dictates whether module-specific low level tests are compiled in or not.
+		/// </summary>
 		public bool WithLowLevelTests
 		{
 			get { return Inner.WithLowLevelTests; }
diff --git a/Engine/Source/Programs/UnrealBuildTool/Configuration/TestModuleRules.cs b/Engine/Source/Programs/UnrealBuildTool/Configuration/TestModuleRules.cs
--- a/Engine/Source/Programs/UnrealBuildTool/Configuration/TestModuleRules.cs
+++ b/Engine/Source/Programs/UnrealBuildTool/Configuration/TestModuleRules.cs
@@ -70,6 +70,7 @@ namespace UnrealBuildTool
 			Context = TestedModule.Context;
 
 			PrivateDependencyModuleNames.AddRange(TestedModule.PrivateDependencyModuleNames);
+			PublicDependencyModuleNames.AddRange(TestedModule.PublicDependencyModuleNames);
 
 			// Tests can refer to tested module's Public and Private paths
 			string ModulePublicDir = Path.Combine(TestedModule.ModuleDirectory, "Public");
diff --git a/Engine/Source/Programs/UnrealBuildTool/Configuration/TestTargetRules.cs b/Engine/Source/Programs/UnrealBuildTool/Configuration/TestTargetRules.cs
--- a/Engine/Source/Programs/UnrealBuildTool/Configuration/TestTargetRules.cs
+++ b/Engine/Source/Programs/UnrealBuildTool/Configuration/TestTargetRules.cs
@@ -46,7 +46,7 @@ namespace UnrealBuildTool
 
 		/// <summary>
 		/// Test target override for bCompileAgainstApplicationCore.
-		/// It is set to true if there is any reference to ApplicationCore in the reference chain.
+		/// It is set to true if there is any reference to ApplicationCore in the dependency graph.
 		/// </summary>
 		public override bool bCompileAgainstApplicationCore
 		{
@@ -54,9 +54,14 @@ namespace UnrealBuildTool
 			set { bTestsRequireApplicationCore = value; }
 		}
 
+		/// <summary>
+		/// If set to true, it will not compile against ApplicationCore even if "ApplicationCore" is in the dependency graph.
+		/// </summary>
+		public static bool bNeverCompileAgainstApplicationCore = false;
+
 		/// <summary>
 		/// Test target override for bCompileAgainstCoreUObject.
-		/// It is set to true if there is any reference to CoreUObject in the reference chain.
+		/// It is set to true if there is any reference to CoreUObject in the dependency graph.
 		/// </summary>
 		public override bool bCompileAgainstCoreUObject
 		{
@@ -64,9 +69,14 @@ namespace UnrealBuildTool
 			set { bTestsRequireCoreUObject = value; }
 		}
 
+		/// <summary>
+		/// If set to true, it will not compile against CoreUObject even if "CoreUObject" is in the dependency graph.
+		/// </summary>
+		public static bool bNeverCompileAgainstCoreUObject = false;
+
 		/// <summary>
 		/// Test target override for bCompileAgainstEngine.
-		/// It is set to true if there is any reference to Engine in the reference chain.
+		/// It is set to true if there is any reference to Engine in the dependency graph.
 		/// </summary>
 		public override bool bCompileAgainstEngine
 		{
@@ -74,9 +84,15 @@ namespace UnrealBuildTool
 			set { bTestsRequireEngine = value; }
 		}
 
+		/// <summary>
+		/// If set to true, it will not compile against engine even if "Engine" is in the dependency graph.
+		/// Note: Currently set to true because it requires cooked Engine assets.
+		/// </summary>
+		public static bool bNeverCompileAgainstEngine = true;
+
 		/// <summary>
 		/// Test target override for bCompileAgainstEditor.
-		/// It is set to true if there is any reference to UnrealEd in the reference chain.
+		/// It is set to true if there is any reference to UnrealEd in the dependency graph.
 		/// </summary>
 		public override bool bCompileAgainstEditor
 		{
@@ -84,6 +100,11 @@ namespace UnrealBuildTool
 			set { bTestsRequireEditor = value; }
 		}
 
+		/// <summary>
+		/// If set to true, it will not compile against editor even if "UnrealEd" is in the dependency graph.
+		/// </summary>
+		public static bool bNeverCompileAgainstEditor = false;
+
 		/// <summary>
 		/// Constructor for TestTargetRules as own target.
 		/// </summary>
@@ -92,12 +113,9 @@ namespace UnrealBuildTool
 		{
 			SetupCommonProperties(Target);
 
-			bExplicitTestsTargetOverride = this.GetType() != typeof(TestTargetRules);
-
-			ExeBinariesSubFolder = LaunchModuleName = Name + (bExplicitTestsTargetOverride ? string.Empty : "Tests");
-			IncludeOrderVersion = EngineIncludeOrderVersion.Latest;
+			ExeBinariesSubFolder = LaunchModuleName = Name + (ExplicitTestsTarget ? string.Empty : "Tests");
 
-			if (bExplicitTestsTargetOverride)
+			if (ExplicitTestsTarget)
 			{
 				bBuildInSolutionByDefault = true;
 				SolutionDirectory = "Programs/LowLevelTests";
@@ -162,6 +180,8 @@ namespace UnrealBuildTool
 
 			this.TestedTarget = TestedTarget;
 
+			TargetFiles = TestedTarget.TargetFiles;
+
 			ExeBinariesSubFolder = Name = TestedTarget.Name + "Tests";
 			TargetSourceFile = File = TestedTarget.File;
 			if (TestedTarget.LaunchModuleName != null)
@@ -178,6 +198,8 @@ namespace UnrealBuildTool
 
 		private void SetupCommonProperties(TargetInfo Target)
 		{
+			IncludeOrderVersion = EngineIncludeOrderVersion.Latest;
+
 			bIsTestTargetOverride = true;
 
 			VSTestRunSettingsFile = FileReference.Combine(Unreal.EngineDirectory, "Source", "Programs", "LowLevelTests", "vstest.runsettings");
@@ -203,18 +225,27 @@ namespace UnrealBuildTool
 			// No need for shaders by default
 			bForceBuildShaderFormats = false;
 
-			// Do not link against the engine, no Chromium Embedded Framework etc.
+			// Do not compile against the engine, editor etc
 			bCompileAgainstEngine = false;
-			bCompileCEF3 = false;
+			bCompileAgainstEditor = false;
 			bCompileAgainstCoreUObject = false;
 			bCompileAgainstApplicationCore = false;
+			bCompileCEF3 = false;
+
+			// No mixing with Functional Test framework
+			bForceDisableAutomationTests = true;
+
+			bDebugBuildsActuallyUseDebugCRT = true;
+
+			// Allow logging in shipping
 			bUseLoggingInShipping = true;
 
 			// Allow exception handling
 			bForceEnableExceptions = true;
 
-			bool bDebugOrDevelopment = Target.Configuration == UnrealTargetConfiguration.Debug || Target.Configuration == UnrealTargetConfiguration.Development;
-			bBuildWithEditorOnlyData = Target.Platform.IsInGroup(UnrealPlatformGroup.Desktop) && bDebugOrDevelopment;
+			bBuildWithEditorOnlyData = false;
+			bBuildRequiresCookedData = true;
+			bBuildDeveloperTools = false;
 
 			// Disable malloc profiling in tests
 			bUseMallocProfiler = false;
@@ -225,6 +256,11 @@ namespace UnrealBuildTool
 				bDebugBuildsActuallyUseDebugCRT = true;
 			}
 
+			if (!ExplicitTestsTarget && TestedTarget != null)
+			{
+				GlobalDefinitions.AddRange(TestedTarget.GlobalDefinitions);
+			}
+
 			GlobalDefinitions.Add("STATS=0");
 			GlobalDefinitions.Add("TEST_FOR_VALID_FILE_SYSTEM_MEMORY=0");
 
diff --git a/Engine/Source/Programs/UnrealBuildTool/Configuration/UEBuildTarget.cs b/Engine/Source/Programs/UnrealBuildTool/Configuration/UEBuildTarget.cs
--- a/Engine/Source/Programs/UnrealBuildTool/Configuration/UEBuildTarget.cs
+++ b/Engine/Source/Programs/UnrealBuildTool/Configuration/UEBuildTarget.cs
@@ -4477,7 +4477,11 @@ namespace UnrealBuildTool
 			GlobalCompileEnvironment.Definitions.Add("WITH_DEV_AUTOMATION_TESTS=" + (bCompileDevTests ? "1" : "0"));
 			GlobalCompileEnvironment.Definitions.Add("WITH_PERF_AUTOMATION_TESTS=" + (bCompilePerfTests ? "1" : "0"));
 
+			// Target contains Low Level Tests
 			GlobalCompileEnvironment.Definitions.Add(String.Format("WITH_LOW_LEVEL_TESTS={0}", Rules.WithLowLevelTests ? "1" : "0"));
+			// Target derived from TestTargetRules containing Low Level Tests, i.e. an explicit target
+			GlobalCompileEnvironment.Definitions.Add(String.Format("EXPLICIT_TESTS_TARGET={0}", Rules.ExplicitTestsTarget ? "1" : "0"));
+			// Target contains Low Level Tests or Functional Tests
 			GlobalCompileEnvironment.Definitions.Add(String.Format("WITH_TESTS={0}", Rules.WithLowLevelTests || bCompileDevTests || bCompilePerfTests ? "1" : "0"));
 
 			GlobalCompileEnvironment.Definitions.Add("UNICODE");
diff --git a/Engine/Source/Runtime/Core/Tests/TestGroupEvents.cpp b/Engine/Source/Runtime/Core/Tests/TestGroupEvents.cpp
new file mode 100644
--- /dev/null
+++ b/Engine/Source/Runtime/Core/Tests/TestGroupEvents.cpp
@@ -0,0 +1,24 @@
+// Copyright Epic Games, Inc. All Rights Reserved.
+
+#if !EXPLICIT_TESTS_TARGET && WITH_LOW_LEVEL_TESTS
+#include "TestCommon/Initialization.h"
+
+#include <catch2/catch_test_macros.hpp>
+
+struct StaticInit {
+	StaticInit()
+	{
+		FCommandLine::Set(TEXT(""));
+	}
+} GStaticInit;
+
+GROUP_BEFORE_GLOBAL(Catch::DefaultGroup)
+{
+	InitAll(true, true);
+}
+
+GROUP_AFTER_GLOBAL(Catch::DefaultGroup)
+{
+	CleanupAll();
+}
+#endif // !EXPLICIT_TESTS_TARGET && WITH_LOW_LEVEL_TESTS
