Update minimum macOS requirements to Xcode14.1 & macOS12.5

diff --git a/Engine/Source/Programs/UnrealBuildTool/Platform/IOS/IOSPlatformSDK.Versions.cs b/Engine/Source/Programs/UnrealBuildTool/Platform/IOS/IOSPlatformSDK.Versions.cs
--- a/Engine/Source/Programs/UnrealBuildTool/Platform/IOS/IOSPlatformSDK.Versions.cs
+++ b/Engine/Source/Programs/UnrealBuildTool/Platform/IOS/IOSPlatformSDK.Versions.cs
@@ -28,7 +28,8 @@ namespace UnrealBuildTool
 			{
 				// return new Version("1200.00.32.29");	// UE 5.0 (Xcode 12.4)
 				// return new Version("1300.00.29.30"); // (old) UE 5.1 sets Xcode13.2.1 as minimum supported
-				return new Version("1316.0.21.2"); // UE 5.1 sets Xcode13.4 as minimum supported
+				// return new Version("1316.0.21.2"); // UE 5.1 sets Xcode13.4 as minimum supported
+				return new Version("1400.0.29.202"); // UE 5.2 sets Xcode14.1 as minimum supported
 			}
 		}
 	}
diff --git a/Engine/Source/Programs/UnrealBuildTool/Platform/Mac/ApplePlatformSDK.Versions.cs b/Engine/Source/Programs/UnrealBuildTool/Platform/Mac/ApplePlatformSDK.Versions.cs
--- a/Engine/Source/Programs/UnrealBuildTool/Platform/Mac/ApplePlatformSDK.Versions.cs
+++ b/Engine/Source/Programs/UnrealBuildTool/Platform/Mac/ApplePlatformSDK.Versions.cs
@@ -14,21 +14,22 @@ namespace UnrealBuildTool
 	{
 		public override string GetMainVersion()
 		{
-			return "13.3";
+			// Xcode prefered version?
+			return "14.1";
 		}
 
 		protected override void GetValidVersionRange(out string MinVersion, out string MaxVersion)
 		{
 			if (RuntimePlatform.IsMac)
 			{
-				MinVersion = "13.4.0";
-				MaxVersion = "14.9.9";
+				MinVersion = "14.1.0";
+				MaxVersion = "15.9.9";
 			}
 			else
 			{
 				// @todo turnkey: these are MobileDevice .dll versions in Windows - to get the iTunes app version (12.3.4.1 etc) would need to hunt down the .exe
 				MinVersion = "1100.0.0.0";
-				MaxVersion = "1499.0";
+				MaxVersion = "8999.0";
 			}
 		}
 
diff --git a/Engine/Source/Programs/UnrealBuildTool/Platform/Mac/MacPlatformSDK.Versions.cs b/Engine/Source/Programs/UnrealBuildTool/Platform/Mac/MacPlatformSDK.Versions.cs
--- a/Engine/Source/Programs/UnrealBuildTool/Platform/Mac/MacPlatformSDK.Versions.cs
+++ b/Engine/Source/Programs/UnrealBuildTool/Platform/Mac/MacPlatformSDK.Versions.cs
@@ -14,7 +14,7 @@ namespace UnrealBuildTool
 
 		protected override void GetValidSoftwareVersionRange(out string MinVersion, out string? MaxVersion)
 		{
-			MinVersion = "12.0.0";
+			MinVersion = "12.5.0";		// macOS Monterey 12.5 (Released 2022/07/22)
 			MaxVersion = null;
 		}
 	}
