Fixed issue where UBT fails to detect C# UBT plugins in the game folder.

diff --git a/Engine/Source/Programs/UnrealBuildTool/ProjectFiles/ProjectFileGenerator.cs b/Engine/Source/Programs/UnrealBuildTool/ProjectFiles/ProjectFileGenerator.cs
--- a/Engine/Source/Programs/UnrealBuildTool/ProjectFiles/ProjectFileGenerator.cs
+++ b/Engine/Source/Programs/UnrealBuildTool/ProjectFiles/ProjectFileGenerator.cs
@@ -437,9 +437,11 @@ namespace UnrealBuildTool
 		void AddRulesModules(Rules.RulesFileType RulesFileType, string ProgramSubDirectory, List<ProjectFile> AddedProjectFiles, 
 			List<FileReference> UnrealProjectFiles, PrimaryProjectFolder RootFolder, PrimaryProjectFolder ProgramsFolder, ILogger Logger)
 		{
+			List<DirectoryReference> GameFolders = new List<DirectoryReference>();
 			List<DirectoryReference> BuildFolders = new List<DirectoryReference>();
 			foreach (FileReference UnrealProjectFile in UnrealProjectFiles)
 			{
+				GameFolders.Add(UnrealProjectFile.Directory);
 				DirectoryReference GameBuildFolder = DirectoryReference.Combine(UnrealProjectFile.Directory, "Build");
 				if (DirectoryReference.Exists(GameBuildFolder))
 				{
@@ -451,7 +453,7 @@ namespace UnrealBuildTool
 			DirectoryReference SamplesDirectory = DirectoryReference.Combine(Unreal.RootDirectory, "Samples");
 
 			// Find all the modules .csproj files to add
-			List<FileReference> ModuleFiles = Rules.FindAllRulesSourceFiles(RulesFileType, null, ForeignPlugins: null, AdditionalSearchPaths: BuildFolders);
+			List<FileReference> ModuleFiles = Rules.FindAllRulesSourceFiles(RulesFileType, GameFolders, ForeignPlugins: null, AdditionalSearchPaths: BuildFolders);
 			foreach (FileReference ProjectFile in ModuleFiles)
 			{
 				if (FileReference.Exists(ProjectFile))
