UnrealBuildTool: Ensure ForceIncludeFiles are added to VCCompileAction.AdditionalPrerequisiteItems

#rnx

diff --git a/Engine/Source/Programs/UnrealBuildTool/Platform/Windows/VCToolChain.cs b/Engine/Source/Programs/UnrealBuildTool/Platform/Windows/VCToolChain.cs
--- a/Engine/Source/Programs/UnrealBuildTool/Platform/Windows/VCToolChain.cs
+++ b/Engine/Source/Programs/UnrealBuildTool/Platform/Windows/VCToolChain.cs
@@ -1548,6 +1548,7 @@ namespace UnrealBuildTool
 					CompileAction.AdditionalPrerequisiteItems.AddRange(InlinedFiles);
 				}
 
+				CompileAction.AdditionalPrerequisiteItems.AddRange(CompileEnvironment.ForceIncludeFiles);
 				CompileAction.AdditionalPrerequisiteItems.AddRange(CompileEnvironment.AdditionalPrerequisites);
 
 				if (!String.IsNullOrEmpty(CompileEnvironment.AdditionalArguments))
