UnrealBuildTool: Use a force include define header for .ixx compiles

#rnx

diff --git a/Engine/Source/Programs/UnrealBuildTool/Configuration/UEBuildModuleCPP.cs b/Engine/Source/Programs/UnrealBuildTool/Configuration/UEBuildModuleCPP.cs
--- a/Engine/Source/Programs/UnrealBuildTool/Configuration/UEBuildModuleCPP.cs
+++ b/Engine/Source/Programs/UnrealBuildTool/Configuration/UEBuildModuleCPP.cs
@@ -445,7 +445,13 @@ namespace UnrealBuildTool
 			// Compile any module interfaces
 			if (InputFiles.IXXFiles.Count > 0 && Target.bEnableCppModules)
 			{
-				CPPOutput ModuleOutput = ToolChain.CompileAllCPPFiles(CompileEnvironment, InputFiles.IXXFiles, IntermediateDirectory, Name, Graph);
+				CppCompileEnvironment IxxCompileEnvironment = CompileEnvironment;
+
+				// Write all the definitions to a separate file for the ixx compile
+				CreateHeaderForDefinitions(IxxCompileEnvironment, IntermediateDirectory, "ixx", Graph);
+
+				CPPOutput ModuleOutput = ToolChain.CompileAllCPPFiles(IxxCompileEnvironment, InputFiles.IXXFiles, IntermediateDirectory, Name, Graph);
+
 				LinkInputFiles.AddRange(ModuleOutput.ObjectFiles);
 				CompileEnvironment.AdditionalPrerequisites.AddRange(ModuleOutput.CompiledModuleInterfaces);
 			}
