UnrealBuildTool: Platform validation is not required when calling QueryTargetMode and is causing an exception on windows if no msvc toolchain is available

#rnx

diff --git a/Engine/Source/Programs/UnrealBuildTool/Modes/QueryTargetsMode.cs b/Engine/Source/Programs/UnrealBuildTool/Modes/QueryTargetsMode.cs
--- a/Engine/Source/Programs/UnrealBuildTool/Modes/QueryTargetsMode.cs
+++ b/Engine/Source/Programs/UnrealBuildTool/Modes/QueryTargetsMode.cs
@@ -125,7 +125,7 @@ namespace UnrealBuildTool
 					try
 					{
 						UnrealArchitectures Architectures = UnrealArchitectureConfig.ForPlatform(BuildHostPlatform.Current.Platform).ActiveArchitectures(ProjectFile, TargetName);
-						TargetRules = Assembly.CreateTargetRules(TargetName, BuildHostPlatform.Current.Platform, UnrealTargetConfiguration.Development, Architectures, ProjectFile, Arguments, Logger);
+						TargetRules = Assembly.CreateTargetRules(TargetName, BuildHostPlatform.Current.Platform, UnrealTargetConfiguration.Development, Architectures, ProjectFile, Arguments, Logger, bSkipValidation: true);
 					}
 					catch (Exception Ex)
 					{
diff --git a/Engine/Source/Programs/UnrealBuildTool/System/RulesAssembly.cs b/Engine/Source/Programs/UnrealBuildTool/System/RulesAssembly.cs
--- a/Engine/Source/Programs/UnrealBuildTool/System/RulesAssembly.cs
+++ b/Engine/Source/Programs/UnrealBuildTool/System/RulesAssembly.cs
@@ -566,8 +566,9 @@ namespace UnrealBuildTool
 		/// <param name="TargetInfo">Target configuration information to pass to the constructor</param>
 		/// <param name="Logger">Logger for output</param>
 		/// <param name="IsTestTarget">If building a low level tests target</param>
+		/// <param name="bSkipValidation">If validation should be skipped (QueryTargetMode)</param>
 		/// <returns>Instance of the corresponding TargetRules or null if requested type name does not exist</returns>
-		protected TargetRules? CreateTargetRulesInstance(string TypeName, TargetInfo TargetInfo, ILogger Logger, bool IsTestTarget = false)
+		protected TargetRules? CreateTargetRulesInstance(string TypeName, TargetInfo TargetInfo, ILogger Logger, bool IsTestTarget = false, bool bSkipValidation = false)
 		{
 			// The build module must define a type named '<TargetName>Target' that derives from our 'TargetRules' type.  
 			Type? BaseRulesType = CompiledAssembly?.GetType(TypeName);
@@ -697,8 +698,11 @@ namespace UnrealBuildTool
 			}
 
 			// Allow the platform to finalize the settings
-			UEBuildPlatform Platform = UEBuildPlatform.GetBuildPlatform(Rules.Platform);
-			Platform.ValidateTarget(Rules);
+			if (!bSkipValidation)
+			{
+				UEBuildPlatform Platform = UEBuildPlatform.GetBuildPlatform(Rules.Platform);
+				Platform.ValidateTarget(Rules);
+			}
 
 			// Some platforms may *require* monolithic compilation...
 			if (Rules.LinkType != TargetLinkType.Monolithic && UEBuildPlatform.PlatformRequiresMonolithicBuilds(Rules.Platform, Rules.Configuration))
@@ -725,8 +729,9 @@ namespace UnrealBuildTool
 		/// <param name="Arguments">Command line arguments for this target</param>
 		/// <param name="Logger"></param>
 		/// <param name="IsTestTarget">If building a low level test target</param>
+		/// <param name="bSkipValidation">If validation should be skipped (QueryTargetMode)</param>
 		/// <returns>The build target rules for the specified target</returns>
-		public TargetRules CreateTargetRules(string TargetName, UnrealTargetPlatform Platform, UnrealTargetConfiguration Configuration, UnrealArchitectures? Architectures, FileReference? ProjectFile, CommandLineArguments? Arguments, ILogger Logger, bool IsTestTarget = false)
+		public TargetRules CreateTargetRules(string TargetName, UnrealTargetPlatform Platform, UnrealTargetConfiguration Configuration, UnrealArchitectures? Architectures, FileReference? ProjectFile, CommandLineArguments? Arguments, ILogger Logger, bool IsTestTarget = false, bool bSkipValidation = false)
 		{
 			if (IsTestTarget)
 			{
@@ -762,7 +767,7 @@ namespace UnrealBuildTool
 				}
 				else
 				{
-					return Parent.CreateTargetRules(TargetName, Platform, Configuration, Architectures, ProjectFile, Arguments, Logger, IsTestTarget);
+					return Parent.CreateTargetRules(TargetName, Platform, Configuration, Architectures, ProjectFile, Arguments, Logger, IsTestTarget, bSkipValidation);
 				}
 			}
 
@@ -770,7 +775,7 @@ namespace UnrealBuildTool
 			string TargetTypeName = TargetName + "Target";
 
 			// The build module must define a type named '<TargetName>Target' that derives from our 'TargetRules' type.  
-			TargetRules? TargetRules = CreateTargetRulesInstance(TargetTypeName, new TargetInfo(TargetName, Platform, Configuration, Architectures, ProjectFile, Arguments), Logger, IsTestTarget);
+			TargetRules? TargetRules = CreateTargetRulesInstance(TargetTypeName, new TargetInfo(TargetName, Platform, Configuration, Architectures, ProjectFile, Arguments), Logger, IsTestTarget, bSkipValidation);
 
 			if (TargetRules == null)
             {
