UnrealBuildTool: Clean up the WindowsSDK min and max version properties to be consistent

#rnx

diff --git a/Engine/Source/Programs/UnrealBuildTool/Platform/Windows/MicrosoftPlatformSDK.Versions.cs b/Engine/Source/Programs/UnrealBuildTool/Platform/Windows/MicrosoftPlatformSDK.Versions.cs
--- a/Engine/Source/Programs/UnrealBuildTool/Platform/Windows/MicrosoftPlatformSDK.Versions.cs
+++ b/Engine/Source/Programs/UnrealBuildTool/Platform/Windows/MicrosoftPlatformSDK.Versions.cs
@@ -16,6 +16,16 @@ namespace UnrealBuildTool
 			VersionNumber.Parse("10.0.18362.0")
 		};
 
+		/// <summary>
+		/// The minimum Windows SDK version to be used. If this is null then it means there is no minimum version
+		/// </summary>
+		static readonly VersionNumber? MinimumWindowsSDKVersion = new VersionNumber(10, 0, 18362, 0);
+
+		/// <summary>
+		/// The maximum Windows SDK version to be used. If this is null then it means "Latest"
+		/// </summary>
+		static readonly VersionNumber? MaximumWindowsSDKVersion = null;
+
 		/// <summary>
 		/// The default compiler version to be used, if installed. 
 		/// </summary>
@@ -57,122 +67,25 @@ namespace UnrealBuildTool
 
 		static readonly VersionNumber MinimumIntelOneApiVersion = new VersionNumber(2022, 2, 0);
 
-		/// <summary>
-		/// The default set of components that should be suggested to be installed for Visual Studio 2019 or 2022.
-		/// This or the 2019\2022 specific components should be updated if the preferred visual cpp version changes
-		/// </summary>
-		static readonly string[] VisualStudioSuggestedComponents = new string[]
-		{
-			"Microsoft.VisualStudio.Workload.CoreEditor",
-			"Microsoft.VisualStudio.Workload.NativeDesktop",
-			"Microsoft.VisualStudio.Workload.NativeGame",
-			"Microsoft.VisualStudio.Component.VC.Tools.x86.x64",
-			"Microsoft.VisualStudio.Component.Windows10SDK",
-		};
-
-		/// <summary>
-		/// Additional set of components that should be suggested to be installed for Visual Studio 2019 or 2022
-		/// to support the HoloLens platform.
-		/// </summary>
-		static readonly string[] VisualStudioSuggestedHololensComponents = new string[]
-		{
-			"Microsoft.VisualStudio.Workload.Universal",
-			"Microsoft.VisualStudio.Component.VC.Tools.ARM64",
-		};
-
-		/// <summary>
-		/// Additional set of components that should be suggested to be installed for Visual Studio 2019 or 2022
-		/// to support the Linux platform.
-		/// </summary>
-		static readonly string[] VisualStudioSuggestedLinuxComponents = new string[]
-		{
-			"Microsoft.VisualStudio.Workload.NativeCrossPlat",
-		};
-
-		/// <summary>
-		/// Additional set of components that should be suggested to be installed for Visual Studio 2019.
-		/// </summary>
-		static readonly string[] VisualStudio2019SuggestedComponents = new string[]
-		{
-		};
-
-		/// <summary>
-		/// Additional set of components that should be suggested to be installed for Visual Studio 2022.
-		/// </summary>
-		static readonly string[] VisualStudio2022SuggestedComponents = new string[]
-		{
-			"Microsoft.VisualStudio.Workload.ManagedDesktop",
-			"Microsoft.VisualStudio.Component.VC.14.34.17.4.x86.x64",
-			"Microsoft.Net.Component.4.6.2.TargetingPack",
-		};
-
-		/// <summary>
-		/// Additional set of components that should be suggested to be installed for Visual Studio 2022
-		/// to support the HoloLens platform.
-		/// </summary>
-		static readonly string[] VisualStudio2022SuggestedHololensComponents = new string[]
-		{
-			"Microsoft.VisualStudio.Component.VC.14.34.17.4.ARM64",
-		};
-
-		/// <summary>
-		/// Returns the list of suggested of components that should be suggested to be installed for Visual Studio.
-		/// Used to generate a .vsconfig file which will prompt Visual Studio to ask the user to install these components.
-		/// </summary>
-		public static IEnumerable<string> GetVisualStudioSuggestedComponents(VCProjectFileFormat Format)
-		{
-			bool LinuxValid = InstalledPlatformInfo.IsValidPlatform(UnrealTargetPlatform.Linux) && UEBuildPlatform.IsPlatformAvailable(UnrealTargetPlatform.Linux);
-			bool HololensValid = InstalledPlatformInfo.IsValidPlatform(UnrealTargetPlatform.HoloLens) && UEBuildPlatform.IsPlatformAvailable(UnrealTargetPlatform.HoloLens);
-
-			SortedSet<string> Components = new SortedSet<string>();
-			Components.UnionWith(VisualStudioSuggestedComponents);
-
-			switch (Format)
-			{
-				case VCProjectFileFormat.VisualStudio2019:
-					Components.UnionWith(VisualStudio2019SuggestedComponents);
-					break;
-				case VCProjectFileFormat.VisualStudio2022:
-					Components.UnionWith(VisualStudio2022SuggestedComponents);
-					if (HololensValid)
-					{
-						Components.UnionWith(VisualStudio2022SuggestedHololensComponents);
-					}
-					break;
-				default:
-					throw new BuildException("Unsupported Visual Studio version");
-			}
-
-			if (LinuxValid)
-			{
-				Components.UnionWith(VisualStudioSuggestedLinuxComponents);
-			}
-
-			if (HololensValid)
-			{
-				Components.UnionWith(VisualStudioSuggestedHololensComponents);
-			}
-
-			return Components;
-		}
-
+		/// <inheritdoc/>
 		public override string GetMainVersion()
 		{
 			// preferred/main version is the top of the Preferred list - 
 			return PreferredWindowsSdkVersions.First().ToString();
 		}
 
+		/// <inheritdoc/>
 		protected override void GetValidVersionRange(out string MinVersion, out string MaxVersion)
 		{
 			MinVersion = "10.0.00000.0";
 			MaxVersion = "10.9.99999.0";
 		}
 
+		/// <inheritdoc/>
 		protected override void GetValidSoftwareVersionRange(out string? MinVersion, out string? MaxVersion)
 		{
-			// minimum version is the oldest version in the Preferred list -
-			MinVersion = PreferredWindowsSdkVersions.Min()?.ToString();
-			MaxVersion = null;
+			MinVersion = MinimumWindowsSDKVersion?.ToString();
+			MaxVersion = MaximumWindowsSDKVersion?.ToString();
 		}
 
 		/// <summary>
diff --git a/Engine/Source/Programs/UnrealBuildTool/Platform/Windows/MicrosoftPlatformSDK.VisualStudio.cs b/Engine/Source/Programs/UnrealBuildTool/Platform/Windows/MicrosoftPlatformSDK.VisualStudio.cs
new file mode 100644
--- /dev/null
+++ b/Engine/Source/Programs/UnrealBuildTool/Platform/Windows/MicrosoftPlatformSDK.VisualStudio.cs
@@ -0,0 +1,111 @@
+﻿using EpicGames.Core;
+using System;
+using System.Collections.Generic;
+using System.Linq;
+using System.Text;
+using System.Threading.Tasks;
+
+namespace UnrealBuildTool
+{
+	partial class MicrosoftPlatformSDK : UEBuildPlatformSDK
+	{
+		/// <summary>
+		/// The default set of components that should be suggested to be installed for Visual Studio 2019 or 2022.
+		/// This or the 2019\2022 specific components should be updated if the preferred visual cpp version changes
+		/// </summary>
+		static readonly string[] VisualStudioSuggestedComponents = new string[]
+		{
+			"Microsoft.VisualStudio.Workload.CoreEditor",
+			"Microsoft.VisualStudio.Workload.NativeDesktop",
+			"Microsoft.VisualStudio.Workload.NativeGame",
+			"Microsoft.VisualStudio.Component.VC.Tools.x86.x64",
+			"Microsoft.VisualStudio.Component.Windows10SDK.22000",
+		};
+
+		/// <summary>
+		/// Additional set of components that should be suggested to be installed for Visual Studio 2019 or 2022
+		/// to support the HoloLens platform.
+		/// </summary>
+		static readonly string[] VisualStudioSuggestedHololensComponents = new string[]
+		{
+			"Microsoft.VisualStudio.Workload.Universal",
+			"Microsoft.VisualStudio.Component.VC.Tools.ARM64",
+		};
+
+		/// <summary>
+		/// Additional set of components that should be suggested to be installed for Visual Studio 2019 or 2022
+		/// to support the Linux platform.
+		/// </summary>
+		static readonly string[] VisualStudioSuggestedLinuxComponents = new string[]
+		{
+			"Microsoft.VisualStudio.Workload.NativeCrossPlat",
+		};
+
+		/// <summary>
+		/// Additional set of components that should be suggested to be installed for Visual Studio 2019.
+		/// </summary>
+		static readonly string[] VisualStudio2019SuggestedComponents = new string[]
+		{
+		};
+
+		/// <summary>
+		/// Additional set of components that should be suggested to be installed for Visual Studio 2022.
+		/// </summary>
+		static readonly string[] VisualStudio2022SuggestedComponents = new string[]
+		{
+			"Microsoft.VisualStudio.Workload.ManagedDesktop",
+			"Microsoft.VisualStudio.Component.VC.14.34.17.4.x86.x64",
+			"Microsoft.Net.Component.4.6.2.TargetingPack",
+		};
+
+		/// <summary>
+		/// Additional set of components that should be suggested to be installed for Visual Studio 2022
+		/// to support the HoloLens platform.
+		/// </summary>
+		static readonly string[] VisualStudio2022SuggestedHololensComponents = new string[]
+		{
+			"Microsoft.VisualStudio.Component.VC.14.34.17.4.ARM64",
+		};
+
+		/// <summary>
+		/// Returns the list of suggested of components that should be suggested to be installed for Visual Studio.
+		/// Used to generate a .vsconfig file which will prompt Visual Studio to ask the user to install these components.
+		/// </summary>
+		public static IEnumerable<string> GetVisualStudioSuggestedComponents(VCProjectFileFormat Format)
+		{
+			bool LinuxValid = InstalledPlatformInfo.IsValidPlatform(UnrealTargetPlatform.Linux) && UEBuildPlatform.IsPlatformAvailable(UnrealTargetPlatform.Linux);
+			bool HololensValid = InstalledPlatformInfo.IsValidPlatform(UnrealTargetPlatform.HoloLens) && UEBuildPlatform.IsPlatformAvailable(UnrealTargetPlatform.HoloLens);
+
+			SortedSet<string> Components = new SortedSet<string>();
+			Components.UnionWith(VisualStudioSuggestedComponents);
+
+			switch (Format)
+			{
+				case VCProjectFileFormat.VisualStudio2019:
+					Components.UnionWith(VisualStudio2019SuggestedComponents);
+					break;
+				case VCProjectFileFormat.VisualStudio2022:
+					Components.UnionWith(VisualStudio2022SuggestedComponents);
+					if (HololensValid)
+					{
+						Components.UnionWith(VisualStudio2022SuggestedHololensComponents);
+					}
+					break;
+				default:
+					throw new BuildException("Unsupported Visual Studio version");
+			}
+
+			if (LinuxValid)
+			{
+				Components.UnionWith(VisualStudioSuggestedLinuxComponents);
+			}
+
+			if (HololensValid)
+			{
+				Components.UnionWith(VisualStudioSuggestedHololensComponents);
+			}
+
+			return Components;
+		}
+	}
+}
