FastBuild updates via UDN from The Coalition

#rnx

diff --git a/Engine/Source/Programs/UnrealBuildTool/Executors/Experimental/FASTBuild.cs b/Engine/Source/Programs/UnrealBuildTool/Executors/Experimental/FASTBuild.cs
--- a/Engine/Source/Programs/UnrealBuildTool/Executors/Experimental/FASTBuild.cs
+++ b/Engine/Source/Programs/UnrealBuildTool/Executors/Experimental/FASTBuild.cs
@@ -152,6 +152,11 @@ namespace UnrealBuildTool
 		[XmlConfigFile]
 		public static String MsvcCRTRedistVersion = "";
 
+		/// <summary>
+		/// Which MSVC Compiler version to use
+		/// </summary>
+		[XmlConfigFile]
+		public static string CompilerVersion = "";
 		//////////////////////////////////////////
 
 		/// <summary>
@@ -262,6 +267,7 @@ namespace UnrealBuildTool
 				string? BrokeragePath = GetBrokeragePath();
 				if (string.IsNullOrEmpty(BrokeragePath) || !Directory.Exists(BrokeragePath))
 				{
+					Logger.LogWarning("BrokeragePath '{BrokeragePath}' doesn't exist! FASTBuild disabled.", BrokeragePath);
 					return false;
 				}
 			}
@@ -348,6 +354,20 @@ namespace UnrealBuildTool
 			"Module.ProxyLODMeshReduction"
 		};
 
+		private readonly HashSet<string> ForceOverwriteCompilerOptionModules = new HashSet<string>()
+		{
+			"Module.USDStageImporter", 
+			"Module.USDUtilities", 
+			"Module.UnrealUSDWrapper", 
+			"Module.USDStage", 
+			"Module.USDSchemas", 
+			"Module.GeometryCacheUSD",
+			"Module.USDStageEditorViewModels", 
+			"Module.USDTests", 
+			"Module.USDStageEditor", 
+			"Module.USDExporter"
+		};
+
 		private enum FBBuildType
 		{
 			Windows,
@@ -388,6 +408,12 @@ namespace UnrealBuildTool
 				Action => Action.CommandPath.FullName,
 				FBBuildType.Windows
 			),
+			Tuple.Create<string, Func<LinkedAction, string>, FBBuildType>
+			(
+				"Win64",
+				Action => Action.CommandPath.FullName,
+				FBBuildType.Windows
+			),
 		};
 
 		private bool DetectBuildType(IEnumerable<LinkedAction> Actions, ILogger Logger)
@@ -396,10 +422,16 @@ namespace UnrealBuildTool
 			{
 				foreach (Tuple<string, Func<LinkedAction, string>, FBBuildType> BuildTypeSearchParam in BuildTypeSearchParams)
 				{
+					if (BuildTypeSearchParam.Item3.Equals(FBBuildType.Apple) &&
+						(BuildTypeSearchParam.Item2(Action).Contains("Win64", StringComparison.OrdinalIgnoreCase) || 
+						BuildTypeSearchParam.Item2(Action).Contains("X64", StringComparison.OrdinalIgnoreCase)))
+					{
+						continue;
+					}
 					if (BuildTypeSearchParam.Item2(Action).Contains(BuildTypeSearchParam.Item1))
 					{
 						BuildType = BuildTypeSearchParam.Item3;
-						Logger.LogInformation("Detected build type as {Type} from '{From}' using search term '{Term}'", BuildTypeSearchParam.Item1.ToString(), BuildTypeSearchParam.Item2(Action), BuildTypeSearchParam.Item1);
+						Logger.LogInformation("Detected build type as {Type} from '{From}' using search term '{Term}'", BuildTypeSearchParam.Item3.ToString(), BuildTypeSearchParam.Item2(Action), BuildTypeSearchParam.Item1);
 						return true;
 					}
 				}
@@ -434,7 +466,6 @@ namespace UnrealBuildTool
 				return true;
 			}
 
-
 			IEnumerable<LinkedAction> CompileActions = Actions.Where(Action => Action.ActionType == ActionType.Compile && Action.bCanExecuteRemotely && Action.bCanExecuteRemotelyWithSNDBS);
 			if (CompileActions.Any() && DetectBuildType(CompileActions, Logger))
 			{
@@ -443,7 +474,7 @@ namespace UnrealBuildTool
 				{
 					return false;
 				}
-
+				
 				return ExecuteBffFile(FASTBuildFilePath, Logger);
 			}
 
@@ -489,7 +520,6 @@ namespace UnrealBuildTool
 			string ResponseFilePath = "";
 			List<string> AllTokens = new List<string>();
 
-
 			int ResponseFileTokenIndex = Array.FindIndex(RawTokens, RawToken => RawToken.StartsWith("@\""));
 			if (ResponseFileTokenIndex == -1)
 			{
@@ -772,9 +802,19 @@ namespace UnrealBuildTool
 				// it probably means we are building for another platform.
                 if(BuildType == FBBuildType.Windows)
                 {
-					VCEnv = VCEnvironment.Create(WindowsPlatform.GetDefaultCompiler(null, UnrealArch.X64, Logger), WindowsCompiler.Default, UnrealTargetPlatform.Win64, UnrealArch.X64, null, null, null, false, false, Logger);
+					VCEnv = VCEnvironment.Create(
+						Compiler: WindowsPlatform.GetDefaultCompiler(null, UnrealArch.X64, Logger),
+						ToolChain: WindowsCompiler.Default,
+						Platform: UnrealTargetPlatform.Win64,
+						Architecture: UnrealArch.X64,
+						CompilerVersion: string.IsNullOrEmpty(CompilerVersion) ? null : CompilerVersion,
+						WindowsSdkVersion: null,
+						SuppliedSdkDirectoryForVersion: null,
+						bUseCPPWinRT: false,
+						bAllowClangLinker: false,
+						Logger);
 				}
-            }
+			}
 			catch (Exception)
 			{
 				Logger.LogWarning("Failed to get Visual Studio environment.");
@@ -840,6 +880,7 @@ namespace UnrealBuildTool
 				DirectoryReference CLFilterDirectory = DirectoryReference.Combine(Unreal.EngineDirectory, "Build", "Windows", "cl-filter");
 
 				AddText($"\t.Root = '{VCEnv.GetToolPath()}'\n");
+
 				if (UsingCLFilter)
 				{
 					AddText($"\t.CLFilterRoot = '{CLFilterDirectory.FullName}'\n");
@@ -870,7 +911,7 @@ namespace UnrealBuildTool
 					{
 						AddText("\t\t'$Root$/{cluiSubDirName}/clui.dll'\n");
 					}
-					cluiDllPath = new FileReference(VCEnv.GetToolPath() + "{cluiSubDirName}/clui.dll");
+					cluiDllPath = new FileReference(VCEnv.GetToolPath() + "{cluiSubDirName}/clui.dll");					
 				}
 				else
 				{
@@ -990,9 +1031,9 @@ namespace UnrealBuildTool
 			//Start Environment
 			AddText("\t.Environment = \n\t{\n");
 			if (VCEnv != null)
-            {
-                AddText(string.Format("\t\t\"PATH={0}\\Common7\\IDE\\;{1};{2}\\bin\\{3}\\x64\",\n", VCEnv.GetVCInstallDirectory(), VCEnv.GetToolPath(), VCEnv.WindowsSdkDir, VCEnv.WindowsSdkVersion));
-            }
+			{
+				AddText(string.Format("\t\t\"PATH={0}\\Common7\\IDE\\;{1};{2}\\bin\\{3}\\x64\",\n", VCEnv.GetVCInstallDirectory(), VCEnv.GetToolPath(), VCEnv.WindowsSdkDir, VCEnv.WindowsSdkVersion));
+			}
 
 			if (!IsApple())
 			{
@@ -1061,6 +1102,12 @@ namespace UnrealBuildTool
 			}
 
 			string OtherCompilerOptions = GetOptionValue(ParsedCompilerOptions, "OtherOptions", Action, Logger);
+
+			if (ForceOverwriteCompilerOptionModules.Any(x => Action.CommandArguments.Contains(x, StringComparison.OrdinalIgnoreCase)))
+			{
+				OtherCompilerOptions = OtherCompilerOptions.Replace("/WX", "");
+			}
+
 			string CompilerOutputExtension = ".unset";
 			string CLFilterParams = "";
 			string ShowIncludesParam = "";
