UnrealBuildTool: Add missing copyright boilerplate

#rnx

diff --git a/Engine/Source/Programs/UnrealBuildTool/Platform/Windows/MicrosoftPlatformSDK.VisualStudio.cs b/Engine/Source/Programs/UnrealBuildTool/Platform/Windows/MicrosoftPlatformSDK.VisualStudio.cs
--- a/Engine/Source/Programs/UnrealBuildTool/Platform/Windows/MicrosoftPlatformSDK.VisualStudio.cs
+++ b/Engine/Source/Programs/UnrealBuildTool/Platform/Windows/MicrosoftPlatformSDK.VisualStudio.cs
@@ -1,4 +1,6 @@
-﻿using EpicGames.Core;
+﻿// Copyright Epic Games, Inc. All Rights Reserved.
+
+using EpicGames.Core;
 using System;
 using System.Collections.Generic;
 using System.Linq;
