UnrealBuildTool: Remove fastbuild warning

#rnx

diff --git a/Engine/Source/Programs/UnrealBuildTool/Executors/Experimental/FASTBuild.cs b/Engine/Source/Programs/UnrealBuildTool/Executors/Experimental/FASTBuild.cs
--- a/Engine/Source/Programs/UnrealBuildTool/Executors/Experimental/FASTBuild.cs
+++ b/Engine/Source/Programs/UnrealBuildTool/Executors/Experimental/FASTBuild.cs
@@ -267,7 +267,6 @@ namespace UnrealBuildTool
 				string? BrokeragePath = GetBrokeragePath();
 				if (string.IsNullOrEmpty(BrokeragePath) || !Directory.Exists(BrokeragePath))
 				{
-					Logger.LogWarning("BrokeragePath '{BrokeragePath}' doesn't exist! FASTBuild disabled.", BrokeragePath);
 					return false;
 				}
 			}
