[UBT]
* Changed so code generating unity file is modifying the right compile environment when adding lookup for source file -> .gen.cpp

Fixes a bug where sndbs fails to remotely build unity files properly

diff --git a/Engine/Source/Programs/UnrealBuildTool/Configuration/UEBuildModuleCPP.cs b/Engine/Source/Programs/UnrealBuildTool/Configuration/UEBuildModuleCPP.cs
--- a/Engine/Source/Programs/UnrealBuildTool/Configuration/UEBuildModuleCPP.cs
+++ b/Engine/Source/Programs/UnrealBuildTool/Configuration/UEBuildModuleCPP.cs
@@ -684,13 +684,13 @@ namespace UnrealBuildTool
 			// Compile CPP files
 			if (bModuleUsesUnityBuild)
 			{
-				Unity.GenerateUnityCPPs(Target, CPPFiles, InputFiles.HeaderFiles, CompileEnvironment, WorkingSet, Rules.ShortName ?? Name, IntermediateDirectory, Graph, SourceFileToUnityFile, 
+				Unity.GenerateUnityCPPs(Target, CPPFiles, InputFiles.HeaderFiles, CodeCompileEnvironment, WorkingSet, Rules.ShortName ?? Name, IntermediateDirectory, Graph, SourceFileToUnityFile, 
 					out List<FileItem> NormalFiles, out List<FileItem> AdaptiveFiles, NumIncludedBytesPerUnityCPP);
 				LinkInputFiles.AddRange(CompileFilesWithToolChain(Target, ToolChain, CodeCompileEnvironment, ModuleCompileEnvironment, NormalFiles, AdaptiveFiles, Graph, Logger).ObjectFiles);
 			}
 			else if (SpecificFilesToCompile.Count == 0)
 			{
-				Unity.GetAdaptiveFiles(Target, CPPFiles, InputFiles.HeaderFiles, CompileEnvironment, WorkingSet, Rules.ShortName ?? Name, IntermediateDirectory, Graph, 
+				Unity.GetAdaptiveFiles(Target, CPPFiles, InputFiles.HeaderFiles, CodeCompileEnvironment, WorkingSet, Rules.ShortName ?? Name, IntermediateDirectory, Graph, 
 					out List<FileItem> NormalFiles, out List<FileItem> AdaptiveFiles);
 				if (NormalFiles.Where(file => !file.HasExtension(".gen.cpp")).Count() == 0)
 				{
@@ -701,7 +701,7 @@ namespace UnrealBuildTool
 			}
 			else
 			{
-				LinkInputFiles.AddRange(CompileFilesWithToolChain(Target, ToolChain, CompileEnvironment, ModuleCompileEnvironment, CPPFiles, new List<FileItem>(), Graph, Logger).ObjectFiles);
+				LinkInputFiles.AddRange(CompileFilesWithToolChain(Target, ToolChain, CodeCompileEnvironment, ModuleCompileEnvironment, CPPFiles, new List<FileItem>(), Graph, Logger).ObjectFiles);
 			}
 
 			// Compile ISPC files directly
