[UBT - SNDBS]
* Made sure sorting of actions is preserved in sndbs.json

Taken from sndbs documentation:

"SNDBS will tend to start the jobs toward the start of the array first. This ordering is a hint and is not
guaranteed. This feature can be used to optimize build performance by submitting long running jobs first"

diff --git a/Engine/Source/Programs/UnrealBuildTool/Executors/SNDBS.cs b/Engine/Source/Programs/UnrealBuildTool/Executors/SNDBS.cs
--- a/Engine/Source/Programs/UnrealBuildTool/Executors/SNDBS.cs
+++ b/Engine/Source/Programs/UnrealBuildTool/Executors/SNDBS.cs
@@ -216,8 +216,9 @@ namespace UnrealBuildTool
 				throw new BuildException($"Failed to create directory \"{IntermediateDir}\".");
 			}
 
+			int IdCounter = 0;
 			// Build the json script file to describe all the actions and their dependencies
-			var ActionIds = Actions.ToDictionary(a => a, a => Guid.NewGuid().ToString());
+			var ActionIds = Actions.ToDictionary(a => a, a => new Guid(++IdCounter, 0, 0, 0, 0, 0, 0, 0, 0, 0, 0).ToString());
 			JsonSerializerOptions JsonOption = new JsonSerializerOptions();
 			JsonOption.Encoder = System.Text.Encodings.Web.JavaScriptEncoder.UnsafeRelaxedJsonEscaping;
 			File.WriteAllText(ScriptFile.FullName, JsonSerializer.Serialize(new Dictionary<string, object>()
