Properly escape preprocessor definitions in AndroidToolChain now it uses common ClangToolChain

#android

diff --git a/Engine/Source/Programs/UnrealBuildTool/Platform/Android/AndroidToolChain.cs b/Engine/Source/Programs/UnrealBuildTool/Platform/Android/AndroidToolChain.cs
--- a/Engine/Source/Programs/UnrealBuildTool/Platform/Android/AndroidToolChain.cs
+++ b/Engine/Source/Programs/UnrealBuildTool/Platform/Android/AndroidToolChain.cs
@@ -454,6 +454,12 @@ namespace UnrealBuildTool
 			Arguments.AddRange(FilteredPaths.Select(IncludePath => GetSystemIncludePathArgument(IncludePath)));
 		}
 
+		/// <inheritdoc/>
+		protected override string EscapePreprocessorDefinition(string Definition)
+		{
+			return Definition.Contains("\"") ? Definition.Replace("\"", "\\\"") : Definition;
+		}
+
 		/// <inheritdoc/>
 		protected override void GetCompileArguments_WarningsAndErrors(CppCompileEnvironment CompileEnvironment, List<string> Arguments)
 		{
