[UBT]
Changed how single files are built. We don't want to invalidate makefile everytime we build single files since it destroys turnaround times. With this change a single file compile can take 2s (+ the actual compile time).

The new behavior injects a special action per module when creating the makefile. These actions can be used to on-the-fly create a proper compile action that follows the rules of the module that the specific file belongs to.. In a normal build these actions are ignored since the logic deciding which actions to build is backtraced from which binaries we want to create.

When a specific file compile is triggered, the logic deciding which files to build search up all these special actions and create a lookup based on which folders the special actions handle. It then try to find the special action that handles the specific file. The matching special action then creates a compile action that can handle that specific file and then queue up the action for execution. If no special action is found it falls back to try to use actions that have this specific file as input (ispc files for example)

Details:
* Removed lots of custom code for "specific files" handling
* Changed so pch (both private and shared) always use definition file. Added #pragma once and change so pch wrapper file include definition file. This made the adaptive path and specific file easy to implement (just disable pch in compile environment and it will just work)
* Added SingleFileAction for both VCToolChain and ClangToolChain. It now works to compile specific headers and cpp files. (It creates wrapper files on the fly to be able to compile all header files (compiling headers directly blow up if there are circular includes)
* Fixed so GenerateClangDatabase mode works with new changes
* Moved the logic that makes sure all (directly) depending cpp files are recompiled when .h are included in the singlefile option

diff --git a/Engine/Source/Programs/UnrealBuildTool/Configuration/TargetDescriptor.cs b/Engine/Source/Programs/UnrealBuildTool/Configuration/TargetDescriptor.cs
--- a/Engine/Source/Programs/UnrealBuildTool/Configuration/TargetDescriptor.cs
+++ b/Engine/Source/Programs/UnrealBuildTool/Configuration/TargetDescriptor.cs
@@ -5,7 +5,6 @@ using System.Diagnostics.CodeAnalysis;
 using System.IO;
 using System.Linq;
 using System.Text;
-using System.Threading.Tasks;
 using EpicGames.Core;
 using UnrealBuildBase;
 using Microsoft.Extensions.Logging;
@@ -62,12 +61,6 @@ namespace UnrealBuildTool
 		[CommandLine("-SingleFile=")]
 		public List<FileReference> SpecificFilesToCompile = new List<FileReference>();
 
-		/// <summary>
-		/// Individual files to compile which may or may not be part of the target. This file set is built from a header scan, and may include files
-		/// not part of the target by design.
-		/// </summary>
-		public List<FileReference> OptionalFilesToCompile = new List<FileReference>();
-
 		/// <summary>
 		/// Whether to perform hot reload for this target
 		/// </summary>
diff --git a/Engine/Source/Programs/UnrealBuildTool/Configuration/UEBuildBinary.cs b/Engine/Source/Programs/UnrealBuildTool/Configuration/UEBuildBinary.cs
--- a/Engine/Source/Programs/UnrealBuildTool/Configuration/UEBuildBinary.cs
+++ b/Engine/Source/Programs/UnrealBuildTool/Configuration/UEBuildBinary.cs
@@ -171,13 +171,12 @@ namespace UnrealBuildTool
 		/// <param name="ToolChain">The toolchain which to use for building</param>
 		/// <param name="CompileEnvironment">The environment to compile the binary in</param>
 		/// <param name="LinkEnvironment">The environment to link the binary in</param>
-		/// <param name="SpecificFilesToCompile">If non-empty, specifies individual cpp files to be compiled</param>
 		/// <param name="WorkingSet">The working set of source files</param>
 		/// <param name="ExeDir">Directory containing the output executable</param>
 		/// <param name="Graph">The graph being built</param>
 		/// <param name="Logger">Logger for output</param>
 		/// <returns>Set of built products</returns>
-		public List<FileItem> Build(ReadOnlyTargetRules Target, UEToolChain ToolChain, CppCompileEnvironment CompileEnvironment, LinkEnvironment LinkEnvironment, List<FileReference> SpecificFilesToCompile, ISourceFileWorkingSet WorkingSet, DirectoryReference ExeDir, IActionGraphBuilder Graph, ILogger Logger)
+		public List<FileItem> Build(ReadOnlyTargetRules Target, UEToolChain ToolChain, CppCompileEnvironment CompileEnvironment, LinkEnvironment LinkEnvironment, ISourceFileWorkingSet WorkingSet, DirectoryReference ExeDir, IActionGraphBuilder Graph, ILogger Logger)
 		{
 			// Return nothing if we're using precompiled binaries. If we're not linking, we might want just one module to be compiled (eg. a foreign plugin), so allow any actions to run.
 			if (bUsePrecompiled && !(Target.LinkType == TargetLinkType.Monolithic && Target.bDisableLinking))
@@ -186,7 +185,7 @@ namespace UnrealBuildTool
 			}
 
 			// Setup linking environment.
-			LinkEnvironment BinaryLinkEnvironment = SetupBinaryLinkEnvironment(Target, ToolChain, LinkEnvironment, CompileEnvironment, SpecificFilesToCompile, WorkingSet, ExeDir, Graph, Logger);
+			LinkEnvironment BinaryLinkEnvironment = SetupBinaryLinkEnvironment(Target, ToolChain, LinkEnvironment, CompileEnvironment, WorkingSet, ExeDir, Graph, Logger);
 
 			// If we're generating projects, we only need include paths and definitions, there is no need to run the linking logic.
 			if (ProjectFileGenerator.bGenerateProjectFiles)
@@ -699,7 +698,7 @@ namespace UnrealBuildTool
 			return BinaryCompileEnvironment;
 		}
 
-		private LinkEnvironment SetupBinaryLinkEnvironment(ReadOnlyTargetRules Target, UEToolChain ToolChain, LinkEnvironment LinkEnvironment, CppCompileEnvironment CompileEnvironment, List<FileReference> SpecificFilesToCompile, ISourceFileWorkingSet WorkingSet, DirectoryReference ExeDir, IActionGraphBuilder Graph, ILogger Logger)
+		private LinkEnvironment SetupBinaryLinkEnvironment(ReadOnlyTargetRules Target, UEToolChain ToolChain, LinkEnvironment LinkEnvironment, CppCompileEnvironment CompileEnvironment, ISourceFileWorkingSet WorkingSet, DirectoryReference ExeDir, IActionGraphBuilder Graph, ILogger Logger)
 		{
 			LinkEnvironment BinaryLinkEnvironment = new LinkEnvironment(LinkEnvironment);
 			HashSet<UEBuildModule> LinkEnvironmentVisitedModules = new HashSet<UEBuildModule>();
@@ -720,7 +719,7 @@ namespace UnrealBuildTool
 				{
 					// Compile each module.
 					Logger.LogDebug("Compile module: {ModuleName}", Module.Name);
-					List<FileItem> LinkInputFiles = Module.Compile(Target, ToolChain, BinaryCompileEnvironment, SpecificFilesToCompile, WorkingSet, Graph, Logger);
+					List<FileItem> LinkInputFiles = Module.Compile(Target, ToolChain, BinaryCompileEnvironment, WorkingSet, Graph, Logger);
 
 					// Save the module outputs. In monolithic builds, this is just the object files.
 					if (Target.LinkType == TargetLinkType.Monolithic)
diff --git a/Engine/Source/Programs/UnrealBuildTool/Configuration/UEBuildModule.cs b/Engine/Source/Programs/UnrealBuildTool/Configuration/UEBuildModule.cs
--- a/Engine/Source/Programs/UnrealBuildTool/Configuration/UEBuildModule.cs
+++ b/Engine/Source/Programs/UnrealBuildTool/Configuration/UEBuildModule.cs
@@ -917,7 +917,7 @@ namespace UnrealBuildTool
 		/// <summary>
 		/// Compiles the module, and returns a list of files output by the compiler.
 		/// </summary>
-		public virtual List<FileItem> Compile(ReadOnlyTargetRules Target, UEToolChain ToolChain, CppCompileEnvironment CompileEnvironment, List<FileReference> SpecificFilesToCompile, ISourceFileWorkingSet WorkingSet, IActionGraphBuilder Graph, ILogger Logger)
+		public virtual List<FileItem> Compile(ReadOnlyTargetRules Target, UEToolChain ToolChain, CppCompileEnvironment CompileEnvironment, ISourceFileWorkingSet WorkingSet, IActionGraphBuilder Graph, ILogger Logger)
 		{
 			// Generate type libraries for Windows
 			foreach(ModuleRules.TypeLibrary TypeLibrary in Rules.TypeLibraries)
diff --git a/Engine/Source/Programs/UnrealBuildTool/Configuration/UEBuildModuleCPP.cs b/Engine/Source/Programs/UnrealBuildTool/Configuration/UEBuildModuleCPP.cs
--- a/Engine/Source/Programs/UnrealBuildTool/Configuration/UEBuildModuleCPP.cs
+++ b/Engine/Source/Programs/UnrealBuildTool/Configuration/UEBuildModuleCPP.cs
@@ -341,11 +341,11 @@ namespace UnrealBuildTool
 		}
 
 		// UEBuildModule interface.
-		public override List<FileItem> Compile(ReadOnlyTargetRules Target, UEToolChain ToolChain, CppCompileEnvironment BinaryCompileEnvironment, List<FileReference> SpecificFilesToCompile, ISourceFileWorkingSet WorkingSet, IActionGraphBuilder Graph, ILogger Logger)
+		public override List<FileItem> Compile(ReadOnlyTargetRules Target, UEToolChain ToolChain, CppCompileEnvironment BinaryCompileEnvironment, ISourceFileWorkingSet WorkingSet, IActionGraphBuilder Graph, ILogger Logger)
 		{
 			//UEBuildPlatform BuildPlatform = UEBuildPlatform.GetBuildPlatform(BinaryCompileEnvironment.Platform);
 
-			List<FileItem> LinkInputFiles = base.Compile(Target, ToolChain, BinaryCompileEnvironment, SpecificFilesToCompile, WorkingSet, Graph, Logger);
+			List<FileItem> LinkInputFiles = base.Compile(Target, ToolChain, BinaryCompileEnvironment, WorkingSet, Graph, Logger);
 
 			CppCompileEnvironment ModuleCompileEnvironment = CreateModuleCompileEnvironment(Target, BinaryCompileEnvironment, Logger);
 
@@ -387,27 +387,6 @@ namespace UnrealBuildTool
 				Graph.AddSourceFiles(Pair.Key, Pair.Value);
 			}
 
-			// If we're compiling only specific files, strip out anything else. This prevents us clobbering response files for anything we're 
-			// not going to build, triggering a larger build than necessary when we do a regular build again.
-			if(SpecificFilesToCompile.Count > 0)
-			{
-				InputFiles.CPPFiles.RemoveAll(x => !SpecificFilesToCompile.Contains(x.Location));
-				InputFiles.CCFiles.RemoveAll(x => !SpecificFilesToCompile.Contains(x.Location));
-				InputFiles.CFiles.RemoveAll(x => !SpecificFilesToCompile.Contains(x.Location));
-				InputFiles.MMFiles.RemoveAll(x => !SpecificFilesToCompile.Contains(x.Location));
-				InputFiles.RCFiles.RemoveAll(x => !SpecificFilesToCompile.Contains(x.Location));
-
-				if (InputFiles.CPPFiles.Count == 0
-					&& InputFiles.CCFiles.Count == 0
-					&& InputFiles.CFiles.Count == 0
-					&& InputFiles.MMFiles.Count == 0
-					&& InputFiles.RCFiles.Count == 0
-					&& !SpecificFilesToCompile.Any(x => ContainsFile(x)))
-				{
-					return new List<FileItem>();
-				}
-			}
-
 			// We are building with IWYU and thismodule does not support it, early out
 			if (Target.bIWYU && Rules.IWYUSupport == IWYUSupport.None)
 			{
@@ -485,36 +464,23 @@ namespace UnrealBuildTool
 			if (GeneratedCppDirectories != null && !CompileEnvironment.bHackHeaderGenerator)
 			{
 				var GeneratedFiles = new Dictionary<string, FileItem>();
-				if (SpecificFilesToCompile.Count == 0)
+				foreach (string GeneratedDir in GeneratedCppDirectories)
 				{
-					foreach (string GeneratedDir in GeneratedCppDirectories)
+					if (!Directory.Exists(GeneratedDir))
 					{
-						if (!Directory.Exists(GeneratedDir))
-						{
-							continue;
-						}
+						continue;
+					}
 
-						string Prefix = Path.GetFileName(GeneratedDir) + '/'; // "UHT/" or "VNI/"
+					string Prefix = Path.GetFileName(GeneratedDir) + '/'; // "UHT/" or "VNI/"
 
-						DirectoryItem DirItem = DirectoryItem.GetItemByPath(GeneratedDir);
-						foreach (FileItem File in DirItem.EnumerateFiles())
-						{
-							string FileName = File.Name;
-							if (FileName.EndsWith(".gen.cpp"))
-							{
-								string Key = Prefix + FileName.Substring(0, FileName.Length - ".gen.cpp".Length);
-								GeneratedFiles.Add(Key, File);
-							}
-						}
-					}
-				}
-				else
-				{
-					foreach (FileReference FileToCompile in SpecificFilesToCompile)
+					DirectoryItem DirItem = DirectoryItem.GetItemByPath(GeneratedDir);
+					foreach (FileItem File in DirItem.EnumerateFiles())
 					{
-						if (GeneratedCppDirectories.Any(x => FileToCompile.IsUnderDirectory(new DirectoryReference(x))))
+						string FileName = File.Name;
+						if (FileName.EndsWith(".gen.cpp"))
 						{
-							GeneratedFiles.Add(FileToCompile.GetFileNameWithoutAnyExtensions(), FileItem.GetItemByFileReference(FileToCompile));
+							string Key = Prefix + FileName.Substring(0, FileName.Length - ".gen.cpp".Length);
+							GeneratedFiles.Add(Key, File);
 						}
 					}
 				}
@@ -593,10 +559,7 @@ namespace UnrealBuildTool
 					// Compile all the generated files
 					foreach (FileItem GeneratedCppFileItem in GeneratedFiles.Values)
 					{
-						if (SpecificFilesToCompile.Count == 0 || SpecificFilesToCompile.Contains(GeneratedCppFileItem.Location))
-						{
-							GeneratedFileItems.Add(GeneratedCppFileItem);
-						}
+						GeneratedFileItems.Add(GeneratedCppFileItem);
 					}
 				}
 			}
@@ -688,7 +651,7 @@ namespace UnrealBuildTool
 					out List<FileItem> NormalFiles, out List<FileItem> AdaptiveFiles, NumIncludedBytesPerUnityCPP);
 				LinkInputFiles.AddRange(CompileFilesWithToolChain(Target, ToolChain, CodeCompileEnvironment, ModuleCompileEnvironment, NormalFiles, AdaptiveFiles, Graph, Logger).ObjectFiles);
 			}
-			else if (SpecificFilesToCompile.Count == 0)
+			else
 			{
 				Unity.GetAdaptiveFiles(Target, CPPFiles, InputFiles.HeaderFiles, CodeCompileEnvironment, WorkingSet, Rules.ShortName ?? Name, IntermediateDirectory, Graph, 
 					out List<FileItem> NormalFiles, out List<FileItem> AdaptiveFiles);
@@ -699,10 +662,6 @@ namespace UnrealBuildTool
 				}
 				LinkInputFiles.AddRange(CompileFilesWithToolChain(Target, ToolChain, CodeCompileEnvironment, ModuleCompileEnvironment, NormalFiles, AdaptiveFiles, Graph, Logger).ObjectFiles);
 			}
-			else
-			{
-				LinkInputFiles.AddRange(CompileFilesWithToolChain(Target, ToolChain, CodeCompileEnvironment, ModuleCompileEnvironment, CPPFiles, new List<FileItem>(), Graph, Logger).ObjectFiles);
-			}
 
 			// Compile ISPC files directly
 			if (InputFiles.ISPCFiles.Count > 0)
@@ -866,6 +825,23 @@ namespace UnrealBuildTool
 			return new PrecompiledHeaderTemplate(this, CompileEnvironment, HeaderFile, PrecompiledHeaderDir);
 		}
 
+		private HashSet<string> GetImmutableDefinitions(List<string> Definitions)
+		{
+			HashSet<string> ImmutableDefinitions = new();
+			foreach (string Definition in Definitions)
+			{
+				if (Definition.Contains("UE_IS_ENGINE_MODULE") ||
+					Definition.Contains("DEPRECATED_FORGAME") ||
+					Definition.Contains("UE_DEPRECATED_FORGAME") ||
+					Definition.Contains("UE_ENABLE_INCLUDE_ORDER_DEPRECATED_IN_"))
+				{
+					continue;
+				}
+				ImmutableDefinitions.Add(Definition);
+			}
+			return ImmutableDefinitions;
+		}
+
 		/// <summary>
 		/// Creates a precompiled header action to generate a new pch file 
 		/// </summary>
@@ -877,8 +853,9 @@ namespace UnrealBuildTool
 		private PrecompiledHeaderInstance CreatePrivatePCH(UEToolChain? ToolChain, FileItem HeaderFile, CppCompileEnvironment ModuleCompileEnvironment, IActionGraphBuilder Graph)
 		{
 			// Create the wrapper file, which sets all the definitions needed to compile it
+			FileItem DefinitionsFileItem = CreateHeaderForDefinitions(ModuleCompileEnvironment, IntermediateDirectory, null, Graph)!;
 			FileReference WrapperLocation = FileReference.Combine(IntermediateDirectory, String.Format("PCH.{0}.h", Name));
-			FileItem WrapperFile = CreatePCHWrapperFile(WrapperLocation, ModuleCompileEnvironment.Definitions, HeaderFile, Graph);
+			FileItem WrapperFile = CreatePCHWrapperFile(WrapperLocation, DefinitionsFileItem!, HeaderFile, Graph);
 
 			// Create a new C++ environment that is used to create the PCH.
 			CppCompileEnvironment CompileEnvironment = new CppCompileEnvironment(ModuleCompileEnvironment);
@@ -897,7 +874,8 @@ namespace UnrealBuildTool
 			{
 				Output = ToolChain.CompileAllCPPFiles(CompileEnvironment, new List<FileItem>() { WrapperFile }, IntermediateDirectory, Name, Graph);
 			}
-			return new PrecompiledHeaderInstance(WrapperFile, CompileEnvironment, Output);
+
+			return new PrecompiledHeaderInstance(WrapperFile, DefinitionsFileItem, CompileEnvironment, Output, GetImmutableDefinitions(ModuleCompileEnvironment.Definitions));
 		}
 
 		/// <summary>
@@ -929,9 +907,16 @@ namespace UnrealBuildTool
 				// Create a suffix to distinguish this shared PCH variant from any others. Currently only optimized and non-optimized shared PCHs are supported.
 				string Variant = GetSuffixForSharedPCH(ModuleCompileEnvironment, Template.BaseCompileEnvironment);
 
+				FileReference SharedDefinitionsLocation = FileReference.Combine(Template.OutputDir, String.Format("SharedDefinitions.{0}{1}.h", Template.Module.Name, Variant));
+				List<string> NewDefinitions = new();
+				StringBuilder Writer = new StringBuilder();
+				Writer.AppendLine("#pragma once");
+				WriteDefinitions(Definitions, Writer);
+				FileItem SharedDefinitionsFileItem = Graph.CreateIntermediateTextFile(SharedDefinitionsLocation, Writer.ToString());
+
 				// Create the wrapper file, which sets all the definitions needed to compile it
 				FileReference WrapperLocation = FileReference.Combine(Template.OutputDir, String.Format("SharedPCH.{0}{1}.h", Template.Module.Name, Variant));
-				FileItem WrapperFile = CreatePCHWrapperFile(WrapperLocation, Definitions, Template.HeaderFile, Graph);
+				FileItem WrapperFile = CreatePCHWrapperFile(WrapperLocation, SharedDefinitionsFileItem, Template.HeaderFile, Graph);
 
 				// Create the compile environment for this PCH
 				CppCompileEnvironment CompileEnvironment = new CppCompileEnvironment(Template.BaseCompileEnvironment);
@@ -950,7 +935,7 @@ namespace UnrealBuildTool
 				{
 					Output = ToolChain.CompileAllCPPFiles(CompileEnvironment, new List<FileItem>() { WrapperFile }, Template.OutputDir, "Shared", Graph);
 				}
-				Instance = new PrecompiledHeaderInstance(WrapperFile, CompileEnvironment, Output);
+				Instance = new PrecompiledHeaderInstance(WrapperFile, SharedDefinitionsFileItem, CompileEnvironment, Output, GetImmutableDefinitions(Template.BaseCompileEnvironment.Definitions));
 				Template.Instances.Add(Instance);
 			}
 
@@ -1187,6 +1172,15 @@ namespace UnrealBuildTool
 				AdaptiveFiles = new List<FileItem>();
 			}
 
+			// Create and register a special action that can be used to compile single files (even when unity is enabled)
+			if (NormalFiles.Count > 0 || AdaptiveFiles.Count > 0)
+			{
+				PrecompiledHeaderAction OldAction = CompileEnvironment.PrecompiledHeaderAction;
+				CompileEnvironment.PrecompiledHeaderAction = PrecompiledHeaderAction.None;
+				ToolChain.CreateSpecificFileAction(CompileEnvironment, ModuleDirectory, IntermediateDirectory, Graph);
+				CompileEnvironment.PrecompiledHeaderAction = OldAction;
+			}
+
 			CPPOutput OutputFiles = new CPPOutput();
 
 			if (NormalFiles.Count > 0)
@@ -1353,8 +1347,8 @@ namespace UnrealBuildTool
 						FileItem PrivateDefinitionsFileItem;
 						{
 							StringBuilder Writer = new();
-							// Remove the module _API definition for cases where there are circular dependencies between the shared PCH module and modules using it
-							Writer.AppendLine("#undef {0}", ModuleApiDefine);
+
+							Writer.AppendLine($"#include \"{GetIncludeString(Instance.DefinitionsFile)}\"");
 
 							// Games may choose to use shared PCHs from the engine, so allow them to change the value of these macros
 							if(!Rules.bTreatAsEngineModule)
@@ -1370,7 +1364,26 @@ namespace UnrealBuildTool
 								}
 							}
 
-							WriteDefinitions(CompileEnvironment.Definitions, Writer);
+							// Only add new definitions that are not already existing in the shared pch
+							List<string> NewDefinitions = new();
+							bool ModuleApiUndef = false;
+							foreach (string Definition in CompileEnvironment.Definitions)
+							{
+								if (Instance.ImmutableDefinitions.Contains(Definition))
+								{
+									continue;
+								}
+
+								NewDefinitions.Add(Definition);
+
+								// Remove the module _API definition for cases where there are circular dependencies between the shared PCH module and modules using it
+								if (!ModuleApiUndef && Definition.StartsWith(ModuleApiDefine))
+								{
+									ModuleApiUndef = true;
+									Writer.AppendLine("#undef {0}", ModuleApiDefine);
+								}
+							}
+							WriteDefinitions(NewDefinitions, Writer);
 							PrivateDefinitionsFileItem = Graph.CreateIntermediateTextFile(PrivateDefinitionsFile, Writer.ToString());
 						}
 
@@ -1398,7 +1411,7 @@ namespace UnrealBuildTool
 		/// <param name="IntermediateDirectory">Directory to create the intermediate file</param>
 		/// <param name="HeaderSuffix">Suffix for the included file</param>
 		/// <param name="Graph">The action graph being built</param>
-		static void CreateHeaderForDefinitions(CppCompileEnvironment CompileEnvironment, DirectoryReference IntermediateDirectory, string? HeaderSuffix, IActionGraphBuilder Graph)
+		static FileItem? CreateHeaderForDefinitions(CppCompileEnvironment CompileEnvironment, DirectoryReference IntermediateDirectory, string? HeaderSuffix, IActionGraphBuilder Graph)
 		{
 			if(CompileEnvironment.Definitions.Count > 0)
 			{
@@ -1418,8 +1431,10 @@ namespace UnrealBuildTool
 
 					FileItem PrivateDefinitionsFileItem = Graph.CreateIntermediateTextFile(PrivateDefinitionsFile, Writer.ToString());
 					CompileEnvironment.ForceIncludeFiles.Add(PrivateDefinitionsFileItem);
+					return PrivateDefinitionsFileItem;
 				}
 			}
+			return null;
 		}
 
 		/// <summary>
@@ -1438,29 +1453,35 @@ namespace UnrealBuildTool
 			CompileEnvironment.UserIncludePaths.Add(IntermediateDirectory);
 		}
 
+		static string GetIncludeString(FileItem FileItem)
+		{
+			string IncludeFileString = FileItem.AbsolutePath;
+			if (FileItem.Location.IsUnderDirectory(Unreal.RootDirectory))
+			{
+				IncludeFileString = FileItem.Location.MakeRelativeTo(Unreal.EngineSourceDirectory);
+			}
+			return IncludeFileString.Replace('\\', '/');
+		}
+
 		/// <summary>
 		/// Create a header file containing the module definitions, which also includes the PCH itself. Including through another file is necessary on 
 		/// Clang, since we get warnings about #pragma once otherwise, but it also allows us to consistently define the preprocessor state on all 
 		/// platforms.
 		/// </summary>
 		/// <param name="OutputFile">The output file to create</param>
-		/// <param name="Definitions">Definitions required by the PCH</param>
+		/// <param name="DefinitionsFile">File containing definitions required by the PCH</param>
 		/// <param name="IncludedFile">The PCH file to include</param>
 		/// <param name="Graph">The action graph builder</param>
 		/// <returns>FileItem for the created file</returns>
-		static FileItem CreatePCHWrapperFile(FileReference OutputFile, IEnumerable<string> Definitions, FileItem IncludedFile, IActionGraphBuilder Graph)
+		static FileItem CreatePCHWrapperFile(FileReference OutputFile, FileItem DefinitionsFile, FileItem IncludedFile, IActionGraphBuilder Graph)
 		{
 			// Build the contents of the wrapper file
 			StringBuilder WrapperContents = new StringBuilder();
 			{
-				string IncludeFileString = IncludedFile.AbsolutePath;
-				if (IncludedFile.Location.IsUnderDirectory(Unreal.RootDirectory))
-				{
-					IncludeFileString = IncludedFile.Location.MakeRelativeTo(Unreal.EngineSourceDirectory);
-				}
-				IncludeFileString = IncludeFileString.Replace('\\', '/');
+				string IncludeFileString = GetIncludeString(IncludedFile);
+				string DefinitionsFileString = GetIncludeString(DefinitionsFile);
 				WrapperContents.AppendLine("// PCH for {0}", IncludeFileString);
-				WriteDefinitions(Definitions, WrapperContents);
+				WrapperContents.AppendLine("#include \"{0}\"", DefinitionsFileString);
 				WrapperContents.AppendLine("#include \"{0}\"", IncludeFileString);
 			}
 
diff --git a/Engine/Source/Programs/UnrealBuildTool/Configuration/UEBuildTarget.cs b/Engine/Source/Programs/UnrealBuildTool/Configuration/UEBuildTarget.cs
--- a/Engine/Source/Programs/UnrealBuildTool/Configuration/UEBuildTarget.cs
+++ b/Engine/Source/Programs/UnrealBuildTool/Configuration/UEBuildTarget.cs
@@ -1224,56 +1224,6 @@ namespace UnrealBuildTool
 				RulesObject.DependencyListFileNames.Add(DependencyListFile);
 			}
 
-			// If we're compiling only specific files, we need to prevent unity builds from running
-			if (Descriptor.SpecificFilesToCompile.Count > 0)
-			{
-				RulesObject.bUseUnityBuild = false;
-				RulesObject.bForceUnityBuild = false;
-				RulesObject.bUsePCHFiles = false;
-				RulesObject.bDisableLinking = true;
-				RulesObject.bIgnoreBuildOutputs = true;
-			}
-
-			// If the files to compile contains headers, figure out the cpp files that include them
-			Func<FileReference, bool> IsHeader = x => x.HasExtension(".h") || x.HasExtension(".hpp") || x.HasExtension(".hxx");
-			List<FileReference> SpecificHeaderFiles = Descriptor.SpecificFilesToCompile.Where(IsHeader).ToList();
-			if (SpecificHeaderFiles.Count > 0)
-			{
-				List<DirectoryReference> BaseDirs = new List<DirectoryReference>();
-				BaseDirs.Add(Unreal.EngineDirectory);
-
-				DirectoryReference CacheDir = Unreal.EngineDirectory;
-				if (RulesObject.ProjectFile != null)
-				{
-					BaseDirs.Add(RulesObject.ProjectFile.Directory);
-					CacheDir = RulesObject.ProjectFile.Directory;
-				}
-
-				FileReference IncludeCache = FileReference.Combine(CacheDir, "Intermediate", "HeaderLookup.dat");
-
-				Logger.LogInformation("Building dependency cache for specified headers");
-				foreach (FileReference HeaderFile in SpecificHeaderFiles)
-				{
-					Logger.LogDebug("  {HeaderFile}", HeaderFile);
-				}
-
-				CppIncludeLookup IncludeLookup = new CppIncludeLookup(IncludeCache);
-				IncludeLookup.Load();
-				IncludeLookup.Update(BaseDirs);
-				IncludeLookup.Save();
-
-				List<FileReference> NewFiles = IncludeLookup.FindFiles(Descriptor.SpecificFilesToCompile.Select(x => x.GetFileName())).Select(x => x.Location).ToList();
-				NewFiles.RemoveAll(x => IsHeader(x));
-				NewFiles = NewFiles.Except(Descriptor.SpecificFilesToCompile).ToList();
-
-				Logger.LogInformation("Found {NewFilesCount} source files", NewFiles.Count);
-				foreach (FileReference NewFile in NewFiles)
-				{
-					Logger.LogDebug("  {NewFile}", NewFile);
-				}
-
-				Descriptor.OptionalFilesToCompile.AddRange(NewFiles.Except(Descriptor.SpecificFilesToCompile));
-			}
 
 			// If we're compiling a plugin, and this target is monolithic, just create the object files
 			if (Descriptor.ForeignPlugin != null && RulesObject.LinkType == TargetLinkType.Monolithic)
@@ -1587,7 +1537,7 @@ namespace UnrealBuildTool
 			RulesAssembly = InRulesAssembly;
 			TargetType = Rules.Type;
 			ForeignPlugin = InDescriptor.ForeignPlugin;
-			bDeployAfterCompile = InRules.bDeployAfterCompile && !InRules.bDisableLinking && InDescriptor.SpecificFilesToCompile.Count == 0 && InDescriptor.OnlyModuleNames.Count == 0;
+			bDeployAfterCompile = InRules.bDeployAfterCompile && !InRules.bDisableLinking && InDescriptor.OnlyModuleNames.Count == 0;
 
 			// now that we have the platform, we can set the intermediate path to include the platform/architecture name
 			PlatformIntermediateFolder = GetPlatformIntermediateFolder(Platform, Architectures, false);
@@ -2433,29 +2383,13 @@ namespace UnrealBuildTool
 				}
 			}
 
-			// Find the set of binaries to build. If we're compiling only specific files, filter the list of binaries to only include the files we're interested in.
-			HashSet<FileReference> SpecificFilesToCompile = new();
-			SpecificFilesToCompile.UnionWith(TargetDescriptor.SpecificFilesToCompile);
-			SpecificFilesToCompile.UnionWith(TargetDescriptor.OptionalFilesToCompile);
-			if (TargetDescriptor.SpecificFilesToCompile.Count > 0 && !BuildConfiguration.bIgnoreInvalidFiles)
-			{
-				foreach (FileReference SpecificFile in TargetDescriptor.SpecificFilesToCompile)
-				{
-					UEBuildBinary? Binary = Binaries.Find(x => x.Modules.Any(y => y.ContainsFile(SpecificFile)));
-					if (Binary == null)
-					{
-						throw new BuildException("Couldn't find any module containing {0} in {1}.", SpecificFile, TargetName);
-					}
-				}
-			}
-
 			// Build the target's binaries.
 			DirectoryReference ExeDir = GetExecutableDir();
 			using (GlobalTracer.Instance.BuildSpan("UEBuildBinary.Build()").StartActive())
 			{
 				foreach (UEBuildBinary Binary in Binaries)
 				{
-					List<FileItem> BinaryOutputItems = Binary.Build(Rules, TargetToolChain, GlobalCompileEnvironment, GlobalLinkEnvironment, SpecificFilesToCompile.ToList(), WorkingSet, ExeDir, MakefileBuilder, Logger);
+					List<FileItem> BinaryOutputItems = Binary.Build(Rules, TargetToolChain, GlobalCompileEnvironment, GlobalLinkEnvironment, WorkingSet, ExeDir, MakefileBuilder, Logger);
 					Makefile.OutputItems.AddRange(BinaryOutputItems);
 				}
 			}
diff --git a/Engine/Source/Programs/UnrealBuildTool/Modes/BuildMode.cs b/Engine/Source/Programs/UnrealBuildTool/Modes/BuildMode.cs
--- a/Engine/Source/Programs/UnrealBuildTool/Modes/BuildMode.cs
+++ b/Engine/Source/Programs/UnrealBuildTool/Modes/BuildMode.cs
@@ -439,11 +439,34 @@ namespace UnrealBuildTool
 					ActionGraph.CheckPathLengths(BuildConfiguration, Makefiles.SelectMany(x => x.Actions), Logger);
 				}
 
+				HashSet<FileItem> MergedOutputItems = new HashSet<FileItem>();
+
 				// Create a QueuedAction instance for each action in the makefiles
 				List<LinkedAction>[] QueuedActions = new List<LinkedAction>[Makefiles.Length];
-				for(int Idx = 0; Idx < Makefiles.Length; Idx++)
+				for (int Idx = 0; Idx < Makefiles.Length; Idx++)
 				{
-					QueuedActions[Idx] = Makefiles[Idx].Actions.ConvertAll(x => new LinkedAction(x, TargetDescriptors[Idx]));
+					TargetDescriptor TargetDescriptor = TargetDescriptors[Idx];
+					List<FileReference> FilesToBuild = TargetDescriptor.SpecificFilesToCompile;
+
+					// No specific files to filter on, queue all actions
+					if (FilesToBuild.Count == 0)
+					{
+						QueuedActions[Idx] = Makefiles[Idx].Actions.ConvertAll(x => new LinkedAction(x, TargetDescriptors[Idx]));
+					}
+					else
+					{
+						// If there are headers in the list, expand the FilesToBuild to also include all files that include those headers
+						FilesToBuild = GetAllSourceFilesIncludingHeader(FilesToBuild, TargetDescriptor.ProjectFile, Makefiles[Idx].Actions, Logger);
+
+						// We have specific files to compile so we will only queue up those files.
+						QueuedActions[Idx] = CreateLinkedActionsFromFileList(TargetDescriptor, FilesToBuild, Makefiles[Idx].Actions, Logger);
+
+						// Add all queued actions to the MergedOutputItems to make sure they are not skipped
+						foreach (LinkedAction Action in QueuedActions[Idx])
+						{
+							MergedOutputItems.Add(Action.ProducedItems.First());
+						}
+					}
 				}
 
 				// Clean up any previous hot reload runs, and reapply the current state if it's already active
@@ -465,7 +488,6 @@ namespace UnrealBuildTool
 				}
 
 				// Gather all the prerequisite actions that are part of the targets
-				HashSet<FileItem> MergedOutputItems = new HashSet<FileItem>();
 				for (int TargetIdx = 0; TargetIdx < TargetDescriptors.Count; TargetIdx++)
 				{
 					GatherOutputItems(TargetDescriptors[TargetIdx], Makefiles[TargetIdx], MergedOutputItems);
@@ -610,17 +632,6 @@ namespace UnrealBuildTool
 
 					// Plan the actions to execute for the build. For single file compiles, always rebuild the source file regardless of whether it's out of date.
 					ActionGraph.GatherAllOutdatedActions(PrerequisiteActions, History, ActionToOutdatedFlag, CppDependencies, BuildConfiguration.bIgnoreOutdatedImportLibraries, Logger);
-					if (TargetDescriptor.SpecificFilesToCompile.Count != 0)
-					{
-						HashSet<FileReference> ForceBuildFiles = new HashSet<FileReference>();
-						ForceBuildFiles.UnionWith(TargetDescriptor.SpecificFilesToCompile);
-						ForceBuildFiles.UnionWith(TargetDescriptor.OptionalFilesToCompile);
-
-						foreach (LinkedAction PrerequisiteAction in PrerequisiteActions.Where(x => x.PrerequisiteItems.Any(y => ForceBuildFiles.Contains(y.Location))))
-						{
-							ActionToOutdatedFlag[PrerequisiteAction] = true;
-						}
-					}
 				}
 
 				// Link the action graph again to sort it
@@ -767,6 +778,188 @@ namespace UnrealBuildTool
 			}
 		}
 
+		internal static List<LinkedAction> CreateLinkedActionsFromFileList(TargetDescriptor? Target, List<FileReference> FileList, List<IExternalAction> Actions, ILogger Logger)
+		{
+			// We have specific files to compile so we will only queue up those files.
+			// We will also add them to the MergedOutputItems to make sure they are not skipped
+			List<LinkedAction> LinkedActions = new();
+
+			// First, find all the SpecificFileActions.
+			// These are used to create a custom action for the source file in case it is inside a unity or normally not part of the build (headers for example)
+			Dictionary<DirectoryReference, ISpecificFileAction> SpecificFileActions = new();
+			foreach (IExternalAction Action in Actions)
+			{
+				if (Action is ISpecificFileAction SpecificFileAction)
+				{
+					SpecificFileActions.TryAdd(SpecificFileAction.RootDirectory, SpecificFileAction);
+				}
+			}
+
+			// Now traverse all specific files and make sure we find or create actions for them.
+			foreach (FileReference FileRef in FileList)
+			{
+				FileItem SourceFile = FileItem.GetItemByFileReference(FileRef);
+
+				ISpecificFileAction? SpecificFileAction = null;
+
+				// Traverse upwards from file directory to try to find a SpecificFileAction.
+				DirectoryItem Dir = SourceFile.Directory;
+				while (true)
+				{
+					if (SpecificFileActions.TryGetValue(Dir.Location, out SpecificFileAction))
+					{
+						break;
+					}
+					DirectoryItem? ParentDir = Dir.GetParentDirectoryItem();
+					if (ParentDir == null)
+					{
+						break;
+					}
+					Dir = ParentDir;
+				}
+
+				// We found an action to take care of this file.
+				if (SpecificFileAction != null)
+				{
+					IExternalAction? NewAction = SpecificFileAction.CreateAction(SourceFile, Logger);
+					if (NewAction != null)
+					{
+						LinkedActions.Add(new LinkedAction(NewAction, Target));
+						continue;
+					}
+				}
+
+				// There is no special action for this file.. let's look for an action that has this exact file as input and use that (for example ispc files)
+				bool FoundAction = false;
+				foreach (IExternalAction Action in Actions)
+				{
+					foreach (FileItem PrereqItem in Action.PrerequisiteItems)
+					{
+						if (PrereqItem == SourceFile)
+						{
+							FoundAction = true;
+							LinkedActions.Add(new LinkedAction(Action, Target));
+							break;
+						}
+					}
+					if (FoundAction)
+					{
+						break;
+					}
+				}
+
+				if (!FoundAction)
+				{
+					Logger.LogError($"{FileRef.FullName} - ERROR: Failed to find an Action that can be used to build this file (does target use this file?)");
+				}
+			}
+
+			return LinkedActions;
+		}
+
+
+		/// <summary>
+		/// If there are headers in the specific list, this function will exand the list to include all files that include those headers
+		/// </summary>
+		/// <param name="SpecificFilesToCompile">List of files to compile</param>
+		/// <param name="ProjectFile">Project file if there is one</param>
+		/// <param name="Actions">Actions to filter out expansion of files</param>
+		/// <param name="Logger">Logger for output</param>
+		internal static List<FileReference> GetAllSourceFilesIncludingHeader(List<FileReference> SpecificFilesToCompile, FileReference? ProjectFile, List<IExternalAction> Actions, ILogger Logger)
+		{
+			Func<FileReference, bool> IsCode = x => IsHeader(x)|| x.HasExtension(".cpp") || x.HasExtension(".c");
+			List<FileReference> SpecificHeaderFiles = SpecificFilesToCompile.Where(IsHeader).ToList();
+			if (SpecificHeaderFiles.Count == 0)
+				return SpecificFilesToCompile;
+
+			// TODO: This code below works very similar to the code before this changelist but is not a great solution for figuring out which files to compile.
+			// Suggested approach is to:
+			// 1. Write out a file with additional include paths per module.
+			// 2. Use the special single file actions to figure out if the specific headers are included by the source files using additional include paths
+			// 3. Fix so CppIncludeLookup is including headers... Current code does not handle this: A.h <- B.h <- C.cpp... C.cpp is not added.
+			// 4. Cache which files unity files include to reduce number of files needed to be read
+
+			List<DirectoryReference> BaseDirs = new List<DirectoryReference>();
+			BaseDirs.Add(Unreal.EngineDirectory);
+
+			DirectoryReference CacheDir = Unreal.EngineDirectory;
+			if (ProjectFile != null)
+			{
+				BaseDirs.Add(ProjectFile.Directory);
+				CacheDir = ProjectFile.Directory;
+			}
+
+			FileReference IncludeCache = FileReference.Combine(CacheDir, "Intermediate", "HeaderLookup.dat");
+
+			Logger.LogInformation("Building dependency cache for specified headers");
+			foreach (FileReference HeaderFile in SpecificHeaderFiles)
+			{
+				Logger.LogDebug("  {HeaderFile}", HeaderFile);
+			}
+
+			CppIncludeLookup IncludeLookup = new CppIncludeLookup(IncludeCache);
+			IncludeLookup.Load();
+			IncludeLookup.Update(BaseDirs);
+			IncludeLookup.Save();
+
+			// Find all files that can be part of this build
+			ConcurrentDictionary<FileReference, int> UsedFiles = new();
+			Parallel.ForEach(Actions, (Action) =>
+			{
+				foreach (FileItem File in Action.PrerequisiteItems)
+				{
+					if (!IsCode(File.Location))
+					{
+						continue;
+					}
+
+					// If file is a unity file we want to traverse the internal files and add them as UsedFiles
+					if (File.Name.StartsWith(Unity.ModulePrefix))
+					{
+						foreach (string Line in System.IO.File.ReadAllLines(File.FullName))
+						{
+							int IncludeStart = Line.IndexOf('"') + 1;
+							if (IncludeStart != 0)
+							{
+								string Include = Line.Substring(IncludeStart, Line.Length - IncludeStart - 1);
+								FileItem SourceFile = FileItem.GetItemByFileReference(FileReference.Combine(Unreal.EngineSourceDirectory, Include));
+								UsedFiles.TryAdd(SourceFile.Location, 0);
+							}
+						}
+					}
+					else
+					{
+						UsedFiles.TryAdd(File.Location, 0);
+					}
+				}
+			});
+
+
+			SortedSet<FileReference> NewFiles = new();
+			List<FileReference> FoundFiles = IncludeLookup.FindFiles(SpecificFilesToCompile.Select(x => x.GetFileName())).Select(x => x.Location).ToList();
+
+			// Found files can be a lot (Commandline.h gives 130.000+ files). And we only want to build the ones that are in existing actions
+			Parallel.ForEach(FoundFiles, (File) =>
+			{
+				if (UsedFiles.ContainsKey(File))
+				{
+					lock (NewFiles)
+					{
+						NewFiles.Add(File);
+					}
+				}
+			});
+
+			Logger.LogInformation("Found {NewFilesCount} source files", NewFiles.Count);
+			foreach (FileReference NewFile in NewFiles)
+			{
+				Logger.LogDebug("  {NewFile}", NewFile);
+			}
+
+			NewFiles.UnionWith(SpecificFilesToCompile);
+			return NewFiles.ToList(); ;
+		}
+
 		/// <summary>
 		/// Outputs the toolchain used to build each target
 		/// </summary>
@@ -803,6 +996,12 @@ namespace UnrealBuildTool
 			}
 		}
 
+		/// <summary>
+		/// Returns true if File is header
+		/// </summary>
+		/// <param name="File">File to check</param>
+		internal static bool IsHeader(FileReference File) { return File.HasExtension(".h") || File.HasExtension(".hpp") || File.HasExtension(".hxx"); }
+
 		/// <summary>
 		/// Creates the makefile for a target. If an existing, valid makefile already exists on disk, loads that instead.
 		/// </summary>
@@ -815,7 +1014,7 @@ namespace UnrealBuildTool
 		{
 			// Get the path to the makefile for this target
 			FileReference? MakefileLocation = null;
-			if(BuildConfiguration.bUseUBTMakefiles && TargetDescriptor.SpecificFilesToCompile.Count == 0)
+			if(BuildConfiguration.bUseUBTMakefiles)
 			{
 				MakefileLocation = TargetMakefile.GetLocation(TargetDescriptor.ProjectFile, TargetDescriptor.Name, TargetDescriptor.Platform, TargetDescriptor.Architectures, TargetDescriptor.Configuration);
 			}
@@ -946,11 +1145,6 @@ namespace UnrealBuildTool
 		{
 			if(TargetDescriptor.SpecificFilesToCompile.Count > 0)
 			{
-				// If we're just compiling a specific files, set the target items to be all the derived items
-				List<FileItem> FilesToCompile = TargetDescriptor.SpecificFilesToCompile.Union(TargetDescriptor.OptionalFilesToCompile).Select(x => FileItem.GetItemByFileReference(x)).ToList();
-				OutputItems.UnionWith(
-					Makefile.Actions.Where(x => x.PrerequisiteItems.Any(y => FilesToCompile.Contains(y)))
-					.SelectMany(x => x.ProducedItems));
 			}
 			else if(TargetDescriptor.OnlyModuleNames.Count > 0)
 			{
@@ -988,6 +1182,11 @@ namespace UnrealBuildTool
 				string GroupPrefix = String.Format("{0}-{1}-{2}", TargetDescriptors[TargetIdx].Name, TargetDescriptors[TargetIdx].Platform, TargetDescriptors[TargetIdx].Configuration);
 				foreach(LinkedAction TargetAction in TargetActions[TargetIdx])
 				{
+					if (!TargetAction.ProducedItems.Any())
+					{
+						continue;
+					}
+					
 					if (TargetAction.IgnoreConflicts())
 					{
 						IgnoreConflictActions.Add(TargetAction);
diff --git a/Engine/Source/Programs/UnrealBuildTool/Modes/GenerateClangDatabase.cs b/Engine/Source/Programs/UnrealBuildTool/Modes/GenerateClangDatabase.cs
--- a/Engine/Source/Programs/UnrealBuildTool/Modes/GenerateClangDatabase.cs
+++ b/Engine/Source/Programs/UnrealBuildTool/Modes/GenerateClangDatabase.cs
@@ -180,55 +180,64 @@ namespace UnrealBuildTool
 								InputFiles.AddRange(InputFileCollection.CPPFiles);
 								InputFiles.AddRange(InputFileCollection.CCFiles);
 
-								var fileList = new List<FileReference>();
+								HashSet<FileItem> FilterFileList = new();
 								foreach (FileItem InputFile in InputFiles)
 								{
 									if (FileFilter == null || FileFilter.Matches(InputFile.Location.MakeRelativeTo(Unreal.RootDirectory)))
 									{
-										var fileRef = new FileReference(InputFile.AbsolutePath);
-										fileList.Add(fileRef);
+										FilterFileList.Add(InputFile);
 									}
 								}
 
-								if (fileList.Any())
-								{
-									CaptureActionGraphBuilder ActionGraphBuilder = new CaptureActionGraphBuilder(Logger);
-
-									Module.Compile(Target.Rules, TargetToolChain, BinaryCompileEnvironment, fileList, WorkingSet, ActionGraphBuilder, Logger);
+								CaptureActionGraphBuilder ActionGraphBuilder = new CaptureActionGraphBuilder(Logger);
 
-									IEnumerable<IExternalAction> CompileActions = ActionGraphBuilder.CapturedActions.Where(Action => Action.ActionType == ActionType.Compile);
+								Module.Compile(Target.Rules, TargetToolChain, BinaryCompileEnvironment, WorkingSet, ActionGraphBuilder, Logger);
 
-									if (CompileActions.Any())
+								List<IExternalAction> ValidActions = new();
+								foreach (IExternalAction Action in ActionGraphBuilder.CapturedActions)
+								{
+									if (Action.ActionType == ActionType.Compile && Action.ProducedItems.Any())
 									{
-										// convert any rsp files
-										Dictionary<string, string> UpdatedResFiles = new Dictionary<string, string>();
-										foreach (Tuple<FileItem, IEnumerable<string>> FileAndContents in ActionGraphBuilder.CapturedTextFiles)
+										foreach (FileItem Prereq in Action.PrerequisiteItems)
 										{
-											if (FileAndContents.Item1.AbsolutePath.EndsWith(".rsp") || FileAndContents.Item1.AbsolutePath.EndsWith(".response"))
+											if (FilterFileList.Contains(Prereq))
 											{
-												string NewResPath = ConvertResponseFile(FileAndContents.Item1, FileAndContents.Item2, Logger);
-												UpdatedResFiles[FileAndContents.Item1.AbsolutePath] = NewResPath;
+												ValidActions.Add(Action);
 											}
 										}
+									}
+								}
 
-										foreach (IExternalAction Action in CompileActions)
+								if (ValidActions.Count != 0)
+								{
+									// convert any rsp files
+									Dictionary<string, string> UpdatedResFiles = new Dictionary<string, string>();
+									foreach (Tuple<FileItem, IEnumerable<string>> FileAndContents in ActionGraphBuilder.CapturedTextFiles)
+									{
+										if (FileAndContents.Item1.AbsolutePath.EndsWith(".rsp") || FileAndContents.Item1.AbsolutePath.EndsWith(".response"))
 										{
-											// Create the command
-											StringBuilder CommandBuilder = new StringBuilder();
-											string CommandArguments = Action.CommandArguments.Replace(".rsp", ".rsp.gcd").Replace(".response", ".response.gcd");
-											CommandBuilder.AppendFormat("\"{0}\" {1}", Action.CommandPath, CommandArguments);
+											string NewResPath = ConvertResponseFile(FileAndContents.Item1, FileAndContents.Item2, Logger);
+											UpdatedResFiles[FileAndContents.Item1.AbsolutePath] = NewResPath;
+										}
+									}
 
-											foreach (string ExtraArgument in GetExtraPlatformArguments(TargetToolChain))
-											{
-												CommandBuilder.AppendFormat(" {0}", ExtraArgument);
-											}
+									foreach (IExternalAction Action in ValidActions)
+									{
+										// Create the command
+										StringBuilder CommandBuilder = new StringBuilder();
+										string CommandArguments = Action.CommandArguments.Replace(".rsp", ".rsp.gcd").Replace(".response", ".response.gcd");
+										CommandBuilder.AppendFormat("\"{0}\" {1}", Action.CommandPath, CommandArguments);
 
-											// find source file
-											var SourceFile = Action.PrerequisiteItems.FirstOrDefault(fi => fi.HasExtension(".cpp") || fi.HasExtension(".c") || fi.HasExtension(".c"));
-											if (SourceFile != null)
-											{
-												FileToCommand[SourceFile.Location] = CommandBuilder.ToString();
-											}
+										foreach (string ExtraArgument in GetExtraPlatformArguments(TargetToolChain))
+										{
+											CommandBuilder.AppendFormat(" {0}", ExtraArgument);
+										}
+
+										// find source file
+										var SourceFile = Action.PrerequisiteItems.FirstOrDefault(fi => fi.HasExtension(".cpp") || fi.HasExtension(".c") || fi.HasExtension(".c"));
+										if (SourceFile != null)
+										{
+											FileToCommand[SourceFile.Location] = CommandBuilder.ToString();
 										}
 									}
 								}
@@ -253,6 +262,8 @@ namespace UnrealBuildTool
 					}
 					Writer.WriteArrayEnd();
 				}
+				Logger.LogInformation($"ClangDatabase written to {DatabaseFile.FullName}");
+
 			}
 
 			return 0;
diff --git a/Engine/Source/Programs/UnrealBuildTool/Platform/Windows/VCSpecificFileAction.cs b/Engine/Source/Programs/UnrealBuildTool/Platform/Windows/VCSpecificFileAction.cs
new file mode 100644
--- /dev/null
+++ b/Engine/Source/Programs/UnrealBuildTool/Platform/Windows/VCSpecificFileAction.cs
@@ -0,0 +1,87 @@
+// Copyright Epic Games, Inc. All Rights Reserved.
+
+using EpicGames.Core;
+using Microsoft.Extensions.Logging;
+using System.Collections.Generic;
+using System.IO;
+using UnrealBuildBase;
+
+namespace UnrealBuildTool
+{
+	internal class VcSpecificFileAction : Action, ISpecificFileAction
+	{
+		DirectoryReference SourceDir;
+		DirectoryReference OutputDir;
+		VCCompileAction BaseAction;
+		int SingleFileCounter;
+
+		internal VcSpecificFileAction(DirectoryReference Source, DirectoryReference Output, VCCompileAction Action) : base(ActionType.Compile)
+		{
+			SourceDir = Source;
+			OutputDir = Output;
+			BaseAction = Action;
+		}
+
+		public VcSpecificFileAction(BinaryArchiveReader Reader) : base(ActionType.Compile)
+		{
+			SourceDir = Reader.ReadCompactDirectoryReference();
+			OutputDir = Reader.ReadCompactDirectoryReference();
+			BaseAction = new VCCompileAction(Reader);
+		}
+
+		public new void Write(BinaryArchiveWriter Writer)
+		{
+			Writer.WriteCompactDirectoryReference(SourceDir);
+			Writer.WriteCompactDirectoryReference(OutputDir);
+			BaseAction.Write(Writer);
+		}
+
+		public DirectoryReference RootDirectory { get => SourceDir; }
+
+		public IExternalAction? CreateAction(FileItem SourceFile, ILogger Logger)
+		{
+			// If it is a header file we need to wrap it from another file.. because otherwise it will fail when there are circular dependencies
+			// (There are a lot of those, very few of the .h files in core can be compiled without errors)
+			if (SourceFile.HasExtension(".h"))
+			{
+				FileItem DummyFile = FileItem.GetItemByFileReference(FileReference.Combine(OutputDir, "SingleFile", SourceFile.Name));
+				Directory.CreateDirectory(DummyFile.Directory.FullName);
+				File.WriteAllText(DummyFile.FullName, $"#include \"{SourceFile.FullName.Replace('\\', '/')}\"");
+				SourceFile = DummyFile;
+			}
+			else if (!SourceFile.HasExtension(".cpp"))
+			{
+				return null;
+			}
+
+			FileItem ResponseFile = FileItem.GetItemByFileReference(FileReference.Combine(OutputDir, $"SingleFile{SingleFileCounter}.rsp"));
+			FileItem ObjectFile = FileItem.GetItemByFileReference(FileReference.Combine(OutputDir, $"SingleFile{SingleFileCounter}.obj"));
+			++SingleFileCounter;
+
+			VCCompileAction Action = new VCCompileAction(BaseAction);
+			Action.ObjectFile = ObjectFile;
+			Action.SourceFile = SourceFile;
+			FileItem TempFile = FileItem.GetItemByPath(System.IO.Path.GetTempFileName());
+			Action.ResponseFile = ResponseFile;
+			List<string> Arguments = Action.GetCompilerArguments(Logger);
+			System.IO.File.WriteAllLines(ResponseFile.FullName, Arguments);
+			return Action;
+		}
+	}
+
+	class VcSpecificFileActionSerializer : ActionSerializerBase<VcSpecificFileAction>
+	{
+		/// <inheritdoc/>
+		public override VcSpecificFileAction Read(BinaryArchiveReader Reader)
+		{
+			return new VcSpecificFileAction(Reader);
+		}
+
+		/// <inheritdoc/>
+		public override void Write(BinaryArchiveWriter Writer, VcSpecificFileAction Action)
+		{
+			Action.Write(Writer);
+		}
+	}
+
+}
\ No newline at end of file
diff --git a/Engine/Source/Programs/UnrealBuildTool/Platform/Windows/VCToolChain.cs b/Engine/Source/Programs/UnrealBuildTool/Platform/Windows/VCToolChain.cs
--- a/Engine/Source/Programs/UnrealBuildTool/Platform/Windows/VCToolChain.cs
+++ b/Engine/Source/Programs/UnrealBuildTool/Platform/Windows/VCToolChain.cs
@@ -2,11 +2,8 @@
 
 using System;
 using System.Collections.Generic;
-using System.Text.RegularExpressions;
-using System.Diagnostics;
 using System.IO;
 using System.Linq;
-using Microsoft.Win32;
 using System.Text;
 using EpicGames.Core;
 using UnrealBuildBase;
@@ -1368,7 +1365,7 @@ namespace UnrealBuildTool
 			}
 		}
 
-		protected override CPPOutput CompileCPPFiles(CppCompileEnvironment CompileEnvironment, List<FileItem> InputFiles, DirectoryReference OutputDir, string ModuleName, IActionGraphBuilder Graph)
+		private VCCompileAction CreateBaseCompileAction(CppCompileEnvironment CompileEnvironment)
 		{
 			VCCompileAction BaseCompileAction = new VCCompileAction(EnvVars);
 
@@ -1422,6 +1419,12 @@ namespace UnrealBuildTool
 					BaseCompileAction.Arguments.Add("/d1reportTime");
 				}
 			}
+			return BaseCompileAction;
+		}
+
+		protected override CPPOutput CompileCPPFiles(CppCompileEnvironment CompileEnvironment, List<FileItem> InputFiles, DirectoryReference OutputDir, string ModuleName, IActionGraphBuilder Graph)
+		{
+			VCCompileAction BaseCompileAction = CreateBaseCompileAction(CompileEnvironment);
 
 			// Create a compile action for each source file.
 			List<VCCompileAction> Actions = new List<VCCompileAction>();
@@ -2082,6 +2085,13 @@ namespace UnrealBuildTool
 			return NewCompileEnvironment;
 		}
 
+		public override void CreateSpecificFileAction(CppCompileEnvironment CompileEnvironment, DirectoryReference SourceDir, DirectoryReference OutputDir, IActionGraphBuilder Graph)
+		{
+			VCCompileAction BaseCompileAction = CreateBaseCompileAction(CompileEnvironment);
+			AppendCLArguments_CPP(CompileEnvironment, BaseCompileAction.Arguments);
+			Graph.AddAction(new VcSpecificFileAction(SourceDir, OutputDir, BaseCompileAction));
+		}
+
 		/// <summary>
 		/// Macros passed via the command line have their quotes stripped, and are tokenized before being re-stringized by the compiler. This conversion
 		/// back and forth is normally ok, but certain characters such as single quotes must always be paired. Remove any such characters here.
diff --git a/Engine/Source/Programs/UnrealBuildTool/System/Action.cs b/Engine/Source/Programs/UnrealBuildTool/System/Action.cs
--- a/Engine/Source/Programs/UnrealBuildTool/System/Action.cs
+++ b/Engine/Source/Programs/UnrealBuildTool/System/Action.cs
@@ -10,6 +10,7 @@ using System.Text;
 using System.Threading;
 using System.Threading.Tasks;
 using EpicGames.Core;
+using Microsoft.Extensions.Logging;
 using UnrealBuildBase;
 
 namespace UnrealBuildTool
@@ -432,6 +433,23 @@ namespace UnrealBuildTool
 		}
 	}
 
+	/// <summary>
+	/// Interface that is used when -SingleFile=xx is used on the cmd line to make specific files compile
+	/// Actions that has this interface can be used to create real actions that can compile the specific files
+	/// </summary>
+	interface ISpecificFileAction
+	{
+		/// <summary>
+		/// The directory this action can create actions for
+		/// </summary>
+		DirectoryReference RootDirectory { get; }
+
+		/// <summary>
+		/// Creates an action for a specific file. It can return null if for example file extension is not handled
+		/// </summary>
+		IExternalAction? CreateAction(FileItem SourceFile, ILogger Logger);
+	}
+
 	/// <summary>
 	/// Extension methods for action classes
 	/// </summary>
diff --git a/Engine/Source/Programs/UnrealBuildTool/System/PrecompiledHeaderInstance.cs b/Engine/Source/Programs/UnrealBuildTool/System/PrecompiledHeaderInstance.cs
--- a/Engine/Source/Programs/UnrealBuildTool/System/PrecompiledHeaderInstance.cs
+++ b/Engine/Source/Programs/UnrealBuildTool/System/PrecompiledHeaderInstance.cs
@@ -19,6 +19,11 @@ namespace UnrealBuildTool
 		/// </summary>
 		public FileItem HeaderFile;
 
+		/// <summary>
+		/// The definitions file
+		/// </summary>
+		public FileItem DefinitionsFile;
+
 		/// <summary>
 		/// The compile environment for this shared PCH
 		/// </summary>
@@ -34,14 +39,21 @@ namespace UnrealBuildTool
 		/// </summary>
 		public List<UEBuildModuleCPP> Modules = new List<UEBuildModuleCPP>();
 
+		/// <summary>
+		/// These are definitions that are immutable and should never be #undef. There are a few exceptions and we make sure those are not ending up in this list
+		/// </summary>
+		public HashSet<string> ImmutableDefinitions;
+
 		/// <summary>
 		/// Constructor
 		/// </summary>
-		public PrecompiledHeaderInstance(FileItem HeaderFile, CppCompileEnvironment CompileEnvironment, CPPOutput Output)
+		public PrecompiledHeaderInstance(FileItem HeaderFile, FileItem DefinitionsFile, CppCompileEnvironment CompileEnvironment, CPPOutput Output, HashSet<string> ImmutableDefinitions)
 		{
 			this.HeaderFile = HeaderFile;
+			this.DefinitionsFile = DefinitionsFile;
 			this.CompileEnvironment = CompileEnvironment;
 			this.Output = Output;
+			this.ImmutableDefinitions = ImmutableDefinitions;
 		}
 
 		/// <summary>
diff --git a/Engine/Source/Programs/UnrealBuildTool/ToolChain/ClangSpecificFileAction.cs b/Engine/Source/Programs/UnrealBuildTool/ToolChain/ClangSpecificFileAction.cs
new file mode 100644
--- /dev/null
+++ b/Engine/Source/Programs/UnrealBuildTool/ToolChain/ClangSpecificFileAction.cs
@@ -0,0 +1,118 @@
+// Copyright Epic Games, Inc. All Rights Reserved.
+
+using EpicGames.Core;
+using Microsoft.Extensions.Logging;
+using System.Collections.Generic;
+using System.IO;
+using System.Linq;
+using UnrealBuildBase;
+
+namespace UnrealBuildTool
+{
+	/// <summary>
+	/// Single
+	/// </summary>
+	internal class ClangSpecificFileAction : Action, ISpecificFileAction
+	{
+		DirectoryReference SourceDir;
+		DirectoryReference OutputDir;
+		IEnumerable<string> RspLines;
+		int SingleFileCounter;
+
+		internal ClangSpecificFileAction(DirectoryReference Source, DirectoryReference Output, Action Action, IEnumerable<string> ContentLines) : base(Action)
+		{
+			SourceDir = Source;
+			OutputDir = Output;
+			RspLines = ContentLines;
+		}
+
+		public ClangSpecificFileAction(BinaryArchiveReader Reader) : base(Reader)
+		{
+			SourceDir = Reader.ReadCompactDirectoryReference();
+			OutputDir = Reader.ReadCompactDirectoryReference();
+			RspLines = Reader.ReadList(() => Reader.ReadString())!;
+		}
+
+		public new void Write(BinaryArchiveWriter Writer)
+		{
+			base.Write(Writer);
+			Writer.WriteCompactDirectoryReference(SourceDir);
+			Writer.WriteCompactDirectoryReference(OutputDir);
+			Writer.WriteList(RspLines.ToList(), (Str) => Writer.WriteString(Str));
+		}
+
+		public DirectoryReference RootDirectory { get => SourceDir; }
+
+		public IExternalAction? CreateAction(FileItem SourceFile, ILogger Logger)
+		{
+			string DummyName = "SingleFile.cpp";
+			string UniqueDummyName = $"SingleFile{SingleFileCounter}.cpp";
+			++SingleFileCounter;
+
+			int FileNameIndex = CommandArguments.IndexOf(DummyName);
+			string DummyPath = CommandArguments.Substring(2, FileNameIndex + DummyName.Length - 2);
+
+			if (SourceFile.HasExtension(".h"))
+			{
+				FileItem DummyFile = FileItem.GetItemByFileReference(FileReference.Combine(OutputDir, "SingleFile", SourceFile.Name));
+				Directory.CreateDirectory(DummyFile.Directory.FullName);
+				File.WriteAllText(DummyFile.FullName, $"#include \"{SourceFile.FullName.Replace('\\', '/')}\"");
+				SourceFile = DummyFile;
+			}
+
+			List<string> NewRspLines = new();
+			foreach (string L in RspLines)
+			{
+				string Line = L;
+				if (Line.Contains(".cpp.o") || Line.Contains(".cpp.d"))
+				{
+					Line = Line.Replace("SingleFile.cpp", UniqueDummyName);
+				}
+				else
+				{
+					Line = Line.Replace(DummyPath, SourceFile.FullName.Replace('\\', '/'));
+				}
+				NewRspLines.Add(Line);
+			}
+
+			Action Action = new Action(this);
+			Action.CommandArguments = CommandArguments.Replace(DummyName, UniqueDummyName);
+			Action.DependencyListFile = null;
+			Action.StatusDescription = SourceFile.Name;
+
+			FileItem ResponseFile = FileItem.GetItemByPath(Action.CommandArguments.Substring(1).Trim('"'));
+			File.WriteAllLines(ResponseFile.FullName, NewRspLines);
+
+			return Action;
+		}
+	}
+
+	class ClangSpecificFileActionSerializer : ActionSerializerBase<ClangSpecificFileAction>
+	{
+		/// <inheritdoc/>
+		public override ClangSpecificFileAction Read(BinaryArchiveReader Reader)
+		{
+			return new ClangSpecificFileAction(Reader);
+		}
+
+		/// <inheritdoc/>
+		public override void Write(BinaryArchiveWriter Writer, ClangSpecificFileAction Action)
+		{
+			Action.Write(Writer);
+		}
+	}
+
+	class ClangSpecificFileActionGraphBuilder : ForwardingActionGraphBuilder
+	{
+		public ClangSpecificFileActionGraphBuilder(ILogger Logger) : base(new NullActionGraphBuilder(Logger))
+		{
+		}
+		public override void CreateIntermediateTextFile(FileItem Location, IEnumerable<string> ContentLines, bool AllowAsync = true)
+		{
+			this.ContentLines = ContentLines;
+		}
+
+		public IEnumerable<string> ContentLines = Enumerable.Empty<string>();
+	}
+
+}
diff --git a/Engine/Source/Programs/UnrealBuildTool/ToolChain/ClangToolChain.cs b/Engine/Source/Programs/UnrealBuildTool/ToolChain/ClangToolChain.cs
--- a/Engine/Source/Programs/UnrealBuildTool/ToolChain/ClangToolChain.cs
+++ b/Engine/Source/Programs/UnrealBuildTool/ToolChain/ClangToolChain.cs
@@ -844,9 +844,6 @@ namespace UnrealBuildTool
 			// Add include paths to the argument list.
 			GetCompileArguments_IncludePaths(CompileEnvironment, Arguments);
 
-			// Add force include paths to the argument list.
-			GetCompileArguments_ForceInclude(CompileEnvironment, Arguments);
-
 			// Add preprocessor definitions to the argument list.
 			GetCompileArguments_PreprocessorDefinitions(CompileEnvironment, Arguments);
 
@@ -894,6 +891,10 @@ namespace UnrealBuildTool
 				Arguments.Add(GetResponseFileArgument(AdditionalResponseFile));
 			}
 
+			// Add force include paths to the argument list.
+			GetCompileArguments_ForceInclude(CompileEnvironment, Arguments);
+
+
 			// Add the C++ source file and its included files to the prerequisite item list.
 			CompileAction.PrerequisiteItems.UnionWith(CompileEnvironment.ForceIncludeFiles);
 			CompileAction.PrerequisiteItems.UnionWith(CompileEnvironment.AdditionalPrerequisites);
@@ -1049,6 +1050,23 @@ namespace UnrealBuildTool
 			return NewCompileEnvironment;
 		}
 
+		public override void CreateSpecificFileAction(CppCompileEnvironment CompileEnvironment, DirectoryReference SourceDir, DirectoryReference OutputDir, IActionGraphBuilder Graph)
+		{
+			List<string> GlobalArguments = new();
+			if (CompileEnvironment.SystemIncludePaths.Count != 0)
+			{
+				GetCompileArguments_Global(CompileEnvironment, GlobalArguments);
+			}
+
+			FileItem DummyFile = FileItem.GetItemByFileReference(FileReference.Combine(OutputDir, "SingleFile.cpp"));
+			CPPOutput Result = new();
+			ClangSpecificFileActionGraphBuilder GraphBuilder = new(Logger);
+			Action Action = CompileCPPFile(CompileEnvironment, DummyFile, OutputDir, "<Unknown>", GraphBuilder, GlobalArguments, Result);
+			Action.PrerequisiteItems.Clear();
+
+			Graph.AddAction(new ClangSpecificFileAction(SourceDir, OutputDir, Action, GraphBuilder.ContentLines));
+		}
+
 		protected virtual Action CompileCPPFile(CppCompileEnvironment CompileEnvironment, FileItem SourceFile, DirectoryReference OutputDir, string ModuleName, IActionGraphBuilder Graph, IReadOnlyCollection<string> GlobalArguments, CPPOutput Result)
 		{
 			Action CompileAction = Graph.CreateAction(ActionType.Compile);
diff --git a/Engine/Source/Programs/UnrealBuildTool/ToolChain/UEToolChain.cs b/Engine/Source/Programs/UnrealBuildTool/ToolChain/UEToolChain.cs
--- a/Engine/Source/Programs/UnrealBuildTool/ToolChain/UEToolChain.cs
+++ b/Engine/Source/Programs/UnrealBuildTool/ToolChain/UEToolChain.cs
@@ -157,6 +157,10 @@ namespace UnrealBuildTool
 			return CompileEnvironment;
 		}
 
+		public virtual void CreateSpecificFileAction(CppCompileEnvironment CompileEnvironment, DirectoryReference SourceDir, DirectoryReference OutputDir, IActionGraphBuilder Graph)
+		{
+		}
+
 		/// <summary>
 		/// Get the name of the response file for the current compile environment and output file
 		/// </summary>
