UnrealBuildTool: Invalidate makefile if -Manifest= is requested but is out of date or missing

#rnx

diff --git a/Engine/Source/Programs/UnrealBuildTool/System/TargetMakefile.cs b/Engine/Source/Programs/UnrealBuildTool/System/TargetMakefile.cs
--- a/Engine/Source/Programs/UnrealBuildTool/System/TargetMakefile.cs
+++ b/Engine/Source/Programs/UnrealBuildTool/System/TargetMakefile.cs
@@ -639,6 +639,24 @@ namespace UnrealBuildTool
 					return null;
 				}
 
+				// Check if any manifests are out of date
+				foreach (FileReference Manifest in Makefile.AdditionalArguments!
+					.Where(x => x.StartsWith("-Manifest=", StringComparison.OrdinalIgnoreCase))
+					.Select(x => FileReference.FromString(x.Substring("-Manifest=".Length)))
+					.Where(x => x != null))
+				{
+					if (!FileReference.Exists(Manifest))
+					{
+						ReasonNotLoaded = $"manifest '{Manifest}' not found";
+						return null;
+					}
+					else if (FileReference.GetLastWriteTimeUtc(Manifest) < Makefile.CreateTimeUtc)
+					{
+						ReasonNotLoaded = $"manifest '{Manifest}' not found";
+						return null;
+					}
+				}
+
 				// Check if any config settings have changed. Ini files contain build settings too.
 				if(!Makefile.ConfigValueTracker.IsValid())
 				{
