[UBT]
* Fixed so .gen.cpp files are also using shared rsp. Also made code a little bit simpler

diff --git a/Engine/Source/Programs/UnrealBuildTool/Configuration/UEBuildModuleCPP.cs b/Engine/Source/Programs/UnrealBuildTool/Configuration/UEBuildModuleCPP.cs
--- a/Engine/Source/Programs/UnrealBuildTool/Configuration/UEBuildModuleCPP.cs
+++ b/Engine/Source/Programs/UnrealBuildTool/Configuration/UEBuildModuleCPP.cs
@@ -452,6 +452,10 @@ namespace UnrealBuildTool
 			// Write all the definitions to a separate file
 			CreateHeaderForDefinitions(CompileEnvironment, IntermediateDirectory, null, Graph);
 
+			// Create shared rsp for the normal cpp files
+			FileReference SharedResponseFile = FileReference.Combine(IntermediateDirectory, $"{Name}.Shared{UEToolChain.ResponseExt}");
+			CompileEnvironment = ToolChain.CreateSharedResponseFile(CompileEnvironment, SharedResponseFile, Graph);
+
 			// Mapping of source file to unity file. We output this to intermediate directories for other tools (eg. live coding) to use.
 			Dictionary<FileItem, FileItem> SourceFileToUnityFile = new Dictionary<FileItem, FileItem>();
 
@@ -640,27 +644,23 @@ namespace UnrealBuildTool
 				LinkInputFiles.AddRange(ToolChain.CompileAllCPPFiles(GeneratedCPPCompileEnvironment, GeneratedFileItems, IntermediateDirectory, Name, Graph).ObjectFiles);
 			}
 
-			// Create shared rsp for the normal cpp files
-			FileReference SharedResponseFile = FileReference.Combine(IntermediateDirectory, $"{Name}.Shared{UEToolChain.ResponseExt}");
-			CppCompileEnvironment CodeCompileEnvironment = ToolChain.CreateSharedResponseFile(CompileEnvironment, SharedResponseFile, Graph);
-
 			// Compile CPP files
 			if (bModuleUsesUnityBuild)
 			{
-				Unity.GenerateUnityCPPs(Target, CPPFiles, InputFiles.HeaderFiles, CodeCompileEnvironment, WorkingSet, Rules.ShortName ?? Name, IntermediateDirectory, Graph, SourceFileToUnityFile, 
+				Unity.GenerateUnityCPPs(Target, CPPFiles, InputFiles.HeaderFiles, CompileEnvironment, WorkingSet, Rules.ShortName ?? Name, IntermediateDirectory, Graph, SourceFileToUnityFile, 
 					out List<FileItem> NormalFiles, out List<FileItem> AdaptiveFiles, NumIncludedBytesPerUnityCPP);
-				LinkInputFiles.AddRange(CompileFilesWithToolChain(Target, ToolChain, CodeCompileEnvironment, ModuleCompileEnvironment, NormalFiles, AdaptiveFiles, Graph, Logger).ObjectFiles);
+				LinkInputFiles.AddRange(CompileFilesWithToolChain(Target, ToolChain, CompileEnvironment, ModuleCompileEnvironment, NormalFiles, AdaptiveFiles, Graph, Logger).ObjectFiles);
 			}
 			else
 			{
-				Unity.GetAdaptiveFiles(Target, CPPFiles, InputFiles.HeaderFiles, CodeCompileEnvironment, WorkingSet, Rules.ShortName ?? Name, IntermediateDirectory, Graph, 
+				Unity.GetAdaptiveFiles(Target, CPPFiles, InputFiles.HeaderFiles, CompileEnvironment, WorkingSet, Rules.ShortName ?? Name, IntermediateDirectory, Graph, 
 					out List<FileItem> NormalFiles, out List<FileItem> AdaptiveFiles);
 				if (NormalFiles.Where(file => !file.HasExtension(".gen.cpp")).Count() == 0)
 				{
 					NormalFiles = CPPFiles;
 					AdaptiveFiles.RemoveAll(new HashSet<FileItem>(NormalFiles).Contains);
 				}
-				LinkInputFiles.AddRange(CompileFilesWithToolChain(Target, ToolChain, CodeCompileEnvironment, ModuleCompileEnvironment, NormalFiles, AdaptiveFiles, Graph, Logger).ObjectFiles);
+				LinkInputFiles.AddRange(CompileFilesWithToolChain(Target, ToolChain, CompileEnvironment, ModuleCompileEnvironment, NormalFiles, AdaptiveFiles, Graph, Logger).ObjectFiles);
 			}
 
 			// Compile ISPC files directly
@@ -740,7 +740,7 @@ namespace UnrealBuildTool
 			if (Target.bIWYU)
 			{
 				// Collect the headers that should be built
-				List<FileItem> HeaderFileItems = GetCompilableHeaders(InputFiles, CodeCompileEnvironment);
+				List<FileItem> HeaderFileItems = GetCompilableHeaders(InputFiles, CompileEnvironment);
 				if (HeaderFileItems.Count > 0)
 				{
 					if (Target.bIWYUHeadersOnly)
@@ -749,7 +749,7 @@ namespace UnrealBuildTool
 					}
 
 					// Add the compile actions
-					LinkInputFiles.AddRange(ToolChain.CompileAllCPPFiles(CodeCompileEnvironment, HeaderFileItems, IntermediateDirectory, Name, Graph).ObjectFiles);
+					LinkInputFiles.AddRange(ToolChain.CompileAllCPPFiles(CompileEnvironment, HeaderFileItems, IntermediateDirectory, Name, Graph).ObjectFiles);
 				}
 			}
 
