[UBT]
* Fixed so BuildConfiguration.bIgnoreInvalidFiles is checked before outputting error when building with -SingleFile containing invalid files

diff --git a/Engine/Source/Programs/UnrealBuildTool/Modes/BuildMode.cs b/Engine/Source/Programs/UnrealBuildTool/Modes/BuildMode.cs
--- a/Engine/Source/Programs/UnrealBuildTool/Modes/BuildMode.cs
+++ b/Engine/Source/Programs/UnrealBuildTool/Modes/BuildMode.cs
@@ -459,7 +459,7 @@ namespace UnrealBuildTool
 						FilesToBuild = GetAllSourceFilesIncludingHeader(FilesToBuild, TargetDescriptor.ProjectFile, Makefiles[Idx].Actions, Logger);
 
 						// We have specific files to compile so we will only queue up those files.
-						QueuedActions[Idx] = CreateLinkedActionsFromFileList(TargetDescriptor, FilesToBuild, Makefiles[Idx].Actions, Logger);
+						QueuedActions[Idx] = CreateLinkedActionsFromFileList(TargetDescriptor, BuildConfiguration, FilesToBuild, Makefiles[Idx].Actions, Logger);
 
 						// Add all queued actions to the MergedOutputItems to make sure they are not skipped
 						foreach (LinkedAction Action in QueuedActions[Idx])
@@ -778,7 +778,7 @@ namespace UnrealBuildTool
 			}
 		}
 
-		internal static List<LinkedAction> CreateLinkedActionsFromFileList(TargetDescriptor? Target, List<FileReference> FileList, List<IExternalAction> Actions, ILogger Logger)
+		internal static List<LinkedAction> CreateLinkedActionsFromFileList(TargetDescriptor Target, BuildConfiguration BuildConfiguration, List<FileReference> FileList, List<IExternalAction> Actions, ILogger Logger)
 		{
 			// We have specific files to compile so we will only queue up those files.
 			// We will also add them to the MergedOutputItems to make sure they are not skipped
@@ -848,7 +848,7 @@ namespace UnrealBuildTool
 					}
 				}
 
-				if (!FoundAction)
+				if (!FoundAction && BuildConfiguration.bIgnoreInvalidFiles)
 				{
 					Logger.LogError($"{FileRef.FullName} - ERROR: Failed to find an Action that can be used to build this file (does target use this file?)");
 				}
