[UBT]
* Fixed inverted logic in BuildConfiguration.bIgnoreInvalidFiles check for singlefile compile

diff --git a/Engine/Source/Programs/UnrealBuildTool/Modes/BuildMode.cs b/Engine/Source/Programs/UnrealBuildTool/Modes/BuildMode.cs
--- a/Engine/Source/Programs/UnrealBuildTool/Modes/BuildMode.cs
+++ b/Engine/Source/Programs/UnrealBuildTool/Modes/BuildMode.cs
@@ -848,7 +848,7 @@ namespace UnrealBuildTool
 					}
 				}
 
-				if (!FoundAction && BuildConfiguration.bIgnoreInvalidFiles)
+				if (!FoundAction && !BuildConfiguration.bIgnoreInvalidFiles)
 				{
 					Logger.LogError($"{FileRef.FullName} - ERROR: Failed to find an Action that can be used to build this file (does target use this file?)");
 				}
