Added support for Clang PCH chaining
https://clang.llvm.org/docs/PCHInternals.html#chained-precompiled-headers

Linux - EngineTestEditor -allmodules
Before:
Total PCH Size: 3.41 GB

After:
Total PCH Size: 2.65 GB

diff --git a/Engine/Source/Programs/UnrealBuildTool/Configuration/TargetRules.cs b/Engine/Source/Programs/UnrealBuildTool/Configuration/TargetRules.cs
--- a/Engine/Source/Programs/UnrealBuildTool/Configuration/TargetRules.cs
+++ b/Engine/Source/Programs/UnrealBuildTool/Configuration/TargetRules.cs
@@ -1630,14 +1630,14 @@ namespace UnrealBuildTool
 		public bool bUsePCHFiles = true;
 
 		/// <summary>
-		/// Set flags require for determinstic linking (experimental, may not be fully supported).
+		/// Set flags require for deterministic linking (experimental, may not be fully supported).
 		/// Deterministic compiling is controlled via ModuleRules.
 		/// </summary>
 		[XmlConfigFile(Category = "BuildConfiguration")]
 		public bool bDeterministic = false;
 
 		/// <summary>
-		/// Force set flags require for determinstic compiling and linking (experimental, may not be fully supported).
+		/// Force set flags require for deterministic compiling and linking (experimental, may not be fully supported).
 		/// This setting is only recommended for testing, instead:
 		/// * Set bDeterministic on a per module basis in ModuleRules to control deterministic compiling.
 		/// * Set bDeterministic on a per target basis in TargetRules to control deterministic linking.
@@ -1646,6 +1646,13 @@ namespace UnrealBuildTool
 		[XmlConfigFile(Category = "BuildConfiguration")]
 		public bool bForceDeterministic = false;
 
+		/// <summary>
+		/// Whether PCHs should be chained when compiling with clang.
+		/// </summary>
+		[CommandLine("-NoPCHChain", Value = "false")]
+		[XmlConfigFile(Category = "BuildConfiguration")]
+		public bool bChainPCHs = true;
+
 		/// <summary>
 		/// Whether PCH headers should be force included for gen.cpp files when PCH is disabled.
 		/// </summary>
@@ -3410,6 +3417,11 @@ namespace UnrealBuildTool
 			get { return Inner.bForceDeterministic; }
 		}
 
+		public bool bChainPCHs
+		{
+			get { return Inner.bChainPCHs; }
+		}
+
 		public bool bForceIncludePCHHeadersForGenCppFilesWhenPCHIsDisabled
 		{
 			get { return Inner.bForceIncludePCHHeadersForGenCppFilesWhenPCHIsDisabled; }
diff --git a/Engine/Source/Programs/UnrealBuildTool/Configuration/UEBuildBinary.cs b/Engine/Source/Programs/UnrealBuildTool/Configuration/UEBuildBinary.cs
--- a/Engine/Source/Programs/UnrealBuildTool/Configuration/UEBuildBinary.cs
+++ b/Engine/Source/Programs/UnrealBuildTool/Configuration/UEBuildBinary.cs
@@ -713,7 +713,7 @@ namespace UnrealBuildTool
 
 			HashSet<FileItem> InputFilesLookup = new();
 
-			foreach (UEBuildModule Module in Modules)
+			foreach (UEBuildModule Module in UEBuildModule.StableTopologicalSort(Modules))
 			{
 				if (Module.Binary == null || Module.Binary == this)
 				{
diff --git a/Engine/Source/Programs/UnrealBuildTool/Configuration/UEBuildModule.cs b/Engine/Source/Programs/UnrealBuildTool/Configuration/UEBuildModule.cs
--- a/Engine/Source/Programs/UnrealBuildTool/Configuration/UEBuildModule.cs
+++ b/Engine/Source/Programs/UnrealBuildTool/Configuration/UEBuildModule.cs
@@ -11,6 +11,7 @@ using System.Xml;
 using EpicGames.Core;
 using UnrealBuildBase;
 using Microsoft.Extensions.Logging;
+using System.Threading.Tasks;
 
 namespace UnrealBuildTool
 {
@@ -1189,6 +1190,7 @@ namespace UnrealBuildTool
 			{
 				Writer.WriteValue("SharedPCH", FileReference.Combine(ModuleDirectory, Rules.SharedPCHHeaderFile).FullName);
 			}
+			Writer.WriteValue("ChainSharedPCH", Rules.Target.bChainPCHs);
 
 			ExportJsonModuleArray(Writer, "PublicDependencyModules", PublicDependencyModules);
 			ExportJsonModuleArray(Writer, "PublicIncludePathModules", PublicIncludePathModules);
@@ -1271,5 +1273,90 @@ namespace UnrealBuildTool
 			}
 			Writer.WriteArrayEnd();
 		}
+
+		/// <summary>
+		/// Returns a copy of Nodes sorted by dependency.  Independent or circularly-dependent nodes should
+		/// remain in their same relative order within the original Nodes sequence.
+		/// </summary>
+		/// <param name="NodeList">The list of nodes to sort.</param>
+		static public List<UEBuildModule> StableTopologicalSort(List<UEBuildModule> NodeList)
+		{
+			int NodeCount = NodeList.Count;
+
+			// For each Node in NodeList, populated with the full circular dependency list from
+			// Node.GetAllDependencyModules()
+			List<Task<HashSet<UEBuildModule>>> NodeDependencies = new List<Task<HashSet<UEBuildModule>>>(NodeCount);
+			// Used to populate an element of NodeDependencies
+			HashSet<UEBuildModule> FetchDependencies(int NodeIndex)
+			{
+				HashSet<UEBuildModule> Dependencies = new HashSet<UEBuildModule>();
+				NodeList[NodeIndex].GetAllDependencyModules(new List<UEBuildModule>(), Dependencies, true, true, false);
+				return Dependencies;
+			}
+
+			// For each Node in NodeList, populated with the nodes with a lower index in NodeList that Node depends on
+			List<Task<HashSet<UEBuildModule>>> PrecedingDependents = new List<Task<HashSet<UEBuildModule>>>(NodeCount);
+			HashSet<UEBuildModule> ComputePrecedingDependents(int NodeIndex)
+			{
+				HashSet<UEBuildModule> Results = new HashSet<UEBuildModule>();
+
+				UEBuildModule Node = NodeList[NodeIndex];
+				HashSet<UEBuildModule> Dependencies = NodeDependencies[NodeIndex].Result;
+
+				for (int I = 0; I < NodeIndex; ++I)
+				{
+					if (NodeDependencies[I].Result.Contains(Node) && !Dependencies.Contains(NodeList[I]))
+					{
+						Results.Add(NodeList[I]);
+					}
+				}
+
+				return Results;
+			}
+
+			for (int I = 0; I < NodeCount; ++I)
+			{
+				int LocalI = I;
+				NodeDependencies.Add(Task.Run(() => FetchDependencies(LocalI)));
+				PrecedingDependents.Add(Task.Run(() => ComputePrecedingDependents(LocalI)));
+			}
+
+			List<UEBuildModule> Out = new List<UEBuildModule>(NodeCount);
+			// Write the ordered output			
+			for (int Index1 = 0; Index1 != NodeCount; ++Index1)
+			{
+				UEBuildModule Node1 = NodeList[Index1];
+				HashSet<UEBuildModule> NodesThatDependOnNode1 = PrecedingDependents[Index1].Result;
+				Out.Add(Node1);
+
+				if (NodesThatDependOnNode1.Count == 0)
+				{
+					continue;
+				}
+
+				for (int Index2 = 0; Index2 != Index1; ++Index2)
+				{
+					UEBuildModule Node2 = Out[Index2];
+
+					if (NodesThatDependOnNode1.Contains(Node2))
+					{
+						// Rotate element at Index1 into position at Index2
+						for (int Index3 = Index1; Index3 != Index2;)
+						{
+							--Index3;
+							Out[Index3 + 1] = Out[Index3];
+						}
+
+						Out[Index2] = Node1;
+
+						// Break out of this loop, because this iteration must have covered all existing cases
+						// involving the node formerly at position Index1
+						break;
+					}
+				}
+			}
+
+			return Out;
+		}
 	};
 }
diff --git a/Engine/Source/Programs/UnrealBuildTool/Configuration/UEBuildModuleCPP.cs b/Engine/Source/Programs/UnrealBuildTool/Configuration/UEBuildModuleCPP.cs
--- a/Engine/Source/Programs/UnrealBuildTool/Configuration/UEBuildModuleCPP.cs
+++ b/Engine/Source/Programs/UnrealBuildTool/Configuration/UEBuildModuleCPP.cs
@@ -8,6 +8,7 @@ using System.Text;
 using EpicGames.Core;
 using Microsoft.Extensions.Logging;
 using UnrealBuildBase;
+using static UnrealBuildTool.ProjectFile;
 
 namespace UnrealBuildTool
 {
@@ -881,12 +882,13 @@ namespace UnrealBuildTool
 		/// <summary>
 		/// Generates a precompiled header instance from the given template, or returns an existing one if it already exists
 		/// </summary>
+		/// <param name="Target">The target which owns this module</param>
 		/// <param name="ToolChain">The toolchain being used to build this module</param>
 		/// <param name="Template">The PCH template</param>
 		/// <param name="ModuleCompileEnvironment">Compile environment for the current module</param>
 		/// <param name="Graph">List of actions to be executed. Additional actions will be added to this list.</param>
 		/// <returns>Instance of a PCH</returns>
-		public PrecompiledHeaderInstance FindOrCreateSharedPCH(UEToolChain? ToolChain, PrecompiledHeaderTemplate Template, CppCompileEnvironment ModuleCompileEnvironment, IActionGraphBuilder Graph)
+		public PrecompiledHeaderInstance FindOrCreateSharedPCH(ReadOnlyTargetRules Target, UEToolChain? ToolChain, PrecompiledHeaderTemplate Template, CppCompileEnvironment ModuleCompileEnvironment, IActionGraphBuilder Graph)
 		{
 			PrecompiledHeaderInstance? Instance = Template.Instances.Find(x => IsCompatibleForSharedPCH(x.CompileEnvironment, ModuleCompileEnvironment));
 			if(Instance == null)
@@ -925,6 +927,52 @@ namespace UnrealBuildTool
 				CompileEnvironment.PrecompiledHeaderIncludeFilename = WrapperFile.Location;
 				CopySettingsForSharedPCH(ModuleCompileEnvironment, CompileEnvironment);
 
+				// Setup PCH chaining
+				PrecompiledHeaderInstance? ParentPCHInstance = null;
+				Dictionary<string, string>? DefinitionsDictionary = null;
+				if (Rules.Target.bChainPCHs)
+				{
+					// Create a lookup table for the definitions
+					DefinitionsDictionary = Definitions.ToHashSet().ToDictionary(x => x.Split('=')[0], x => (x.Split("=").Length >= 2) ? x.Split("=")[1] : string.Empty);
+
+					// Find all the dependencies of this module
+					HashSet<UEBuildModule> ReferencedModules = new HashSet<UEBuildModule>();
+					Template.Module.GetAllDependencyModules(new List<UEBuildModule>(), ReferencedModules, bIncludeDynamicallyLoaded: false, bForceCircular: false, bOnlyDirectDependencies: false);
+
+					foreach (PrecompiledHeaderTemplate ParentTemplate in CompileEnvironment.SharedPCHs.Where(x => ReferencedModules.Contains(x.Module)))
+					{
+						if (ParentTemplate.IsValidFor(CompileEnvironment))
+						{
+							bool AreSharedPCHDefinitionsCompatible(PrecompiledHeaderInstance ParentInstance)
+							{
+								if (ParentInstance.DefinitionsDictionary == null)
+								{
+									return false;
+								}
+
+								foreach (var Definition in DefinitionsDictionary)
+								{
+									if (ParentInstance.DefinitionsDictionary.TryGetValue(Definition.Key, out string? ParentValue))
+									{
+										if (ParentValue != Definition.Value)
+										{
+											return false;
+										}
+									}
+								}
+								return true;
+							}
+
+							ParentPCHInstance = ParentTemplate.Instances.Find(x => IsCompatibleForSharedPCH(x.CompileEnvironment, ModuleCompileEnvironment) && AreSharedPCHDefinitionsCompatible(x));
+							if (ParentPCHInstance != null)
+							{
+								CompileEnvironment.ParentPrecompiledHeaderFile = ParentPCHInstance.Output.PrecompiledHeaderFile;
+								break;
+							}
+						}
+					}
+				}
+
 				// Create the PCH
 				CPPOutput Output;
 				if (ToolChain == null)
@@ -935,7 +983,11 @@ namespace UnrealBuildTool
 				{
 					Output = ToolChain.CompileAllCPPFiles(CompileEnvironment, new List<FileItem>() { WrapperFile }, Template.OutputDir, "Shared", Graph);
 				}
-				Instance = new PrecompiledHeaderInstance(WrapperFile, SharedDefinitionsFileItem, CompileEnvironment, Output, GetImmutableDefinitions(Template.BaseCompileEnvironment.Definitions));
+				Instance = new PrecompiledHeaderInstance(WrapperFile, SharedDefinitionsFileItem, CompileEnvironment, Output, GetImmutableDefinitions(Template.BaseCompileEnvironment.Definitions))
+				{
+					DefinitionsDictionary = DefinitionsDictionary,
+					ParentPCHInstance = ParentPCHInstance
+				};
 				Template.Instances.Add(Instance);
 			}
 
@@ -1340,7 +1392,7 @@ namespace UnrealBuildTool
 					PrecompiledHeaderTemplate? Template = CompileEnvironment.SharedPCHs.FirstOrDefault(x => ReferencedModules.Contains(x.Module));
 					if(Template != null && Template.IsValidFor(CompileEnvironment))
 					{
-						PrecompiledHeaderInstance Instance = FindOrCreateSharedPCH(ToolChain, Template, CompileEnvironment, Graph);
+						PrecompiledHeaderInstance Instance = FindOrCreateSharedPCH(Target, ToolChain, Template, CompileEnvironment, Graph);
 
 						FileReference PrivateDefinitionsFile = FileReference.Combine(IntermediateDirectory, String.Format("Definitions.{0}.h", Name));
 
diff --git a/Engine/Source/Programs/UnrealBuildTool/Modes/PrintBuildGraphInfo.cs b/Engine/Source/Programs/UnrealBuildTool/Modes/PrintBuildGraphInfo.cs
--- a/Engine/Source/Programs/UnrealBuildTool/Modes/PrintBuildGraphInfo.cs
+++ b/Engine/Source/Programs/UnrealBuildTool/Modes/PrintBuildGraphInfo.cs
@@ -164,6 +164,10 @@ namespace UnrealBuildTool
 							foreach (PrecompiledHeaderInstance Instance in SortedInstances)
 							{
 								Logger.LogInformation("   {InstanceName} - Used by {TimesUsed} modules:", Instance.CompileEnvironment.PrecompiledHeaderIncludeFilename, Instance.Modules.Count);
+								if (Instance.ParentPCHInstance != null)
+								{
+									Logger.LogInformation("    ParentPCH: {ParentInstanceName}", Instance.ParentPCHInstance.CompileEnvironment.PrecompiledHeaderIncludeFilename);
+								}
 
 								List<UEBuildModuleCPP> SortedModules = new(Instance.Modules);
 								SortedModules.SortBy(Module => Module.Name);
diff --git a/Engine/Source/Programs/UnrealBuildTool/Platform/Windows/UEBuildWindows.cs b/Engine/Source/Programs/UnrealBuildTool/Platform/Windows/UEBuildWindows.cs
--- a/Engine/Source/Programs/UnrealBuildTool/Platform/Windows/UEBuildWindows.cs
+++ b/Engine/Source/Programs/UnrealBuildTool/Platform/Windows/UEBuildWindows.cs
@@ -1061,6 +1061,13 @@ namespace UnrealBuildTool
 			{
 				Target.bUsePCHFiles = false;
 			}
+			
+			// Disable chaining PCH files if not using clang
+			if (!Target.WindowsPlatform.Compiler.IsClang())
+			{
+				Target.bChainPCHs = false;
+			}
+
 
 			// E&C support.
 			if (Target.bSupportEditAndContinue || Target.bAdaptiveUnityEnablesEditAndContinue)
diff --git a/Engine/Source/Programs/UnrealBuildTool/System/CppCompileEnvironment.cs b/Engine/Source/Programs/UnrealBuildTool/System/CppCompileEnvironment.cs
--- a/Engine/Source/Programs/UnrealBuildTool/System/CppCompileEnvironment.cs
+++ b/Engine/Source/Programs/UnrealBuildTool/System/CppCompileEnvironment.cs
@@ -497,6 +497,11 @@ namespace UnrealBuildTool
 		/// </summary>
 		public FileItem? PrecompiledHeaderFile = null;
 
+		/// <summary>
+		/// The parent file containing the precompiled header data.
+		/// </summary>
+		public FileItem? ParentPrecompiledHeaderFile = null;
+
 		/// <summary>
 		/// A dictionary of PCH files for multiple architectures
 		/// </summary>
@@ -645,6 +650,7 @@ namespace UnrealBuildTool
 			AdditionalArguments = Other.AdditionalArguments;
 			AdditionalFrameworks.AddRange(Other.AdditionalFrameworks);
 			PrecompiledHeaderFile = Other.PrecompiledHeaderFile;
+			ParentPrecompiledHeaderFile = Other.ParentPrecompiledHeaderFile;
 			bHackHeaderGenerator = Other.bHackHeaderGenerator;
 			bHideSymbolsByDefault = Other.bHideSymbolsByDefault;
 			CppStandard = Other.CppStandard;
diff --git a/Engine/Source/Programs/UnrealBuildTool/System/ExternalExecution.cs b/Engine/Source/Programs/UnrealBuildTool/System/ExternalExecution.cs
--- a/Engine/Source/Programs/UnrealBuildTool/System/ExternalExecution.cs
+++ b/Engine/Source/Programs/UnrealBuildTool/System/ExternalExecution.cs
@@ -289,91 +289,6 @@ namespace UnrealBuildTool
 			return Type.GetValueOrDefault();
 		}
 
-		/// <summary>
-		/// Returns a copy of Nodes sorted by dependency.  Independent or circularly-dependent nodes should
-		/// remain in their same relative order within the original Nodes sequence.
-		/// </summary>
-		/// <param name="NodeList">The list of nodes to sort.</param>
-		static List<UEBuildModuleCPP> StableTopologicalSort(List<UEBuildModuleCPP> NodeList)
-		{
-			int NodeCount = NodeList.Count;
-
-			// For each Node in NodeList, populated with the full circular dependency list from
-			// Node.GetAllDependencyModules()
-			List<Task<HashSet<UEBuildModule>>> NodeDependencies = new List<Task<HashSet<UEBuildModule>>>(NodeCount);
-			// Used to populate an element of NodeDependencies
-			HashSet<UEBuildModule> FetchDependencies(int NodeIndex)
-			{
-				HashSet<UEBuildModule> Dependencies = new HashSet<UEBuildModule>();
-				NodeList[NodeIndex].GetAllDependencyModules(new List<UEBuildModule>(), Dependencies, true, true, false);
-				return Dependencies;
-			}
-
-			// For each Node in NodeList, populated with the nodes with a lower index in NodeList that Node depends on
-			List<Task<HashSet<UEBuildModuleCPP>>> PrecedingDependents = new List<Task<HashSet<UEBuildModuleCPP>>>(NodeCount);
-			HashSet<UEBuildModuleCPP> ComputePrecedingDependents(int NodeIndex)
-			{
-				HashSet<UEBuildModuleCPP> Results = new HashSet<UEBuildModuleCPP>();
-
-				UEBuildModuleCPP Node = NodeList[NodeIndex];
-				HashSet<UEBuildModule> Dependencies = NodeDependencies[NodeIndex].Result;
-
-				for (int I = 0; I < NodeIndex; ++I)
-				{
-					if (NodeDependencies[I].Result.Contains(Node) && !Dependencies.Contains(NodeList[I]))
-					{
-						Results.Add(NodeList[I]);
-					}
-				}
-				
-				return Results;
-			}
-			
-			for (int I = 0; I < NodeCount; ++I)
-			{
-				int LocalI = I;
-				NodeDependencies.Add(Task.Run(() => FetchDependencies(LocalI)));
-				PrecedingDependents.Add(Task.Run(() => ComputePrecedingDependents(LocalI)));
-			}
-			
-			List<UEBuildModuleCPP> Out = new List<UEBuildModuleCPP>(NodeCount);
-			// Write the ordered output			
-			for (int Index1 = 0; Index1 != NodeCount; ++Index1)
-			{
-				UEBuildModuleCPP Node1 = NodeList[Index1];
-				HashSet<UEBuildModuleCPP> NodesThatDependOnNode1 = PrecedingDependents[Index1].Result;
-				Out.Add(Node1);
-
-				if (NodesThatDependOnNode1.Count == 0)
-				{
-					continue;
-				}
-
-				for (int Index2 = 0; Index2 != Index1; ++Index2)
-				{
-					UEBuildModuleCPP Node2 = Out[Index2];
-
-					if (NodesThatDependOnNode1.Contains(Node2))
-					{
-						// Rotate element at Index1 into position at Index2
-						for (int Index3 = Index1; Index3 != Index2;)
-						{
-							--Index3;
-							Out[Index3 + 1] = Out[Index3];
-						}
-
-						Out[Index2] = Node1;
-
-						// Break out of this loop, because this iteration must have covered all existing cases
-						// involving the node formerly at position Index1
-						break;
-					}
-				}
-			}
-
-			return Out;
-		}
-
 		/// <summary>
 		/// Gets the module type for a given rules object
 		/// </summary>
@@ -506,7 +421,7 @@ namespace UnrealBuildTool
 
 			// Sort modules by type, then by dependency
 			List<UEBuildModuleCPP> ModulesSortedByType = ModulesToGenerateHeadersFor.OrderBy(c => ModuleToType[c]).ToList();
-			ModulesSortedByType = StableTopologicalSort(ModulesSortedByType);
+			ModulesSortedByType = UEBuildModule.StableTopologicalSort(ModulesSortedByType.Cast<UEBuildModule>().ToList()).Cast<UEBuildModuleCPP>().ToList();
 
 			// Create the info for each module in parallel
 			UHTModuleInfo[] ModuleInfoArray = new UHTModuleInfo[ModulesSortedByType.Count];
diff --git a/Engine/Source/Programs/UnrealBuildTool/System/PrecompiledHeaderInstance.cs b/Engine/Source/Programs/UnrealBuildTool/System/PrecompiledHeaderInstance.cs
--- a/Engine/Source/Programs/UnrealBuildTool/System/PrecompiledHeaderInstance.cs
+++ b/Engine/Source/Programs/UnrealBuildTool/System/PrecompiledHeaderInstance.cs
@@ -44,6 +44,16 @@ namespace UnrealBuildTool
 		/// </summary>
 		public HashSet<string> ImmutableDefinitions;
 
+		/// <summary>
+		/// Parent PCH instance used in PCH chaining
+		/// </summary>
+		public PrecompiledHeaderInstance? ParentPCHInstance;
+
+		/// <summary>
+		/// Dictionary of definitions use in the CppCompileEnvironment
+		/// </summary>
+		public Dictionary<string,string>? DefinitionsDictionary;
+
 		/// <summary>
 		/// Constructor
 		/// </summary>
diff --git a/Engine/Source/Programs/UnrealBuildTool/ToolChain/ClangToolChain.cs b/Engine/Source/Programs/UnrealBuildTool/ToolChain/ClangToolChain.cs
--- a/Engine/Source/Programs/UnrealBuildTool/ToolChain/ClangToolChain.cs
+++ b/Engine/Source/Programs/UnrealBuildTool/ToolChain/ClangToolChain.cs
@@ -950,6 +950,12 @@ namespace UnrealBuildTool
 			{
 				OutputFile = FileItem.GetItemByFileReference(FileReference.Combine(OutputDir, GetFileNameFromExtension(FileName, ".gch")));
 				CompileResult.PrecompiledHeaderFile = OutputFile;
+
+				if (CompileEnvironment.ParentPrecompiledHeaderFile != null)
+				{
+					Arguments.Add(GetIncludePCHFileArgument(CompileEnvironment.ParentPrecompiledHeaderFile));
+					CompileAction.PrerequisiteItems.Add(CompileEnvironment.ParentPrecompiledHeaderFile!);
+				}
 			}
 			else if (CompileEnvironment.bPreprocessOnly)
 			{
