UnrealBuildTool: Fix issue that is causing -Compiler= to not be applied to WindowsTargetRules by applying arch configs and command lines in the TargetRules. Partial back out of 23931789

#rnx

diff --git a/Engine/Source/Programs/UnrealBuildTool/Configuration/TargetRules.cs b/Engine/Source/Programs/UnrealBuildTool/Configuration/TargetRules.cs
--- a/Engine/Source/Programs/UnrealBuildTool/Configuration/TargetRules.cs
+++ b/Engine/Source/Programs/UnrealBuildTool/Configuration/TargetRules.cs
@@ -2558,6 +2558,10 @@ namespace UnrealBuildTool
 					}
 				}
 			}
+			foreach (UnrealArch Arch in Architectures.Architectures)
+			{
+				yield return Arch;
+			}
 		}
 
 		/// <summary>
diff --git a/Engine/Source/Programs/UnrealBuildTool/Configuration/UEBuildPlatform.cs b/Engine/Source/Programs/UnrealBuildTool/Configuration/UEBuildPlatform.cs
--- a/Engine/Source/Programs/UnrealBuildTool/Configuration/UEBuildPlatform.cs
+++ b/Engine/Source/Programs/UnrealBuildTool/Configuration/UEBuildPlatform.cs
@@ -55,17 +55,6 @@ namespace UnrealBuildTool
 			return UEBuildPlatform.GetBuildPlatform(Platform).ArchitectureConfig;
 		}
 
-		/// <summary>
-		/// Get all known ArchitectureConfig objects
-		/// </summary>
-		/// <returns></returns>
-		public static IEnumerable<UnrealArchitectureConfig> AllConfigs()
-		{
-			// return ArchConfigs for platforms that have a BuildPlatform (which is where the Configs are stored)
-			return UnrealTargetPlatform.GetValidPlatforms().Where(x => UEBuildPlatform.TryGetBuildPlatform(x, out _)).Select(x => UEBuildPlatform.GetBuildPlatform(x).ArchitectureConfig);
-		}
-
-
 		/// <summary>
 		/// The multi-architecture mode for this platform (potentially single-architecture)
 		/// </summary>
diff --git a/Engine/Source/Programs/UnrealBuildTool/Modes/BuildMode.cs b/Engine/Source/Programs/UnrealBuildTool/Modes/BuildMode.cs
--- a/Engine/Source/Programs/UnrealBuildTool/Modes/BuildMode.cs
+++ b/Engine/Source/Programs/UnrealBuildTool/Modes/BuildMode.cs
@@ -119,12 +119,6 @@ namespace UnrealBuildTool
 			// Read the XML configuration files
 			XmlConfig.ApplyTo(this);
 
-			// Apply to architecture configs that need to read commandline arguments and didn't have the Arguments passed in during construction
-			foreach (UnrealArchitectureConfig Config in UnrealArchitectureConfig.AllConfigs())
-			{
-				Arguments.ApplyTo(Config);
-			}
-
 			// Fixup the log path if it wasn't overridden by a config file
 			if (BaseLogFileName == null)
 			{
diff --git a/Engine/Source/Programs/UnrealBuildTool/Modes/GenerateProjectFilesMode.cs b/Engine/Source/Programs/UnrealBuildTool/Modes/GenerateProjectFilesMode.cs
--- a/Engine/Source/Programs/UnrealBuildTool/Modes/GenerateProjectFilesMode.cs
+++ b/Engine/Source/Programs/UnrealBuildTool/Modes/GenerateProjectFilesMode.cs
@@ -69,12 +69,6 @@ namespace UnrealBuildTool
 			// Apply the XML config to this class
 			XmlConfig.ApplyTo(this);
 
-			// Apply to architecture configs that need to read commandline arguments and didn't have the Arguments passed in during construction
-			foreach (UnrealArchitectureConfig Config in UnrealArchitectureConfig.AllConfigs())
-			{
-				Arguments.ApplyTo(Config);
-			}
-
 			// set up logging (taken from BuildMode)
 			FileReference LogFile = FileReference.Combine(Unreal.EngineProgramSavedDirectory, "UnrealBuildTool", "Log_GPF.txt");
 			Log.AddFileWriter("DefaultLogTraceListener", LogFile);
