Disabled pch chaining on Linux because it is crashing the compiler when compiling the CitySample.

diff --git a/Engine/Source/Programs/UnrealBuildTool/Platform/Linux/UEBuildLinux.cs b/Engine/Source/Programs/UnrealBuildTool/Platform/Linux/UEBuildLinux.cs
--- a/Engine/Source/Programs/UnrealBuildTool/Platform/Linux/UEBuildLinux.cs
+++ b/Engine/Source/Programs/UnrealBuildTool/Platform/Linux/UEBuildLinux.cs
@@ -269,6 +269,9 @@ namespace UnrealBuildTool
 			{
 				IWYUToolChain.ValidateTarget(Target);
 			}
+
+			// Disable chaining PCHs for the moment because it is crashing clang
+			Target.bChainPCHs = false;
 		}
 
 		public override bool CanUseXGE()
