UBT: AppleToolChain temporary workaround for failing builds.

-SingleFile=... will not work for the time being, but this at least resolves horde build failures.

diff --git a/Engine/Source/Programs/UnrealBuildTool/ToolChain/AppleToolChain.cs b/Engine/Source/Programs/UnrealBuildTool/ToolChain/AppleToolChain.cs
--- a/Engine/Source/Programs/UnrealBuildTool/ToolChain/AppleToolChain.cs
+++ b/Engine/Source/Programs/UnrealBuildTool/ToolChain/AppleToolChain.cs
@@ -259,6 +259,11 @@ namespace UnrealBuildTool
 			return DestFile;
 		}
 
+		public override void CreateSpecificFileAction(CppCompileEnvironment CompileEnvironment, DirectoryReference SourceDir, DirectoryReference OutputDir, IActionGraphBuilder Graph)
+		{
+			// Temporarily turn of single file compile in apple toolchain
+		}
+
 
 		/// <inheritdoc/>
 		protected override string EscapePreprocessorDefinition(string Definition)
