[UBT]
* Disabled -SingleFile option when building for multiple architectures at once... this can be implemented if someone needs it but require a bit of work
* Enabled -SingleFile option for apple platforms

diff --git a/Engine/Source/Programs/UnrealBuildTool/Platform/Windows/VCToolChain.cs b/Engine/Source/Programs/UnrealBuildTool/Platform/Windows/VCToolChain.cs
--- a/Engine/Source/Programs/UnrealBuildTool/Platform/Windows/VCToolChain.cs
+++ b/Engine/Source/Programs/UnrealBuildTool/Platform/Windows/VCToolChain.cs
@@ -2087,6 +2087,10 @@ namespace UnrealBuildTool
 
 		public override void CreateSpecificFileAction(CppCompileEnvironment CompileEnvironment, DirectoryReference SourceDir, DirectoryReference OutputDir, IActionGraphBuilder Graph)
 		{
+			// This is not supported for now.. If someone wants it we can implement it
+			if (CompileEnvironment.Architectures.bIsMultiArch)
+				return;
+
 			VCCompileAction BaseCompileAction = CreateBaseCompileAction(CompileEnvironment);
 			AppendCLArguments_CPP(CompileEnvironment, BaseCompileAction.Arguments);
 			Graph.AddAction(new VcSpecificFileAction(SourceDir, OutputDir, BaseCompileAction));
diff --git a/Engine/Source/Programs/UnrealBuildTool/ToolChain/AppleToolChain.cs b/Engine/Source/Programs/UnrealBuildTool/ToolChain/AppleToolChain.cs
--- a/Engine/Source/Programs/UnrealBuildTool/ToolChain/AppleToolChain.cs
+++ b/Engine/Source/Programs/UnrealBuildTool/ToolChain/AppleToolChain.cs
@@ -259,12 +259,6 @@ namespace UnrealBuildTool
 			return DestFile;
 		}
 
-		public override void CreateSpecificFileAction(CppCompileEnvironment CompileEnvironment, DirectoryReference SourceDir, DirectoryReference OutputDir, IActionGraphBuilder Graph)
-		{
-			// Temporarily turn of single file compile in apple toolchain
-		}
-
-
 		/// <inheritdoc/>
 		protected override string EscapePreprocessorDefinition(string Definition)
 		{
diff --git a/Engine/Source/Programs/UnrealBuildTool/ToolChain/ClangToolChain.cs b/Engine/Source/Programs/UnrealBuildTool/ToolChain/ClangToolChain.cs
--- a/Engine/Source/Programs/UnrealBuildTool/ToolChain/ClangToolChain.cs
+++ b/Engine/Source/Programs/UnrealBuildTool/ToolChain/ClangToolChain.cs
@@ -1058,6 +1058,10 @@ namespace UnrealBuildTool
 
 		public override void CreateSpecificFileAction(CppCompileEnvironment CompileEnvironment, DirectoryReference SourceDir, DirectoryReference OutputDir, IActionGraphBuilder Graph)
 		{
+			// This is not supported for now.. If someone wants it we can implement it
+			if (CompileEnvironment.Architectures.bIsMultiArch)
+				return;
+
 			List<string> GlobalArguments = new();
 			if (CompileEnvironment.SystemIncludePaths.Count != 0)
 			{
