UnrealBuildTool: Trim quotes from -Manifest=

#rnx

diff --git a/Engine/Source/Programs/UnrealBuildTool/System/TargetMakefile.cs b/Engine/Source/Programs/UnrealBuildTool/System/TargetMakefile.cs
--- a/Engine/Source/Programs/UnrealBuildTool/System/TargetMakefile.cs
+++ b/Engine/Source/Programs/UnrealBuildTool/System/TargetMakefile.cs
@@ -642,7 +642,7 @@ namespace UnrealBuildTool
 				// Check if any manifests are out of date
 				foreach (FileReference Manifest in Makefile.AdditionalArguments!
 					.Where(x => x.StartsWith("-Manifest=", StringComparison.OrdinalIgnoreCase))
-					.Select(x => FileReference.FromString(x.Substring("-Manifest=".Length)))
+					.Select(x => FileReference.FromString(x.Substring("-Manifest=".Length).Trim('"').Trim('\'')))
 					.Where(x => x != null))
 				{
 					if (!FileReference.Exists(Manifest))
