[UBT]
* Changed so clang single file placeholder action is not having any produced items.. this means that ubt modes that looks for produced items won't find anything on the placeholder actions

diff --git a/Engine/Source/Programs/UnrealBuildTool/ToolChain/ClangSpecificFileAction.cs b/Engine/Source/Programs/UnrealBuildTool/ToolChain/ClangSpecificFileAction.cs
--- a/Engine/Source/Programs/UnrealBuildTool/ToolChain/ClangSpecificFileAction.cs
+++ b/Engine/Source/Programs/UnrealBuildTool/ToolChain/ClangSpecificFileAction.cs
@@ -21,6 +21,9 @@ namespace UnrealBuildTool
 
 		internal ClangSpecificFileAction(DirectoryReference Source, DirectoryReference Output, Action Action, IEnumerable<string> ContentLines) : base(Action)
 		{
+			ProducedItems.Clear();
+			DependencyListFile = null;
+
 			SourceDir = Source;
 			OutputDir = Output;
 			RspLines = ContentLines;
@@ -80,6 +83,12 @@ namespace UnrealBuildTool
 			Action.DependencyListFile = null;
 			Action.StatusDescription = SourceFile.Name;
 
+			// We have to add a produced item so this action is not skipped.
+			// Note we on purpose use a different extension than what the compiler produce because otherwise up-to-date checker might see it as up-to-date
+			// even though we want it to always be built
+			FileItem ProducedItem = FileItem.GetItemByFileReference(FileReference.Combine(OutputDir, UniqueDummyName + ".n"));
+			Action.ProducedItems.Add(ProducedItem);
+
 			FileItem ResponseFile = FileItem.GetItemByPath(Action.CommandArguments.Substring(1).Trim('"'));
 			File.WriteAllLines(ResponseFile.FullName, NewRspLines);
 
