UnrealBuildTool: Return target file from GetCompileArguments_FileType rather then getting it from the ProducedItem list

#rnx

diff --git a/Engine/Source/Programs/UnrealBuildTool/Platform/IOS/IOSToolChain.cs b/Engine/Source/Programs/UnrealBuildTool/Platform/IOS/IOSToolChain.cs
--- a/Engine/Source/Programs/UnrealBuildTool/Platform/IOS/IOSToolChain.cs
+++ b/Engine/Source/Programs/UnrealBuildTool/Platform/IOS/IOSToolChain.cs
@@ -300,9 +300,9 @@ namespace UnrealBuildTool
 		}
 
 		/// <inheritdoc/>
-		protected override void GetCompileArguments_FileType(CppCompileEnvironment CompileEnvironment, FileItem SourceFile, DirectoryReference OutputDir, List<string> Arguments, Action CompileAction, CPPOutput CompileResult)
+		protected override FileItem GetCompileArguments_FileType(CppCompileEnvironment CompileEnvironment, FileItem SourceFile, DirectoryReference OutputDir, List<string> Arguments, Action CompileAction, CPPOutput CompileResult)
 		{
-			base.GetCompileArguments_FileType(CompileEnvironment, SourceFile, OutputDir, Arguments, CompileAction, CompileResult);
+			FileItem TargetFile = base.GetCompileArguments_FileType(CompileEnvironment, SourceFile, OutputDir, Arguments, CompileAction, CompileResult);
 
 			string Extension = Path.GetExtension(SourceFile.AbsolutePath).ToUpperInvariant();
 
@@ -310,6 +310,8 @@ namespace UnrealBuildTool
 			{
 				Arguments.Add(GetObjCExceptionsFlag(CompileEnvironment));
 			}
+
+			return TargetFile;
 		}
 
 		// Conditionally enable (default disabled) Objective-C exceptions
diff --git a/Engine/Source/Programs/UnrealBuildTool/ToolChain/ClangToolChain.cs b/Engine/Source/Programs/UnrealBuildTool/ToolChain/ClangToolChain.cs
--- a/Engine/Source/Programs/UnrealBuildTool/ToolChain/ClangToolChain.cs
+++ b/Engine/Source/Programs/UnrealBuildTool/ToolChain/ClangToolChain.cs
@@ -485,6 +485,10 @@ namespace UnrealBuildTool
 			{
 				Arguments.Add(GetIncludePCHFileArgument(CompileEnvironment.PrecompiledHeaderFile!));
 			}
+			else if (CompileEnvironment.PrecompiledHeaderAction == PrecompiledHeaderAction.Create && CompileEnvironment.ParentPrecompiledHeaderFile != null)
+			{
+				Arguments.Add(GetIncludePCHFileArgument(CompileEnvironment.ParentPrecompiledHeaderFile));
+			}
 
 			Arguments.AddRange(CompileEnvironment.ForceIncludeFiles.Select(ForceIncludeFile => GetForceIncludeFileArgument(ForceIncludeFile)));
 		}
@@ -883,7 +887,8 @@ namespace UnrealBuildTool
 		/// <param name="Arguments"></param>
 		/// <param name="CompileAction"></param>
 		/// <param name="CompileResult"></param>
-		protected virtual void GetCompileArguments_FileType(CppCompileEnvironment CompileEnvironment, FileItem SourceFile, DirectoryReference OutputDir, List<string> Arguments, Action CompileAction, CPPOutput CompileResult)
+		/// <returns>Path to the target file (such as .o)</returns>
+		protected virtual FileItem GetCompileArguments_FileType(CppCompileEnvironment CompileEnvironment, FileItem SourceFile, DirectoryReference OutputDir, List<string> Arguments, Action CompileAction, CPPOutput CompileResult)
 		{
 			// Add the additional response files
 			foreach (FileItem AdditionalResponseFile in CompileEnvironment.AdditionalResponseFiles)
@@ -894,7 +899,6 @@ namespace UnrealBuildTool
 			// Add force include paths to the argument list.
 			GetCompileArguments_ForceInclude(CompileEnvironment, Arguments);
 
-
 			// Add the C++ source file and its included files to the prerequisite item list.
 			CompileAction.PrerequisiteItems.UnionWith(CompileEnvironment.ForceIncludeFiles);
 			CompileAction.PrerequisiteItems.UnionWith(CompileEnvironment.AdditionalPrerequisites);
@@ -953,7 +957,6 @@ namespace UnrealBuildTool
 
 				if (CompileEnvironment.ParentPrecompiledHeaderFile != null)
 				{
-					Arguments.Add(GetIncludePCHFileArgument(CompileEnvironment.ParentPrecompiledHeaderFile));
 					CompileAction.PrerequisiteItems.Add(CompileEnvironment.ParentPrecompiledHeaderFile!);
 				}
 			}
@@ -1008,6 +1011,8 @@ namespace UnrealBuildTool
 
 			// Add the parameters needed to compile the output file to the command-line.
 			Arguments.Add(GetOutputFileArgument(OutputFile));
+
+			return OutputFile;
 		}
 
 		protected virtual List<string> ExpandResponseFileContents(List<string> ResponseFileContents)
@@ -1084,11 +1089,8 @@ namespace UnrealBuildTool
 			// copy the global arguments into the file arguments, so GetCompileArguments_FileType can remove entries if needed (special case but can be important)
 			List<string> FileArguments = new(GlobalArguments);
 
-			// Add C or C++ specific compiler arguments.
-			GetCompileArguments_FileType(CompileEnvironment, SourceFile, OutputDir, FileArguments, CompileAction, Result);
-
-			// Gets the target file so we can get the correct output path.
-			FileItem TargetFile = CompileAction.ProducedItems.First();
+			// Add C or C++ specific compiler arguments and get the target file so we can get the correct output path.
+			FileItem TargetFile = GetCompileArguments_FileType(CompileEnvironment, SourceFile, OutputDir, FileArguments, CompileAction, Result);
 
 			// Creates the path to the response file using the name of the output file and creates its contents.
 			FileReference ResponseFileName = GetResponseFileName(CompileEnvironment, TargetFile);
@@ -1166,7 +1168,7 @@ namespace UnrealBuildTool
 				PrepassAction.DeleteItems.UnionWith(PrepassAction.ProducedItems);
 
 				// Gets the target file so we can get the correct output path.
-				FileItem PreprocessTargetFile = PrepassAction.ProducedItems.First();
+				FileItem PreprocessTargetFile = PrepassAction.DependencyListFile;
 
 				// Creates the path to the response file using the name of the output file and creates its contents.
 				FileReference PreprocessResponseFileName = GetResponseFileName(CompileEnvironment, PreprocessTargetFile);
@@ -1224,4 +1226,4 @@ namespace UnrealBuildTool
 		{
 		}
 	}
-}
+}
\ No newline at end of file
