UnrealBuildTool: Add return type to AndroidToolChain.GetCompileArguments_FileType to match change from merge

#rnx

diff --git a/Engine/Source/Programs/UnrealBuildTool/Platform/Android/AndroidToolChain.cs b/Engine/Source/Programs/UnrealBuildTool/Platform/Android/AndroidToolChain.cs
--- a/Engine/Source/Programs/UnrealBuildTool/Platform/Android/AndroidToolChain.cs
+++ b/Engine/Source/Programs/UnrealBuildTool/Platform/Android/AndroidToolChain.cs
@@ -693,9 +693,9 @@ namespace UnrealBuildTool
 		}
 
 		/// <inheritdoc/>
-		protected override void GetCompileArguments_FileType(CppCompileEnvironment CompileEnvironment, FileItem SourceFile, DirectoryReference OutputDir, List<string> Arguments, Action CompileAction, CPPOutput CompileResult)
+		protected override FileItem GetCompileArguments_FileType(CppCompileEnvironment CompileEnvironment, FileItem SourceFile, DirectoryReference OutputDir, List<string> Arguments, Action CompileAction, CPPOutput CompileResult)
 		{
-			base.GetCompileArguments_FileType(CompileEnvironment, SourceFile, OutputDir, Arguments, CompileAction, CompileResult);
+			FileItem TargetFile = base.GetCompileArguments_FileType(CompileEnvironment, SourceFile, OutputDir, Arguments, CompileAction, CompileResult);
 
 			string Extension = Path.GetExtension(SourceFile.AbsolutePath).ToUpperInvariant();
 			if (Extension == ".C")
@@ -708,6 +708,8 @@ namespace UnrealBuildTool
 				// remove any PCH includes - mostly for the force-added .c files in Launch as those will attempt to have the PCH used that was made with .cpp language
 				Arguments.RemoveAll(x => x.StartsWith("-include-pch"));
 			}
+
+			return TargetFile;
 		}
 
 
