Add support for /GENPROFILE:NOTRACKEH and /FASTGENPROFILE on MSVC

#rnx

diff --git a/Engine/Source/Programs/UnrealBuildTool/Platform/Windows/UEBuildWindows.cs b/Engine/Source/Programs/UnrealBuildTool/Platform/Windows/UEBuildWindows.cs
--- a/Engine/Source/Programs/UnrealBuildTool/Platform/Windows/UEBuildWindows.cs
+++ b/Engine/Source/Programs/UnrealBuildTool/Platform/Windows/UEBuildWindows.cs
@@ -216,6 +216,21 @@ namespace UnrealBuildTool
 		/// </summary>
 		TargetRules Target;
 
+		/// <summary>
+		/// If specified along with -PGOProfile, then /FASTGENPROFILE will be used instead of /GENPROFILE. This usually means that the PGO data is generated faster, but the resulting data may not yield as efficient optimizations during -PGOOptimize
+		/// </summary>
+		[XmlConfigFile(Category = "WindowsPlatform")]
+		[CommandLine("-PGOFastGen")]
+		public bool bUseFastGenProfile = false;
+
+		/// <summary>
+		/// If specified along with -PGOProfile, prevent the usage of extra counters. Please note that by default /FASTGENPROFILE doesnt use extra counters
+		/// https://learn.microsoft.com/en-us/cpp/build/reference/genprofile-fastgenprofile-generate-profiling-instrumented-build?view=msvc-170
+		/// </summary>
+		[XmlConfigFile(Category = "WindowsPlatform")]
+		[CommandLine("-PGONoExtraCounters")]
+		public bool bPGONoExtraCounters = false;
+
 		/// <summary>
 		/// Version of the compiler toolchain to use on Windows platform. A value of "default" will be changed to a specific version at UBT start up.
 		/// </summary>
@@ -635,6 +650,16 @@ namespace UnrealBuildTool
 		#region Read-only accessor properties 
 		#pragma warning disable CS1591
 
+		public bool bUseFastGenProfile 
+		{
+			get { return Inner.bUseFastGenProfile; }
+		}
+
+		public bool bPGONoExtraCounters 
+		{
+			get { return Inner.bPGONoExtraCounters; }
+		}
+
 		public WindowsCompiler Compiler
 		{
 			get { return Inner.Compiler; }
diff --git a/Engine/Source/Programs/UnrealBuildTool/Platform/Windows/VCToolChain.cs b/Engine/Source/Programs/UnrealBuildTool/Platform/Windows/VCToolChain.cs
--- a/Engine/Source/Programs/UnrealBuildTool/Platform/Windows/VCToolChain.cs
+++ b/Engine/Source/Programs/UnrealBuildTool/Platform/Windows/VCToolChain.cs
@@ -2519,6 +2519,7 @@ namespace UnrealBuildTool
 				//Arguments.Add("/GENPROFILE:PGD=" + Path.Combine(LinkEnvironment.PGODirectory, LinkEnvironment.PGOFilenamePrefix + ".pgd"));
 				Arguments.Add("/LTCG");
 				//Arguments.Add("/GENPROFILE:PGD=" + LinkEnvironment.PGOFilenamePrefix + ".pgd");
+
 				if (Target.WindowsPlatform.bUseFastGenProfile)
 				{
 					Arguments.Add("/FASTGENPROFILE");
@@ -2527,7 +2528,23 @@ namespace UnrealBuildTool
 				{
 					Arguments.Add("/GENPROFILE");
 				}
-				Log.TraceInformationOnce("Enabling Profile Guided Optimization (PGO). Linking will take a while.");
+					Log.TraceInformationOnce("Enabling Fastgen Profile Guided Optimization (PGO). Linking will take a while.");
+					Arguments.Add("/FASTGENPROFILE");
+				}
+				else 
+				{
+					if (Target.WindowsPlatform.bPGONoExtraCounters)
+					{
+						Log.TraceInformationOnce("Enabling Profile Guided Optimization No Extra Counters(PGO). Linking will take a while.");
+						Arguments.Add("/GENPROFILE:NOTRACKEH");
+					}
+					else
+					{
+						Log.TraceInformationOnce("Enabling Profile Guided Optimization (PGO). Linking will take a while.");
+						Arguments.Add("/GENPROFILE");
+					}	
+					
+				}
 			}
 		}
 
