Fix bad merge

#rnx

diff --git a/Engine/Source/Programs/UnrealBuildTool/Platform/Windows/UEBuildWindows.cs b/Engine/Source/Programs/UnrealBuildTool/Platform/Windows/UEBuildWindows.cs
--- a/Engine/Source/Programs/UnrealBuildTool/Platform/Windows/UEBuildWindows.cs
+++ b/Engine/Source/Programs/UnrealBuildTool/Platform/Windows/UEBuildWindows.cs
@@ -476,14 +476,6 @@ namespace UnrealBuildTool
 		/// </summary>
 		public bool bOptimizeGlobalData = true;
 
-		/// <summary>
-		/// If specified along with -PGOProfile, then /FASTGENPROFILE will be used instead of /GENPROFILE. 
-		/// This usually means that the PGO data is generated faster, but the resulting data may not yield as efficient optimizations during -PGOOptimize
-		/// </summary>
-		[XmlConfigFile(Category = "WindowsPlatform")]
-		[CommandLine("-PGOFastGen")]
-		public bool bUseFastGenProfile = false;
-
 		/// <summary>
 		/// (Experimental) Appends the -ftime-trace argument to the command line for Clang to output a JSON file containing a timeline for the compile. 
 		/// See http://aras-p.info/blog/2019/01/16/time-trace-timeline-flame-chart-profiler-for-Clang/ for more info.
@@ -827,11 +819,6 @@ namespace UnrealBuildTool
 			get { return Inner.bOptimizeGlobalData; }
 		}
 
-		public bool bUseFastGenProfile
-		{
-			get { return Inner.bUseFastGenProfile; }
-		}
-
 		public bool bClangTimeTrace
 		{
 			get { return Inner.bClangTimeTrace; }
diff --git a/Engine/Source/Programs/UnrealBuildTool/Platform/Windows/VCToolChain.cs b/Engine/Source/Programs/UnrealBuildTool/Platform/Windows/VCToolChain.cs
--- a/Engine/Source/Programs/UnrealBuildTool/Platform/Windows/VCToolChain.cs
+++ b/Engine/Source/Programs/UnrealBuildTool/Platform/Windows/VCToolChain.cs
@@ -2523,13 +2523,7 @@ namespace UnrealBuildTool
 				if (Target.WindowsPlatform.bUseFastGenProfile)
 				{
 					Arguments.Add("/FASTGENPROFILE");
-				}
-				else
-				{
-					Arguments.Add("/GENPROFILE");
-				}
 					Log.TraceInformationOnce("Enabling Fastgen Profile Guided Optimization (PGO). Linking will take a while.");
-					Arguments.Add("/FASTGENPROFILE");
 				}
 				else 
 				{
@@ -2543,7 +2537,6 @@ namespace UnrealBuildTool
 						Log.TraceInformationOnce("Enabling Profile Guided Optimization (PGO). Linking will take a while.");
 						Arguments.Add("/GENPROFILE");
 					}	
-					
 				}
 			}
 		}
