UnrealBuildTool: Command line arg to enable edit and continue support

#rnx

diff --git a/Engine/Source/Programs/UnrealBuildTool/Configuration/TargetRules.cs b/Engine/Source/Programs/UnrealBuildTool/Configuration/TargetRules.cs
--- a/Engine/Source/Programs/UnrealBuildTool/Configuration/TargetRules.cs
+++ b/Engine/Source/Programs/UnrealBuildTool/Configuration/TargetRules.cs
@@ -1765,6 +1765,7 @@ namespace UnrealBuildTool
 		/// <summary>
 		/// Whether to support edit and continue.
 		/// </summary>
+		[CommandLine("-SupportEditAndContinue")]
 		[XmlConfigFile(Category = "BuildConfiguration")]
 		public bool bSupportEditAndContinue = false;
 
