[UBT]
* Fixed so xge is not counting actions wrong

diff --git a/Engine/Source/Programs/UnrealBuildTool/Executors/XGE.cs b/Engine/Source/Programs/UnrealBuildTool/Executors/XGE.cs
--- a/Engine/Source/Programs/UnrealBuildTool/Executors/XGE.cs
+++ b/Engine/Source/Programs/UnrealBuildTool/Executors/XGE.cs
@@ -483,7 +483,7 @@ namespace UnrealBuildTool
 				string OutputPrefix = "";
 				if (bProgressMarkup)
 				{
-					OutputPrefix += ProgressMarkupPrefix;
+					OutputPrefix += $"{ProgressMarkupPrefix}_{ActionIndex} ";
 				}
 				if (Action.bShouldOutputStatusDescription)
 				{
@@ -699,6 +699,7 @@ namespace UnrealBuildTool
 			{
 				int NumCompletedActions = 0;
 				string ProgressText = string.Empty;
+				HashSet<int> ReportedActionIndices = new();
 
 				// Create a wrapper delegate that will parse the output actions
 				DataReceivedEventHandler EventHandlerWrapper = (Sender, Args) =>
@@ -708,16 +709,31 @@ namespace UnrealBuildTool
 						string Text = Args.Data;
 						if (Text.StartsWith(ProgressMarkupPrefix))
 						{
+							// Code below should not need to be tested for success but if some logging from XGE is wrong we just gracefully ignore it and accept that counting might end up wrong
+							int MarkupLength = ProgressMarkupPrefix.Length;
+							int EndOfMarkupPrefix = Text.IndexOf(' ', MarkupLength);
+							if (EndOfMarkupPrefix != -1)
+							{
+								MarkupLength = EndOfMarkupPrefix + 1;
+								int ActionIndex;
+								if (int.TryParse(Text.Substring(ProgressMarkupPrefix.Length + 1, EndOfMarkupPrefix - ProgressMarkupPrefix.Length - 1), out ActionIndex))
+								{
+									// We keep track of the actions that we have already reported so NumCompletedActions match up with NumActions
+									if (ReportedActionIndices.Add(ActionIndex))
+									{
+										Writer.Write(++NumCompletedActions, NumActions);
+									}
+								}
+							}
 							// Flush old progress text
 							if (!string.IsNullOrEmpty(ProgressText))
 							{
 								Logger.LogInformation("[{NumCompletedActions}/{NumActions}] Complete {ProgressText}", NumCompletedActions, NumActions, ProgressText);
 								ProgressText = string.Empty;
 							}
-							Writer.Write(++NumCompletedActions, NumActions);
 
 							// Strip out anything that is just an XGE timer. Some programs don't output anything except the progress text.
-							Text = Args.Data.Substring(ProgressMarkupPrefix.Length);
+							Text = Args.Data.Substring(MarkupLength);
 							if(Text.StartsWith(" (") && Text.EndsWith(")"))
 							{
 								// Write the progress text with the next line of output if the current doesn't have any status.
