UnrealBuildTool: Simplify action conflict filename

#rnx

diff --git a/Engine/Source/Programs/UnrealBuildTool/System/ActionGraph.cs b/Engine/Source/Programs/UnrealBuildTool/System/ActionGraph.cs
--- a/Engine/Source/Programs/UnrealBuildTool/System/ActionGraph.cs
+++ b/Engine/Source/Programs/UnrealBuildTool/System/ActionGraph.cs
@@ -247,8 +247,8 @@ namespace UnrealBuildTool
 
 			string AJson = JsonSerializer.Serialize(A, Options);
 			string BJson = JsonSerializer.Serialize(B, Options);
-			string AJsonPath = Path.Combine(Path.GetTempPath(), "UnrealBuildTool", $"{A.StatusDescription}-{IoHash.Compute(Encoding.Default.GetBytes(AJson)).GetHashCode():X}") + ".json";
-			string BJsonPath = Path.Combine(Path.GetTempPath(), "UnrealBuildTool", $"{B.StatusDescription}-{IoHash.Compute(Encoding.Default.GetBytes(BJson)).GetHashCode():X}") + ".json";
+			string AJsonPath = Path.Combine(Path.GetTempPath(), "UnrealBuildTool", $"{IoHash.Compute(Encoding.Default.GetBytes(AJson)).GetHashCode():X}") + ".json";
+			string BJsonPath = Path.Combine(Path.GetTempPath(), "UnrealBuildTool", $"{IoHash.Compute(Encoding.Default.GetBytes(BJson)).GetHashCode():X}") + ".json";
 
 			Directory.CreateDirectory(Path.Combine(Path.GetTempPath(), "UnrealBuildTool"));
 			File.WriteAllText(AJsonPath, AJson);
