- Fixing location for UHT generated headers to not have architecture in the path, because they are per-arch, and with InstalledBuilds, if we made a Universal build, but then user does non-universal, the UHT intermediates would look in a different location than where we put them ('arm64' vs 'arm64+x64')

diff --git a/Engine/Source/Programs/UnrealBuildTool/Configuration/UEBuildTarget.cs b/Engine/Source/Programs/UnrealBuildTool/Configuration/UEBuildTarget.cs
--- a/Engine/Source/Programs/UnrealBuildTool/Configuration/UEBuildTarget.cs
+++ b/Engine/Source/Programs/UnrealBuildTool/Configuration/UEBuildTarget.cs
@@ -1420,10 +1420,15 @@ namespace UnrealBuildTool
 		public UnrealArchitectures Architectures;
 
 		/// <summary>
-		/// Relative path for platform-specific intermediates (eg. Intermediate/Build/Win64)
+		/// Relative path for platform-specific intermediates (eg. Intermediate/Build/Win64/x64)
 		/// </summary>
 		public string PlatformIntermediateFolder;
 
+		/// <summary>
+		/// Relative path for platform-specific intermediates that explicitly do not need an architeecture (eg. Intermediate/Build/Win64)
+		/// </summary>
+		public string PlatformIntermediateFolderNoArch;
+
 		/// <summary>
 		/// Root directory for the active project. Typically contains the .uproject file, or the engine root.
 		/// </summary>
@@ -1434,6 +1439,11 @@ namespace UnrealBuildTool
 		/// </summary>
 		public DirectoryReference ProjectIntermediateDirectory;
 
+		/// <summary>
+		/// Default directory for intermediate files. Typically underneath ProjectDirectory.
+		/// </summary>
+		public DirectoryReference ProjectIntermediateDirectoryNoArch;
+
 		/// <summary>
 		/// Directory for engine intermediates. For an agnostic editor/game executable, this will be under the engine directory. For monolithic executables this will be the same as the project intermediate directory.
 		/// </summary>
@@ -1541,6 +1551,7 @@ namespace UnrealBuildTool
 
 			// now that we have the platform, we can set the intermediate path to include the platform/architecture name
 			PlatformIntermediateFolder = GetPlatformIntermediateFolder(Platform, Architectures, false);
+			PlatformIntermediateFolderNoArch = GetPlatformIntermediateFolder(Platform, null, false);
 
 			TargetRulesFile = InRules.File;
 
@@ -1582,6 +1593,7 @@ namespace UnrealBuildTool
 
 			// Build the project intermediate directory
 			ProjectIntermediateDirectory = DirectoryReference.Combine(OutputRootDirectory, PlatformIntermediateFolder, TargetName, Configuration.ToString());
+			ProjectIntermediateDirectoryNoArch = DirectoryReference.Combine(OutputRootDirectory, PlatformIntermediateFolderNoArch, TargetName, Configuration.ToString());
 
 			// Build the engine intermediate directory. If we're building agnostic engine binaries, we can use the engine intermediates folder. Otherwise we need to use the project intermediates directory.
 			if (!bUseSharedBuildEnvironment)
@@ -1614,10 +1626,15 @@ namespace UnrealBuildTool
 		/// <param name="Architectures">Architectures to get the folder for</param>
 		/// <param name="External">Insert External after Intermediate - for out-of-tree plugins</param>
 		/// <returns>The output directory for intermediates</returns>
-		public static string GetPlatformIntermediateFolder(UnrealTargetPlatform Platform, UnrealArchitectures Architectures, bool External)
+		public static string GetPlatformIntermediateFolder(UnrealTargetPlatform Platform, UnrealArchitectures? Architectures, bool External)
 		{
 			// now that we have the platform, we can set the intermediate path to include the platform/architecture name
-			return Path.Combine("Intermediate", External ? "External" : String.Empty, "Build", Platform.ToString(), UnrealArchitectureConfig.ForPlatform(Platform).GetFolderNameForArchitectures(Architectures));
+			string FolderPath = Path.Combine("Intermediate", External ? "External" : String.Empty, "Build", Platform.ToString());
+			if (Architectures != null)
+			{
+				FolderPath = Path.Combine(FolderPath, UnrealArchitectureConfig.ForPlatform(Platform).GetFolderNameForArchitectures(Architectures));
+			}
+			return FolderPath;
 		}
 
 		/// <summary>
@@ -2250,7 +2267,8 @@ namespace UnrealBuildTool
 
 			// Create the makefile
 			string ExternalMetadata = UEBuildPlatform.GetBuildPlatform(Platform).GetExternalBuildMetadata(ProjectFile);
-			TargetMakefile Makefile = new TargetMakefile(ExternalMetadata, Binaries[0].OutputFilePaths[0], ReceiptFileName, ProjectIntermediateDirectory, TargetType, 
+			TargetMakefile Makefile = new TargetMakefile(ExternalMetadata, Binaries[0].OutputFilePaths[0], ReceiptFileName, 
+				ProjectIntermediateDirectory, ProjectIntermediateDirectoryNoArch, TargetType, 
 				Rules.ConfigValueTracker, bDeployAfterCompile, UbtPlugins?.ToArray(), EnabledUbtPlugins?.ToArray(), 
 				EnabledUhtPlugins?.ToArray());
 			Makefile.IsTestTarget = Rules.IsTestTarget;
@@ -4903,8 +4921,8 @@ namespace UnrealBuildTool
 						GeneratedCodeDirectory = ProjectDirectory;
 					}
 
-					// Get the subfolder containing generated code
-					GeneratedCodeDirectory = DirectoryReference.Combine(GeneratedCodeDirectory, PlatformIntermediateFolder, AppName, "Inc");
+					// Get the subfolder containing generated code - we don't need architeceture information since these are shared between all arches for a platform
+					GeneratedCodeDirectory = DirectoryReference.Combine(GeneratedCodeDirectory, PlatformIntermediateFolderNoArch, AppName, "Inc");
 
 					// Append the binaries subfolder, if present. We rely on this to ensure that build products can be filtered correctly.
 					if (RulesObject.BinariesSubFolder != null)
diff --git a/Engine/Source/Programs/UnrealBuildTool/System/ExternalExecution.cs b/Engine/Source/Programs/UnrealBuildTool/System/ExternalExecution.cs
--- a/Engine/Source/Programs/UnrealBuildTool/System/ExternalExecution.cs
+++ b/Engine/Source/Programs/UnrealBuildTool/System/ExternalExecution.cs
@@ -930,7 +930,7 @@ namespace UnrealBuildTool
 		/// <returns>Manifest file name</returns>
 		public static FileReference GetUHTModuleInfoFileName(TargetMakefile Makefile, string TargetName)
 		{
-			return FileReference.Combine(Makefile.ProjectIntermediateDirectory, TargetName + ".uhtmanifest");
+			return FileReference.Combine(Makefile.ProjectIntermediateDirectoryNoArch, TargetName + ".uhtmanifest");
 		}
 
 		/// <summary>
diff --git a/Engine/Source/Programs/UnrealBuildTool/System/TargetMakefile.cs b/Engine/Source/Programs/UnrealBuildTool/System/TargetMakefile.cs
--- a/Engine/Source/Programs/UnrealBuildTool/System/TargetMakefile.cs
+++ b/Engine/Source/Programs/UnrealBuildTool/System/TargetMakefile.cs
@@ -209,6 +209,11 @@ namespace UnrealBuildTool
 		/// </summary>
 		public DirectoryReference ProjectIntermediateDirectory;
 
+		/// <summary>
+		/// The project intermediate directory without architecture info
+		/// </summary>
+		public DirectoryReference ProjectIntermediateDirectoryNoArch;
+
 		/// <summary>
 		/// Type of the target
 		/// </summary>
@@ -359,13 +364,15 @@ namespace UnrealBuildTool
 		/// <param name="ExecutableFile">Path to the executable or primary output binary for this target</param>
 		/// <param name="ReceiptFile">Path to the receipt file</param>
 		/// <param name="ProjectIntermediateDirectory">Path to the project intermediate directory</param>
+		/// <param name="ProjectIntermediateDirectoryNoArch">Path to the project intermediate directory</param>
 		/// <param name="TargetType">The type of target</param>
 		/// <param name="ConfigValueTracker">Set of dependencies on config files</param>
 		/// <param name="bDeployAfterCompile">Whether to deploy the target after compiling</param>
 		/// <param name="UbtPlugins">Collection of UBT plugins</param>
 		/// <param name="EnabledUbtPlugins">Collection of UBT plugins</param>
 		/// <param name="EnabledUhtPlugins">Collection of UBT plugins for UHT</param>
-		public TargetMakefile(string ExternalMetadata, FileReference ExecutableFile, FileReference ReceiptFile, DirectoryReference ProjectIntermediateDirectory, 
+		public TargetMakefile(string ExternalMetadata, FileReference ExecutableFile, FileReference ReceiptFile, 
+			DirectoryReference ProjectIntermediateDirectory, DirectoryReference ProjectIntermediateDirectoryNoArch,
 			TargetType TargetType, ConfigValueTracker ConfigValueTracker, bool bDeployAfterCompile, 
 			FileReference[]? UbtPlugins, FileReference[]? EnabledUbtPlugins, FileReference[]? EnabledUhtPlugins)
 		{
@@ -376,6 +383,7 @@ namespace UnrealBuildTool
 			this.ExecutableFile = ExecutableFile;
 			this.ReceiptFile = ReceiptFile;
 			this.ProjectIntermediateDirectory = ProjectIntermediateDirectory;
+			this.ProjectIntermediateDirectoryNoArch = ProjectIntermediateDirectoryNoArch;
 			this.TargetType = TargetType;
 			this.ConfigValueTracker = ConfigValueTracker;
 			this.bDeployAfterCompile = bDeployAfterCompile;
@@ -414,6 +422,7 @@ namespace UnrealBuildTool
 			ExecutableFile = Reader.ReadFileReference();
 			ReceiptFile = Reader.ReadFileReference();
 			ProjectIntermediateDirectory = Reader.ReadDirectoryReferenceNotNull();
+			ProjectIntermediateDirectoryNoArch = Reader.ReadDirectoryReferenceNotNull();
 			TargetType = (TargetType)Reader.ReadInt();
 			IsTestTarget = Reader.ReadBool();
 			ConfigValueTracker = new ConfigValueTracker(Reader);
@@ -457,6 +466,7 @@ namespace UnrealBuildTool
 			Writer.WriteFileReference(ExecutableFile);
 			Writer.WriteFileReference(ReceiptFile);
 			Writer.WriteDirectoryReference(ProjectIntermediateDirectory);
+			Writer.WriteDirectoryReference(ProjectIntermediateDirectoryNoArch);
 			Writer.WriteInt((int)TargetType);
 			Writer.WriteBool(IsTestTarget);
 			ConfigValueTracker.Write(Writer);
