UnrealBuildTool: Ignore conflicts in CreateAppBundle Actions

#rnx

diff --git a/Engine/Source/Programs/UnrealBuildTool/System/Action.cs b/Engine/Source/Programs/UnrealBuildTool/System/Action.cs
--- a/Engine/Source/Programs/UnrealBuildTool/System/Action.cs
+++ b/Engine/Source/Programs/UnrealBuildTool/System/Action.cs
@@ -511,7 +511,8 @@ namespace UnrealBuildTool
 		/// <returns>True if conflicts are ignored, else false.</returns>
 		public static bool IgnoreConflicts(this IExternalAction Action)
 		{
-			return Action.ActionType == ActionType.WriteMetadata;
+			return Action.ActionType == ActionType.WriteMetadata ||
+				Action.ActionType == ActionType.CreateAppBundle;
 		}
 	}
 
