UnrealBuildTool: Add Editor to ConfigHierarchyType

#rnx

diff --git a/Engine/Source/Programs/UnrealBuildTool/System/ConfigHierarchy.cs b/Engine/Source/Programs/UnrealBuildTool/System/ConfigHierarchy.cs
--- a/Engine/Source/Programs/UnrealBuildTool/System/ConfigHierarchy.cs
+++ b/Engine/Source/Programs/UnrealBuildTool/System/ConfigHierarchy.cs
@@ -28,6 +28,11 @@ namespace UnrealBuildTool
 		/// </summary>
 		Engine,
 
+		/// <summary>
+		/// BaseEditor.ini, DefaultEditor.ini, etc...
+		/// </summary>
+		Editor,
+
 		/// <summary>
 		/// BaseEditorPerProjectUserSettings.ini, DefaultEditorPerProjectUserSettings.ini, etc..
 		/// </summary>
