Fix remaining UnrealGame Win64 V3 pathing issues

#rnx

diff --git a/Engine/Plugins/Experimental/AnimNext/Source/AnimNextInterfaceGraph/Private/RigUnit_AnimNextInterfaceParameter.h b/Engine/Plugins/Experimental/AnimNext/Source/AnimNextInterfaceGraph/Private/RigUnit_AnimNextInterfaceParameter.h
--- a/Engine/Plugins/Experimental/AnimNext/Source/AnimNextInterfaceGraph/Private/RigUnit_AnimNextInterfaceParameter.h
+++ b/Engine/Plugins/Experimental/AnimNext/Source/AnimNextInterfaceGraph/Private/RigUnit_AnimNextInterfaceParameter.h
@@ -2,9 +2,9 @@
 
 #pragma once
 
+#include "AnimNextInterfaceExecuteContext.h"
 #include "IAnimNextInterface.h"
 #include "Units/RigUnit.h"
-#include "AnimNextInterfaceGraph/Internal/AnimNextInterfaceExecuteContext.h"
 #include "RigUnit_AnimNextInterfaceParameter.generated.h"
 
 struct FAnimNextInterfaceUnitContext;
diff --git a/Engine/Source/Developer/TraceAnalysis/Private/TraceAnalysisModule.cpp b/Engine/Source/Developer/TraceAnalysis/Private/TraceAnalysisModule.cpp
--- a/Engine/Source/Developer/TraceAnalysis/Private/TraceAnalysisModule.cpp
+++ b/Engine/Source/Developer/TraceAnalysis/Private/TraceAnalysisModule.cpp
@@ -2,7 +2,7 @@
 
 #include "TraceAnalysisModule.h"
 
-#include "MessageLog/Public/MessageLogModule.h"
+#include "MessageLogModule.h"
 #include "Modules/ModuleManager.h"
 
 ////////////////////////////////////////////////////////////////////////////////
diff --git a/Engine/Source/Developer/TraceAnalysis/TraceAnalysis.Build.cs b/Engine/Source/Developer/TraceAnalysis/TraceAnalysis.Build.cs
--- a/Engine/Source/Developer/TraceAnalysis/TraceAnalysis.Build.cs
+++ b/Engine/Source/Developer/TraceAnalysis/TraceAnalysis.Build.cs
@@ -11,5 +11,7 @@ public class TraceAnalysis : ModuleRules
 		PrivateDependencyModuleNames.Add("Core");
 		PrivateDependencyModuleNames.Add("Sockets");
 		PrivateDependencyModuleNames.Add("TraceLog");
+
+		PrivateIncludePathModuleNames.Add("MessageLog");
 	}
 }
diff --git a/Engine/Source/Editor/SourceControlWindows/Private/SSourceControlChangelists.cpp b/Engine/Source/Editor/SourceControlWindows/Private/SSourceControlChangelists.cpp
--- a/Engine/Source/Editor/SourceControlWindows/Private/SSourceControlChangelists.cpp
+++ b/Engine/Source/Editor/SourceControlWindows/Private/SSourceControlChangelists.cpp
@@ -42,7 +42,7 @@
 #include "HAL/PlatformTime.h"
 #include "HAL/FileManager.h"
 #include "ToolMenus.h"
-#include "UnsavedAssetsTracker/Source/Public/UnsavedAssetsTrackerModule.h"
+#include "UnsavedAssetsTrackerModule.h"
 #include "UObject/UObjectGlobals.h"
 #include "UObject/Package.h"
 #include "SSourceControlSubmit.h"
diff --git a/Engine/Source/Programs/UnrealBuildTool/Modes/FixIncludePathsMode.cs b/Engine/Source/Programs/UnrealBuildTool/Modes/FixIncludePathsMode.cs
--- a/Engine/Source/Programs/UnrealBuildTool/Modes/FixIncludePathsMode.cs
+++ b/Engine/Source/Programs/UnrealBuildTool/Modes/FixIncludePathsMode.cs
@@ -7,6 +7,7 @@ using System.Text.RegularExpressions;
 using EpicGames.Core;
 using UnrealBuildBase;
 using Microsoft.Extensions.Logging;
+using System.Diagnostics;
 
 namespace UnrealBuildTool
 {
@@ -180,6 +181,8 @@ namespace UnrealBuildTool
 											continue;
 										}
 
+										//Debugger.Launch();
+
 										string? PreferredInclude = null;
 										if (!PreferredPathCache.TryGetValue(Include, out PreferredInclude))
 										{
@@ -204,7 +207,7 @@ namespace UnrealBuildTool
 
 											if (FoundIncludeFile != null && !IncludeFileList.Contains(FoundIncludeFile))
 											{
-												Logger.LogDebug("{FileName}({LineNumber}): Skipping '{Include}' because it is filtered out.", InputFile.FullName, LineNumber, Include);
+												Logger.LogInformation("{FileName}({LineNumber}): Skipping '{Include}' because it is filtered out.", InputFile.FullName, LineNumber, Include);
 												PreferredInclude = Include;
 												PreferredPathCache[Include] = PreferredInclude;
 												continue;
@@ -215,7 +218,7 @@ namespace UnrealBuildTool
 												string FullPath = FoundIncludeFile.FullName.Replace('\\', '/');
 												if (FullPath.Contains("ThirdParty"))
 												{
-													Logger.LogDebug("{FileName}({LineNumber}): Skipping '{Include}' because it is a third party header.", InputFile.FullName, LineNumber, Include);
+													Logger.LogInformation("{FileName}({LineNumber}): Skipping '{Include}' because it is a third party header.", InputFile.FullName, LineNumber, Include);
 													PreferredInclude = Include;
 													PreferredPathCache[Include] = PreferredInclude;
 													continue;
@@ -225,14 +228,14 @@ namespace UnrealBuildTool
 												if (string.Equals(System.IO.Directory.GetParent(FullPath)?.FullName, System.IO.Directory.GetParent(InputFile.FullName)?.FullName, StringComparison.CurrentCultureIgnoreCase) &&
 													string.Equals(Include, System.IO.Path.GetFileName(FullPath), StringComparison.CurrentCultureIgnoreCase))
 												{
-													Logger.LogDebug("{FileName}({LineNumber}): Using '{Include}' because it is in the same directory.", InputFile.FullName, LineNumber, Include);
+													Logger.LogInformation("{FileName}({LineNumber}): Using '{Include}' because it is in the same directory.", InputFile.FullName, LineNumber, Include);
 													PreferredInclude = Include;
 												}
 												else
 												{
 													if (!FullPath.Contains(UnrealRootDirectory))
 													{
-														Logger.LogDebug("{FileName}({LineNumber}): Skipping '{Include}' because it isn't under the Unreal root directory.", InputFile.FullName, LineNumber, Include);
+														Logger.LogInformation("{FileName}({LineNumber}): Skipping '{Include}' because it isn't under the Unreal root directory.", InputFile.FullName, LineNumber, Include);
 													}
 													else
 													{
@@ -245,7 +248,7 @@ namespace UnrealBuildTool
 															// Is the current include a shortened version of the preferred include path?
 															if (PreferredInclude != Include && PreferredInclude.Contains(Include))
 															{
-																Logger.LogDebug("{FileName}({LineNumber}): Using '{Include}' because it is shorter than '{PreferredInclude}'.", InputFile.FullName, LineNumber, Include, PreferredInclude);
+																Logger.LogInformation("{FileName}({LineNumber}): Using '{Include}' because it is shorter than '{PreferredInclude}'.", InputFile.FullName, LineNumber, Include, PreferredInclude);
 																PreferredInclude = Include;
 															}
 														}
@@ -282,7 +285,7 @@ namespace UnrealBuildTool
 												
 											if (PreferredInclude == null)
 											{
-												Logger.LogDebug("{FileName}({LineNumber}): Could not find path to '{IncludePath}'", InputFile.FullName, LineNumber, Include);
+												Logger.LogInformation("{FileName}({LineNumber}): Could not find path to '{IncludePath}'", InputFile.FullName, LineNumber, Include);
 											}
 										}
 
diff --git a/Engine/Source/Runtime/RigVM/Public/RigVMFunctions/Math/RigVMMathLibrary.h b/Engine/Source/Runtime/RigVM/Public/RigVMFunctions/Math/RigVMMathLibrary.h
--- a/Engine/Source/Runtime/RigVM/Public/RigVMFunctions/Math/RigVMMathLibrary.h
+++ b/Engine/Source/Runtime/RigVM/Public/RigVMFunctions/Math/RigVMMathLibrary.h
@@ -4,8 +4,8 @@
 
 
 #include "CoreMinimal.h"
+#include "EulerTransform.h"
 #include "RigVMFunctions/RigVMFunctionDefines.h"
-#include "AnimationCore/Public/EulerTransform.h"
 #include "RigVMMathLibrary.generated.h"
 
 UENUM()
