Add support for overriding module optimization mode from a target. This allows game projects to override engine module optimization modes.

Note: This also implicitly disables PCH usage if a private PCH isn't set (this was previously fatal, but we need to be more flexible if we're overriding from the target).

diff --git a/Engine/Source/Programs/UnrealBuildTool/Configuration/ModuleRules.cs b/Engine/Source/Programs/UnrealBuildTool/Configuration/ModuleRules.cs
--- a/Engine/Source/Programs/UnrealBuildTool/Configuration/ModuleRules.cs
+++ b/Engine/Source/Programs/UnrealBuildTool/Configuration/ModuleRules.cs
@@ -642,6 +642,14 @@ namespace UnrealBuildTool
 		{
 			get
 			{
+				if (Target.OptimizeForSizeModules?.Contains(Name) ?? false)
+				{
+					return OptimizationMode.Size;
+				}
+				if (Target.OptimizeForSizeAndSpeedModules?.Contains(Name) ?? false)
+				{
+					return OptimizationMode.SizeAndSpeed;
+				}
 				if (OptimizationLevelOverride.HasValue)
 				{
 					return OptimizationLevelOverride.Value;
diff --git a/Engine/Source/Programs/UnrealBuildTool/Configuration/TargetRules.cs b/Engine/Source/Programs/UnrealBuildTool/Configuration/TargetRules.cs
--- a/Engine/Source/Programs/UnrealBuildTool/Configuration/TargetRules.cs
+++ b/Engine/Source/Programs/UnrealBuildTool/Configuration/TargetRules.cs
@@ -2210,6 +2210,20 @@ namespace UnrealBuildTool
 		[XmlConfigFile(Category = "ModuleConfiguration", Name = "DisableOptimizeCode")]
 		public string[]? DisableOptimizeCodeForModules = null;
 
+		/// <summary>
+		///  List of modules to optimize for size. This allows the target to override module optimization level
+		///  Note that this may disable PCH usage if a private PCH is not provided
+		/// </summary>
+		[XmlConfigFile(Category = "ModuleConfiguration", Name = "OptimizeForSize")]
+		public string[]? OptimizeForSizeModules = null;
+
+		/// <summary>
+		///  List of modules to optimize for size and speed. This allows the target to override module optimization level
+		///  Note that this may disable PCH usage if a private PCH is not provided
+		/// </summary>
+		[XmlConfigFile(Category = "ModuleConfiguration", Name = "OptimizeForSizeAndSpeed")]
+		public string[]? OptimizeForSizeAndSpeedModules = null;
+
 		/// <summary>
 		/// When generating project files, specifies the name of the project file to use when there are multiple targets of the same type.
 		/// </summary>
@@ -3842,6 +3856,16 @@ namespace UnrealBuildTool
 			get { return Inner.DisableOptimizeCodeForModules; }
 		}
 
+		public IReadOnlyList<string>? OptimizeForSizeModules
+		{
+			get { return Inner.OptimizeForSizeModules; }
+		}
+
+		public IReadOnlyList<string>? OptimizeForSizeAndSpeedModules
+		{
+			get { return Inner.OptimizeForSizeAndSpeedModules; }
+		}
+
 		public IReadOnlyList<UnrealTargetPlatform>? OptedInModulePlatforms
 		{
 			get { return Inner.OptedInModulePlatforms; } 
diff --git a/Engine/Source/Programs/UnrealBuildTool/Configuration/UEBuildModuleCPP.cs b/Engine/Source/Programs/UnrealBuildTool/Configuration/UEBuildModuleCPP.cs
--- a/Engine/Source/Programs/UnrealBuildTool/Configuration/UEBuildModuleCPP.cs
+++ b/Engine/Source/Programs/UnrealBuildTool/Configuration/UEBuildModuleCPP.cs
@@ -1699,11 +1699,20 @@ namespace UnrealBuildTool
 
 			if (Result.OptimizationLevel != Rules.OptimizationLevel)
 			{
-				if ( Rules.PCHUsage != ModuleRules.PCHUsageMode.NoPCHs && Rules.PrivatePCHHeaderFile == null )
+				Logger.LogInformation("Module {0} - Optimization level changed for module due to override. Old: {1} New: {2}", Name, Result.OptimizationLevel, Rules.OptimizationLevel);
+				if (Rules.PrivatePCHHeaderFile == null)
 				{
-					throw new BuildException("Module {0} - Overriding OptimizationLevel requires a private PCH", Name );
+					if (Rules.PCHUsage != ModuleRules.PCHUsageMode.NoPCHs)
+					{
+						Logger.LogInformation("  Overriding OptimizationLevel requires a private PCH. Disabling PCH usage for {0}", Name);
+						Rules.PCHUsage = ModuleRules.PCHUsageMode.NoPCHs;
+					}
+				}
+				else if (Rules.PCHUsage == ModuleRules.PCHUsageMode.UseSharedPCHs)
+				{
+					Logger.LogInformation("  Overriding OptimizationLevel requires a private PCH. A private PCH exists but UseSharedPCHs was specified. Overriding to NoSharedPCHs for {0}", Name);
+					Rules.PCHUsage = ModuleRules.PCHUsageMode.NoSharedPCHs;
 				}
-				Logger.LogInformation("Optimization level changed for module {0} due to override. Old: {1} New: {2}", Name, Result.OptimizationLevel, Rules.OptimizationLevel);
 				Result.OptimizationLevel = Rules.OptimizationLevel;
 			}
 
