UnrealBuildTool: Always make a unity file for small modules when enabled unless requested to not do so

#rnx

diff --git a/Engine/Source/Programs/UnrealBuildTool/Configuration/UEBuildModuleCPP.cs b/Engine/Source/Programs/UnrealBuildTool/Configuration/UEBuildModuleCPP.cs
--- a/Engine/Source/Programs/UnrealBuildTool/Configuration/UEBuildModuleCPP.cs
+++ b/Engine/Source/Programs/UnrealBuildTool/Configuration/UEBuildModuleCPP.cs
@@ -573,7 +573,7 @@ namespace UnrealBuildTool
 			// the module rules file.  By default, game modules only use unity of they have enough source files for that
 			// to be worthwhile.  If you have a lot of small game modules, consider specifying MinSourceFilesForUnityBuildOverride=0
 			// in the modules that you don't typically iterate on source files in very frequently.
-			int MinSourceFilesForUnityBuild = 2;
+			int MinSourceFilesForUnityBuild = 0;
 			if (Rules.MinSourceFilesForUnityBuildOverride != 0)
 			{
 				MinSourceFilesForUnityBuild = Rules.MinSourceFilesForUnityBuildOverride;
