UnrealBuildTool: Make module include path warning debugging output instead
#rnx

diff --git a/Engine/Source/Programs/UnrealBuildTool/System/ExternalExecution.cs b/Engine/Source/Programs/UnrealBuildTool/System/ExternalExecution.cs
--- a/Engine/Source/Programs/UnrealBuildTool/System/ExternalExecution.cs
+++ b/Engine/Source/Programs/UnrealBuildTool/System/ExternalExecution.cs
@@ -398,12 +398,12 @@ namespace UnrealBuildTool
 				if (ModuleIncludeBase == null && Module.ModuleDirectory.IsUnderDirectory(Unreal.EngineSourceDirectory))
 				{
 					ModuleIncludeBase = Unreal.EngineSourceDirectory;
-					Logger.LogWarning("Unable to find a module include path for {Module}, using engine source path '{Path}' because no standard Public/Private/Internal include paths were added. Please resolve by updating the module's .Build.cs", Module.Name, ModuleIncludeBase);
+					Logger.LogDebug("Unable to find a module include path for {Module}, using engine source path '{Path}' because no standard Public/Private/Internal include paths were added. Please resolve by updating the module's .Build.cs", Module.Name, ModuleIncludeBase);
 				}
 				if (ModuleIncludeBase == null)
 				{
 					ModuleIncludeBase = Module.GeneratedCodeDirectoryUHT!;
-					Logger.LogWarning("Unable to find a module include path for {Module}, using generated path '{Path}' because no standard Public/Private/Internal include paths were added. Please resolve by updating the module's .Build.cs", Module.Name, ModuleIncludeBase);
+					Logger.LogDebug("Unable to find a module include path for {Module}, using generated path '{Path}' because no standard Public/Private/Internal include paths were added. Please resolve by updating the module's .Build.cs", Module.Name, ModuleIncludeBase);
 				}
 			}
 
