[UnrealBuildTool]
* Fixed so -SingleFile includes running prereq actions if they are not up to date.

diff --git a/Engine/Source/Programs/UnrealBuildTool/Modes/BuildMode.cs b/Engine/Source/Programs/UnrealBuildTool/Modes/BuildMode.cs
--- a/Engine/Source/Programs/UnrealBuildTool/Modes/BuildMode.cs
+++ b/Engine/Source/Programs/UnrealBuildTool/Modes/BuildMode.cs
@@ -448,24 +448,17 @@ namespace UnrealBuildTool
 					TargetDescriptor TargetDescriptor = TargetDescriptors[Idx];
 					List<FileReference> FilesToBuild = TargetDescriptor.SpecificFilesToCompile;
 
-					// No specific files to filter on, queue all actions
-					if (FilesToBuild.Count == 0)
-					{
-						QueuedActions[Idx] = Makefiles[Idx].Actions.ConvertAll(x => new LinkedAction(x, TargetDescriptors[Idx]));
-					}
-					else
+					QueuedActions[Idx] = Makefiles[Idx].Actions.ConvertAll(x => new LinkedAction(x, TargetDescriptors[Idx]));
+
+					// If we have specific files to build we need to put those as MergedOutputItems (and create special single file actions if needed)
+					if (FilesToBuild.Count != 0)
 					{
 						// If there are headers in the list, expand the FilesToBuild to also include all files that include those headers
 						FilesToBuild = GetAllSourceFilesIncludingHeader(FilesToBuild, TargetDescriptor.ProjectFile, Makefiles[Idx].Actions, Logger);
 
 						// We have specific files to compile so we will only queue up those files.
-						QueuedActions[Idx] = CreateLinkedActionsFromFileList(TargetDescriptor, BuildConfiguration, FilesToBuild, Makefiles[Idx].Actions, Logger);
-
-						// Add all queued actions to the MergedOutputItems to make sure they are not skipped
-						foreach (LinkedAction Action in QueuedActions[Idx])
-						{
-							MergedOutputItems.Add(Action.ProducedItems.First());
-						}
+						List<FileItem> ProducedItems = CreateLinkedActionsFromFileList(TargetDescriptor, BuildConfiguration, FilesToBuild, QueuedActions[Idx], Logger);
+						MergedOutputItems.UnionWith(ProducedItems);
 					}
 				}
 
@@ -778,18 +771,18 @@ namespace UnrealBuildTool
 			}
 		}
 
-		internal static List<LinkedAction> CreateLinkedActionsFromFileList(TargetDescriptor Target, BuildConfiguration BuildConfiguration, List<FileReference> FileList, List<IExternalAction> Actions, ILogger Logger)
+		internal static List<FileItem> CreateLinkedActionsFromFileList(TargetDescriptor Target, BuildConfiguration BuildConfiguration, List<FileReference> FileList, List<LinkedAction> Actions, ILogger Logger)
 		{
 			// We have specific files to compile so we will only queue up those files.
 			// We will also add them to the MergedOutputItems to make sure they are not skipped
-			List<LinkedAction> LinkedActions = new();
+			List<FileItem> ProducedItems = new();
 
 			// First, find all the SpecificFileActions.
 			// These are used to create a custom action for the source file in case it is inside a unity or normally not part of the build (headers for example)
 			Dictionary<DirectoryReference, ISpecificFileAction> SpecificFileActions = new();
-			foreach (IExternalAction Action in Actions)
+			foreach (LinkedAction Action in Actions)
 			{
-				if (Action is ISpecificFileAction SpecificFileAction)
+				if (Action.Inner is ISpecificFileAction SpecificFileAction)
 				{
 					SpecificFileActions.TryAdd(SpecificFileAction.RootDirectory, SpecificFileAction);
 				}
@@ -824,21 +817,22 @@ namespace UnrealBuildTool
 					IExternalAction? NewAction = SpecificFileAction.CreateAction(SourceFile, Logger);
 					if (NewAction != null)
 					{
-						LinkedActions.Add(new LinkedAction(NewAction, Target));
+						Actions.Add(new LinkedAction(NewAction, Target));
+						ProducedItems.AddRange(NewAction.ProducedItems);
 						continue;
 					}
 				}
 
 				// There is no special action for this file.. let's look for an action that has this exact file as input and use that (for example ispc files)
 				bool FoundAction = false;
-				foreach (IExternalAction Action in Actions)
+				foreach (LinkedAction Action in Actions)
 				{
 					foreach (FileItem PrereqItem in Action.PrerequisiteItems)
 					{
 						if (PrereqItem == SourceFile)
 						{
 							FoundAction = true;
-							LinkedActions.Add(new LinkedAction(Action, Target));
+							ProducedItems.AddRange(Action.ProducedItems);
 							break;
 						}
 					}
@@ -854,7 +848,7 @@ namespace UnrealBuildTool
 				}
 			}
 
-			return LinkedActions;
+			return ProducedItems;
 		}
 
 
diff --git a/Engine/Source/Programs/UnrealBuildTool/Platform/Windows/VCToolChain.cs b/Engine/Source/Programs/UnrealBuildTool/Platform/Windows/VCToolChain.cs
--- a/Engine/Source/Programs/UnrealBuildTool/Platform/Windows/VCToolChain.cs
+++ b/Engine/Source/Programs/UnrealBuildTool/Platform/Windows/VCToolChain.cs
@@ -2093,6 +2093,7 @@ namespace UnrealBuildTool
 
 			VCCompileAction BaseCompileAction = CreateBaseCompileAction(CompileEnvironment);
 			AppendCLArguments_CPP(CompileEnvironment, BaseCompileAction.Arguments);
+			BaseCompileAction.AdditionalPrerequisiteItems.AddRange(CompileEnvironment.AdditionalPrerequisites); // Primarily for ispc.generated.h files
 			Graph.AddAction(new VcSpecificFileAction(SourceDir, OutputDir, BaseCompileAction));
 		}
 
diff --git a/Engine/Source/Programs/UnrealBuildTool/ToolChain/ClangToolChain.cs b/Engine/Source/Programs/UnrealBuildTool/ToolChain/ClangToolChain.cs
--- a/Engine/Source/Programs/UnrealBuildTool/ToolChain/ClangToolChain.cs
+++ b/Engine/Source/Programs/UnrealBuildTool/ToolChain/ClangToolChain.cs
@@ -1073,11 +1073,12 @@ namespace UnrealBuildTool
 				GetCompileArguments_Global(CompileEnvironment, GlobalArguments);
 			}
 
-			FileItem DummyFile = FileItem.GetItemByFileReference(FileReference.Combine(OutputDir, "SingleFile.cpp"));
+			string DummyName = "SingleFile.cpp";
+			FileItem DummyFile = FileItem.GetItemByFileReference(FileReference.Combine(OutputDir, DummyName));
 			CPPOutput Result = new();
 			ClangSpecificFileActionGraphBuilder GraphBuilder = new(Logger);
 			Action Action = CompileCPPFile(CompileEnvironment, DummyFile, OutputDir, "<Unknown>", GraphBuilder, GlobalArguments, Result);
-			Action.PrerequisiteItems.Clear();
+			Action.PrerequisiteItems.RemoveWhere(File => File.Name.Contains(DummyName));
 
 			Graph.AddAction(new ClangSpecificFileAction(SourceDir, OutputDir, Action, GraphBuilder.ContentLines));
 		}
