UnrealBuildTool: Mounting a cpp cache for engine is always required and should use the writable engine dir when installed

#rnx

diff --git a/Engine/Source/Programs/UnrealBuildTool/System/CppDependencyCache.cs b/Engine/Source/Programs/UnrealBuildTool/System/CppDependencyCache.cs
--- a/Engine/Source/Programs/UnrealBuildTool/System/CppDependencyCache.cs
+++ b/Engine/Source/Programs/UnrealBuildTool/System/CppDependencyCache.cs
@@ -346,9 +346,9 @@ namespace UnrealBuildTool
 		/// <returns>Dependency cache hierarchy for the given project</returns>
 		public void Mount(FileReference? ProjectFile, string TargetName, UnrealTargetPlatform Platform, UnrealTargetConfiguration Configuration, TargetType TargetType, UnrealArchitectures Architectures, ILogger Logger)
 		{
+			string AppName = TargetName;
 			if (ProjectFile == null || !Unreal.IsEngineInstalled())
 			{
-				string AppName;
 				if (TargetType == TargetType.Program)
 				{
 					AppName = TargetName;
@@ -357,11 +357,11 @@ namespace UnrealBuildTool
 				{
 					AppName = UEBuildTarget.GetAppNameForTargetType(TargetType);
 				}
-
-				FileReference EngineCacheLocation = FileReference.Combine(Unreal.EngineDirectory, UEBuildTarget.GetPlatformIntermediateFolder(Platform, Architectures, false), AppName, Configuration.ToString(), "DependencyCache.bin");
-				FindOrAddPartition(EngineCacheLocation, Unreal.EngineDirectory, Logger);
 			}
 
+			FileReference EngineCacheLocation = FileReference.Combine(Unreal.WritableEngineDirectory, UEBuildTarget.GetPlatformIntermediateFolder(Platform, Architectures, false), AppName, Configuration.ToString(), "DependencyCache.bin");
+			FindOrAddPartition(EngineCacheLocation, Unreal.EngineDirectory, Logger);
+
 			if (ProjectFile != null)
 			{
 				FileReference ProjectCacheLocation = FileReference.Combine(ProjectFile.Directory, UEBuildTarget.GetPlatformIntermediateFolder(Platform, Architectures, false), TargetName, Configuration.ToString(), "DependencyCache.bin");
