Add missing copyright boilerplate

#rnx

diff --git a/Engine/Source/Programs/UnrealBuildTool/ProjectFiles/ProjectFile.cs b/Engine/Source/Programs/UnrealBuildTool/ProjectFiles/ProjectFile.cs
--- a/Engine/Source/Programs/UnrealBuildTool/ProjectFiles/ProjectFile.cs
+++ b/Engine/Source/Programs/UnrealBuildTool/ProjectFiles/ProjectFile.cs
@@ -1,4 +1,6 @@
-﻿using EpicGames.Core;
+﻿// Copyright Epic Games, Inc. All Rights Reserved.
+
+using EpicGames.Core;
 using Microsoft.Extensions.Logging;
 using System;
 using System.Collections.Generic;
