Fixed compile issues when compiling with clangcl.

diff --git a/Engine/Source/Programs/UnrealBuildTool/Platform/Windows/VCToolChain.cs b/Engine/Source/Programs/UnrealBuildTool/Platform/Windows/VCToolChain.cs
--- a/Engine/Source/Programs/UnrealBuildTool/Platform/Windows/VCToolChain.cs
+++ b/Engine/Source/Programs/UnrealBuildTool/Platform/Windows/VCToolChain.cs
@@ -1058,6 +1058,7 @@ namespace UnrealBuildTool
 				Arguments.Add("-Wno-tautological-compare");
 				Arguments.Add("-Wno-switch");
 				Arguments.Add("-Wno-invalid-offsetof"); // needed to suppress warnings about using offsetof on non-POD types.
+				Arguments.Add("-Wno-return-type");				// needed for external headers we shan't change 
 
 				// @todo clang: Sorry for adding more of these, but I couldn't read my output log. Most should probably be looked at
 				Arguments.Add("-Wno-unused-parameter");			// Unused function parameter. A lot are named 'bUnused'...
