Possible fix for SNDBS when compiling with PCH chains.

diff --git a/Engine/Source/Programs/UnrealBuildTool/Configuration/UEBuildModuleCPP.cs b/Engine/Source/Programs/UnrealBuildTool/Configuration/UEBuildModuleCPP.cs
--- a/Engine/Source/Programs/UnrealBuildTool/Configuration/UEBuildModuleCPP.cs
+++ b/Engine/Source/Programs/UnrealBuildTool/Configuration/UEBuildModuleCPP.cs
@@ -826,7 +826,7 @@ namespace UnrealBuildTool
 			return new PrecompiledHeaderTemplate(this, CompileEnvironment, HeaderFile, PrecompiledHeaderDir);
 		}
 
-		private HashSet<string> GetImmutableDefinitions(List<string> Definitions)
+		static HashSet<string> GetImmutableDefinitions(List<string> Definitions)
 		{
 			HashSet<string> ImmutableDefinitions = new();
 			foreach (string Definition in Definitions)
@@ -966,7 +966,7 @@ namespace UnrealBuildTool
 							ParentPCHInstance = ParentTemplate.Instances.Find(x => IsCompatibleForSharedPCH(x.CompileEnvironment, ModuleCompileEnvironment) && AreSharedPCHDefinitionsCompatible(x));
 							if (ParentPCHInstance != null)
 							{
-								CompileEnvironment.ParentPrecompiledHeaderFile = ParentPCHInstance.Output.PrecompiledHeaderFile;
+								CompileEnvironment.ParentPCHInstance = ParentPCHInstance;
 								break;
 							}
 						}
@@ -1344,8 +1344,7 @@ namespace UnrealBuildTool
 				FileEnvironment.Definitions.Clear();
 				FileEnvironment.PrecompiledHeaderAction = PrecompiledHeaderAction.Include;
 				FileEnvironment.PrecompiledHeaderIncludeFilename = DedicatedPchFile.Location;
-				FileEnvironment.PrecompiledHeaderFile = PchOutput.PrecompiledHeaderFile;
-				FileEnvironment.PerArchPrecompiledHeaderFiles = PchOutput.PerArchPrecompiledHeaderFiles;
+				FileEnvironment.PCHInstance = new PrecompiledHeaderInstance(DedicatedPchFile, DedicatedPchFile, PchEnvironment, PchOutput, GetImmutableDefinitions(CompileEnvironment.Definitions));
 
 				// Create the action to compile the PCH file.
 				CPPOutput FileOutput = ToolChain.CompileAllCPPFiles(FileEnvironment, new List<FileItem>() { File }, IntermediateDirectory, ModuleName, Graph);
@@ -1375,8 +1374,7 @@ namespace UnrealBuildTool
 					CompileEnvironment.Definitions.Clear();
 					CompileEnvironment.PrecompiledHeaderAction = PrecompiledHeaderAction.Include;
 					CompileEnvironment.PrecompiledHeaderIncludeFilename = Instance.HeaderFile.Location;
-					CompileEnvironment.PrecompiledHeaderFile = Instance.Output.PrecompiledHeaderFile;
-					CompileEnvironment.PerArchPrecompiledHeaderFiles = Instance.Output.PerArchPrecompiledHeaderFiles;
+					CompileEnvironment.PCHInstance = Instance;
 
 					LinkInputFiles.AddRange(Instance.Output.ObjectFiles);
 				}
@@ -1444,8 +1442,7 @@ namespace UnrealBuildTool
 						CompileEnvironment.ForceIncludeFiles.Add(PrivateDefinitionsFileItem);
 						CompileEnvironment.PrecompiledHeaderAction = PrecompiledHeaderAction.Include;
 						CompileEnvironment.PrecompiledHeaderIncludeFilename = Instance.HeaderFile.Location;
-						CompileEnvironment.PrecompiledHeaderFile = Instance.Output.PrecompiledHeaderFile;
-						CompileEnvironment.PerArchPrecompiledHeaderFiles = Instance.Output.PerArchPrecompiledHeaderFiles;
+						CompileEnvironment.PCHInstance = Instance;
 
 						LinkInputFiles.AddRange(Instance.Output.ObjectFiles);
 					}
diff --git a/Engine/Source/Programs/UnrealBuildTool/Modes/PrintBuildGraphInfo.cs b/Engine/Source/Programs/UnrealBuildTool/Modes/PrintBuildGraphInfo.cs
--- a/Engine/Source/Programs/UnrealBuildTool/Modes/PrintBuildGraphInfo.cs
+++ b/Engine/Source/Programs/UnrealBuildTool/Modes/PrintBuildGraphInfo.cs
@@ -60,7 +60,7 @@ namespace UnrealBuildTool
 					foreach (UEBuildBinary Binary in Target.Binaries)
 					{
 						CppCompileEnvironment BinaryCompileEnvironment = Binary.CreateBinaryCompileEnvironment(GlobalCompileEnvironment);
-						foreach (UEBuildModule Module in Binary.Modules)
+						foreach (UEBuildModule Module in UEBuildModule.StableTopologicalSort(Binary.Modules))
 						{
 							if (Module is UEBuildModuleCPP ModuleCPP)
 							{
@@ -163,10 +163,10 @@ namespace UnrealBuildTool
 							SortedInstances.SortBy(Instance => Instance.HeaderFile.Name);
 							foreach (PrecompiledHeaderInstance Instance in SortedInstances)
 							{
-								Logger.LogInformation("   {InstanceName} - Used by {TimesUsed} modules:", Instance.CompileEnvironment.PrecompiledHeaderIncludeFilename, Instance.Modules.Count);
+								Logger.LogInformation("   {InstanceName} - Used by {TimesUsed} modules:", Instance.HeaderFile, Instance.Modules.Count);
 								if (Instance.ParentPCHInstance != null)
 								{
-									Logger.LogInformation("    ParentPCH: {ParentInstanceName}", Instance.ParentPCHInstance.CompileEnvironment.PrecompiledHeaderIncludeFilename);
+									Logger.LogInformation("    ParentPCH: {ParentInstanceName}", Instance.ParentPCHInstance.HeaderFile);
 								}
 
 								List<UEBuildModuleCPP> SortedModules = new(Instance.Modules);
diff --git a/Engine/Source/Programs/UnrealBuildTool/System/CppCompileEnvironment.cs b/Engine/Source/Programs/UnrealBuildTool/System/CppCompileEnvironment.cs
--- a/Engine/Source/Programs/UnrealBuildTool/System/CppCompileEnvironment.cs
+++ b/Engine/Source/Programs/UnrealBuildTool/System/CppCompileEnvironment.cs
@@ -495,17 +495,22 @@ namespace UnrealBuildTool
 		/// <summary>
 		/// The file containing the precompiled header data.
 		/// </summary>
-		public FileItem? PrecompiledHeaderFile = null;
+		public FileItem? PrecompiledHeaderFile => PCHInstance?.Output.PrecompiledHeaderFile;
 
 		/// <summary>
-		/// The parent file containing the precompiled header data.
+		/// A dictionary of PCH files for multiple architectures
 		/// </summary>
-		public FileItem? ParentPrecompiledHeaderFile = null;
+		public Dictionary<UnrealArch, FileItem>? PerArchPrecompiledHeaderFiles => PCHInstance?.Output.PerArchPrecompiledHeaderFiles;
 
 		/// <summary>
-		/// A dictionary of PCH files for multiple architectures
+		/// The instance containing the precompiled header data.
+		/// </summary>
+		public PrecompiledHeaderInstance? PCHInstance = null;
+
+		/// <summary>
+		/// The parent PCH instance used when creating this PCH.
 		/// </summary>
-		public Dictionary<UnrealArch, FileItem>? PerArchPrecompiledHeaderFiles = null;
+		public PrecompiledHeaderInstance? ParentPCHInstance = null;
 
 		/// <summary>
 		/// True if a single PRecompiledHeader exists, or at least one PerArchPrecompiledHeaderFile exists
@@ -586,10 +591,6 @@ namespace UnrealBuildTool
 			MetadataCache = Other.MetadataCache;
 			SharedPCHs = Other.SharedPCHs;
 			PrecompiledHeaderIncludeFilename = Other.PrecompiledHeaderIncludeFilename;
-			if (Other.PerArchPrecompiledHeaderFiles != null)
-			{
-				PerArchPrecompiledHeaderFiles = new(Other.PerArchPrecompiledHeaderFiles);
-			}
 			PrecompiledHeaderAction = Other.PrecompiledHeaderAction;
 			bUseSharedBuildEnvironment = Other.bUseSharedBuildEnvironment;
 			bUseRTTI = Other.bUseRTTI;
@@ -649,8 +650,8 @@ namespace UnrealBuildTool
 			AdditionalResponseFiles.AddRange(Other.AdditionalResponseFiles);
 			AdditionalArguments = Other.AdditionalArguments;
 			AdditionalFrameworks.AddRange(Other.AdditionalFrameworks);
-			PrecompiledHeaderFile = Other.PrecompiledHeaderFile;
-			ParentPrecompiledHeaderFile = Other.ParentPrecompiledHeaderFile;
+			PCHInstance = Other.PCHInstance;
+			ParentPCHInstance = Other.ParentPCHInstance;
 			bHackHeaderGenerator = Other.bHackHeaderGenerator;
 			bHideSymbolsByDefault = Other.bHideSymbolsByDefault;
 			CppStandard = Other.CppStandard;
@@ -665,6 +666,8 @@ namespace UnrealBuildTool
 			: this(Other)
 		{
 			Architectures = new UnrealArchitectures(OverrideArchitecture);
+
+			FileItem? PrecompiledHeaderFile;
 			if (Other.PerArchPrecompiledHeaderFiles != null)
 			{
 				Other.PerArchPrecompiledHeaderFiles.TryGetValue(OverrideArchitecture, out PrecompiledHeaderFile);
diff --git a/Engine/Source/Programs/UnrealBuildTool/ToolChain/ClangToolChain.cs b/Engine/Source/Programs/UnrealBuildTool/ToolChain/ClangToolChain.cs
--- a/Engine/Source/Programs/UnrealBuildTool/ToolChain/ClangToolChain.cs
+++ b/Engine/Source/Programs/UnrealBuildTool/ToolChain/ClangToolChain.cs
@@ -485,9 +485,9 @@ namespace UnrealBuildTool
 			{
 				Arguments.Add(GetIncludePCHFileArgument(CompileEnvironment.PrecompiledHeaderFile!));
 			}
-			else if (CompileEnvironment.PrecompiledHeaderAction == PrecompiledHeaderAction.Create && CompileEnvironment.ParentPrecompiledHeaderFile != null)
+			else if (CompileEnvironment.PrecompiledHeaderAction == PrecompiledHeaderAction.Create && CompileEnvironment.ParentPCHInstance != null)
 			{
-				Arguments.Add(GetIncludePCHFileArgument(CompileEnvironment.ParentPrecompiledHeaderFile));
+				Arguments.Add(GetIncludePCHFileArgument(CompileEnvironment.ParentPCHInstance.Output.PrecompiledHeaderFile!));
 			}
 
 			Arguments.AddRange(CompileEnvironment.ForceIncludeFiles.Select(ForceIncludeFile => GetForceIncludeFileArgument(ForceIncludeFile)));
@@ -938,8 +938,14 @@ namespace UnrealBuildTool
 
 			if (CompileEnvironment.PrecompiledHeaderAction == PrecompiledHeaderAction.Include)
 			{
-				CompileAction.PrerequisiteItems.Add(CompileEnvironment.PrecompiledHeaderFile!);
 				CompileAction.PrerequisiteItems.Add(FileItem.GetItemByFileReference(CompileEnvironment.PrecompiledHeaderIncludeFilename!));
+
+				PrecompiledHeaderInstance? PCHInstance = CompileEnvironment.PCHInstance;
+				while (PCHInstance != null)
+				{
+					CompileAction.PrerequisiteItems.Add(PCHInstance.Output.PrecompiledHeaderFile!);
+					PCHInstance = PCHInstance.ParentPCHInstance;
+				}
 			}
 
 			string FileName = SourceFile.Name;
@@ -955,9 +961,11 @@ namespace UnrealBuildTool
 				OutputFile = FileItem.GetItemByFileReference(FileReference.Combine(OutputDir, GetFileNameFromExtension(FileName, ".gch")));
 				CompileResult.PrecompiledHeaderFile = OutputFile;
 
-				if (CompileEnvironment.ParentPrecompiledHeaderFile != null)
+				PrecompiledHeaderInstance? ParentPCHInstance = CompileEnvironment.ParentPCHInstance;
+				while (ParentPCHInstance != null)
 				{
-					CompileAction.PrerequisiteItems.Add(CompileEnvironment.ParentPrecompiledHeaderFile!);
+					CompileAction.PrerequisiteItems.Add(ParentPCHInstance.Output.PrecompiledHeaderFile!);
+					ParentPCHInstance = ParentPCHInstance.ParentPCHInstance;
 				}
 			}
 			else if (CompileEnvironment.bPreprocessOnly)
