UnrealBuildTool: Marketplace plugins shouldn't default to the engine BuildSettingsVersion and should match other plugins

#rnx

diff --git a/Engine/Source/Programs/UnrealBuildTool/System/RulesCompiler.cs b/Engine/Source/Programs/UnrealBuildTool/System/RulesCompiler.cs
--- a/Engine/Source/Programs/UnrealBuildTool/System/RulesCompiler.cs
+++ b/Engine/Source/Programs/UnrealBuildTool/System/RulesCompiler.cs
@@ -215,7 +215,7 @@ namespace UnrealBuildTool
 			if (ModuleFileToContext.Count > 0)
 			{
 				FileReference AssemblyFileName = FileReference.Combine(Unreal.WritableEngineDirectory, "Intermediate", "Build", "BuildRules", "MarketplaceRules.dll");
-				Result = new RulesAssembly(MarketplaceScope, new List<DirectoryReference> { DirectoryReference.Combine(Unreal.EngineDirectory, "Plugins", "Marketplace") }, Plugins, ModuleFileToContext, new List<FileReference>(), AssemblyFileName, bContainsEngineModules: true, DefaultBuildSettings: DefaultEngineBuildSettingsVersion, bReadOnly: bReadOnly, bSkipCompile: bSkipCompile, bForceCompile: bForceCompile, Parent: Parent, Logger: Logger);
+				Result = new RulesAssembly(MarketplaceScope, new List<DirectoryReference> { DirectoryReference.Combine(Unreal.EngineDirectory, "Plugins", "Marketplace") }, Plugins, ModuleFileToContext, new List<FileReference>(), AssemblyFileName, bContainsEngineModules: true, DefaultBuildSettings: null, bReadOnly: bReadOnly, bSkipCompile: bSkipCompile, bForceCompile: bForceCompile, Parent: Parent, Logger: Logger);
 			}
 			return Result;
 		}
