Possible fix for android having issues with PCHs

diff --git a/Engine/Source/Programs/UnrealBuildTool/System/CppCompileEnvironment.cs b/Engine/Source/Programs/UnrealBuildTool/System/CppCompileEnvironment.cs
--- a/Engine/Source/Programs/UnrealBuildTool/System/CppCompileEnvironment.cs
+++ b/Engine/Source/Programs/UnrealBuildTool/System/CppCompileEnvironment.cs
@@ -8,6 +8,7 @@ using System.Text.RegularExpressions;
 using System.IO;
 using EpicGames.Core;
 using UnrealBuildBase;
+using System.Runtime.InteropServices;
 
 namespace UnrealBuildTool
 {
@@ -131,6 +132,19 @@ namespace UnrealBuildTool
 				PerArchPrecompiledHeaderFiles[Architecture] = Other.PrecompiledHeaderFile;
 			}
 		}
+
+		public FileItem? GetPrecompiledHeaderFile(UnrealArch Architecture)
+		{
+			if (PerArchPrecompiledHeaderFiles != null)
+			{
+				PerArchPrecompiledHeaderFiles.TryGetValue(Architecture, out FileItem? PerArchPrecompiledHeaderFile);
+				if (PerArchPrecompiledHeaderFile != null)
+				{
+					return PerArchPrecompiledHeaderFile;
+				}
+			}
+			return this.PrecompiledHeaderFile;
+		}
 	}
 
 	/// <summary>
@@ -492,11 +506,6 @@ namespace UnrealBuildTool
 		/// </summary>
 		public List<UEBuildFramework> AdditionalFrameworks = new List<UEBuildFramework>();
 
-		/// <summary>
-		/// The file containing the precompiled header data.
-		/// </summary>
-		public FileItem? PrecompiledHeaderFile => PCHInstance?.Output.PrecompiledHeaderFile;
-
 		/// <summary>
 		/// A dictionary of PCH files for multiple architectures
 		/// </summary>
@@ -507,11 +516,33 @@ namespace UnrealBuildTool
 		/// </summary>
 		public PrecompiledHeaderInstance? PCHInstance = null;
 
+		/// <summary>
+		/// The file containing the precompiled header data.
+		/// </summary>
+		public FileItem? PrecompiledHeaderFile
+		{
+			get
+			{
+				return GetPrecompiledHeaderFile(PCHInstance);
+			}
+		}
+
 		/// <summary>
 		/// The parent PCH instance used when creating this PCH.
 		/// </summary>
 		public PrecompiledHeaderInstance? ParentPCHInstance = null;
 
+		/// <summary>
+		/// The parent's PCH header file.
+		/// </summary>
+		public FileItem? ParentPrecompiledHeaderFile
+		{
+			get
+			{
+				return GetPrecompiledHeaderFile(ParentPCHInstance);
+			}
+		}
+
 		/// <summary>
 		/// True if a single PRecompiledHeader exists, or at least one PerArchPrecompiledHeaderFile exists
 		/// </summary>
@@ -666,20 +697,11 @@ namespace UnrealBuildTool
 			: this(Other)
 		{
 			Architectures = new UnrealArchitectures(OverrideArchitecture);
+		}
 
-			FileItem? PrecompiledHeaderFile;
-			if (Other.PerArchPrecompiledHeaderFiles != null)
-			{
-				Other.PerArchPrecompiledHeaderFiles.TryGetValue(OverrideArchitecture, out PrecompiledHeaderFile);
-			}
-			else
-			{
-				PrecompiledHeaderFile = null;
-			}
-			if (PrecompiledHeaderFile == null && PrecompiledHeaderAction == PrecompiledHeaderAction.Include)
-			{
-				Console.WriteLine("badness");
-			}
+		private FileItem? GetPrecompiledHeaderFile(PrecompiledHeaderInstance? Instance)
+		{
+			return Instance?.Output.GetPrecompiledHeaderFile(Architectures.SingleArchitecture);
 		}
 	}
 }
diff --git a/Engine/Source/Programs/UnrealBuildTool/ToolChain/ClangToolChain.cs b/Engine/Source/Programs/UnrealBuildTool/ToolChain/ClangToolChain.cs
--- a/Engine/Source/Programs/UnrealBuildTool/ToolChain/ClangToolChain.cs
+++ b/Engine/Source/Programs/UnrealBuildTool/ToolChain/ClangToolChain.cs
@@ -487,7 +487,7 @@ namespace UnrealBuildTool
 			}
 			else if (CompileEnvironment.PrecompiledHeaderAction == PrecompiledHeaderAction.Create && CompileEnvironment.ParentPCHInstance != null)
 			{
-				Arguments.Add(GetIncludePCHFileArgument(CompileEnvironment.ParentPCHInstance.Output.PrecompiledHeaderFile!));
+				Arguments.Add(GetIncludePCHFileArgument(CompileEnvironment.ParentPrecompiledHeaderFile!));
 			}
 
 			Arguments.AddRange(CompileEnvironment.ForceIncludeFiles.Select(ForceIncludeFile => GetForceIncludeFileArgument(ForceIncludeFile)));
@@ -943,7 +943,7 @@ namespace UnrealBuildTool
 				PrecompiledHeaderInstance? PCHInstance = CompileEnvironment.PCHInstance;
 				while (PCHInstance != null)
 				{
-					CompileAction.PrerequisiteItems.Add(PCHInstance.Output.PrecompiledHeaderFile!);
+					CompileAction.PrerequisiteItems.Add(PCHInstance.Output.GetPrecompiledHeaderFile(CompileEnvironment.Architecture)!);
 					PCHInstance = PCHInstance.ParentPCHInstance;
 				}
 			}
@@ -964,7 +964,7 @@ namespace UnrealBuildTool
 				PrecompiledHeaderInstance? ParentPCHInstance = CompileEnvironment.ParentPCHInstance;
 				while (ParentPCHInstance != null)
 				{
-					CompileAction.PrerequisiteItems.Add(ParentPCHInstance.Output.PrecompiledHeaderFile!);
+					CompileAction.PrerequisiteItems.Add(ParentPCHInstance.Output.GetPrecompiledHeaderFile(CompileEnvironment.Architecture)!);
 					ParentPCHInstance = ParentPCHInstance.ParentPCHInstance;
 				}
 			}
