Fix for changed container error during project file generation.

#rnx

diff --git a/Engine/Source/Programs/UnrealBuildTool/Configuration/UEBuildModule.cs b/Engine/Source/Programs/UnrealBuildTool/Configuration/UEBuildModule.cs
--- a/Engine/Source/Programs/UnrealBuildTool/Configuration/UEBuildModule.cs
+++ b/Engine/Source/Programs/UnrealBuildTool/Configuration/UEBuildModule.cs
@@ -1125,8 +1125,9 @@ namespace UnrealBuildTool
 			if (Modules == null)
 			{
 				Modules = new List<UEBuildModule>();
-				foreach (string ModuleName in ModuleNames)
+				for (int i = 0; i < ModuleNames.Count; i++)
 				{
+					string ModuleName = ModuleNames[i];
 					UEBuildModule Module = CreateModule(ModuleName, ReferenceChain);
 					if (!Modules.Contains(Module))
 					{
