[UnrealBuildTool]
* Removed /PDBCompress. This option only enables use of NTFS file compression and almost double link times. It also does not help with any network transfers or anything since it is only a local optimization. It does also not help fix large pdb linker crashes

diff --git a/Engine/Source/Programs/UnrealBuildTool/Platform/Windows/VCToolChain.cs b/Engine/Source/Programs/UnrealBuildTool/Platform/Windows/VCToolChain.cs
--- a/Engine/Source/Programs/UnrealBuildTool/Platform/Windows/VCToolChain.cs
+++ b/Engine/Source/Programs/UnrealBuildTool/Platform/Windows/VCToolChain.cs
@@ -1243,7 +1243,7 @@ namespace UnrealBuildTool
 					throw new BuildException($"Monolithic editors now require VC Toolchain 14.29.30133 (the toolchain for Visual Studio 16.11) ({EnvVars.CompilerVersion} < {VersionNumber.Parse("14.29.30133")})");
 				}
 				Arguments.Add("/PDBPAGESIZE:16384");
-				Arguments.Add("/PDBCompress");
+				//Arguments.Add("/PDBCompress"); // Do not turn this on, it makes link times almost 2x slower. This is _only_ to save local disk space. Will _not_ make actual file smaller for network transfer
 			}
 
 			//
