UnrealBuildTool: Enable BuildSettingsVersion.V3 as default for Engine modules

diff --git a/Engine/Source/Programs/UnrealBuildTool/Configuration/TargetRules.cs b/Engine/Source/Programs/UnrealBuildTool/Configuration/TargetRules.cs
--- a/Engine/Source/Programs/UnrealBuildTool/Configuration/TargetRules.cs
+++ b/Engine/Source/Programs/UnrealBuildTool/Configuration/TargetRules.cs
@@ -138,7 +138,7 @@ namespace UnrealBuildTool
 		/// <summary>
 		/// Always use the defaults for the current engine version. Note that this may cause compatibility issues when upgrading.
 		/// </summary>
-		Latest = V2,
+		Latest = V3,
 	}
 
 	/// <summary>
