Updated UBT to create the shared PCH instances before setting up the modules so all the PCHs and parent PCHs are setup beforehand.

diff --git a/Engine/Source/Programs/UnrealBuildTool/Configuration/UEBuildModuleCPP.cs b/Engine/Source/Programs/UnrealBuildTool/Configuration/UEBuildModuleCPP.cs
--- a/Engine/Source/Programs/UnrealBuildTool/Configuration/UEBuildModuleCPP.cs
+++ b/Engine/Source/Programs/UnrealBuildTool/Configuration/UEBuildModuleCPP.cs
@@ -882,13 +882,12 @@ namespace UnrealBuildTool
 		/// <summary>
 		/// Generates a precompiled header instance from the given template, or returns an existing one if it already exists
 		/// </summary>
-		/// <param name="Target">The target which owns this module</param>
 		/// <param name="ToolChain">The toolchain being used to build this module</param>
 		/// <param name="Template">The PCH template</param>
 		/// <param name="ModuleCompileEnvironment">Compile environment for the current module</param>
 		/// <param name="Graph">List of actions to be executed. Additional actions will be added to this list.</param>
 		/// <returns>Instance of a PCH</returns>
-		public PrecompiledHeaderInstance FindOrCreateSharedPCH(ReadOnlyTargetRules Target, UEToolChain? ToolChain, PrecompiledHeaderTemplate Template, CppCompileEnvironment ModuleCompileEnvironment, IActionGraphBuilder Graph)
+		public PrecompiledHeaderInstance FindOrCreateSharedPCH(UEToolChain? ToolChain, PrecompiledHeaderTemplate Template, CppCompileEnvironment ModuleCompileEnvironment, IActionGraphBuilder Graph)
 		{
 			PrecompiledHeaderInstance? Instance = Template.Instances.Find(x => IsCompatibleForSharedPCH(x.CompileEnvironment, ModuleCompileEnvironment));
 			if(Instance == null)
@@ -1001,7 +1000,7 @@ namespace UnrealBuildTool
 		/// <param name="ModuleCompileEnvironment">The module compile environment</param>
 		/// <param name="CompileEnvironment">The shared PCH compile environment</param>
 		/// <returns>True if the two compile enviroments are compatible</returns>
-		private bool IsCompatibleForSharedPCH(CppCompileEnvironment ModuleCompileEnvironment, CppCompileEnvironment CompileEnvironment)
+		static private bool IsCompatibleForSharedPCH(CppCompileEnvironment ModuleCompileEnvironment, CppCompileEnvironment CompileEnvironment)
 		{
 			if(ModuleCompileEnvironment.bOptimizeCode != CompileEnvironment.bOptimizeCode)
 			{
@@ -1050,7 +1049,7 @@ namespace UnrealBuildTool
 		/// <param name="CompileEnvironment">The shared PCH compile environment</param>
 		/// <param name="BaseCompileEnvironment">The base compile environment</param>
 		/// <returns>The unique suffix for the shared PCH</returns>
-		private string GetSuffixForSharedPCH(CppCompileEnvironment CompileEnvironment, CppCompileEnvironment BaseCompileEnvironment)
+		static private string GetSuffixForSharedPCH(CppCompileEnvironment CompileEnvironment, CppCompileEnvironment BaseCompileEnvironment)
 		{
 			string Variant = "";
 			if(CompileEnvironment.bOptimizeCode != BaseCompileEnvironment.bOptimizeCode)
@@ -1156,7 +1155,7 @@ namespace UnrealBuildTool
 		/// </summary>
 		/// <param name="ModuleCompileEnvironment">The module compile environment</param>
 		/// <param name="CompileEnvironment">The shared PCH compile environment</param>
-		private void CopySettingsForSharedPCH(CppCompileEnvironment ModuleCompileEnvironment, CppCompileEnvironment CompileEnvironment)
+		static private void CopySettingsForSharedPCH(CppCompileEnvironment ModuleCompileEnvironment, CppCompileEnvironment CompileEnvironment)
 		{
 			CompileEnvironment.bOptimizeCode = ModuleCompileEnvironment.bOptimizeCode;
 			CompileEnvironment.bUseRTTI = ModuleCompileEnvironment.bUseRTTI;
@@ -1390,7 +1389,7 @@ namespace UnrealBuildTool
 					PrecompiledHeaderTemplate? Template = CompileEnvironment.SharedPCHs.FirstOrDefault(x => ReferencedModules.Contains(x.Module));
 					if(Template != null && Template.IsValidFor(CompileEnvironment))
 					{
-						PrecompiledHeaderInstance Instance = FindOrCreateSharedPCH(Target, ToolChain, Template, CompileEnvironment, Graph);
+						PrecompiledHeaderInstance Instance = FindOrCreateSharedPCH(ToolChain, Template, CompileEnvironment, Graph);
 
 						FileReference PrivateDefinitionsFile = FileReference.Combine(IntermediateDirectory, String.Format("Definitions.{0}.h", Name));
 
diff --git a/Engine/Source/Programs/UnrealBuildTool/Configuration/UEBuildTarget.cs b/Engine/Source/Programs/UnrealBuildTool/Configuration/UEBuildTarget.cs
--- a/Engine/Source/Programs/UnrealBuildTool/Configuration/UEBuildTarget.cs
+++ b/Engine/Source/Programs/UnrealBuildTool/Configuration/UEBuildTarget.cs
@@ -2369,10 +2369,13 @@ namespace UnrealBuildTool
 					FileMetadataPrefetch.QueueDirectoryTree(Module.GeneratedCodeDirectoryUHT);
 			}
 
-			// Find all the shared PCHs.
 			if (Rules.bUseSharedPCHs)
 			{
+				// Find all the shared PCHs.
 				FindSharedPCHs(OriginalBinaries, GlobalCompileEnvironment, Logger);
+
+				// Create all the shared PCH instances before processing the modules
+				CreateSharedPCHInstances(Rules, TargetToolChain, OriginalBinaries, GlobalCompileEnvironment, MakefileBuilder, Logger);
 			}
 
 			// Can probably be moved further down?
@@ -3393,6 +3396,87 @@ namespace UnrealBuildTool
 			}
 		}
 
+		/// <summary>
+		/// Creates all the shared PCH instances before processing the modules
+		/// </summary>
+		/// <param name="Target">The target being built</param>
+		/// <param name="ToolChain">The toolchain to build with</param>
+		/// <param name="OriginalBinaries">The list of binaries</param>
+		/// <param name="GlobalCompileEnvironment">The compile environment. The shared PCHs will be added to the SharedPCHs list in this.</param>
+		/// <param name="Logger">Logger for output</param>
+		/// <param name="Graph">List of build actions</param>
+		void CreateSharedPCHInstances(ReadOnlyTargetRules Target, UEToolChain ToolChain, List<UEBuildBinary> OriginalBinaries, CppCompileEnvironment GlobalCompileEnvironment, IActionGraphBuilder Graph, ILogger Logger)
+		{
+			if (!Target.bUsePCHFiles || GlobalCompileEnvironment.SharedPCHs.Count == 0)
+			{
+				return;
+			}
+
+			List<UEBuildModuleCPP> Modules = new List<UEBuildModuleCPP>();
+			Dictionary<UEBuildBinary, CppCompileEnvironment> BinaryCompileEnvironments = new Dictionary<UEBuildBinary, CppCompileEnvironment>();
+			foreach (UEBuildBinary Binary in OriginalBinaries)
+			{
+				BinaryCompileEnvironments.Add(Binary, Binary.CreateBinaryCompileEnvironment(GlobalCompileEnvironment));
+				Modules.AddRange(Binary.Modules.OfType<UEBuildModuleCPP>());
+			}
+
+			// Find which instances are needed
+			Dictionary<UEBuildModuleCPP, CppCompileEnvironment> ModuleCompileEnvironments = new Dictionary<UEBuildModuleCPP, CppCompileEnvironment>();
+			Dictionary<PrecompiledHeaderTemplate, List<UEBuildModuleCPP>> PCHsAndModulesDict = new Dictionary<PrecompiledHeaderTemplate, List<UEBuildModuleCPP>>();
+			foreach (UEBuildModuleCPP Module in Modules)
+			{
+				CppCompileEnvironment BinaryCompileEnvironment = BinaryCompileEnvironments[Module.Binary];
+				CppCompileEnvironment ModuleCompileEnvironment = Module.CreateModuleCompileEnvironment(Target, BinaryCompileEnvironment, Logger);
+
+				// Is it using a private PCH?
+				if (Module.Rules.PrivatePCHHeaderFile != null && (Module.Rules.PCHUsage == ModuleRules.PCHUsageMode.NoSharedPCHs || Module.Rules.PCHUsage == ModuleRules.PCHUsageMode.UseExplicitOrSharedPCHs))
+				{
+					continue;
+				}
+
+				// Is it using a shared PCH?
+				if (Module.Rules.PCHUsage != ModuleRules.PCHUsageMode.NoPCHs)
+				{
+					if (!ModuleCompileEnvironment.bIsBuildingLibrary && Module.Rules.PCHUsage != ModuleRules.PCHUsageMode.NoSharedPCHs)
+					{
+						// Find all the dependencies of this module
+						HashSet<UEBuildModule> ReferencedModules = new HashSet<UEBuildModule>();
+						Module.GetAllDependencyModules(new List<UEBuildModule>(), ReferencedModules, bIncludeDynamicallyLoaded: false, bForceCircular: false, bOnlyDirectDependencies: true);
+
+						// Find the first shared PCH module we can use
+						PrecompiledHeaderTemplate? Template = ModuleCompileEnvironment.SharedPCHs.FirstOrDefault(x => ReferencedModules.Contains(x.Module));
+						if (Template != null && Template.IsValidFor(ModuleCompileEnvironment))
+						{
+							ModuleCompileEnvironments[Module] = ModuleCompileEnvironment;
+
+							if (PCHsAndModulesDict.TryGetValue(Template, out List<UEBuildModuleCPP>? ModulesUsingSharedPCHs))
+							{
+								ModulesUsingSharedPCHs.Add(Module);
+							}
+							else
+							{
+								PCHsAndModulesDict[Template] = new List<UEBuildModuleCPP>() { Module };
+							}
+						}
+					}
+				}
+			}
+
+			// Create the PCH instances
+			List<PrecompiledHeaderTemplate> SharedPCHs = new List<PrecompiledHeaderTemplate>(GlobalCompileEnvironment.SharedPCHs);
+			SharedPCHs.Reverse();
+			foreach (var SharedPCH in SharedPCHs)
+			{
+				if (PCHsAndModulesDict.TryGetValue(SharedPCH, out List<UEBuildModuleCPP>? ModulesUsingSharedPCHs))
+				{
+					foreach (var Module in ModulesUsingSharedPCHs)
+					{
+						Module.FindOrCreateSharedPCH(ToolChain, SharedPCH, ModuleCompileEnvironments[Module], Graph);
+					}
+				}
+			}
+		}
+
 		/// <summary>
 		/// When building a target, this is called to add any additional modules that should be compiled along
 		/// with the main target.  If you override this in a derived class, remember to call the base implementation!
diff --git a/Engine/Source/Programs/UnrealBuildTool/System/PrecompiledHeaderInstance.cs b/Engine/Source/Programs/UnrealBuildTool/System/PrecompiledHeaderInstance.cs
--- a/Engine/Source/Programs/UnrealBuildTool/System/PrecompiledHeaderInstance.cs
+++ b/Engine/Source/Programs/UnrealBuildTool/System/PrecompiledHeaderInstance.cs
@@ -37,7 +37,7 @@ namespace UnrealBuildTool
 		/// <summary>
 		/// List of modules using this instance
 		/// </summary>
-		public List<UEBuildModuleCPP> Modules = new List<UEBuildModuleCPP>();
+		public HashSet<UEBuildModuleCPP> Modules = new HashSet<UEBuildModuleCPP>();
 
 		/// <summary>
 		/// These are definitions that are immutable and should never be #undef. There are a few exceptions and we make sure those are not ending up in this list
