Fixed a PCH issue when compiling with multiple architectures at the same time.

diff --git a/Engine/Source/Programs/UnrealBuildTool/Configuration/UEBuildModuleCPP.cs b/Engine/Source/Programs/UnrealBuildTool/Configuration/UEBuildModuleCPP.cs
--- a/Engine/Source/Programs/UnrealBuildTool/Configuration/UEBuildModuleCPP.cs
+++ b/Engine/Source/Programs/UnrealBuildTool/Configuration/UEBuildModuleCPP.cs
@@ -438,17 +438,17 @@ namespace UnrealBuildTool
 
 			// Configure the precompiled headers for this module
 			CompileEnvironment = SetupPrecompiledHeaders(Target, ToolChain, CompileEnvironment, LinkInputFiles, Graph);
-			if (CompileEnvironment.PrecompiledHeaderFile != null)
-			{
-				Logger.LogDebug("Module '{ModuleName}' uses PCH '{PCHIncludeFilename}'", this.Name, CompileEnvironment.PrecompiledHeaderFile);
-			}
-			else if (CompileEnvironment.PerArchPrecompiledHeaderFiles != null)
+			if (CompileEnvironment.PerArchPrecompiledHeaderFiles != null)
 			{
 				foreach (UnrealArch Arch in CompileEnvironment.PerArchPrecompiledHeaderFiles.Keys)
 				{
 					Logger.LogDebug("Module '{ModuleName}' uses PCH '{PCHIncludeFilename}' for Architecture '{Arch}'", this.Name, Arch, CompileEnvironment.PerArchPrecompiledHeaderFiles[Arch]);
 				}
 			}
+			else if (CompileEnvironment.PrecompiledHeaderFile != null)
+			{
+				Logger.LogDebug("Module '{ModuleName}' uses PCH '{PCHIncludeFilename}'", this.Name, CompileEnvironment.PrecompiledHeaderFile);
+			}
 
 			// Write all the definitions to a separate file
 			CreateHeaderForDefinitions(CompileEnvironment, IntermediateDirectory, null, Graph);
diff --git a/Engine/Source/Programs/UnrealBuildTool/System/CppCompileEnvironment.cs b/Engine/Source/Programs/UnrealBuildTool/System/CppCompileEnvironment.cs
--- a/Engine/Source/Programs/UnrealBuildTool/System/CppCompileEnvironment.cs
+++ b/Engine/Source/Programs/UnrealBuildTool/System/CppCompileEnvironment.cs
@@ -546,7 +546,7 @@ namespace UnrealBuildTool
 		/// <summary>
 		/// True if a single PRecompiledHeader exists, or at least one PerArchPrecompiledHeaderFile exists
 		/// </summary>
-		public bool bHasPrecompiledHeader => PrecompiledHeaderFile != null || (PerArchPrecompiledHeaderFiles != null && PerArchPrecompiledHeaderFiles.Count > 0);
+		public bool bHasPrecompiledHeader => (PerArchPrecompiledHeaderFiles != null && PerArchPrecompiledHeaderFiles.Count > 0) || PrecompiledHeaderFile != null;
 
 		/// <summary>
 		/// Whether or not UHT is being built
