- Added ability to have a IDE project for each Target (UnrealGame vs UnrealEditor) and remove the Type from the Configurations (Debug, Development, etc - no more Development Editor config)
- Enabled for Xcode, added support but left disabled for VisualStudio
- For Games and Samples with multiple projects, a solution folder now contains them

diff --git a/Engine/Source/Programs/UnrealBuildTool/Platform/Windows/WindowsProjectGenerator.cs b/Engine/Source/Programs/UnrealBuildTool/Platform/Windows/WindowsProjectGenerator.cs
--- a/Engine/Source/Programs/UnrealBuildTool/Platform/Windows/WindowsProjectGenerator.cs
+++ b/Engine/Source/Programs/UnrealBuildTool/Platform/Windows/WindowsProjectGenerator.cs
@@ -65,7 +65,7 @@ namespace UnrealBuildTool
 					DebugOptions += ForeignUProjectPath;
 					DebugOptions += " -skipcompile";
 				}
-				else if (InTargetRules.Type == TargetType.Editor && ProjectName != ProjectFileGenerator.EngineProjectFileNameBase)
+				else if (InTargetRules.Type == TargetType.Editor && InTargetRules.ProjectFile != null)
 				{
 					DebugOptions += ProjectName;
 				}
diff --git a/Engine/Source/Programs/UnrealBuildTool/ProjectFiles/ProjectFileGenerator.cs b/Engine/Source/Programs/UnrealBuildTool/ProjectFiles/ProjectFileGenerator.cs
--- a/Engine/Source/Programs/UnrealBuildTool/ProjectFiles/ProjectFileGenerator.cs
+++ b/Engine/Source/Programs/UnrealBuildTool/ProjectFiles/ProjectFileGenerator.cs
@@ -336,6 +336,11 @@ namespace UnrealBuildTool
 		[XmlConfigFile(Name = "Configurations")]
 		string[]? ConfigurationNames = null;
 
+		/// <summary>
+		/// If true, this generator wants to have a project for each target type (UnrealGame, UnrealEditor, etc) and has only the straight configs (Debug, Development)
+		/// </summary>
+		protected virtual bool bMakeProjectPerTarget => false;
+
 		/// <summary>
 		/// Relative path to the directory where the primary project file will be saved to
 		/// </summary>
@@ -344,7 +349,17 @@ namespace UnrealBuildTool
 		/// <summary>
 		/// Name of the UE engine project that contains all of the engine code, config files and other files
 		/// </summary>
-		public const string EngineProjectFileNameBase = "UE5";
+		public const string EngineRulesAssemblyName = "UE5";
+
+		/// <summary>
+		/// Name of the UE engine project that contains all of the engine code, config files and other files
+		/// </summary>
+		public string EngineProjectFileNameBase => bMakeProjectPerTarget ? "UnrealGame" : EngineRulesAssemblyName;
+
+		/// <summary>
+		/// Name of the UE engine project that contains all of the engine code, config files and other files
+		/// </summary>
+		public string EngineEditorProjectFileNameBase => bMakeProjectPerTarget ? "UnrealEditor" : EngineRulesAssemblyName;
 
 		/// <summary>
 		/// When ProjectsAreIntermediate is true, this is the directory to store generated project files
@@ -379,7 +394,12 @@ namespace UnrealBuildTool
 		/// <summary>
 		/// Maps all module files that were included in generated project files, to actual project file objects.
 		/// </summary>
-		protected Dictionary<FileReference, ProjectFile> ModuleToProjectFileMap = new Dictionary<FileReference, ProjectFile>();
+		protected Dictionary<FileReference, ProjectFile> ModuleToEditorProjectFileMap = new Dictionary<FileReference, ProjectFile>();
+
+		/// <summary>
+		/// Maps all module files that were included in generated project files, to actual project file objects.
+		/// </summary>
+		protected Dictionary<FileReference, ProjectFile> ModuleToNonEditorProjectFileMap = new Dictionary<FileReference, ProjectFile>();
 
 		/// <summary>
 		/// If generating project files for a single project, the path to its .uproject file.
@@ -474,11 +494,13 @@ namespace UnrealBuildTool
 
 						if (ProjectFile.IsUnderDirectory(SamplesDirectory))
 						{
-							RootFolder.AddSubFolder("Samples").ChildProjects.Add(Project);
+							RootFolder.AddSubFolder("Samples").AddSubFolder(ProjectFile.GetFileNameWithoutAnyExtensions()).ChildProjects.Add(Project);
 						}
 						else
 						{
-							RootFolder.AddSubFolder("Games").ChildProjects.Add(Project);
+							// @todo: adding to a group that is the name of the game's automation - if the name is not the name of the game,
+							// we'd have to search up looking for the .uproject. not sure if this happens
+							RootFolder.AddSubFolder("Games").AddSubFolder(ProjectFile.GetFileNameWithoutAnyExtensions()).ChildProjects.Add(Project);
 						}
 					}
 					else
@@ -904,7 +926,7 @@ namespace UnrealBuildTool
 			// we're generating project files for.
 			List<FileReference> AllModuleFiles = DiscoverModules(AllGameProjects, AdditionalSearchPaths.Values.SelectMany(x => x).ToList());
 
-			ProjectFile? EngineProject = null;
+			List<ProjectFile> EngineProjects = new List<ProjectFile>();
 			List<ProjectFile> GameProjects = new List<ProjectFile>();
 			List<ProjectFile> ModProjects = new List<ProjectFile>();
 			Dictionary<FileReference, ProjectFile> ProgramProjects = new Dictionary<FileReference, ProjectFile>();
@@ -913,7 +935,7 @@ namespace UnrealBuildTool
 			if(IncludeCppSource)
 			{
 				// Setup buildable projects for all targets
-				AddProjectsForAllTargets(PlatformProjectGenerators, AllGameProjects, AllTargetFiles, Arguments, ref EngineProject, GameProjects, ProjectFileToUProjectFile, ProgramProjects, RulesAssemblies, Logger);
+				AddProjectsForAllTargets(PlatformProjectGenerators, AllGameProjects, AllTargetFiles, Arguments, EngineProjects, GameProjects, ProjectFileToUProjectFile, ProgramProjects, RulesAssemblies, Logger);
 
 				// Add projects for mods
 				AddProjectsForMods(GameProjects, ModProjects);
@@ -938,21 +960,30 @@ namespace UnrealBuildTool
 				if (bIncludeEngineSource)
 				{
 					// If we're still missing an engine project because we don't have any targets for it, make one up.
-					if (EngineProject == null)
+					if (EngineProjects.Count == 0)
 					{
 						FileReference ProjectFilePath = FileReference.Combine(IntermediateProjectFilesPath, EngineProjectFileNameBase + ProjectFileExtension);
 
-						bool bAlreadyExisted;
-						EngineProject = FindOrAddProject(ProjectFilePath, Unreal.EngineDirectory, true, out bAlreadyExisted);
+						ProjectFile EngineProject = FindOrAddProject(ProjectFilePath, Unreal.EngineDirectory, true, out _);
 
 						EngineProject.IsForeignProject = false;
 						EngineProject.IsGeneratedProject = true;
 						EngineProject.IsStubProject = true;
+
+						EngineProjects.Add(EngineProject);
 					}
 
-					if (EngineProject != null)
+					if (EngineProjects.Count > 0)
 					{
-						RootFolder.AddSubFolder("Engine").ChildProjects.Add(EngineProject);
+						// put all engine projects into solution
+						PrimaryProjectFolder EngineFolder = RootFolder.AddSubFolder("Engine");
+						foreach (ProjectFile Project in EngineProjects)
+						{
+							EngineFolder.ChildProjects.Add(Project);
+						}
+
+						// put all this stuff into one of the engine projects
+						ProjectFile EngineProject = EngineProjects.First(x => x.ProjectFilePath.GetFileNameWithoutAnyExtensions() == EngineProjectFileNameBase);
 
 						// Engine config files
 						if (bIncludeConfigFiles)
@@ -1031,11 +1062,11 @@ namespace UnrealBuildTool
 						else if (UnrealProjectFile.IsUnderDirectory(SamplesDirectory))
 						{
 							DirectoryReference SampleGameDirectory = CurGameProject.BaseDir;
-							RootFolder.AddSubFolder("Samples").ChildProjects.Add(CurGameProject);
+							RootFolder.AddSubFolder("Samples").AddSubFolder(UnrealProjectFile.GetFileNameWithoutExtension()).ChildProjects.Add(CurGameProject);
 						}
 						else
 						{
-							RootFolder.AddSubFolder("Games").ChildProjects.Add(CurGameProject);
+							RootFolder.AddSubFolder("Games").AddSubFolder(UnrealProjectFile.GetFileNameWithoutExtension()).ChildProjects.Add(CurGameProject);
 						}
 
 						List<Tuple<ProjectFile, string>>? NewProjectFiles = CurGameProject.WriteDebugProjectFiles(SupportedPlatforms, SupportedConfigurations, PlatformProjectGenerators, Logger);
@@ -1140,18 +1171,16 @@ namespace UnrealBuildTool
 				List<Tuple<ProjectFile, ProjectTarget>> IntelliSenseTargetFiles = new List<Tuple<ProjectFile, ProjectTarget>>();
 				{
 					// Engine targets
-					if (EngineProject != null)
+					foreach (ProjectFile EngineProject in EngineProjects)
 					{
-						foreach (ProjectTarget ProjectTarget in EngineProject.ProjectTargets)
+						//ProjectTarget ProjectTarget = (ProjectTarget)EngineProject.ProjectTarget!;
+						foreach (ProjectTarget EngineTarget in EngineProject.ProjectTargets)
 						{
-							if (ProjectTarget.TargetFilePath != null)
+							// Only bother with the editor target.  We want to make sure that definitions are setup to be as inclusive as possible
+							// for good quality IntelliSense.  For example, we want WITH_EDITORONLY_DATA=1, so using the editor targets works well.
+							if (bMakeProjectPerTarget || EngineTarget.TargetRules!.Type == TargetType.Editor) // @todo <-- try removing first term
 							{
-								// Only bother with the editor target.  We want to make sure that definitions are setup to be as inclusive as possible
-								// for good quality IntelliSense.  For example, we want WITH_EDITORONLY_DATA=1, so using the editor targets works well.
-								if (ProjectTarget.TargetRules!.Type == TargetType.Editor)
-								{
-									IntelliSenseTargetFiles.Add(Tuple.Create(EngineProject, ProjectTarget));
-								}
+								IntelliSenseTargetFiles.Add(Tuple.Create(EngineProject, EngineTarget));
 							}
 						}
 					}
@@ -1175,7 +1204,7 @@ namespace UnrealBuildTool
 						{
 							// Only bother with the editor target.  We want to make sure that definitions are setup to be as inclusive as possible
 							// for good quality IntelliSense.  For example, we want WITH_EDITORONLY_DATA=1, so using the editor targets works well.
-							if (ProjectTarget.TargetRules!.Type == TargetType.Editor)
+							if (bMakeProjectPerTarget || ProjectTarget.TargetRules!.Type == TargetType.Editor)
 							{
 								IntelliSenseTargetFiles.Add(Tuple.Create(GameProject, ProjectTarget));
 							}
@@ -1985,7 +2014,7 @@ namespace UnrealBuildTool
 									foreach (UEBuildModuleCPP Module in Binary.Modules.OfType<UEBuildModuleCPP>())
 									{
 										ProjectFile? ProjectFileForIDE;
-										if (ModuleToProjectFileMap.TryGetValue(Module.RulesFile, out ProjectFileForIDE) && ProjectFileForIDE == TargetProjectFile)
+										if (ModuleToEditorProjectFileMap.TryGetValue(Module.RulesFile, out ProjectFileForIDE) && ProjectFileForIDE == TargetProjectFile)
 										{
 											Utils.WriteFileIfChangedContext = $"{Target.TargetName} {Target.Configuration} {Target.Platform} {Binary.OutputFilePaths[0].GetFileName()} {Module.Name}";
 
@@ -2188,10 +2217,10 @@ namespace UnrealBuildTool
 				if (WantProjectFileForModule)
 				{
 					DirectoryReference BaseFolder;
-					ProjectFile ProjectFile = FindProjectForModule(CurModuleFile, AllGames, ProgramProjects, ModProjects, AdditionalSearchPaths, out BaseFolder);
+					ProjectFile ProjectFile = FindProjectForModule(CurModuleFile, TargetType.Editor, AllGames, ProgramProjects, ModProjects, AdditionalSearchPaths, out BaseFolder);
 
 					// Update our module map
-					ModuleToProjectFileMap[CurModuleFile] = ProjectFile;
+					ModuleToEditorProjectFileMap[CurModuleFile] = ProjectFile;
 					ProjectFile.IsGeneratedProject = true;
 
 					// Only search subdirectories for non-external modules.  We don't want to add all of the source and header files
@@ -2281,11 +2310,10 @@ namespace UnrealBuildTool
 		{
 			// Setup a project file entry for this module's project.  Remember, some projects may host multiple modules!
 			FileReference ProjectFileName = FileReference.Combine(IntermediateProjectFilesPath, InProjectFileNameBase + ProjectFileExtension);
-			bool bProjectAlreadyExisted;
-			return FindOrAddProject(ProjectFileName, InBaseFolder, IncludeInGeneratedProjects: true, bAlreadyExisted: out bProjectAlreadyExisted);
+			return FindOrAddProject(ProjectFileName, InBaseFolder, IncludeInGeneratedProjects: true, bAlreadyExisted: out _);
 		}
 
-		private ProjectFile FindProjectForModule(FileReference CurModuleFile, List<FileReference> AllGames, Dictionary<FileReference, ProjectFile> ProgramProjects, List<ProjectFile> ModProjects, Dictionary<FileReference, List<DirectoryReference>> AdditionalSearchPaths, out DirectoryReference BaseFolder)
+		private ProjectFile FindProjectForModule(FileReference CurModuleFile, TargetType TargetType, List<FileReference> AllGames, Dictionary<FileReference, ProjectFile> ProgramProjects, List<ProjectFile> ModProjects, Dictionary<FileReference, List<DirectoryReference>> AdditionalSearchPaths, out DirectoryReference BaseFolder)
 		{
 			// Starting at the base directory of the module find a project which has the same directory as base, walking up the directory hierarchy until a match is found
 
@@ -2300,7 +2328,22 @@ namespace UnrealBuildTool
 					{
 						FileReference ProjectInfo = Game;
 						BaseFolder = ProjectInfo.Directory;
-						return FindOrAddProjectHelper(ProjectInfo.GetFileNameWithoutExtension(), BaseFolder);
+
+						if (bMakeProjectPerTarget)
+						{
+							// find the project that the module is under, and has a TargetType target (useful with bMakeProjectPerTarget)
+							foreach (KeyValuePair<FileReference, ProjectFile> Pair in ProjectFileMap)
+							{
+								if (((ProjectTarget)Pair.Value.ProjectTargets[0]).TargetFilePath.Directory.ParentDirectory == Path && Pair.Value.ProjectTargets.Any(x => x.TargetRules!.Type == TargetType))
+								{
+									return Pair.Value;
+								}
+							}
+						}
+						else
+						{
+							return FindOrAddProjectHelper(ProjectInfo.GetFileNameWithoutExtension(), BaseFolder);
+						}
 					}
 				}
 				// Check if it's a mod
@@ -2332,7 +2375,7 @@ namespace UnrealBuildTool
 					if (Path == ExtensionDir)
 					{
 						BaseFolder = Unreal.EngineDirectory;
-						return FindOrAddProjectHelper(EngineProjectFileNameBase, BaseFolder);
+						return FindOrAddProjectHelper(EngineEditorProjectFileNameBase, BaseFolder);
 					}
 				}
 
@@ -2362,7 +2405,7 @@ namespace UnrealBuildTool
 		/// <param name="AllGames">All game folders</param>
 		/// <param name="AllTargetFiles">All the target files to add</param>
 		/// <param name="Arguments">The commandline arguments used</param>
-		/// <param name="EngineProject">The engine project we created</param>
+		/// <param name="EngineProjects">The engine projects we created</param>
 		/// <param name="GameProjects">Map of game folder name to all of the game projects we created</param>
 		/// <param name="ProjectFileToUProjectFile">Map of generated Project File to the .uproject file for a given game.</param>
 		/// <param name="ProgramProjects">Map of program names to all of the program projects we created</param>
@@ -2373,7 +2416,7 @@ namespace UnrealBuildTool
 			List<FileReference> AllGames,
 			List<FileReference> AllTargetFiles,
 			String[] Arguments,
-			ref ProjectFile? EngineProject,
+			List<ProjectFile> EngineProjects,
 			List<ProjectFile> GameProjects,
 			Dictionary<ProjectFile, FileReference> ProjectFileToUProjectFile,
 			Dictionary<FileReference, ProjectFile> ProgramProjects,
@@ -2474,51 +2517,59 @@ namespace UnrealBuildTool
 						ProjectFileNameBase = ProjectInfo.GetFileNameWithoutExtension();
 					}
 
-					// Look at the project engine config to see if it has specified a default editor target
-					FileReference ProjectLocation = GetProjectLocation(ProjectFileNameBase);
-					string? DefaultEditorTarget = null;
-					if (!DefaultProjectEditorTargetCache.TryGetValue(ProjectFileNameBase, out DefaultEditorTarget))
+					FileReference ProjectFilePath;
+					if (bMakeProjectPerTarget)
 					{
-						if (GameFolder != null)
-						{
-							DefaultEditorTarget = GetProjectDefaultTargetNameForType(GameFolder, TargetType.Editor);
-						}
-
-						DefaultProjectEditorTargetCache.Add(ProjectFileNameBase, DefaultEditorTarget);
+						ProjectFilePath = GetProjectLocation(TargetRulesObject.Name);
 					}
-
-					// Get the suffix to use for this project file. If we have multiple targets of the same type, we'll have to split them out into separate IDE project files.
-					string? GeneratedProjectName = TargetRulesObject.GeneratedProjectName;
-					if (GeneratedProjectName == null)
+					else
 					{
-						ProjectFile? ExistingProjectFile;
-						// We should create a separate project for targets which aren't the default for this target type.
-						// Note that we currently only support changing the default editor target and not any other types, so
-						// if we aren't an editor, we'll just fall back to previous behavior and assume the first one we encounter
-						// should be added to the main project
-						bool bIsDefaultTargetForType = (DefaultEditorTarget == null) || (TargetRulesObject.Type != TargetType.Editor) || (TargetRulesObject.Name == DefaultEditorTarget);
-						if (!bIsDefaultTargetForType || (ProjectFileMap.TryGetValue(ProjectLocation, out ExistingProjectFile) && ExistingProjectFile.ProjectTargets.Any(x => x.TargetRules!.Type == TargetRulesObject.Type)))
+						// Look at the project engine config to see if it has specified a default editor target
+						FileReference ProjectLocation = GetProjectLocation(ProjectFileNameBase);
+						string? DefaultEditorTarget = null;
+						if (!DefaultProjectEditorTargetCache.TryGetValue(ProjectFileNameBase, out DefaultEditorTarget))
 						{
-							GeneratedProjectName = TargetRulesObject.Name;
+							if (GameFolder != null)
+							{
+								DefaultEditorTarget = GetProjectDefaultTargetNameForType(GameFolder, TargetType.Editor);
+							}
+
+							DefaultProjectEditorTargetCache.Add(ProjectFileNameBase, DefaultEditorTarget);
 						}
-						else
+
+						// Get the suffix to use for this project file. If we have multiple targets of the same type, we'll have to split them out into separate IDE project files.
+						string? GeneratedProjectName = TargetRulesObject.GeneratedProjectName;
+						if (GeneratedProjectName == null)
 						{
-							GeneratedProjectName = ProjectFileNameBase;
+							ProjectFile? ExistingProjectFile;
+							// We should create a separate project for targets which aren't the default for this target type.
+							// Note that we currently only support changing the default editor target and not any other types, so
+							// if we aren't an editor, we'll just fall back to previous behavior and assume the first one we encounter
+							// should be added to the main project
+							bool bIsDefaultTargetForType = (DefaultEditorTarget == null) || (TargetRulesObject.Type != TargetType.Editor) || (TargetRulesObject.Name == DefaultEditorTarget);
+							if (!bIsDefaultTargetForType || (ProjectFileMap.TryGetValue(ProjectLocation, out ExistingProjectFile) && ExistingProjectFile.ProjectTargets.Any(x => x.TargetRules!.Type == TargetRulesObject.Type)))
+							{
+								GeneratedProjectName = TargetRulesObject.Name;
+							}
+							else
+							{
+								GeneratedProjectName = ProjectFileNameBase;
+							}
 						}
-					}
 
-					// @todo projectfiles: We should move all of the Target.cs files out of sub-folders to clean up the project directories a bit (e.g. GameUncooked folder)
-					FileReference ProjectFilePath = GetProjectLocation(GeneratedProjectName);
-					if (TargetRulesObject.Type == TargetType.Game || TargetRulesObject.Type == TargetType.Client || TargetRulesObject.Type == TargetType.Server)
-					{
-						// Allow platforms to generate stub projects here...
-						PlatformProjectGenerators.GenerateGameProjectStubs(
-							InGenerator: this,
-							InTargetName: TargetName,
-							InTargetFilepath: TargetFilePath.FullName,
-							InTargetRules: TargetRulesObject,
-							InPlatforms: SupportedPlatforms,
-							InConfigurations: SupportedConfigurations);
+						// @todo projectfiles: We should move all of the Target.cs files out of sub-folders to clean up the project directories a bit (e.g. GameUncooked folder)
+						ProjectFilePath = GetProjectLocation(GeneratedProjectName);
+						if (TargetRulesObject.Type == TargetType.Game || TargetRulesObject.Type == TargetType.Client || TargetRulesObject.Type == TargetType.Server)
+						{
+							// Allow platforms to generate stub projects here...
+							PlatformProjectGenerators.GenerateGameProjectStubs(
+								InGenerator: this,
+								InTargetName: TargetName,
+								InTargetFilepath: TargetFilePath.FullName,
+								InTargetRules: TargetRulesObject,
+								InPlatforms: SupportedPlatforms,
+								InConfigurations: SupportedConfigurations);
+						}
 					}
 
 					DirectoryReference BaseFolder;
@@ -2554,13 +2605,19 @@ namespace UnrealBuildTool
 					}
 					else if (IsEngineTarget)
 					{
-						EngineProject = ProjectFile;
+						// if we want only one project for all types, then remove any previous ones
+						if (bMakeProjectPerTarget == false)
+						{
+							EngineProjects.Clear();
+						}
+						EngineProjects.Add(ProjectFile);
+
 						if (Unreal.IsEngineInstalled())
 						{
 							// Allow engine projects to be created but not built for Installed Engine builds
-							EngineProject.IsForeignProject = false;
-							EngineProject.IsGeneratedProject = true;
-							EngineProject.IsStubProject = true;
+							ProjectFile.IsForeignProject = false;
+							ProjectFile.IsGeneratedProject = true;
+							ProjectFile.IsStubProject = true;
 						}
 					}
 					else
diff --git a/Engine/Source/Programs/UnrealBuildTool/ProjectFiles/Rider/RiderProjectFileGenerator.cs b/Engine/Source/Programs/UnrealBuildTool/ProjectFiles/Rider/RiderProjectFileGenerator.cs
--- a/Engine/Source/Programs/UnrealBuildTool/ProjectFiles/Rider/RiderProjectFileGenerator.cs
+++ b/Engine/Source/Programs/UnrealBuildTool/ProjectFiles/Rider/RiderProjectFileGenerator.cs
@@ -513,7 +513,7 @@ namespace UnrealBuildTool
 				HashSet<UnrealTargetPlatform> ValidPlatforms;
 				List<VCProjectFileGenerator.VCSolutionConfigCombination> ValidSolutionConfigs;
 				VCProjectFileGenerator.CollectSolutionConfigurations(SupportedConfigurations,
-					SupportedPlatforms, AllProjectFiles, Logger, out ValidPlatforms, out ValidSolutionConfigs);
+					SupportedPlatforms, AllProjectFiles, bMakeProjectPerTarget, Logger, out ValidPlatforms, out ValidSolutionConfigs);
 				FileReference OutputFile = GetRiderProjectLocation("root.json");
 				DirectoryReference.CreateDirectory(OutputFile.Directory);
 				using (JsonWriter Writer = new JsonWriter(OutputFile))
diff --git a/Engine/Source/Programs/UnrealBuildTool/ProjectFiles/VisualStudio/VCProject.cs b/Engine/Source/Programs/UnrealBuildTool/ProjectFiles/VisualStudio/VCProject.cs
--- a/Engine/Source/Programs/UnrealBuildTool/ProjectFiles/VisualStudio/VCProject.cs
+++ b/Engine/Source/Programs/UnrealBuildTool/ProjectFiles/VisualStudio/VCProject.cs
@@ -281,6 +281,7 @@ namespace UnrealBuildTool
 	{
 		VCProjectFileFormat ProjectFileFormat;
 		bool bUsePrecompiled;
+		bool bMakeProjectPerTarget;
 		string? BuildToolOverride;
 		Dictionary<DirectoryReference, string> ModuleDirToForceIncludePaths = new Dictionary<DirectoryReference, string>();
 		Dictionary<DirectoryReference, string> ModuleDirToPchHeaderFile = new Dictionary<DirectoryReference, string>();
@@ -304,13 +305,15 @@ namespace UnrealBuildTool
 		/// <param name="BaseDir">The base directory for files within this project</param>
 		/// <param name="ProjectFileFormat">Visual C++ project file version</param>
 		/// <param name="bUsePrecompiled">Whether to add the -UsePrecompiled argumemnt when building targets</param>
+		/// <param name="bMakeProjectPerTarget">Whether to add roll the target type into the config (ie "Development Editor")</param>
 		/// <param name="BuildToolOverride">Optional arguments to pass to UBT when building</param>
 		/// <param name="Settings">Other settings</param>
-		public VCProjectFile(FileReference FilePath, DirectoryReference BaseDir, VCProjectFileFormat ProjectFileFormat, bool bUsePrecompiled, string? BuildToolOverride, VCProjectFileSettings Settings)
+		public VCProjectFile(FileReference FilePath, DirectoryReference BaseDir, VCProjectFileFormat ProjectFileFormat, bool bUsePrecompiled, bool bMakeProjectPerTarget, string? BuildToolOverride, VCProjectFileSettings Settings)
 			: base(FilePath, BaseDir)
 		{
 			this.ProjectFileFormat = ProjectFileFormat;
 			this.bUsePrecompiled = bUsePrecompiled;
+			this.bMakeProjectPerTarget = bMakeProjectPerTarget;
 			this.BuildToolOverride = BuildToolOverride;
 			this.Settings = Settings;
 		}
@@ -360,7 +363,7 @@ namespace UnrealBuildTool
 				ProjectPlatformName = DefaultPlatformName;
 			}
 
-			if (TargetConfigurationName != TargetType.Game)
+			if (!bMakeProjectPerTarget && TargetConfigurationName != TargetType.Game)
 			{
 				ProjectConfigurationName += "_" + TargetConfigurationName.ToString();
 			}
@@ -415,7 +418,7 @@ namespace UnrealBuildTool
 					{
 						if (ProjectTarget.TargetRules != null)
 						{
-							if (TargetConfigurationName == ProjectTarget.TargetRules.Type)
+							if (bMakeProjectPerTarget || TargetConfigurationName == ProjectTarget.TargetRules.Type)
 							{
 								MatchingProjectTargets.Add(ProjectTarget);
 							}
diff --git a/Engine/Source/Programs/UnrealBuildTool/ProjectFiles/VisualStudio/VCProjectFileGenerator.cs b/Engine/Source/Programs/UnrealBuildTool/ProjectFiles/VisualStudio/VCProjectFileGenerator.cs
--- a/Engine/Source/Programs/UnrealBuildTool/ProjectFiles/VisualStudio/VCProjectFileGenerator.cs
+++ b/Engine/Source/Programs/UnrealBuildTool/ProjectFiles/VisualStudio/VCProjectFileGenerator.cs
@@ -98,6 +98,11 @@ namespace UnrealBuildTool
 	/// </summary>
 	class VCProjectFileGenerator : ProjectFileGenerator
 	{
+		/// <summary>
+		/// Set to true to enable a project for each target, and do not put the target type into the configuration
+		/// </summary>
+		protected override bool bMakeProjectPerTarget => true;
+
 		/// <summary>
 		/// The settings object
 		/// </summary>
@@ -211,7 +216,7 @@ namespace UnrealBuildTool
 		/// <returns>The newly allocated project file object</returns>
 		protected override ProjectFile AllocateProjectFile(FileReference InitFilePath, DirectoryReference BaseDir)
 		{
-			return new VCProjectFile(InitFilePath, BaseDir, Settings.ProjectFileFormat, bUsePrecompiled, BuildToolOverride, Settings);
+			return new VCProjectFile(InitFilePath, BaseDir, Settings.ProjectFileFormat, bUsePrecompiled, bMakeProjectPerTarget, BuildToolOverride, Settings);
 		}
 
 
@@ -404,16 +409,20 @@ namespace UnrealBuildTool
 		/// </summary>
 		/// <param name="Configuration">The build configuration</param>
 		/// <param name="TargetType">The type of target being built</param>
+		/// <param name="bMakeProjectPerTarget">True if we are making one project per target type, instead of rolling them into the configs</param>
 		/// <returns>The generated solution configuration name</returns>
-		static string MakeSolutionConfigurationName(UnrealTargetConfiguration Configuration, TargetType TargetType)
+		static string MakeSolutionConfigurationName(UnrealTargetConfiguration Configuration, TargetType TargetType, bool bMakeProjectPerTarget)
 		{
 			string SolutionConfigName = Configuration.ToString();
 
-			// Don't bother postfixing "Game" or "Program" -- that will be the default when using "Debug", "Development", etc.
-			// Also don't postfix "RocketGame" when we're building Rocket game projects.  That's the only type of game there is in that case!
-			if (TargetType != TargetType.Game && TargetType != TargetType.Program)
+			if (!bMakeProjectPerTarget)
 			{
-				SolutionConfigName += " " + TargetType.ToString();
+				// Don't bother postfixing "Game" or "Program" -- that will be the default when using "Debug", "Development", etc.
+				// Also don't postfix "RocketGame" when we're building Rocket game projects.  That's the only type of game there is in that case!
+				if (TargetType != TargetType.Game && TargetType != TargetType.Program)
+				{
+					SolutionConfigName += " " + TargetType.ToString();
+				}
 			}
 
 			return SolutionConfigName;
@@ -667,7 +676,7 @@ namespace UnrealBuildTool
 					{
 						VCSolutionFileContent.AppendLine("	GlobalSection(SolutionConfigurationPlatforms) = preSolution");
 
-						CollectSolutionConfigurations(SupportedConfigurations, SupportedPlatforms, AllProjectFiles,
+						CollectSolutionConfigurations(SupportedConfigurations, SupportedPlatforms, AllProjectFiles, bMakeProjectPerTarget,
 							Logger, out PlatformsValidForProjects, out SolutionConfigCombinations);
 
 						HashSet<string> AppendedSolutionConfigAndPlatformNames = new HashSet<string>(StringComparer.InvariantCultureIgnoreCase);
@@ -822,7 +831,10 @@ namespace UnrealBuildTool
 					VCSolutionOptions Options = new VCSolutionOptions(Settings.ProjectFileFormat);
 
 					// Set the default configuration and startup project
-					VCSolutionConfigCombination? DefaultConfig = SolutionConfigCombinations.Find(x => x.Configuration == UnrealTargetConfiguration.Development && x.Platform == UnrealTargetPlatform.Win64 && x.TargetConfigurationName == TargetType.Editor);
+					VCSolutionConfigCombination? DefaultConfig = SolutionConfigCombinations.Find(x => 
+						x.Configuration == UnrealTargetConfiguration.Development &&
+						x.Platform == UnrealTargetPlatform.Win64 && 
+						(bMakeProjectPerTarget || x.TargetConfigurationName == TargetType.Editor));
 					if (DefaultConfig != null)
 					{
 						List<VCBinarySetting> Settings = new List<VCBinarySetting>();
@@ -871,7 +883,7 @@ namespace UnrealBuildTool
 		}
 
 		public static void CollectSolutionConfigurations(List<UnrealTargetConfiguration> AllConfigurations,
-			List<UnrealTargetPlatform> AllPlatforms, List<ProjectFile> AllProjectFiles, ILogger Logger,
+			List<UnrealTargetPlatform> AllPlatforms, List<ProjectFile> AllProjectFiles, bool bMakeProjectPerTarget, ILogger Logger,
 			out HashSet<UnrealTargetPlatform> OutValidPlatforms, out List<VCSolutionConfigCombination> OutSolutionConfigs)
 		{
 			OutValidPlatforms = new HashSet<UnrealTargetPlatform>();
@@ -913,7 +925,7 @@ namespace UnrealBuildTool
 											}
 
 											string SolutionConfigName =
-												MakeSolutionConfigurationName(CurConfiguration, TargetType);
+												MakeSolutionConfigurationName(CurConfiguration, TargetType, bMakeProjectPerTarget);
 											SolutionConfigurationsValidForProjects[SolutionConfigName] =
 												new Tuple<UnrealTargetConfiguration, TargetType>(CurConfiguration, TargetType);
 										}
diff --git a/Engine/Source/Programs/UnrealBuildTool/ProjectFiles/Xcode/XcodeProject.cs b/Engine/Source/Programs/UnrealBuildTool/ProjectFiles/Xcode/XcodeProject.cs
--- a/Engine/Source/Programs/UnrealBuildTool/ProjectFiles/Xcode/XcodeProject.cs
+++ b/Engine/Source/Programs/UnrealBuildTool/ProjectFiles/Xcode/XcodeProject.cs
@@ -1,4 +1,4 @@
-﻿// Copyright Epic Games, Inc. All Rights Reserved.
+// Copyright Epic Games, Inc. All Rights Reserved.
 
 using System;
 using System.Collections.Generic;
@@ -80,6 +80,11 @@ namespace UnrealBuildTool.XcodeProjectXcconfig
 			BuildTarget = InBuildTarget;
 			ProjectTarget = InProjectTarget;
 			BuildConfig = InBuildConfig;
+
+			if (BuildTarget != ProjectTarget?.Name)
+			{
+				throw new BuildException($"Name exepcted to match - {BuildTarget} != {ProjectTarget?.Name}");
+			}
 		}
 
 		public string DisplayName;
@@ -261,6 +266,9 @@ namespace UnrealBuildTool.XcodeProjectXcconfig
 
 	class UnrealData
 	{
+		public bool bIsStubProject;
+		public bool bMakeProjectPerTarget;
+
 		public bool bIsAppBundle;
 		public bool bHasEditorConfiguration;
 		public bool bUseAutomaticSigning = false;
@@ -330,7 +338,7 @@ namespace UnrealBuildTool.XcodeProjectXcconfig
 		}
 
 
-		public UnrealData(FileReference XcodeProjectFileLocation, bool bIsForDistribution, string BundleID, string AppName)
+		public UnrealData(FileReference XcodeProjectFileLocation, bool bIsForDistribution, string BundleID, string AppName, bool bMakeProjectPerTarget, bool bIsStubProject)
 		{
 			// the .xcodeproj is actually a directory
 			this.XcodeProjectFileLocation = new DirectoryReference(XcodeProjectFileLocation.FullName);
@@ -342,6 +350,8 @@ namespace UnrealBuildTool.XcodeProjectXcconfig
 				ProductName = "UnrealGame";
 			}
 
+			this.bIsStubProject = bIsStubProject;
+			this.bMakeProjectPerTarget = bMakeProjectPerTarget;
 			this.bForDistribution = bIsForDistribution;
 			this.BundleIdentifier = string.IsNullOrEmpty(BundleID) ? "$(UE_SIGNING_PREFIX).$(UE_PRODUCT_NAME_STRIPPED)" : BundleID;
 			this.DisplayName = string.IsNullOrEmpty(AppName) ? "$(UE_PRODUCT_NAME)" : AppName;
@@ -622,9 +632,12 @@ namespace UnrealBuildTool.XcodeProjectXcconfig
 										bShouldCompileMonolithic |= (ProjectTarget.CreateRulesDelegate(Platform, Configuration).LinkType == TargetLinkType.Monolithic);
 
 										string ConfigName = Configuration.ToString();
-										if (ProjectTarget.TargetRules!.Type != TargetType.Game && ProjectTarget.TargetRules.Type != TargetType.Program)
+										if (!bMakeProjectPerTarget)
 										{
-											ConfigName += " " + ProjectTarget.TargetRules.Type.ToString();
+											if (ProjectTarget.TargetRules!.Type != TargetType.Game && ProjectTarget.TargetRules.Type != TargetType.Program)
+											{
+												ConfigName += " " + ProjectTarget.TargetRules.Type.ToString();
+											}
 										}
 
 										if (BuildConfigs.Where(Config => Config.DisplayName == ConfigName).ToList().Count == 0)
@@ -636,14 +649,14 @@ namespace UnrealBuildTool.XcodeProjectXcconfig
 											// Get the output directory
 											DirectoryReference RootDirectory;
 											if (UProjectDirectory != null && 
-												(bShouldCompileMonolithic || ProjectTarget.TargetRules.BuildEnvironment == TargetBuildEnvironment.Unique) && 
-												ProjectTarget.TargetRules.File!.IsUnderDirectory(UProjectDirectory))
+												(bShouldCompileMonolithic || ProjectTarget.TargetRules!.BuildEnvironment == TargetBuildEnvironment.Unique) && 
+												ProjectTarget.TargetRules!.File!.IsUnderDirectory(UProjectDirectory))
 											{
 												RootDirectory = UEBuildTarget.GetOutputDirectoryForExecutable(UProjectDirectory, ProjectTarget.TargetRules.File!);
 											}
 											else
 											{
-												RootDirectory = UEBuildTarget.GetOutputDirectoryForExecutable(Unreal.EngineDirectory, ProjectTarget.TargetRules.File!);
+												RootDirectory = UEBuildTarget.GetOutputDirectoryForExecutable(Unreal.EngineDirectory, ProjectTarget.TargetRules!.File!);
 											}
 
 											string ExeName = TargetName;
@@ -1136,6 +1149,7 @@ namespace UnrealBuildTool.XcodeProjectXcconfig
 	{ 
 		public string Guid = XcodeProjectFileGenerator.MakeXcodeGuid();
 		public string TargetName;
+		private UnrealData UnrealData;
 
 		public List<XcodeBuildConfig> BuildConfigs = new();
 
@@ -1148,11 +1162,12 @@ namespace UnrealBuildTool.XcodeProjectXcconfig
 		}
 
 		private UnrealTargetPlatform? Platform;
-		public XcodeBuildConfigList(UnrealTargetPlatform? Platform, string TargetName, List<UnrealBuildConfig> BuildConfigInfos)
+		public XcodeBuildConfigList(UnrealTargetPlatform? Platform, string TargetName, UnrealData UnrealData)
 		{
 			this.Platform = Platform;
+			this.UnrealData = UnrealData;
 
-			if (BuildConfigInfos.Count == 0)
+			if (UnrealData.AllConfigs.Count == 0)
 			{
 				throw new BuildException("Created a XcodeBuildConfigList with no BuildConfigs. This likely means a target was created too early");
 			}
@@ -1160,7 +1175,7 @@ namespace UnrealBuildTool.XcodeProjectXcconfig
 			this.TargetName = TargetName;
 
 			// create build config objects for each info passed in, and them as references
-			IEnumerable<XcodeBuildConfig> Configs = BuildConfigInfos.Select(x => new XcodeBuildConfig(x));
+			IEnumerable<XcodeBuildConfig> Configs = UnrealData.AllConfigs.Select(x => new XcodeBuildConfig(x));
 			// filter out configs that dont match a platform if we are single-platform mode
 			Configs = Configs.Where(x => Platform == null || x.Info.Supports((UnrealTargetPlatform)Platform));
 			BuildConfigs = Configs.ToList();
@@ -1170,7 +1185,11 @@ namespace UnrealBuildTool.XcodeProjectXcconfig
 		public override void Write(StringBuilder Content)
 		{
 			// figure out the default configuration to use
-			string Default = BuildConfigs.Any(x => x.Info.DisplayName.Contains(" Editor")) ? "Development Editor" : "Development";
+			string Default = "Development";
+			if (!UnrealData.bMakeProjectPerTarget && BuildConfigs.Any(x => x.Info.DisplayName.Contains(" Editor")))
+			{
+				Default = "Development Editor";
+			}
 
 			Content.WriteLine(2, $"{Guid} /* Build configuration list for target {TargetName} */ = {{");
 			Content.WriteLine(3,	"isa = XCConfigurationList;");
@@ -1254,7 +1273,7 @@ namespace UnrealBuildTool.XcodeProjectXcconfig
 			}
 
 			// set up names
-			Name = UnrealData.XcodeProjectName + (OverrideName ?? ConfigName);
+			Name = OverrideName ?? (UnrealData.XcodeProjectName + ConfigName);
 		}
 
 		public void AddDependency(XcodeTarget Target, XcodeProject Project)
@@ -1320,13 +1339,15 @@ namespace UnrealBuildTool.XcodeProjectXcconfig
 		private string ProductGuid = XcodeProjectFileGenerator.MakeXcodeGuid();
 		private UnrealData UnrealData;
 		public UnrealTargetPlatform Platform;
+		public TargetType TargetType;
 
-		public XcodeRunTarget(XcodeProject Project, UnrealTargetPlatform Platform, XcodeBuildTarget? BuildTarget)
-			: base(Project.UnrealData.bIsAppBundle ? XcodeTarget.Type.Run_App : XcodeTarget.Type.Run_Tool, Project.UnrealData, 
-				  XcodeProjectFileGenerator.PerPlatformMode == XcodePerPlatformMode.OneWorkspacePerPlatform ? "" : $"_{Platform}")
+		public XcodeRunTarget(XcodeProject Project, string TargetName, TargetType TargetType, UnrealTargetPlatform Platform, XcodeBuildTarget? BuildTarget)
+			: base(Project.UnrealData.bIsAppBundle ? XcodeTarget.Type.Run_App : XcodeTarget.Type.Run_Tool, Project.UnrealData,
+				  TargetName + (XcodeProjectFileGenerator.PerPlatformMode == XcodePerPlatformMode.OneWorkspacePerPlatform ? "" : $"_{Platform}"))
 		{
+			this.TargetType = TargetType;
 			this.Platform = Platform;
-			BuildConfigList = new XcodeBuildConfigList(Platform, Name, Project.UnrealData.AllConfigs);
+			BuildConfigList = new XcodeBuildConfigList(Platform, Name, Project.UnrealData);
 			References.Add(BuildConfigList);
 
 			// add the Product item to the project to be visible in left pane
@@ -1429,7 +1450,7 @@ namespace UnrealBuildTool.XcodeProjectXcconfig
 		public XcodeBuildTarget(XcodeProject Project)
 			: base(XcodeTarget.Type.Build, Project.UnrealData)
 		{
-			BuildConfigList = new XcodeBuildConfigList(Project.Platform, Name, Project.UnrealData.AllConfigs);
+			BuildConfigList = new XcodeBuildConfigList(Project.Platform, Name, Project.UnrealData);
 			References.Add(BuildConfigList);
 		}
 	}
@@ -1444,7 +1465,7 @@ namespace UnrealBuildTool.XcodeProjectXcconfig
 		{
 			UnrealData = Project.UnrealData;
 
-			BuildConfigList = new XcodeBuildConfigList(Project.Platform, Name, UnrealData.AllConfigs);
+			BuildConfigList = new XcodeBuildConfigList(Project.Platform, Name, UnrealData);
 			References.Add(BuildConfigList);
 
 			CreateXcconfigFile(Project, Name);
@@ -1510,25 +1531,21 @@ namespace UnrealBuildTool.XcodeProjectXcconfig
 
 			// if we are run-only, then we don't need a build target (this is shared between platforms if we are doing multi-target)
 			XcodeBuildTarget? BuildTarget = null;
-			if (!XcodeProjectFileGenerator.bGenerateRunOnlyProject)
+			if (!XcodeProjectFileGenerator.bGenerateRunOnlyProject && !UnrealData.bIsStubProject)
 			{
-				// create a biuld target only if we have source files to build
-				if (UnrealData.BatchedFiles.Count != 0)
-				{
-					BuildTarget = new XcodeBuildTarget(this);
-				}
+				BuildTarget = new XcodeBuildTarget(this);
 			}
 
 			// create one run target for each platform if our platform is null (ie XcodeProjectGenerator.PerPlatformMode is RunTargetPerPlatform)
 			List<UnrealTargetPlatform> TargetPlatforms = Platform == null ? XcodeProjectFileGenerator.XcodePlatforms : new() { Platform.Value };
 			foreach (UnrealTargetPlatform TargetPlatform in TargetPlatforms)
 			{
-				XcodeRunTarget RunTarget = new XcodeRunTarget(this, TargetPlatform, BuildTarget);
+				XcodeRunTarget RunTarget = new XcodeRunTarget(this, UnrealData.ProductName, UnrealData.AllConfigs[0].ProjectTarget!.TargetRules!.Type, TargetPlatform, BuildTarget);
 				RunTargets.Add(RunTarget);
 				References.Add(RunTarget);
 			}
 
-			ProjectBuildConfigs = new XcodeBuildConfigList(Platform, UnrealData.ProductName, UnrealData.AllConfigs);
+			ProjectBuildConfigs = new XcodeBuildConfigList(Platform, UnrealData.ProductName, UnrealData);
 			References.Add(ProjectBuildConfigs);
 
 			// create the Projet xcconfig
@@ -2498,10 +2515,11 @@ namespace UnrealBuildTool.XcodeProjectXcconfig
 		/// <param name="bIsForDistribution">True for distribution builds</param>
 		/// <param name="BundleID">Override option for bundle identifier</param>
 		/// <param name="AppName"></param>
-		public XcodeProjectFile(FileReference InitFilePath, DirectoryReference BaseDir, bool bIsForDistribution, string BundleID, string AppName)
+		/// <param name="bMakeProjectPerTarget"></param>
+		public XcodeProjectFile(FileReference InitFilePath, DirectoryReference BaseDir, bool bIsForDistribution, string BundleID, string AppName, bool bMakeProjectPerTarget)
 			: base(InitFilePath, BaseDir)
 		{
-			UnrealData = new UnrealData(InitFilePath, bIsForDistribution, BundleID, AppName);
+			UnrealData = new UnrealData(InitFilePath, bIsForDistribution, BundleID, AppName, bMakeProjectPerTarget, IsStubProject);
 
 			// create the container for all the files that will 
 			SharedFileCollection = new XcodeFileCollection(this);
@@ -2549,14 +2567,14 @@ namespace UnrealBuildTool.XcodeProjectXcconfig
 				bHasCheckedForLegacy = true;
 				if (ProjectTargets.Count == 0)
 				{
-					throw new BuildException("Expected to have targets before AddModule is called");
+					throw new BuildException("Expected to have a target before AddModule is called");
 				}
 				FileReference? UProjectFileLocation = UnrealData.FindUProjectFileLocation(this);
 				ConfigHierarchy Ini = ConfigCache.ReadHierarchy(ConfigHierarchyType.Engine, UProjectFileLocation?.Directory, UnrealTargetPlatform.Mac);
 				bool bUseModernXcode;
 				if (!(Ini.TryGetValue("XcodeConfiguration", "bUseModernXcode", out bUseModernXcode) && bUseModernXcode))
 				{
-					LegacyProjectFile = new XcodeProjectLegacy.XcodeProjectFile(ProjectFilePath, BaseDir, UnrealData.bForDistribution, UnrealData.BundleIdentifier, UnrealData.AppName);
+					LegacyProjectFile = new XcodeProjectLegacy.XcodeProjectFile(ProjectFilePath, BaseDir, UnrealData.bForDistribution, UnrealData.BundleIdentifier, UnrealData.AppName, UnrealData.bMakeProjectPerTarget);
 					LegacyProjectFile.ProjectTargets.AddRange(ProjectTargets);
 					LegacyProjectFile.SourceFiles.AddRange(SourceFiles);
 					LegacyProjectFile.IsGeneratedProject = IsGeneratedProject;
@@ -2735,7 +2753,7 @@ namespace UnrealBuildTool.XcodeProjectXcconfig
 						string ProjectName = ProjectFilePathForPlatform(Platform).GetFileNameWithoutAnyExtensions();
 						string? BuildTargetGuid = XcodeProjectNode.GetNodesOfType<XcodeBuildTarget>(RootProject).FirstOrDefault()?.Guid;
 						string? IndexTargetGuid = XcodeProjectNode.GetNodesOfType<XcodeIndexTarget>(RootProject).FirstOrDefault()?.Guid;
-						WriteSchemeFile(Platform, ProjectName, UnrealData.ProductName, RootProject.RunTargets, BuildTargetGuid, IndexTargetGuid, UnrealData.bHasEditorConfiguration, UnrealData.UProjectFileLocation != null ? UnrealData.UProjectFileLocation.FullName : "");
+						WriteSchemeFile(Platform, ProjectName, UnrealData.ProductName, RootProject.RunTargets, BuildTargetGuid, IndexTargetGuid, UnrealData.UProjectFileLocation != null ? UnrealData.UProjectFileLocation.FullName : "");
 					}
 				}
 				else
@@ -3320,7 +3338,11 @@ namespace UnrealBuildTool.XcodeProjectXcconfig
 					string ConfigName = Plist($"Print :objects:{ConfigGuid}:name");
 					if (ConfigName == "Release")
 					{
-						ConfigName = UnrealData.bHasEditorConfiguration ? "Development Editor" : "Development";
+						ConfigName = "Development";
+						if (UnrealData.bMakeProjectPerTarget && UnrealData.bHasEditorConfiguration)
+						{
+							ConfigName = "Development Editor";
+						}
 						Plist($"Set :objects:{ConfigGuid}:name \"{ConfigName}\"");
 					}
 
@@ -3388,7 +3410,7 @@ namespace UnrealBuildTool.XcodeProjectXcconfig
 			return FileReference.Combine(GetProjectSchemeDirectory(Platform), TargetName + ".xcscheme");
 		}
 
-		private void WriteSchemeFile(UnrealTargetPlatform? Platform, string ProjectName, string ProductName, List<XcodeRunTarget> RunTargets, string? BuildTargetGuid, string? IndexTargetGuid, bool bHasEditorConfiguration, string GameProjectPath)
+		private void WriteSchemeFile(UnrealTargetPlatform? Platform, string ProjectName, string ProductName, List<XcodeRunTarget> RunTargets, string? BuildTargetGuid, string? IndexTargetGuid, string GameProjectPath)
 		{
 			StringBuilder Content = new StringBuilder();
 
@@ -3396,8 +3418,11 @@ namespace UnrealBuildTool.XcodeProjectXcconfig
 			{
 				string TargetName = RunTarget.Name;
 				string TargetGuid = RunTarget.Guid;
-				bHasEditorConfiguration = RunTarget.BuildConfigList!.BuildConfigs.Any(x => x.Info.ProjectTarget!.TargetRules!.Type == TargetType.Editor);
-
+				
+				bool bHasEditorConfiguration = RunTarget.BuildConfigList!.BuildConfigs.Any(x => x.Info.ProjectTarget!.TargetRules!.Type == TargetType.Editor);
+				bool bUseEditorConfiguration = bHasEditorConfiguration && !UnrealData.bMakeProjectPerTarget && !XcodeProjectFileGenerator.bGenerateRunOnlyProject;
+				string DefaultConfiguration = bUseEditorConfiguration ? "Development Editor" : "Development";
+				
 				FileReference SchemeFilePath = GetProjectSchemeFilePathForTarget(Platform, TargetName);
 
 				DirectoryReference.CreateDirectory(SchemeFilePath.Directory);
@@ -3414,7 +3439,6 @@ namespace UnrealBuildTool.XcodeProjectXcconfig
 					}
 				}
 
-				string DefaultConfiguration = (bHasEditorConfiguration && !XcodeProjectFileGenerator.bGenerateRunOnlyProject) ? "Development Editor" : "Development";
 
 				Content.WriteLine("<?xml version=\"1.0\" encoding=\"UTF-8\"?>");
 				Content.WriteLine("<Scheme");
diff --git a/Engine/Source/Programs/UnrealBuildTool/ProjectFiles/Xcode/XcodeProjectFileGenerator.cs b/Engine/Source/Programs/UnrealBuildTool/ProjectFiles/Xcode/XcodeProjectFileGenerator.cs
--- a/Engine/Source/Programs/UnrealBuildTool/ProjectFiles/Xcode/XcodeProjectFileGenerator.cs
+++ b/Engine/Source/Programs/UnrealBuildTool/ProjectFiles/Xcode/XcodeProjectFileGenerator.cs
@@ -1,4 +1,4 @@
-﻿// Copyright Epic Games, Inc. All Rights Reserved.
+// Copyright Epic Games, Inc. All Rights Reserved.
 
 using System;
 using System.Collections.Generic;
@@ -21,6 +21,11 @@ namespace UnrealBuildTool
 	/// </summary>
 	class XcodeProjectFileGenerator : ProjectFileGenerator
 	{
+		/// <summary>
+		///  Xcode makes a project per target, for modern to be able to pull in Cooked data
+		/// </summary>
+		protected override bool bMakeProjectPerTarget => true;
+
 		public DirectoryReference? XCWorkspace;
 
 		// always seed the random number the same, so multiple runs of the generator will generate the same project
@@ -124,7 +129,7 @@ namespace UnrealBuildTool
 			// unfortunately, we can't read the project configs now because we don't have enough information to 
 			// find the .uproject file for that would make this project (we could change the high level to pass it 
 			// down but it would touch all project generators - not worth it if we end up removing the legacy)
-			return new XcodeProjectXcconfig.XcodeProjectFile(InitFilePath, BaseDir, bForDistribution, BundleIdentifier, AppName);
+			return new XcodeProjectXcconfig.XcodeProjectFile(InitFilePath, BaseDir, bForDistribution, BundleIdentifier, AppName, bMakeProjectPerTarget);
 		}
 
 		private bool WriteWorkspaceSettingsFile(string Path, ILogger Logger)
diff --git a/Engine/Source/Programs/UnrealBuildTool/ProjectFiles/Xcode/XcodeProjectLegacy.cs b/Engine/Source/Programs/UnrealBuildTool/ProjectFiles/Xcode/XcodeProjectLegacy.cs
--- a/Engine/Source/Programs/UnrealBuildTool/ProjectFiles/Xcode/XcodeProjectLegacy.cs
+++ b/Engine/Source/Programs/UnrealBuildTool/ProjectFiles/Xcode/XcodeProjectLegacy.cs
@@ -97,12 +97,14 @@ namespace UnrealBuildTool.XcodeProjectLegacy
 		/// <param name="IsForDistribution">True for distribution builds</param>
 		/// <param name="BundleID">Override option for bundle identifier</param>
 		/// <param name="InAppName"></param>
-		public XcodeProjectFile(FileReference InitFilePath, DirectoryReference BaseDir, bool IsForDistribution, string BundleID, string InAppName)
+		/// <param name="bMakeProjectPerTarget"></param>
+		public XcodeProjectFile(FileReference InitFilePath, DirectoryReference BaseDir, bool IsForDistribution, string BundleID, string InAppName, bool bMakeProjectPerTarget)
 			: base(InitFilePath, BaseDir)
 		{
 			bForDistribution = IsForDistribution;
 			BundleIdentifier = BundleID;
 			AppName = InAppName;
+			this.bMakeProjectPerTarget = bMakeProjectPerTarget;
 		}
 
 		public override string ToString()
@@ -110,6 +112,11 @@ namespace UnrealBuildTool.XcodeProjectLegacy
 			return ProjectFilePath.GetFileNameWithoutExtension();
 		}
 
+		/// <summary>
+		/// Temporary for developing this feature
+		/// </summary>
+		bool bMakeProjectPerTarget;
+
 		/// <summary>
 		///  Used to mark the project for distribution (some platforms require this)
 		/// </summary>
@@ -1456,7 +1463,7 @@ namespace UnrealBuildTool.XcodeProjectLegacy
 										bShouldCompileMonolithic |= (ProjectTarget.CreateRulesDelegate(Platform, Configuration).LinkType == TargetLinkType.Monolithic);
 
 										string ConfigName = Configuration.ToString();
-										if (ProjectTarget.TargetRules!.Type != TargetType.Game && ProjectTarget.TargetRules.Type != TargetType.Program)
+										if (!bMakeProjectPerTarget && ProjectTarget.TargetRules!.Type != TargetType.Game && ProjectTarget.TargetRules.Type != TargetType.Program)
 										{
 											ConfigName += " " + ProjectTarget.TargetRules.Type.ToString();
 										}
@@ -1470,14 +1477,14 @@ namespace UnrealBuildTool.XcodeProjectLegacy
 											// Get the output directory
 											DirectoryReference RootDirectory;
 											if (UProjectDirectory != null &&
-												(bShouldCompileMonolithic || ProjectTarget.TargetRules.BuildEnvironment == TargetBuildEnvironment.Unique) &&
-												ProjectTarget.TargetRules.File!.IsUnderDirectory(UProjectDirectory))
+												(bShouldCompileMonolithic || ProjectTarget.TargetRules!.BuildEnvironment == TargetBuildEnvironment.Unique) &&
+												ProjectTarget.TargetRules!.File!.IsUnderDirectory(UProjectDirectory))
 											{
 												RootDirectory = UEBuildTarget.GetOutputDirectoryForExecutable(UProjectDirectory, ProjectTarget.TargetRules.File!);
 											}
 											else
 											{
-												RootDirectory = UEBuildTarget.GetOutputDirectoryForExecutable(Unreal.EngineDirectory, ProjectTarget.TargetRules.File!);
+												RootDirectory = UEBuildTarget.GetOutputDirectoryForExecutable(Unreal.EngineDirectory, ProjectTarget.TargetRules!.File!);
 											}
 
 											string ExeName = TargetName;
diff --git a/Engine/Source/Programs/UnrealBuildTool/System/RulesCompiler.cs b/Engine/Source/Programs/UnrealBuildTool/System/RulesCompiler.cs
--- a/Engine/Source/Programs/UnrealBuildTool/System/RulesCompiler.cs
+++ b/Engine/Source/Programs/UnrealBuildTool/System/RulesCompiler.cs
@@ -99,7 +99,7 @@ namespace UnrealBuildTool
 					}
 				}
 
-				EngineRulesAssembly = CreateEngineRulesAssemblyInternal(Unreal.GetExtensionDirs(Unreal.EngineDirectory), ProjectFileGenerator.EngineProjectFileNameBase, EnginePlugins, Unreal.IsEngineInstalled() || bUsePrecompiled, bSkipCompile, bForceCompile, null, Logger);
+				EngineRulesAssembly = CreateEngineRulesAssemblyInternal(Unreal.GetExtensionDirs(Unreal.EngineDirectory), ProjectFileGenerator.EngineRulesAssemblyName, EnginePlugins, Unreal.IsEngineInstalled() || bUsePrecompiled, bSkipCompile, bForceCompile, null, Logger);
 
 				if (MarketplacePlugins.Count > 0)
 				{
