Add the ability to define new documentation policies and specify a default policy

diff --git a/Engine/Programs/UnrealHeaderTool/Config/DefaultEngine.ini b/Engine/Programs/UnrealHeaderTool/Config/DefaultEngine.ini
--- a/Engine/Programs/UnrealHeaderTool/Config/DefaultEngine.ini
+++ b/Engine/Programs/UnrealHeaderTool/Config/DefaultEngine.ini
@@ -169,4 +169,8 @@ NonEngineObjectPtrMemberBehavior=AllowSilently
 AutomaticSettersAndGetters=False
 AreRigVMUObjectProeprtiesEnabled=False
 AreRigVMUInterfaceProeprtiesEnabled=True
-ShowDeprecations=True
\ No newline at end of file
+ShowDeprecations=True
+
+[UnrealHeaderTool]
++DocumentationPolicies=(Name="Strict", ClassOrStructCommentRequired=true, FunctionToolTipsRequired=true, MemberToolTipsRequired=true, ParameterToolTipsRequired=true, FloatRangesRequired=true)
+DefaultDocumentationPolicy=""
diff --git a/Engine/Source/Programs/Shared/EpicGames.UHT/Types/UhtType.cs b/Engine/Source/Programs/Shared/EpicGames.UHT/Types/UhtType.cs
--- a/Engine/Source/Programs/Shared/EpicGames.UHT/Types/UhtType.cs
+++ b/Engine/Source/Programs/Shared/EpicGames.UHT/Types/UhtType.cs
@@ -610,31 +610,6 @@ namespace EpicGames.UHT.Types
 	/// </summary>
 	public class UhtDocumentationPolicy
 	{
-		private static readonly UhtDocumentationPolicy s_strictDocumentationPolicy = new()
-		{
-			PolicySet = true,
-			ClassOrStructCommentRequired = true,
-			FunctionToolTipsRequired = true,
-			MemberToolTipsRequired = true,
-			ParameterToolTipsRequired = true,
-			FloatRangesRequired = true,
-		};
-
-		private static readonly UhtDocumentationPolicy s_noDocumentationPolicy = new()
-		{
-			PolicySet = false,
-			ClassOrStructCommentRequired = false,
-			FunctionToolTipsRequired = false,
-			MemberToolTipsRequired = false,
-			ParameterToolTipsRequired = false,
-			FloatRangesRequired = false,
-		};
-
-		/// <summary>
-		/// If true, this is a valid documentation policy that should be respected
-		/// </summary>
-		public bool PolicySet { get; set; }
-
 		/// <summary>
 		/// A tool tip must be defined
 		/// </summary>
@@ -659,24 +634,6 @@ namespace EpicGames.UHT.Types
 		/// Float properties must have a range specified
 		/// </summary>
 		public bool FloatRangesRequired { get; set; }
-
-		/// <summary>
-		/// Return the documentation policy given the name
-		/// </summary>
-		/// <param name="documentationPolicyName">Name of the policy</param>
-		/// <param name="policy">Policy settings</param>
-		/// <returns>True if the name is valid, false if not</returns>
-		public static bool TryGet(string documentationPolicyName, out UhtDocumentationPolicy policy)
-		{
-			if (documentationPolicyName == "Strict")
-			{
-				policy = s_strictDocumentationPolicy;
-				return true;
-			}
-
-			policy = s_noDocumentationPolicy;
-			return documentationPolicyName.Length == 0;
-		}
 	};
 
 	/// <summary>
@@ -1136,16 +1093,28 @@ namespace EpicGames.UHT.Types
 
 		private void ValidateDocumentationPolicy()
 		{
-			if (UhtDocumentationPolicy.TryGet(DocumentationPolicyName, out UhtDocumentationPolicy policy))
+			if (Session.Config == null)
 			{
-				if (policy.PolicySet)
-				{
-					ValidateDocumentationPolicy(policy);
-				}
+				return;
+			}
+
+			string policyName = DocumentationPolicyName;
+			if (String.IsNullOrEmpty(policyName))
+			{
+				policyName = Session.Config.DefaultDocumentationPolicy;
+			}
+
+			if (String.IsNullOrEmpty(policyName))
+			{
+				return;
+			}
+			else if (Session.Config.DocumentationPolicies.TryGetValue(policyName, out UhtDocumentationPolicy? policy))
+			{
+				ValidateDocumentationPolicy(policy);
 			}
 			else
 			{
-				this.LogError(MetaData.LineNumber, $"Documentation policy '{DocumentationPolicyName}' is not known");
+				this.LogError(MetaData.LineNumber, $"Documentation policy '{policyName}' is not known");
 			}
 		}
 
diff --git a/Engine/Source/Programs/Shared/EpicGames.UHT/Utils/UhtConfig.cs b/Engine/Source/Programs/Shared/EpicGames.UHT/Utils/UhtConfig.cs
--- a/Engine/Source/Programs/Shared/EpicGames.UHT/Utils/UhtConfig.cs
+++ b/Engine/Source/Programs/Shared/EpicGames.UHT/Utils/UhtConfig.cs
@@ -1,7 +1,9 @@
 // Copyright Epic Games, Inc. All Rights Reserved.
 
+using System.Collections.Generic;
 using EpicGames.Core;
 using EpicGames.UHT.Tokenizer;
+using EpicGames.UHT.Types;
 
 namespace EpicGames.UHT.Utils
 {
@@ -103,6 +105,16 @@ namespace EpicGames.UHT.Utils
 		/// </summary>
 		public UhtIssueBehavior NonEngineEnumUnderlyingTypeNotSet { get; }
 
+		/// <summary>
+		/// Collection of all known documentation policies
+		/// </summary>
+		public IReadOnlyDictionary<string, UhtDocumentationPolicy> DocumentationPolicies { get; }
+
+		/// <summary>
+		/// Default documentation policy to be used if none is specified
+		/// </summary>
+		public string DefaultDocumentationPolicy { get; }
+
 		/// <summary>
 		/// If the token references a remapped identifier, update the value in the token 
 		/// </summary>
diff --git a/Engine/Source/Programs/Shared/EpicGames.UHT/Utils/UhtSession.cs b/Engine/Source/Programs/Shared/EpicGames.UHT/Utils/UhtSession.cs
--- a/Engine/Source/Programs/Shared/EpicGames.UHT/Utils/UhtSession.cs
+++ b/Engine/Source/Programs/Shared/EpicGames.UHT/Utils/UhtSession.cs
@@ -1020,6 +1020,18 @@ namespace EpicGames.UHT.Utils
 					break;
 			}
 
+			{
+				string defaultPolicyName = Config.DefaultDocumentationPolicy;
+				if (!String.IsNullOrEmpty(defaultPolicyName))
+				{
+					if (!Config.DocumentationPolicies.TryGetValue(defaultPolicyName, out UhtDocumentationPolicy? policy))
+					{
+						Log.Logger.LogError("The default documentation policy '{DefaultPolicyName}' isn't known", defaultPolicyName);
+						return;
+					}
+				}
+			}
+
 			if (ReferenceMode != UhtReferenceMode.None)
 			{
 				string directoryToDelete = ReferenceMode == UhtReferenceMode.Reference ? ReferenceDirectory : VerifyDirectory;
diff --git a/Engine/Source/Programs/UnrealBuildTool/Modes/UnrealHeaderToolMode.cs b/Engine/Source/Programs/UnrealBuildTool/Modes/UnrealHeaderToolMode.cs
--- a/Engine/Source/Programs/UnrealBuildTool/Modes/UnrealHeaderToolMode.cs
+++ b/Engine/Source/Programs/UnrealBuildTool/Modes/UnrealHeaderToolMode.cs
@@ -30,7 +30,7 @@ namespace UnrealBuildTool.Modes
 		private readonly IReadOnlyDictionary<StringView, StringView> _typeRedirectMap;
 
 		/// <summary>
-		/// Metadata that have been renamed, treat the old deprecated name as the new name for code generation
+		/// Meta data that have been renamed, treat the old deprecated name as the new name for code generation
 		/// </summary>
 		private readonly IReadOnlyDictionary<string, string> _metaDataRedirectMap;
 
@@ -126,6 +126,16 @@ namespace UnrealBuildTool.Modes
 		/// </summary>
 		private readonly UhtIssueBehavior _nonEngineEnumUnderlyingTypeNotSet = UhtIssueBehavior.AllowSilently;
 
+		/// <summary>
+		/// Collection of known documentation policies
+		/// </summary>
+		public Dictionary<string, UhtDocumentationPolicy> _documentationPolicies = new(StringComparer.OrdinalIgnoreCase);
+
+		/// <summary>
+		/// Default documentation policy (usually empty)
+		/// </summary>
+		private readonly string _defaultDocumentationPolicy = "";
+
 		#region IUhtConfig Implementation
 		/// <inheritdoc/>
 		public EGeneratedCodeVersion DefaultGeneratedCodeVersion => this._defaultGeneratedCodeVersion;
@@ -175,6 +185,12 @@ namespace UnrealBuildTool.Modes
 		/// <inheritdoc/>
 		public UhtIssueBehavior NonEngineEnumUnderlyingTypeNotSet => this._nonEngineEnumUnderlyingTypeNotSet;
 
+		/// <inheritdoc/>
+		public IReadOnlyDictionary<string, UhtDocumentationPolicy> DocumentationPolicies => this._documentationPolicies;
+
+		/// <inheritdoc/>
+		public string DefaultDocumentationPolicy => this._defaultDocumentationPolicy;
+
 		/// <inheritdoc/>
 		public void RedirectTypeIdentifier(ref UhtToken Token)
 		{
@@ -259,19 +275,22 @@ namespace UnrealBuildTool.Modes
 			this._units = GetHashSet("UnrealHeaderTool", "Units", StringViewComparer.OrdinalIgnoreCase);
 			this._delegateParameterCountStrings = GetList("UnrealHeaderTool", "DelegateParameterCountStrings");
 			this._defaultGeneratedCodeVersion = GetGeneratedCodeVersion("UnrealHeaderTool", "DefaultGeneratedCodeVersion", EGeneratedCodeVersion.V1);
-			this._engineNativePointerMemberBehavior = GetPointerMemberBehavior("UnrealHeaderTool", "EngineNativePointerMemberBehavior", UhtIssueBehavior.AllowSilently);
-			this._engineObjectPtrMemberBehavior = GetPointerMemberBehavior("UnrealHeaderTool", "EngineObjectPtrMemberBehavior", UhtIssueBehavior.AllowSilently);
-			this._enginePluginNativePointerMemberBehavior = GetPointerMemberBehavior("UnrealHeaderTool", "EnginePluginNativePointerMemberBehavior", UhtIssueBehavior.AllowSilently);
-			this._enginePluginObjectPtrMemberBehavior = GetPointerMemberBehavior("UnrealHeaderTool", "EnginePluginObjectPtrMemberBehavior", UhtIssueBehavior.AllowSilently);
-			this._nonEngineNativePointerMemberBehavior = GetPointerMemberBehavior("UnrealHeaderTool", "NonEngineNativePointerMemberBehavior", UhtIssueBehavior.AllowSilently);
-			this._nonEngineObjectPtrMemberBehavior = GetPointerMemberBehavior("UnrealHeaderTool", "NonEngineObjectPtrMemberBehavior", UhtIssueBehavior.AllowSilently);
+			this._engineNativePointerMemberBehavior = GetIssueBehavior("UnrealHeaderTool", "EngineNativePointerMemberBehavior", UhtIssueBehavior.AllowSilently);
+			this._engineObjectPtrMemberBehavior = GetIssueBehavior("UnrealHeaderTool", "EngineObjectPtrMemberBehavior", UhtIssueBehavior.AllowSilently);
+			this._enginePluginNativePointerMemberBehavior = GetIssueBehavior("UnrealHeaderTool", "EnginePluginNativePointerMemberBehavior", UhtIssueBehavior.AllowSilently);
+			this._enginePluginObjectPtrMemberBehavior = GetIssueBehavior("UnrealHeaderTool", "EnginePluginObjectPtrMemberBehavior", UhtIssueBehavior.AllowSilently);
+			this._nonEngineNativePointerMemberBehavior = GetIssueBehavior("UnrealHeaderTool", "NonEngineNativePointerMemberBehavior", UhtIssueBehavior.AllowSilently);
+			this._nonEngineObjectPtrMemberBehavior = GetIssueBehavior("UnrealHeaderTool", "NonEngineObjectPtrMemberBehavior", UhtIssueBehavior.AllowSilently);
 			this._areRigVMUObjectProeprtiesEnabled = GetBoolean("UnrealHeaderTool", "AreRigVMUObjectProeprtiesEnabled", false);
 			this._areRigVMUInterfaceProeprtiesEnabled = GetBoolean("UnrealHeaderTool", "AreRigVMUInterfaceProeprtiesEnabled", false);
 			this._showDeprecations = GetBoolean("UnrealHeaderTool", "ShowDeprecations", true);
-			this._engineMissingGeneratedHeaderIncludeBehavior = GetPointerMemberBehavior("UnrealHeaderTool", "EngineMissingGeneratedHeaderIncludeBehavior", UhtIssueBehavior.AllowSilently);
-			this._nonEngineMissingGeneratedHeaderIncludeBehavior = GetPointerMemberBehavior("UnrealHeaderTool", "NonEngineMissingGeneratedHeaderIncludeBehavior", UhtIssueBehavior.AllowSilently);
-			this._engineEnumUnderlyingTypeNotSet = GetPointerMemberBehavior("UnrealHeaderTool", "EngineEnumUnderlyingTypeNotSet", UhtIssueBehavior.AllowSilently);
-			this._nonEngineEnumUnderlyingTypeNotSet = GetPointerMemberBehavior("UnrealHeaderTool", "NonEngineEnumUnderlyingTypeNotSet", UhtIssueBehavior.AllowSilently);
+			this._engineMissingGeneratedHeaderIncludeBehavior = GetIssueBehavior("UnrealHeaderTool", "EngineMissingGeneratedHeaderIncludeBehavior", UhtIssueBehavior.AllowSilently);
+			this._nonEngineMissingGeneratedHeaderIncludeBehavior = GetIssueBehavior("UnrealHeaderTool", "NonEngineMissingGeneratedHeaderIncludeBehavior", UhtIssueBehavior.AllowSilently);
+			this._engineEnumUnderlyingTypeNotSet = GetIssueBehavior("UnrealHeaderTool", "EngineEnumUnderlyingTypeNotSet", UhtIssueBehavior.AllowSilently);
+			this._nonEngineEnumUnderlyingTypeNotSet = GetIssueBehavior("UnrealHeaderTool", "NonEngineEnumUnderlyingTypeNotSet", UhtIssueBehavior.AllowSilently);
+
+			GetDocumentationPolicies("UnrealHeaderTool", "DocumentationPolicies");
+			_defaultDocumentationPolicy = GetString("UnrealHeaderTool", "DefaultDocumentationPolicy", "");
 		}
 
 		private bool GetBoolean(string SectionName, string KeyName, bool bDefault)
@@ -283,13 +302,22 @@ namespace UnrealBuildTool.Modes
 			return bDefault;
 		}
 
-		private UhtIssueBehavior GetPointerMemberBehavior(string SectionName, string KeyName, UhtIssueBehavior Default)
+		private string GetString(string SectionName, string KeyName, string Default)
+		{
+			if (this._ini.TryGetValue(SectionName, KeyName, out string? value))
+			{
+				return value;
+			}
+			return Default;
+		}
+
+		private UhtIssueBehavior GetIssueBehavior(string SectionName, string KeyName, UhtIssueBehavior Default)
 		{
 			if (this._ini.TryGetValue(SectionName, KeyName, out string? BehaviorStr))
 			{
 				if (!Enum.TryParse(BehaviorStr, out UhtIssueBehavior Value))
 				{
-					throw new Exception(string.Format("Unrecognized native pointer member behavior '{0}'", BehaviorStr));
+					throw new Exception(string.Format("Unrecognized issue behavior '{0}'", BehaviorStr));
 				}
 				return Value;
 			}
@@ -386,6 +414,50 @@ namespace UnrealBuildTool.Modes
 			}
 			return Set;
 		}
+
+		private void GetDocumentationPolicies(string Section, string Key)
+		{
+			_documentationPolicies["Strict"] = new()
+			{
+				ClassOrStructCommentRequired = true,
+				FunctionToolTipsRequired = true,
+				MemberToolTipsRequired = true,
+				ParameterToolTipsRequired = true,
+				FloatRangesRequired = true,
+			};
+
+			if (this._ini.TryGetValues(Section, Key, out IReadOnlyList<string>? StringList))
+			{
+				foreach (string Value in StringList)
+				{
+					if (ConfigHierarchy.TryParse(Value, out Dictionary<string, string>? Properties))
+					{
+						if (Properties.TryGetValue("Name", out string? PolicyName))
+						{
+							UhtDocumentationPolicy? Policy = null;
+							if (!_documentationPolicies.TryGetValue(PolicyName, out Policy))
+							{
+								Policy = new UhtDocumentationPolicy();
+							}
+							Policy.ClassOrStructCommentRequired = GetPropertyBool(Properties, "ClassOrStructCommentRequired", Policy.ClassOrStructCommentRequired);
+							Policy.FunctionToolTipsRequired = GetPropertyBool(Properties, "FunctionToolTipsRequired", Policy.FunctionToolTipsRequired);
+							Policy.MemberToolTipsRequired = GetPropertyBool(Properties, "MemberToolTipsRequired", Policy.MemberToolTipsRequired);
+							Policy.ParameterToolTipsRequired = GetPropertyBool(Properties, "ParameterToolTipsRequired", Policy.ParameterToolTipsRequired);
+							Policy.FloatRangesRequired = GetPropertyBool(Properties, "FloatRangesRequired", Policy.FloatRangesRequired);
+						}
+					}
+				}
+			}
+		}
+
+		private static bool GetPropertyBool(Dictionary<string, string> Properties, string Key, bool DefaultValue)
+		{
+			if (Properties.TryGetValue(Key, out string? PropValueString) && ConfigHierarchy.TryParse(PropValueString, out bool PropValue))
+			{
+				return PropValue;
+			}
+			return DefaultValue;
+		}
 	}
 
 	/// <summary>
