- Actually disabled bMakeProjectPerTarget, it was left on mistakenly in my last checkin

diff --git a/Engine/Source/Programs/UnrealBuildTool/ProjectFiles/VisualStudio/VCProjectFileGenerator.cs b/Engine/Source/Programs/UnrealBuildTool/ProjectFiles/VisualStudio/VCProjectFileGenerator.cs
--- a/Engine/Source/Programs/UnrealBuildTool/ProjectFiles/VisualStudio/VCProjectFileGenerator.cs
+++ b/Engine/Source/Programs/UnrealBuildTool/ProjectFiles/VisualStudio/VCProjectFileGenerator.cs
@@ -101,7 +101,7 @@ namespace UnrealBuildTool
 		/// <summary>
 		/// Set to true to enable a project for each target, and do not put the target type into the configuration
 		/// </summary>
-		protected override bool bMakeProjectPerTarget => true;
+		protected override bool bMakeProjectPerTarget => false;
 
 		/// <summary>
 		/// The settings object
