- Missed a check for bMakeProjectPerTarget hceck in legacy xcode projects

diff --git a/Engine/Source/Programs/UnrealBuildTool/ProjectFiles/Xcode/XcodeProjectLegacy.cs b/Engine/Source/Programs/UnrealBuildTool/ProjectFiles/Xcode/XcodeProjectLegacy.cs
--- a/Engine/Source/Programs/UnrealBuildTool/ProjectFiles/Xcode/XcodeProjectLegacy.cs
+++ b/Engine/Source/Programs/UnrealBuildTool/ProjectFiles/Xcode/XcodeProjectLegacy.cs
@@ -1593,7 +1593,7 @@ namespace UnrealBuildTool.XcodeProjectLegacy
 				}
 			}
 
-			string DefaultConfiguration = bHasEditorConfiguration && !XcodeProjectFileGenerator.bGenerateRunOnlyProject ? "Development Editor" : "Development";
+			string DefaultConfiguration = bMakeProjectPerTarget && bHasEditorConfiguration && !XcodeProjectFileGenerator.bGenerateRunOnlyProject ? "Development Editor" : "Development";
 
 			StringBuilder Content = new StringBuilder();
 
