Removed VC check because it is no longer needed.

diff --git a/Engine/Source/Programs/UnrealBuildTool/Platform/Windows/VCToolChain.cs b/Engine/Source/Programs/UnrealBuildTool/Platform/Windows/VCToolChain.cs
--- a/Engine/Source/Programs/UnrealBuildTool/Platform/Windows/VCToolChain.cs
+++ b/Engine/Source/Programs/UnrealBuildTool/Platform/Windows/VCToolChain.cs
@@ -1238,10 +1238,6 @@ namespace UnrealBuildTool
 			//   			WindowsPlatform.AdditionalLinkerOptions = "/PDBPAGESIZE:8192";
 			if (Target.LinkType == TargetLinkType.Monolithic && Target.Type == TargetType.Editor)
 			{
-				if (EnvVars.CompilerVersion < VersionNumber.Parse("14.29.30133"))
-				{
-					throw new BuildException($"Monolithic editors now require VC Toolchain 14.29.30133 (the toolchain for Visual Studio 16.11) ({EnvVars.CompilerVersion} < {VersionNumber.Parse("14.29.30133")})");
-				}
 				Arguments.Add("/PDBPAGESIZE:16384");
 				//Arguments.Add("/PDBCompress"); // Do not turn this on, it makes link times almost 2x slower. This is _only_ to save local disk space. Will _not_ make actual file smaller for network transfer
 			}
