[UBT]
* IWYU - Added table for common types that should not be forward declared but instead use a forwarding header

diff --git a/Engine/Source/Programs/UnrealBuildTool/Modes/IWYUMode.cs b/Engine/Source/Programs/UnrealBuildTool/Modes/IWYUMode.cs
--- a/Engine/Source/Programs/UnrealBuildTool/Modes/IWYUMode.cs
+++ b/Engine/Source/Programs/UnrealBuildTool/Modes/IWYUMode.cs
@@ -175,6 +175,11 @@ namespace UnrealBuildTool
 		/// Is true if this file was included inside a declaration such as a namespace or class/struct/function
 		/// </summary>
 		public bool IncludedInsideDecl;
+
+		/// <summary>
+		/// Is true if this file ends with .cpp
+		/// </summary>
+		public bool IsCpp;
 	}
 
 	/// <summary>
@@ -567,10 +572,6 @@ namespace UnrealBuildTool
 			Dictionary<string, IWYUInfo> Infos = new();
 			HashSet<string> GeneratedHeaderInfos = new();
 			List<IWYUInfo> GeneratedCppInfos = new();
-			IWYUInfo? ObjectMacrosInfo = null;
-			IWYUInfo? ScriptMacrosInfo = null;
-			IWYUInfo? VerseInteropUtilsInfo = null;
-			IWYUInfo? PlatformInfo = null;
 
 			using (ISourceFileWorkingSet WorkingSet = new EmptySourceFileWorkingSet())
 			{
@@ -669,6 +670,8 @@ namespace UnrealBuildTool
 						foreach (var Info in IWYUFile!.Files)
 						{
 							Info.Source = IWYUFile;
+							Info.IsCpp = Info.File.EndsWith(".cpp");
+
 
 							// We track .gen.cpp in a special list, they need special treatment later
 							if (Info.File.Contains(".gen.cpp", StringComparison.Ordinal))
@@ -723,16 +726,6 @@ namespace UnrealBuildTool
 							// Some special logic for headers.
 							if (Info.File.EndsWith(".h"))
 							{
-								// Find ObjectMacros and ScriptMacros. These are always included in .generated.h files which will generate IWYUInfos for further down
-								if (ObjectMacrosInfo == null && Info.File.EndsWith("UObject/ObjectMacros.h"))
-									ObjectMacrosInfo = Info;
-								else if (ScriptMacrosInfo == null && Info.File.EndsWith("UObject/ScriptMacros.h"))
-									ScriptMacrosInfo = Info;
-								else if (VerseInteropUtilsInfo == null && Info.File.EndsWith("VerseInteropUtils.h"))
-									VerseInteropUtilsInfo = Info;
-								else if (PlatformInfo == null && Info.File.EndsWith("HAL/Platform.h"))
-									PlatformInfo = Info;
-
 								// We need to add entries for .generated.h. They are not in the iwyu files
 								foreach (var Include in Info.Includes)
 								{
@@ -775,22 +768,50 @@ namespace UnrealBuildTool
 				}
 			}
 
+			KeyValuePair<string, IWYUIncludeEntry?> GetInfo(string Include, string Path)
+			{
+				string FullPath = DirectoryReference.Combine(Unreal.EngineDirectory, Path).FullName.Replace('\\', '/');
+				IWYUInfo? Out = null;
+				if (!Infos.TryGetValue(FullPath, out Out))
+				{
+					return new(Include, null);
+				}
+
+				IWYUIncludeEntry Entry = new();
+				Entry.Printable = Include;
+				Entry.Full = FullPath;
+				Entry.Resolved = Out;
+				return new(Include, Entry);
+			}
+
+			var SpecialIncludes = new Dictionary<string, IWYUIncludeEntry?>(
+				new[]
+				{
+					GetInfo("UObject/ObjectMacros.h", "Source/Runtime/CoreUObject/Public/UObject/ObjectMacros.h"),
+					GetInfo("UObject/ScriptMacros.h", "Source/Runtime/CoreUObject/Public/UObject/ScriptMacros.h"),
+					GetInfo("VerseInteropUtils.h", "Restricted/NotForLicensees/Plugins/Solaris/Source/VerseNative/Public/VerseInteropTypes.h"),
+					GetInfo("Containers/ContainersFwd.h", "Source/Runtime/Core/Public/Containers/ContainersFwd.h"),
+					GetInfo("Misc/OptionalFwd.h", "Source/Runtime/Core/Public/Misc/OptionalFwd.h"),
+				}
+			);
+
+			var ForwardingHeaders = new Dictionary<string, IWYUIncludeEntry?>()
+			{
+				{ "TMap", SpecialIncludes["Containers/ContainersFwd.h"] },
+				{ "TSet", SpecialIncludes["Containers/ContainersFwd.h"] },
+				{ "TArray", SpecialIncludes["Containers/ContainersFwd.h"] },
+				{ "TArrayView", SpecialIncludes["Containers/ContainersFwd.h"] },
+				{ "TOptional", SpecialIncludes["Misc/OptionalFwd.h"] },
+			};
+
+
 			// Add all .generated.h files as entries in the lookup and explicitly add the includes they have which will never be removed
 			Logger.LogInformation($"Generating infos for .generated.h files...");
 			if (GeneratedHeaderInfos.Count > 0)
 			{
-				IWYUIncludeEntry ObjectMacrosInclude = new();
-				IWYUIncludeEntry ScriptMacrosInclude = new();
-				IWYUIncludeEntry VerseInteropUtilsInclude = new();
-				ObjectMacrosInclude.Printable = "UObject/ObjectMacros.h";
-				ObjectMacrosInclude.Full = ObjectMacrosInfo!.File;
-				ObjectMacrosInclude.Resolved = ObjectMacrosInfo;
-				ScriptMacrosInclude.Printable = "UObject/ScriptMacros.h";
-				ScriptMacrosInclude.Full = ScriptMacrosInfo!.File;
-				ScriptMacrosInclude.Resolved = ScriptMacrosInfo;
-				VerseInteropUtilsInclude.Printable = "VerseInteropUtils.h";
-				VerseInteropUtilsInclude.Full = VerseInteropUtilsInfo!.File;
-				VerseInteropUtilsInclude.Resolved = VerseInteropUtilsInfo;
+				IWYUIncludeEntry ObjectMacrosInclude = SpecialIncludes["UObject/ObjectMacros.h"]!;
+				IWYUIncludeEntry ScriptMacrosInclude = SpecialIncludes["UObject/ScriptMacros.h"]!;
+				IWYUIncludeEntry VerseInteropUtilsInclude = SpecialIncludes["VerseInteropUtils.h"]!;
 
 				foreach (var Gen in GeneratedHeaderInfos)
 				{
@@ -838,6 +859,25 @@ namespace UnrealBuildTool
 					}
 				}
 
+
+				if (!Info.IsCpp)
+				{
+					// See if we need to replace forward declarations with includes
+					Info.ForwardDeclarations.RemoveAll(Entry =>
+					{
+						int LastSpace = Entry.Printable.LastIndexOf(' ');
+						string TypeName = Entry.Printable.Substring(LastSpace + 1, Entry.Printable.Length - LastSpace - 2);
+						IWYUIncludeEntry? IncludeEntry;
+						ForwardingHeaders.TryGetValue(TypeName, out IncludeEntry);
+						if (IncludeEntry == null)
+						{
+							return false;
+						}
+						Info.Includes.Add(IncludeEntry);
+						return true;
+					});
+				}
+
 				// We don't want to mess around with third party includes.. since we can't see the hierarchy we just assumes that they are optimally included
 				Info.Includes.RemoveAll(Entry => Entry.Full.Contains("/ThirdParty/", StringComparison.Ordinal) || Entry.Full.StartsWith(SystemPath));
 
@@ -1103,12 +1143,17 @@ namespace UnrealBuildTool
 					foreach (string Line in Lines)
 					{
 						ReadOnlySpan<char> LineSpan = Line.AsSpan().TrimStart();
-						if (Line.StartsWith("#if"))
+						if (LineSpan.Length == 0 || LineSpan[0] != '#')
+						{
+							continue;
+						}
+						LineSpan = LineSpan.Slice(1).TrimStart();
+						if (LineSpan.StartsWith("if"))
 						{
 							// Include guards are special
 							if (!HasIncludeGuard)
 							{
-								if (Line.StartsWith("#ifndef ") && Line.EndsWith("_H"))
+								if (LineSpan.StartsWith("ifndef ") && LineSpan.EndsWith("_H"))
 								{
 									HasIncludeGuard = true;
 									continue;
@@ -1116,13 +1161,13 @@ namespace UnrealBuildTool
 							}
 							++IfCount;
 						}
-						else if (Line.StartsWith("#endif"))
+						else if (LineSpan.StartsWith("endif"))
 						{
 							--IfCount;
 						}
-						else if (Line.StartsWith("#include"))
+						else if (LineSpan.StartsWith("include"))
 						{
-							var IncludeSpan = LineSpan.Slice("#include".Length).TrimStart();
+							var IncludeSpan = LineSpan.Slice("include".Length).TrimStart();
 							char LeadingIncludeChar = IncludeSpan[0];
 							if (LeadingIncludeChar == '"' || LeadingIncludeChar == '<')
 							{
@@ -1218,7 +1263,7 @@ namespace UnrealBuildTool
 				if (!IsValidForUpdate(Info, ValidPaths, true))
 					return;
 
-				bool IsCpp = Info.File.EndsWith(".cpp");
+				bool IsCpp = Info.IsCpp;
 				bool IsPrivate = IsCpp || Info.File.Contains("/Private/");
 
 				// If we only want to update private files we early out for non-private headers
@@ -1354,19 +1399,22 @@ namespace UnrealBuildTool
 				}
 
 				// Traverse all lines in file and figure out which includes that should be added or removed
-				foreach (var L in ExistingLines)
+				foreach (var Line in ExistingLines)
 				{
 					++LineIndex;
 
-					string Line = L.Trim();
-					if (!Line.StartsWith("#include"))
+					ReadOnlySpan<char> LineSpan = Line.AsSpan().Trim();
+					bool StartsWithHash = LineSpan.Length > 0 && LineSpan[0] == '#';
+					if (StartsWithHash)
+						LineSpan = LineSpan.Slice(1).TrimStart();
+					if (!StartsWithHash || !LineSpan.StartsWith("include"))
 					{
 						// Might be forward declaration.. 
 						if (ForwardDeclarationsToAdd.Remove(Line)) // Skip adding the ones that already exists
 						{
 							if (FirstForwardDeclareLine == null)
 							{
-								FirstForwardDeclareLine = L;
+								FirstForwardDeclareLine = Line;
 							}
 						}
 
@@ -1401,7 +1449,7 @@ namespace UnrealBuildTool
 
 					bool ForceKeep = false;
 
-					ReadOnlySpan<char> IncludeSpan = Line.AsSpan(8).TrimStart();
+					ReadOnlySpan<char> IncludeSpan = LineSpan.Slice("include".Length).TrimStart();
 					char LeadingIncludeChar = IncludeSpan[0];
 					if (LeadingIncludeChar != '"' && LeadingIncludeChar != '<')
 					{
@@ -1426,8 +1474,8 @@ namespace UnrealBuildTool
 							// For now, let's assume that if there is one line in the file that is an include 
 							if (IncludesToAdd.Count == 1)
 							{
-								IncludesToAdd.Clear();
-								ErrorOnMoreIncludes = true;
+								//IncludesToAdd.Clear();
+								//ErrorOnMoreIncludes = true;
 							}
 						}
 						continue;
@@ -1452,7 +1500,7 @@ namespace UnrealBuildTool
 
 					if (!ForceKeep && !ForceKeepScope && !CleanedupIncludes.Contains(Include))
 					{
-						LinesToRemove.TryAdd(L, Include);
+						LinesToRemove.TryAdd(Line, Include);
 					}
 					else
 					{
@@ -1554,9 +1602,15 @@ namespace UnrealBuildTool
 
 					foreach (string OldLine in ExistingLines)
 					{
+						ReadOnlySpan<char> OldLineSpan = OldLine.AsSpan().TrimStart();
+						bool StartsWithHash = OldLineSpan.Length > 0 && OldLineSpan[0] == '#';
+						if (StartsWithHash)
+							OldLineSpan = OldLineSpan.Slice(1).TrimStart();
+						bool IsInclude = StartsWithHash && OldLineSpan.StartsWith("include");
+
 						string OldLineTrimmedStart = OldLine.TrimStart();
 
-						if (!OldLineTrimmedStart.StartsWith("#include"))
+						if (!IsInclude)
 						{
 							if (IsInFirstIncludeBlock)
 							{
