Additional UBT features for supporting a VS Code extension for intellisense.

diff --git a/Engine/Source/Programs/UnrealBuildTool/Modes/GenerateProjectFilesMode.cs b/Engine/Source/Programs/UnrealBuildTool/Modes/GenerateProjectFilesMode.cs
--- a/Engine/Source/Programs/UnrealBuildTool/Modes/GenerateProjectFilesMode.cs
+++ b/Engine/Source/Programs/UnrealBuildTool/Modes/GenerateProjectFilesMode.cs
@@ -261,7 +261,7 @@ namespace UnrealBuildTool
 		/// <param name="Logger">Logger for output</param>
 		/// <param name="ProjectFile">The project file that was parsed</param>
 		/// <returns>True if the project file was parsed, false otherwise</returns>
-		private static bool TryParseProjectFileArgument(CommandLineArguments Arguments, ILogger Logger, [NotNullWhen(true)] out FileReference? ProjectFile)
+		public static bool TryParseProjectFileArgument(CommandLineArguments Arguments, ILogger Logger, [NotNullWhen(true)] out FileReference? ProjectFile)
 		{
 			string? CandidateProjectPath = null;
 
diff --git a/Engine/Source/Programs/UnrealBuildTool/Modes/QueryMode.cs b/Engine/Source/Programs/UnrealBuildTool/Modes/QueryMode.cs
new file mode 100644
--- /dev/null
+++ b/Engine/Source/Programs/UnrealBuildTool/Modes/QueryMode.cs
@@ -0,0 +1,349 @@
+// Copyright Epic Games, Inc. All Rights Reserved.
+
+using EpicGames.Core;
+using Microsoft.Extensions.Logging;
+using System;
+using System.Linq;
+using System.Text.Json;
+using System.Text.Json.Serialization;
+using System.Collections.Generic;
+using OpenTracing.Util;
+using UnrealBuildTool;
+using UnrealBuildBase;
+
+namespace UnrealBuildTool
+{
+	// TODO: Namespacing?
+	enum QueryType
+	{
+		Capabilities,
+		AvailableTargets,
+		TargetDetails,
+	}
+	internal class TargetIntellisenseInfo
+	{
+		internal class CompileSettings
+		{
+			public List<string> IncludePaths { get; set; } = new();
+			public List<string> Defines { get; set; } = new();
+			public string? Standard { get; set; }
+			public List<string> ForcedIncludes { get; set; } = new();
+			public string? CompilerPath { get; set; }
+			public List<string> CompilerArgs { get; set; } = new();
+			public string? WindowsSdkVersion { get; set; }
+		}
+
+		public Dictionary<UEBuildModule, CompileSettings> ModuleToCompileSettings = new();
+		public Dictionary<DirectoryReference, UEBuildModule> DirToModule = new();
+
+		public UEBuildModule? FindModuleForFile(FileReference File)
+		{
+			DirectoryReference? Dir = File.Directory;
+			while (Dir != null)
+			{
+				if (DirToModule.TryGetValue(Dir, out UEBuildModule? Module))
+				{
+					return Module;
+				}
+				Dir = Dir.ParentDirectory;
+			}
+			return null;
+		}
+		public UEBuildModule? FindModuleForDirectory(DirectoryReference Directory)
+		{
+			DirectoryReference? Dir = Directory;
+			while (Dir != null)
+			{
+				if (DirToModule.TryGetValue(Dir, out UEBuildModule? Module))
+				{
+					return Module;
+				}
+				Dir = Dir.ParentDirectory;
+			}
+			return null;
+		}
+	}
+
+	internal class TargetConfigs
+	{
+		public string ProjectPath { get; set; } = "";
+		public List<string> Platforms { get; set; } = new();
+		public List<string> Configurations { get; set; } = new();
+	}
+
+	[ToolMode("Query", ToolModeOptions.BuildPlatforms | ToolModeOptions.XmlConfig | ToolModeOptions.UseStartupTraceListener)]
+	class QueryMode : ToolMode
+	{ 
+		[CommandLine("-LogDirectory=")]
+		public DirectoryReference? LogDirectory = null;
+
+		[CommandLine("-Query=")]
+		public QueryType? Query = null;
+
+		[CommandLine("-IncludeEngineSource=")]
+		public bool bIncludeEngineSource = true;
+
+		[CommandLine("-Target=")]
+		public string? TargetName;
+
+		[CommandLine("-Configuration=")]
+		public string? TargetConfiguration;
+
+		[CommandLine("-Platform=")]
+		public string? TargetPlatform;
+
+		[CommandLine("-Indented")]
+		public bool bIndented;
+
+		private BuildConfiguration BuildConfiguration = new();
+
+		public override int Execute(CommandLineArguments Arguments, ILogger Logger)
+		{
+			Arguments.ApplyTo(this);
+
+			if (LogDirectory == null)
+			{
+				LogDirectory = DirectoryReference.Combine(Unreal.EngineProgramSavedDirectory, "UnrealBuildTool");
+			}
+
+			DirectoryReference.CreateDirectory(LogDirectory);
+
+			FileReference LogFile = FileReference.Combine(LogDirectory, "Log_Query.txt"); // TODO: More history? Pass a log file path on the cmd line from extension and prune from there?
+			Log.AddFileWriter("DefaultLogTraceListener", LogFile);
+
+			XmlConfig.ApplyTo(BuildConfiguration);
+			Arguments.ApplyTo(BuildConfiguration);
+
+			// TODO: Document this hack 
+			ProjectFileGenerator.bGenerateProjectFiles = true;
+
+			if (Query == null)
+			{
+				return 0;
+			}
+
+			var ResponseOptions = new JsonSerializerOptions
+			{
+				PropertyNamingPolicy = JsonNamingPolicy.CamelCase,
+				WriteIndented = bIndented,
+			};
+			switch(Query)
+			{
+				case QueryType.Capabilities:
+					Logger.LogInformation("QueryCapabilities");
+					return QueryCapabilities(Arguments, Logger, ResponseOptions);
+				case QueryType.AvailableTargets:
+					Logger.LogInformation("QueryAvailableTargets");
+					return QueryAvailableTargets(Arguments, Logger, ResponseOptions);
+				case QueryType.TargetDetails:
+					Logger.LogInformation("QueryTargetDetails");
+					return QueryTargetDetails(Arguments, Logger, ResponseOptions);
+			}
+
+			return 0;
+		}
+
+		private int QueryCapabilities(CommandLineArguments Arguments, ILogger Logger, JsonSerializerOptions JsonOptions)
+		{
+			var Reply = new
+			{
+				Queries = new List<string>{ QueryType.Capabilities.ToString(), QueryType.AvailableTargets.ToString(), QueryType.TargetDetails.ToString() }
+			};
+			Console.WriteLine(JsonSerializer.Serialize(Reply, JsonOptions));
+			return 0;		
+		}
+
+		private int QueryAvailableTargets(CommandLineArguments Arguments, ILogger Logger, JsonSerializerOptions JsonOptions)
+		{
+			try
+			{
+				GenerateProjectFilesMode.TryParseProjectFileArgument(Arguments, Logger, out FileReference? ProjectFileArg);
+
+				List<UnrealTargetPlatform> Platforms = new();
+				foreach (UnrealTargetPlatform Platform in UnrealTargetPlatform.GetValidPlatforms())
+				{
+					// Only include desktop platforms for now 
+					if (ProjectFileGenerator.IsValidDesktopPlatform(Platform))
+					{
+						// If there is a build platform present, add it to the SupportedPlatforms list
+						UEBuildPlatform? BuildPlatform;
+						if (UEBuildPlatform.TryGetBuildPlatform(Platform, out BuildPlatform))
+						{
+							if (InstalledPlatformInfo.IsValidPlatform(Platform, EProjectType.Code))
+							{
+								Platforms.Add(Platform);
+							}
+						}
+					}
+				}
+
+				List<UnrealTargetConfiguration> AllowedTargetConfigurations = new List<UnrealTargetConfiguration>();
+				AllowedTargetConfigurations = Enum.GetValues(typeof(UnrealTargetConfiguration)).Cast<UnrealTargetConfiguration>().ToList();
+
+				List<string> Configurations = new();
+				foreach (UnrealTargetConfiguration CurConfiguration in AllowedTargetConfigurations)
+				{
+					if (CurConfiguration != UnrealTargetConfiguration.Unknown)
+					{
+						if (InstalledPlatformInfo.IsValidConfiguration(CurConfiguration, EProjectType.Code))
+						{
+							Configurations.Add(CurConfiguration.ToString());
+						}
+					}
+				}
+
+				List<FileReference> Projects = ProjectFileArg != null ? new List<FileReference>(new[] { ProjectFileArg }) : NativeProjects.EnumerateProjectFiles(Logger).ToList();
+				List<FileReference> AllTargetFiles = ProjectFileGenerator.DiscoverTargets(Projects, Logger, null, Platforms, bIncludeEngineSource: bIncludeEngineSource, bIncludeTempTargets: false);
+
+				// TODO: Check valid configurations/platforms per target 
+				Dictionary<string, TargetConfigs> Targets = new();
+				string? DefaultTarget = null;
+				foreach (FileReference TargetFilePath in AllTargetFiles)
+				{
+					string TargetName = TargetFilePath.GetFileNameWithoutAnyExtensions();
+					FileReference? ProjectPath = Projects.FirstOrDefault(p => TargetFilePath.IsUnderDirectory( p.Directory));
+					Targets.Add(TargetName, new TargetConfigs() { ProjectPath = ProjectPath?.ToString() ?? "", Configurations = Configurations, Platforms = Platforms.Select(x => x.ToString()).ToList() });
+					if ( DefaultTarget == null || TargetName == "UnrealEditor")
+					{
+						DefaultTarget = TargetName;
+					}
+				}
+
+				var Reply = new
+				{
+					Targets = Targets,
+
+					DefaultTarget = DefaultTarget,
+					DefaultPlatform = Platforms[0].ToString(),
+					DefaultConfiguration = UnrealTargetConfiguration.Development.ToString(),
+				};
+				Console.WriteLine(JsonSerializer.Serialize(Reply, JsonOptions));
+				return 0;
+			}
+			catch (Exception e)
+			{
+				Logger.LogError("Failed to query available targets: {0}", e.Message);
+				return 1;
+			}
+		}
+		private int QueryTargetDetails(CommandLineArguments Arguments, ILogger Logger, JsonSerializerOptions JsonOptions)
+		{
+			if (TargetName == null || TargetConfiguration == null || TargetPlatform == null)
+			{
+				return 1;
+			}
+
+			GenerateProjectFilesMode.TryParseProjectFileArgument(Arguments, Logger, out FileReference? ProjectFileArg);
+			var RawArgs = new List<string> { TargetName, TargetConfiguration, TargetPlatform };
+			if (ProjectFileArg != null)
+			{
+				RawArgs.Add(ProjectFileArg.ToString());
+			}
+			CommandLineArguments Args = new CommandLineArguments(RawArgs.ToArray());
+			List<TargetDescriptor> TargetDescriptors = new();
+
+			TargetDescriptor.ParseSingleCommandLine(Args, false, false, false, TargetDescriptors, Logger);
+			if (TargetDescriptors.Count == 0)
+			{
+				// TOOD: Error 
+				return 1;
+			}
+
+			HashSet<string> BrowseConfigurationFolders = new HashSet<string>();
+			try
+			{
+				UEBuildTarget CurrentTarget;
+				using (GlobalTracer.Instance.BuildSpan("UEBuildTarget.Create()").StartActive())
+				{
+					bool bUsePrecompiled = false;
+					CurrentTarget = UEBuildTarget.Create(TargetDescriptors[0], false, false, bUsePrecompiled, Logger);
+				}
+
+				var CurrentTargetIntellisenseInfo = new TargetIntellisenseInfo();
+				var CurrentBrowseConfiguration = new GetBrowseConfigurationResponse { Success = true };
+
+				// Partially duplicated from UEBuildTarget.Build because we just want to get C++ compile actions without running UHT
+				// or generating link actions / full dependency graph 
+				CppConfiguration CppConfiguration = UEBuildTarget.GetCppConfiguration(CurrentTarget.Configuration);
+				SourceFileMetadataCache MetadataCache = SourceFileMetadataCache.CreateHierarchy(null, Logger);
+				CppCompileEnvironment GlobalCompileEnvironment = new CppCompileEnvironment(CurrentTarget.Platform, CppConfiguration, CurrentTarget.Architecture, MetadataCache);
+				LinkEnvironment GlobalLinkEnvironment = new LinkEnvironment(GlobalCompileEnvironment.Platform, GlobalCompileEnvironment.Configuration, GlobalCompileEnvironment.Architecture);
+
+				UEToolChain TargetToolChain = CurrentTarget.CreateToolchain(CurrentTarget.Platform);
+				TargetToolChain.SetEnvironmentVariables();
+				CurrentTarget.SetupGlobalEnvironment(TargetToolChain, GlobalCompileEnvironment, GlobalLinkEnvironment);
+
+				// TODO: For installed builds, filter out all the binaries that aren't in mods
+				foreach (UEBuildBinary Binary in CurrentTarget.Binaries)
+				{
+					HashSet<UEBuildModule> LinkEnvironmentVisitedModules = new HashSet<UEBuildModule>();
+					CppCompileEnvironment BinaryCompileEnvironment = Binary.CreateBinaryCompileEnvironment(GlobalCompileEnvironment);
+					CurrentBrowseConfiguration.Standard = BinaryCompileEnvironment.CppStandard.ToString();
+					CurrentBrowseConfiguration.WindowsSdkVersion = CurrentTarget.Rules.WindowsPlatform.WindowsSdkVersion;
+
+					foreach (UEBuildModuleCPP Module in Binary.Modules.OfType<UEBuildModuleCPP>())
+					{
+						if (Module.Binary != null && Module.Binary != Binary)
+						{
+							continue;
+						}
+
+						CppCompileEnvironment ModuleCompileEnvironment = Module.CreateModuleCompileEnvironment(CurrentTarget.Rules, BinaryCompileEnvironment, Logger);
+						foreach (DirectoryReference Dir in Module.ModuleDirectories)
+						{
+							BrowseConfigurationFolders.Add(Dir.ToString());
+							CurrentTargetIntellisenseInfo.DirToModule.TryAdd(Dir, Module);
+						}
+						if (Module.GeneratedCodeDirectory != null)
+						{
+							CurrentTargetIntellisenseInfo.DirToModule.TryAdd(Module.GeneratedCodeDirectory, Module);
+						}
+
+						foreach (DirectoryReference Dir in ModuleCompileEnvironment.SystemIncludePaths)
+						{
+							BrowseConfigurationFolders.Add(Dir.ToString());
+						}
+						foreach (DirectoryReference Dir in ModuleCompileEnvironment.UserIncludePaths)
+						{
+							BrowseConfigurationFolders.Add(Dir.ToString());
+						}
+
+						var Settings = new TargetIntellisenseInfo.CompileSettings();
+						if (OperatingSystem.IsWindows())
+						{
+							if (CurrentTarget.Platform == UnrealTargetPlatform.Win64)
+							{
+								// TODO: Correct compiler
+								Settings.IncludePaths.AddRange(VCToolChain.GetVCIncludePaths(UnrealTargetPlatform.Win64, WindowsCompiler.VisualStudio2022, null, Logger).Split(";"));
+							}
+						}
+						Settings.IncludePaths.AddRange(ModuleCompileEnvironment.SystemIncludePaths.Select(x => x.ToString()));
+						Settings.IncludePaths.AddRange(ModuleCompileEnvironment.UserIncludePaths.Select(x => x.ToString()));
+						Settings.Defines = ModuleCompileEnvironment.Definitions;
+						Settings.Standard = ModuleCompileEnvironment.CppStandard.ToString();
+						Settings.ForcedIncludes = ModuleCompileEnvironment.ForceIncludeFiles.Select(x => x.ToString()).ToList();
+						Settings.CompilerPath = TargetToolChain.GetCppCompilerPath()?.ToString();
+						Settings.WindowsSdkVersion = CurrentTarget.Rules.WindowsPlatform.WindowsSdkVersion;
+						CurrentTargetIntellisenseInfo.ModuleToCompileSettings.Add(Module, Settings);
+					}
+				}
+
+				CurrentBrowseConfiguration.Paths = BrowseConfigurationFolders.ToList();
+
+				var Result = new
+				{
+					DirToModule = CurrentTargetIntellisenseInfo.DirToModule.ToDictionary(x => x.Key.ToString(), x => x.Value.Name),
+					ModuleToCompileSettings = CurrentTargetIntellisenseInfo.ModuleToCompileSettings.ToDictionary(x => x.Key.Name, x => x.Value),
+				};
+
+				Console.WriteLine(JsonSerializer.Serialize(Result, JsonOptions));
+				return 0;
+			}
+			catch (Exception e)
+			{
+				Logger.LogError("Caught exception setting up target: {0}", e);
+				return 1;
+			}
+		}
+	}
+}
diff --git a/Engine/Source/Programs/UnrealBuildTool/Modes/ServerMode.cs b/Engine/Source/Programs/UnrealBuildTool/Modes/ServerMode.cs
--- a/Engine/Source/Programs/UnrealBuildTool/Modes/ServerMode.cs
+++ b/Engine/Source/Programs/UnrealBuildTool/Modes/ServerMode.cs
@@ -96,7 +96,7 @@ namespace UnrealBuildTool
 		public CommandType Type { get { return CommandType.SetTarget; } }
 		public long Sequence { get; set; }
 	}
-	struct QueryFileResponse : ICommandResponse 
+	struct QueryFileResponse : ICommandResponse
 	{
 		[JsonConverter(typeof(JsonStringEnumConverter))]
 		public CommandType Type { get { return CommandType.QueryFile; } }
@@ -110,22 +110,10 @@ namespace UnrealBuildTool
 		public CommandType Type { get { return CommandType.GetBrowseConfiguration; } }
 		public long Sequence { get; set; }
 		public bool Success { get; set; }
-		
-		public List<string> Paths { get; set; }
-		public string? Standard { get; set; }
-		public string? WindowsSdkVersion { get; set; } 
-	}
 
-	class CompileSettings
-	{
-		public List<string> IncludePaths { get; set; } = new();
-		public List<string> Defines { get; set; } = new();
+		public List<string> Paths { get; set; }
 		public string? Standard { get; set; }
-		public List<string> ForcedIncludes { get; set; } = new();
-		public string? CompilerPath { get; set; }
-		public List<string> CompilerArgs { get; set; } = new();
 		public string? WindowsSdkVersion { get; set; }
-
 	}
 
 	class GetCompileSettingsResponse : ICommandResponse
@@ -133,42 +121,12 @@ namespace UnrealBuildTool
 		[JsonConverter(typeof(JsonStringEnumConverter))]
 		public CommandType Type { get { return CommandType.GetCompileSettings; } }
 		public long Sequence { get; set; }
-		public List<CompileSettings?> Settings { get; set; } = new();
+		public List<TargetIntellisenseInfo.CompileSettings?> Settings { get; set; } = new();
 	}
 
-	internal class TargetIntellisenseInfo
-	{
-		public Dictionary<UEBuildModule, CompileSettings> ModuleToCompileSettings = new();
-		public Dictionary<DirectoryReference, UEBuildModule> DirToModule = new();
-
-		public UEBuildModule? FindModuleForFile(FileReference File)
-		{
-			DirectoryReference? Dir = File.Directory;
-			while (Dir != null)
-			{
-				if (DirToModule.TryGetValue(Dir, out UEBuildModule? Module))
-				{
-					return Module;
-				}
-				Dir = Dir.ParentDirectory;
-			}
-			return null;
-		}
-		public UEBuildModule? FindModuleForDirectory(DirectoryReference Directory)
-		{
-			DirectoryReference? Dir = Directory;
-			while (Dir != null)
-			{
-				if (DirToModule.TryGetValue(Dir, out UEBuildModule? Module))
-				{
-					return Module;
-				}
-				Dir = Dir.ParentDirectory;
-			}
-			return null;
-		}
-	}
 
+	// TODO: Deprecate 
+	[Obsolete]
 	[ToolMode("Server", ToolModeOptions.BuildPlatforms | ToolModeOptions.XmlConfig | ToolModeOptions.UseStartupTraceListener)]
 	class ServerMode : ToolMode
 	{
@@ -417,7 +375,7 @@ namespace UnrealBuildTool
 							BrowseConfigurationFolders.Add(Dir.ToString());
 						}
 
-						CompileSettings Settings = new CompileSettings();
+						var Settings = new TargetIntellisenseInfo.CompileSettings();
 						Settings.IncludePaths.AddRange(ModuleCompileEnvironment.SystemIncludePaths.Select(x => x.ToString()));
 						Settings.IncludePaths.AddRange(ModuleCompileEnvironment.UserIncludePaths.Select(x => x.ToString()));
 						Settings.Defines = ModuleCompileEnvironment.Definitions;
@@ -500,7 +458,7 @@ namespace UnrealBuildTool
 				GetCompileSettingsResponse Response = new GetCompileSettingsResponse { };
 				foreach (string AbsolutePath in Command.AbsolutePaths)
 				{
-					CompileSettings? Settings = null;
+					TargetIntellisenseInfo.CompileSettings? Settings = null;
 					var Path = System.IO.Path.GetFullPath(AbsolutePath); // convert absolute path to dos? 
 					UEBuildModule? Module = CurrentTargetIntellisenseInfo!.FindModuleForFile(new FileReference(Path));
 					if (Module is not null)
diff --git a/Engine/Source/Programs/UnrealBuildTool/ProjectFiles/VisualStudioCode/VSCodeProjectFileGenerator.cs b/Engine/Source/Programs/UnrealBuildTool/ProjectFiles/VisualStudioCode/VSCodeProjectFileGenerator.cs
--- a/Engine/Source/Programs/UnrealBuildTool/ProjectFiles/VisualStudioCode/VSCodeProjectFileGenerator.cs
+++ b/Engine/Source/Programs/UnrealBuildTool/ProjectFiles/VisualStudioCode/VSCodeProjectFileGenerator.cs
@@ -61,6 +61,13 @@ namespace UnrealBuildTool
 		[CommandLine("-NoCompileCommands")]
 		private bool bNoCompileCommands = false;
 
+		/// <summary>
+		/// Create a workspace file for use with VS Code extension that communicates directly with UBT. 
+		/// </summary>
+		[XmlConfigFile(Name = "UseVSCodeExtension")]
+		[CommandLine("-UseVSCodeExtension")]
+		private bool bUseVSCodeExtension = false;
+
 		private enum EPathType
 		{
 			Absolute,
@@ -246,7 +253,7 @@ namespace UnrealBuildTool
 
 		public override bool ShouldGenerateIntelliSenseData()
 		{
-			return !bNoCompileCommands;
+			return !bNoCompileCommands && !bUseVSCodeExtension;
 		}
 
 		protected override ProjectFile AllocateProjectFile(FileReference InitFilePath, DirectoryReference BaseDir)
@@ -734,11 +741,17 @@ namespace UnrealBuildTool
 				OutFile.AddField("intelliSenseMode", "clang-x64");
 			}
 
-			OutFile.AddField("configurationProvider", "epic.ue");
+			if (bUseVSCodeExtension)
+			{
+				OutFile.AddField("configurationProvider", "epic.ue");
+			}
 
-			FileReference CompileCommands = FileReference.Combine(OutputDirectory, string.Format("compileCommands_{0}.json", ProjectName));
-			WriteCompileCommands(CompileCommands, SourceFiles, CompilerPath!, ModuleCommandLines);
-			OutFile.AddField("compileCommands", MakePathString(CompileCommands, bInAbsolute: true, bForceSkipQuotes: true));
+			if (ShouldGenerateIntelliSenseData())
+			{
+				FileReference CompileCommands = FileReference.Combine(OutputDirectory, string.Format("compileCommands_{0}.json", ProjectName));
+				WriteCompileCommands(CompileCommands, SourceFiles, CompilerPath!, ModuleCommandLines);
+				OutFile.AddField("compileCommands", MakePathString(CompileCommands, bInAbsolute: true, bForceSkipQuotes: true));
+			}
 
 			OutFile.EndObject();
 		}
@@ -1058,9 +1071,12 @@ namespace UnrealBuildTool
 
 				OutFile.BeginArray("tasks");
 				{
-					foreach (ProjectData.Project NativeProject in ProjectData.NativeProjects)
+					if (!bUseVSCodeExtension)
 					{
-						WriteNativeTask(NativeProject, OutFile);
+						foreach (ProjectData.Project NativeProject in ProjectData.NativeProjects)
+						{
+							WriteNativeTask(NativeProject, OutFile);
+						}
 					}
 
 					foreach (ProjectData.Project CSharpProject in ProjectData.CSharpProjects)
@@ -1468,9 +1484,12 @@ namespace UnrealBuildTool
 
 				OutFile.BeginArray("configurations");
 				{
-					foreach (ProjectData.Project Project in ProjectData.NativeProjects)
+					if (!bUseVSCodeExtension)
 					{
-						WriteNativeLaunchConfig(Project, OutFile, Logger);
+						foreach (ProjectData.Project Project in ProjectData.NativeProjects)
+						{
+							WriteNativeLaunchConfig(Project, OutFile, Logger);
+						}
 					}
 
 					foreach (ProjectData.Project Project in ProjectData.CSharpProjects)
@@ -1623,6 +1642,21 @@ namespace UnrealBuildTool
 			{
 				// disable autodetect for typescript files to workaround slowdown in vscode as a result of parsing all files
 				WorkspaceFile.AddField("typescript.tsc.autoDetect", "off");
+				// disable npm script autodetect to avoid lag populating tasks list 
+				WorkspaceFile.AddField("npm.autoDetect", "off");
+
+				if (bUseVSCodeExtension)
+				{
+					WorkspaceFile.AddField("UE.loggingLevel", "Warning");
+					if (HostPlatform == UnrealTargetPlatform.Win64)
+					{
+						WorkspaceFile.AddField("UE.UBTScriptPath", MakePathString(FileReference.Combine(ProjectRoot, "Engine", "Build", "BatchFiles", "Build.bat"), true));
+					}
+					else
+					{
+						WorkspaceFile.AddField("UE.UBTScriptPath", MakePathString(FileReference.Combine(ProjectRoot, "Engine", "Build", "BatchFiles", HostPlatform.ToString(), "Build.sh"), true));
+					}
+				}
 			}
 			WorkspaceFile.EndObject();
 			
@@ -1637,6 +1671,11 @@ namespace UnrealBuildTool
 					WorkspaceFile.AddUnnamedField("ms-vscode.cpptools");
 					WorkspaceFile.AddUnnamedField("ms-dotnettools.csharp");
 
+					if (bUseVSCodeExtension)
+					{
+						WorkspaceFile.AddUnnamedField("epic.vscode-ue");
+					}
+
 					// If the platform we run the generator on uses mono, there are additional debugging extensions to add.
 					if (!RuntimePlatform.IsWindows)
 					{
