Architecture -> Architectures

diff --git a/Engine/Source/Programs/UnrealBuildTool/Modes/QueryMode.cs b/Engine/Source/Programs/UnrealBuildTool/Modes/QueryMode.cs
--- a/Engine/Source/Programs/UnrealBuildTool/Modes/QueryMode.cs
+++ b/Engine/Source/Programs/UnrealBuildTool/Modes/QueryMode.cs
@@ -266,8 +266,8 @@ namespace UnrealBuildTool
 				// or generating link actions / full dependency graph 
 				CppConfiguration CppConfiguration = UEBuildTarget.GetCppConfiguration(CurrentTarget.Configuration);
 				SourceFileMetadataCache MetadataCache = SourceFileMetadataCache.CreateHierarchy(null, Logger);
-				CppCompileEnvironment GlobalCompileEnvironment = new CppCompileEnvironment(CurrentTarget.Platform, CppConfiguration, CurrentTarget.Architecture, MetadataCache);
-				LinkEnvironment GlobalLinkEnvironment = new LinkEnvironment(GlobalCompileEnvironment.Platform, GlobalCompileEnvironment.Configuration, GlobalCompileEnvironment.Architecture);
+				CppCompileEnvironment GlobalCompileEnvironment = new CppCompileEnvironment(CurrentTarget.Platform, CppConfiguration, CurrentTarget.Architectures, MetadataCache);
+				LinkEnvironment GlobalLinkEnvironment = new LinkEnvironment(GlobalCompileEnvironment.Platform, GlobalCompileEnvironment.Configuration, GlobalCompileEnvironment.Architectures);
 
 				UEToolChain TargetToolChain = CurrentTarget.CreateToolchain(CurrentTarget.Platform);
 				TargetToolChain.SetEnvironmentVariables();
