- Fixed a Mac-only test that was running on Windows/Linux

diff --git a/Engine/Source/Programs/UnrealBuildTool/Platform/Mac/UEBuildMac.cs b/Engine/Source/Programs/UnrealBuildTool/Platform/Mac/UEBuildMac.cs
--- a/Engine/Source/Programs/UnrealBuildTool/Platform/Mac/UEBuildMac.cs
+++ b/Engine/Source/Programs/UnrealBuildTool/Platform/Mac/UEBuildMac.cs
@@ -223,7 +223,7 @@ namespace UnrealBuildTool
 			}
 			else if (DefaultArchitecture.Contains("host"))
 			{
-				if (MacExports.IsRunningOnAppleArchitecture && bSupportsArm64)
+				if (BuildHostPlatform.Current.Platform == UnrealTargetPlatform.Mac && MacExports.IsRunningOnAppleArchitecture && bSupportsArm64)
 				{
 					Architectures.Add(UnrealArch.Arm64);
 				}
