UAT: Parse ini overrides when reading config values from CreatePaks in CopyBuildToStaginDirectory.
Add ConfigOverrideParams to ProjectParams constructor.
Add OverrideStrings to ConfigHierarchyKey to get the right data from ConfigHierarchy.ReadHierarchy().

#rnx

diff --git a/Engine/Source/Programs/AutomationTool/AutomationUtils/ProjectParams.cs b/Engine/Source/Programs/AutomationTool/AutomationUtils/ProjectParams.cs
--- a/Engine/Source/Programs/AutomationTool/AutomationUtils/ProjectParams.cs
+++ b/Engine/Source/Programs/AutomationTool/AutomationUtils/ProjectParams.cs
@@ -298,6 +298,7 @@ namespace AutomationTool
             this.ClientDependentPlatformMap = InParams.ClientDependentPlatformMap;
 			this.ServerTargetPlatforms = InParams.ServerTargetPlatforms;
             this.ServerDependentPlatformMap = InParams.ServerDependentPlatformMap;
+			this.ConfigOverrideParams = InParams.ConfigOverrideParams;
 			this.Build = InParams.Build;
 			this.SkipBuildClient = InParams.SkipBuildClient;
 			this.SkipBuildEditor = InParams.SkipBuildEditor;
@@ -454,6 +455,7 @@ namespace AutomationTool
             Dictionary<TargetPlatformDescriptor, TargetPlatformDescriptor> ClientDependentPlatformMap = null,
 			List<TargetPlatformDescriptor> ServerTargetPlatforms = null,
             Dictionary<TargetPlatformDescriptor, TargetPlatformDescriptor> ServerDependentPlatformMap = null,
+			List<string> ConfigOverrideParams = null,
 			bool? Build = null,
 			bool? SkipBuildClient = null,
 			bool? SkipBuildEditor = null,
@@ -1237,6 +1239,10 @@ namespace AutomationTool
                 this.RunTimeoutSeconds = Command.ParseParamInt("runtimeoutseconds");
             }
 
+			if (ConfigOverrideParams != null)
+			{
+				this.ConfigOverrideParams = ConfigOverrideParams;
+			}
 			// Gather up any '-ini:' arguments and save them. We'll pass these along to other tools that may be spawned in a new process as part of the command.
 			if(Command != null)
 			{
@@ -3087,6 +3093,7 @@ namespace AutomationTool
 				CommandUtils.LogLog("CreateAppBundle={0}", CreateAppBundle);
 				CommandUtils.LogLog("BaseArchiveDirectory={0}", BaseArchiveDirectory);
 				CommandUtils.LogLog("BaseStageDirectory={0}", BaseStageDirectory);
+				CommandUtils.LogLog("ConfigOverrideParams=-{0}", string.Join(" -", ConfigOverrideParams));
 				CommandUtils.LogLog("Build={0}", Build);
 				CommandUtils.LogLog("SkipBuildClient={0}", SkipBuildClient);
 				CommandUtils.LogLog("SkipBuildEditor={0}", SkipBuildEditor);
diff --git a/Engine/Source/Programs/AutomationTool/Scripts/CopyBuildToStagingDirectory.Automation.cs b/Engine/Source/Programs/AutomationTool/Scripts/CopyBuildToStagingDirectory.Automation.cs
--- a/Engine/Source/Programs/AutomationTool/Scripts/CopyBuildToStagingDirectory.Automation.cs
+++ b/Engine/Source/Programs/AutomationTool/Scripts/CopyBuildToStagingDirectory.Automation.cs
@@ -2660,7 +2660,12 @@ namespace AutomationScripts
 
 			// read some compression settings from the project (once, shared across all pak commands)
 
-			ConfigHierarchy PlatformGameConfig = ConfigCache.ReadHierarchy(ConfigHierarchyType.Game, DirectoryReference.FromFile(Params.RawProjectPath), SC.StageTargetPlatform.IniPlatformType, SC.CustomConfig);
+			ConfigHierarchy PlatformGameConfig = ConfigCache.ReadHierarchy(
+					ConfigHierarchyType.Game,
+					DirectoryReference.FromFile(Params.RawProjectPath),
+					SC.StageTargetPlatform.IniPlatformType,
+					SC.CustomConfig,
+					Params.ConfigOverrideParams.ToArray());
 
 			// in standard runs, ProjectPackagingSettings/bCompressed is read by the program invoking this script
 			//	and used to pass "-compressed" on the command line
@@ -2793,8 +2798,12 @@ namespace AutomationScripts
 				// it's the DLL to use during packaging, not during runtime, so it needs to be for the platform the packager runs on, not the platform you're packaging for.
 				// Most of the compression options come from the platform you're packaging FOR
 
-				ConfigHierarchy BuildPlatformEngineConfig = ConfigCache.ReadHierarchy(ConfigHierarchyType.Engine, DirectoryReference.FromFile(Params.RawProjectPath),
-					BuildHostPlatform.Current.Platform);
+				ConfigHierarchy BuildPlatformEngineConfig = ConfigCache.ReadHierarchy(
+					ConfigHierarchyType.Engine,
+					DirectoryReference.FromFile(Params.RawProjectPath),
+					BuildHostPlatform.Current.Platform,
+					SC.CustomConfig,
+					Params.ConfigOverrideParams.ToArray());
 
 				string OodleCompressDLL;
 				BuildPlatformEngineConfig.GetString("OodleDataCompressionFormat", "OodleCompressDLL", out OodleCompressDLL);
diff --git a/Engine/Source/Programs/UnrealBuildTool/System/ConfigCache.cs b/Engine/Source/Programs/UnrealBuildTool/System/ConfigCache.cs
--- a/Engine/Source/Programs/UnrealBuildTool/System/ConfigCache.cs
+++ b/Engine/Source/Programs/UnrealBuildTool/System/ConfigCache.cs
@@ -170,6 +170,11 @@ namespace UnrealBuildTool
 			/// </summary>
 			public string CustomConfig;
 
+			/// <summary>
+			/// Custom override strings
+			/// </summary>
+			public List<string> OverrideStrings;
+
 			/// <summary>
 			/// Constructor
 			/// </summary>
@@ -177,12 +182,14 @@ namespace UnrealBuildTool
 			/// <param name="ProjectDir">The project directory to read from</param>
 			/// <param name="Platform">Which platform-specific files to read</param>
 			/// <param name="CustomConfig">Custom config subdirectory to load</param>
-			public ConfigHierarchyKey(ConfigHierarchyType Type, DirectoryReference? ProjectDir, UnrealTargetPlatform Platform, string CustomConfig)
+			/// <param name="OverrideStrings">Custom override strings</param>
+			public ConfigHierarchyKey(ConfigHierarchyType Type, DirectoryReference? ProjectDir, UnrealTargetPlatform Platform, string CustomConfig, List<string> OverrideStrings)
 			{
 				this.Type = Type;
 				this.ProjectDir = ProjectDir;
 				this.Platform = Platform;
 				this.CustomConfig = CustomConfig;
+				this.OverrideStrings = OverrideStrings;
 			}
 
 			/// <summary>
@@ -193,7 +200,12 @@ namespace UnrealBuildTool
 			public override bool Equals(object? Other)
 			{
 				ConfigHierarchyKey? OtherKey = Other as ConfigHierarchyKey;
-				return (OtherKey != null && OtherKey.Type == Type && OtherKey.ProjectDir == ProjectDir && OtherKey.Platform == Platform && OtherKey.CustomConfig == CustomConfig);
+				return OtherKey != null &&
+					OtherKey.Type == Type &&
+					OtherKey.ProjectDir == ProjectDir &&
+					OtherKey.Platform == Platform &&
+					OtherKey.CustomConfig == CustomConfig &&
+					OtherKey.OverrideStrings != OverrideStrings;
 			}
 
 			/// <summary>
@@ -202,7 +214,11 @@ namespace UnrealBuildTool
 			/// <returns>Hash value for this object</returns>
 			public override int GetHashCode()
 			{
-				return ((ProjectDir != null) ? ProjectDir.GetHashCode() : 0) + ((int)Type * 123) + (Platform.GetHashCode() * 345) + (CustomConfig.GetHashCode() * 789);
+				return HashCode.Combine(
+					(ProjectDir != null) ? ProjectDir.GetHashCode() : 0,
+					Platform.GetHashCode(),
+					CustomConfig.GetHashCode(),
+					OverrideStrings.GetHashCode());
 			}
 		}
 
@@ -261,10 +277,20 @@ namespace UnrealBuildTool
 		{
 			// Handle command line overrides
 			List<String> OverrideStrings = new List<String>();
-			string[] CmdLine = Environment.GetCommandLineArgs();
+			List<String> CmdLine = Environment.GetCommandLineArgs().ToList();
 			if (CustomArgs != null)
 			{
-				CmdLine = CmdLine.Concat(CustomArgs).ToArray();
+				foreach (string CustomArg in CustomArgs)
+				{
+					if (CustomArg.StartsWith("-", StringComparison.InvariantCultureIgnoreCase))
+					{
+						CmdLine.Add(CustomArg);
+					}
+					else
+					{
+						CmdLine.Add("-" + CustomArg);
+					}
+				}
 			}
 			string IniConfigArgPrefix = "-ini:" + Enum.GetName(typeof(ConfigHierarchyType), Type) + ":";
 			string CustomConfigPrefix = "-CustomConfig=";
@@ -286,7 +312,7 @@ namespace UnrealBuildTool
 			}
 
 			// Get the key to use for the cache. It cannot be null, so we use the engine directory if a project directory is not given.
-			ConfigHierarchyKey Key = new ConfigHierarchyKey(Type, ProjectDir, Platform, CustomConfig);
+			ConfigHierarchyKey Key = new ConfigHierarchyKey(Type, ProjectDir, Platform, CustomConfig, OverrideStrings);
 
 			// Try to get the cached hierarchy with this key
 			ConfigHierarchy? Hierarchy;
