Compile fix for Android LowLevelTests

diff --git a/Engine/Source/Programs/UnrealBuildTool/Platform/Android/AndroidToolChain.cs b/Engine/Source/Programs/UnrealBuildTool/Platform/Android/AndroidToolChain.cs
--- a/Engine/Source/Programs/UnrealBuildTool/Platform/Android/AndroidToolChain.cs
+++ b/Engine/Source/Programs/UnrealBuildTool/Platform/Android/AndroidToolChain.cs
@@ -955,16 +955,6 @@ namespace UnrealBuildTool
 			return Environment.GetEnvironmentVariable("NDKROOT") + "/sources/android/cpufeatures/cpu-features.c";
 		}
 
-		protected virtual void ModifySourceFiles(CppCompileEnvironment CompileEnvironment, List<FileItem> SourceFiles, string ModuleName)
-		{
-			// We need to add the extra glue and cpu code only to Launch module.
-			if (ModuleName.Equals("Launch") || ModuleName.Equals("AndroidLauncher"))
-			{
-				SourceFiles.Add(FileItem.GetItemByPath(GetNativeGluePath()));
-				SourceFiles.Add(FileItem.GetItemByPath(GetCpuFeaturesPath()));
-			}
-		}
-
 		public override CppCompileEnvironment CreateSharedResponseFile(CppCompileEnvironment CompileEnvironment, FileReference OutResponseFile, IActionGraphBuilder Graph)
 		{
 			// Seems like Android clang toolchain does not handle response files including response files
@@ -1030,6 +1020,8 @@ namespace UnrealBuildTool
 		}
 
 		private bool bHasHandledLaunchModule = false;
+		static private bool bHasHandledCoreModule = false;
+
 		protected override CPPOutput CompileCPPFiles(CppCompileEnvironment CompileEnvironment, List<FileItem> InputFiles, DirectoryReference OutputDir, string ModuleName, IActionGraphBuilder Graph)
 		{
 			if (ShouldSkipModule(ModuleName, CompileEnvironment.Architecture))
@@ -1042,13 +1034,20 @@ namespace UnrealBuildTool
 			if (!bHasHandledLaunchModule && (ModuleName.Equals("Launch") || ModuleName.Equals("AndroidLauncher")))
 			{
 				// Directly added NDK files for NDK extensions
-				ModifySourceFiles(CompileEnvironment, InputFiles, ModuleName);
+				InputFiles.Add(FileItem.GetItemByPath(GetNativeGluePath()));
 				// Deal with dynamic modules removed by architecture
 				GenerateEmptyLinkFunctionsForRemovedModules(InputFiles, CompileEnvironment.Architecture, ModuleName, OutputDir, Graph, Logger);
 
 				bHasHandledLaunchModule = true;
 			}
 
+			if (!bHasHandledCoreModule && ModuleName.Equals("Core"))
+			{
+				// This is used by Crypto code in Core
+				InputFiles.Add(FileItem.GetItemByPath(GetCpuFeaturesPath()));
+				bHasHandledCoreModule = true;
+			}
+
 
 			return base.CompileCPPFiles(CompileEnvironment, InputFiles, OutputDir, ModuleName, Graph);
 		}
