Fixed Android compilation

diff --git a/Engine/Source/Programs/UnrealBuildTool/Platform/Android/AndroidToolChain.cs b/Engine/Source/Programs/UnrealBuildTool/Platform/Android/AndroidToolChain.cs
--- a/Engine/Source/Programs/UnrealBuildTool/Platform/Android/AndroidToolChain.cs
+++ b/Engine/Source/Programs/UnrealBuildTool/Platform/Android/AndroidToolChain.cs
@@ -1020,7 +1020,7 @@ namespace UnrealBuildTool
 		}
 
 		private bool bHasHandledLaunchModule = false;
-		static private bool bHasHandledCoreModule = false;
+		private bool bHasHandledCoreModule = false;
 
 		protected override CPPOutput CompileCPPFiles(CppCompileEnvironment CompileEnvironment, List<FileItem> InputFiles, DirectoryReference OutputDir, string ModuleName, IActionGraphBuilder Graph)
 		{
@@ -1029,7 +1029,6 @@ namespace UnrealBuildTool
 				return new CPPOutput();
 			}
 
-
 			// Deal with Launch module special if first time seen
 			if (!bHasHandledLaunchModule && (ModuleName.Equals("Launch") || ModuleName.Equals("AndroidLauncher")))
 			{
@@ -1041,14 +1040,13 @@ namespace UnrealBuildTool
 				bHasHandledLaunchModule = true;
 			}
 
-			if (!bHasHandledCoreModule && ModuleName.Equals("Core"))
+			if (!bHasHandledCoreModule && ModuleName.Equals("Core") && (CompileEnvironment.PrecompiledHeaderAction == PrecompiledHeaderAction.None))
 			{
 				// This is used by Crypto code in Core
 				InputFiles.Add(FileItem.GetItemByPath(GetCpuFeaturesPath()));
 				bHasHandledCoreModule = true;
 			}
 
-
 			return base.CompileCPPFiles(CompileEnvironment, InputFiles, OutputDir, ModuleName, Graph);
 		}
 
