Add support to deprecate when move key to another section in the same .ini config file.

diff --git a/Engine/Config/ConfigRedirects.ini b/Engine/Config/ConfigRedirects.ini
--- a/Engine/Config/ConfigRedirects.ini
+++ b/Engine/Config/ConfigRedirects.ini
@@ -47,4 +47,7 @@ EncryptionKey=ClientEncryptionKey
 EncryptionKey=ClientEncryptionKey
 
 [/Script/OnlineSubsystemEOS.EOSSettings]
-EncryptionKey=ClientEncryptionKey
\ No newline at end of file
+EncryptionKey=ClientEncryptionKey
+
+[HTTP]
+AllowedDomains=Online.HttpManager https:AllowedDomains
\ No newline at end of file
diff --git a/Engine/Source/Programs/UnrealBuildTool/System/ConfigFile.cs b/Engine/Source/Programs/UnrealBuildTool/System/ConfigFile.cs
--- a/Engine/Source/Programs/UnrealBuildTool/System/ConfigFile.cs
+++ b/Engine/Source/Programs/UnrealBuildTool/System/ConfigFile.cs
@@ -145,7 +145,7 @@ namespace UnrealBuildTool
 		static private Dictionary<string, Dictionary<string, string>> SectionKeyRemap = new();
 		static private HashSet<string> WarnedKeys = new(StringComparer.InvariantCultureIgnoreCase);
 
-		static private string Remap(Dictionary<string, string>? Remap, string Key, string Context)
+		static private string RemapSectionOrKey(Dictionary<string, string>? Remap, string Key, string Context)
 		{
 			if (Remap != null)
 			{
@@ -286,7 +286,7 @@ namespace UnrealBuildTool
 									string SectionName = Line.Substring(StartIdx + 1, EndIdx - StartIdx - 2);
 									
 									// lookup remaps
-									SectionName = Remap(SectionNameRemap, SectionName, $"which is a config section in '{Location.FullName}'");
+									SectionName = RemapSectionOrKey(SectionNameRemap, SectionName, $"which is a config section in '{Location.FullName}'");
 									SectionKeyRemap.TryGetValue(SectionName, out CurrentRemap);
 
 									if (!Sections.TryGetValue(SectionName, out CurrentSection))
@@ -301,7 +301,7 @@ namespace UnrealBuildTool
 							// Otherwise add it to the current section or add a new one
 							if(CurrentSection != null)
 							{
-								TryAddConfigLine(CurrentSection, CurrentRemap, Location.FullName, Line, StartIdx, EndIdx, DefaultAction);
+								TryAddConfigLine(CurrentSection, CurrentRemap, Location.FullName, Line, StartIdx, EndIdx, DefaultAction, Sections);
 								break;
 							}
 
@@ -329,7 +329,7 @@ namespace UnrealBuildTool
 				// Locate and break off the section name
 				string SectionName = Setting.Remove(Setting.IndexOf(':')).Trim(new char[] { '[', ']' });
 				// lookup remaps
-				Remap(SectionNameRemap, SectionName, $"which is a config section (found in a string in code - search your .ini files for it, the source file is unknown)");
+				RemapSectionOrKey(SectionNameRemap, SectionName, $"which is a config section (found in a string in code - search your .ini files for it, the source file is unknown)");
 				SectionKeyRemap.TryGetValue(SectionName, out CurrentRemap);
 
 				if (SectionName.Length > 0)
@@ -344,7 +344,7 @@ namespace UnrealBuildTool
 					if (CurrentSection != null)
 					{
 						string IniKeyValue = Setting.Substring(Setting.IndexOf(':') + 1);
-						TryAddConfigLine(CurrentSection, CurrentRemap, "unknown source file", IniKeyValue, 0, IniKeyValue.Length, DefaultAction);
+						TryAddConfigLine(CurrentSection, CurrentRemap, "unknown source file", IniKeyValue, 0, IniKeyValue.Length, DefaultAction, Sections);
 					}
 				}
 			}
@@ -360,8 +360,9 @@ namespace UnrealBuildTool
 		/// <param name="StartIdx">Index of the first non-whitespace character in this line</param>
 		/// <param name="EndIdx">Index of the last (exclusive) non-whitespace character in this line</param>
 		/// <param name="DefaultAction">The default action to take if '+' or '-' is not specified on a config line</param>
+		/// <param name="Sections">The sections to find the redirected section for receiving the config line</param>
 		/// <returns>True if the line was parsed correctly, false otherwise</returns>
-		static bool TryAddConfigLine(ConfigFileSection Section, Dictionary<string, string>? KeyRemap, string Filename, string Line, int StartIdx, int EndIdx, ConfigLineAction DefaultAction)
+		static bool TryAddConfigLine(ConfigFileSection Section, Dictionary<string, string>? KeyRemap, string Filename, string Line, int StartIdx, int EndIdx, ConfigLineAction DefaultAction, Dictionary<string, ConfigFileSection> Sections)
 		{
 			// Find the '=' character separating key and value
 			int EqualsIdx = Line.IndexOf('=', StartIdx, EndIdx - StartIdx);
@@ -431,9 +432,26 @@ namespace UnrealBuildTool
 			string Value = Line.Substring(ValueStartIdx, ValueEndIdx - ValueStartIdx);
 
 			// remap the key if needed
-			Key = Remap(KeyRemap, Key, $"which is a config key in section [{Section.Name}], in '{Filename}'");
+			string NewKey = RemapSectionOrKey(KeyRemap, Key, $"which is a config key in section [{Section.Name}], in '{Filename}'");
 
-			Section.Lines.Add(new ConfigLine(Action, Key, Value));
+			// look for a section:name remap
+			if (!NewKey.Equals(Key) && NewKey.IndexOf(":") != -1)
+			{
+				string SectionName = NewKey.Substring(0, NewKey.IndexOf(':'));
+				ConfigFileSection? CurrentSection;
+				if (!Sections.TryGetValue(SectionName, out CurrentSection))
+				{
+					CurrentSection = new ConfigFileSection(SectionName);
+					Sections.Add(SectionName, CurrentSection);
+				}
+
+				string KeyName = NewKey.Substring(NewKey.IndexOf(':') + 1);
+				CurrentSection.Lines.Add(new ConfigLine(Action, KeyName, Value));
+
+				return true;
+			}
+
+			Section.Lines.Add(new ConfigLine(Action, NewKey, Value));
 			return true;
 		}
 
diff --git a/Engine/Source/Runtime/Core/Private/Misc/ConfigCacheIni.cpp b/Engine/Source/Runtime/Core/Private/Misc/ConfigCacheIni.cpp
--- a/Engine/Source/Runtime/Core/Private/Misc/ConfigCacheIni.cpp
+++ b/Engine/Source/Runtime/Core/Private/Misc/ConfigCacheIni.cpp
@@ -836,6 +836,7 @@ void FConfigFile::CombineFromBuffer(const FString& Buffer, const FString& FileHi
 				}
 
 				const TCHAR* KeyName = Start;
+				FConfigSection* OriginalCurrentSection = CurrentSection;
 				// look up for key remap
 				if (CurrentKeyRemap != nullptr)
 				{
@@ -846,6 +847,16 @@ void FConfigFile::CombineFromBuffer(const FString& Buffer, const FString& FileHi
 
 						// the Remap will not ever reallocate, so we can just point right into the FString
 						KeyName = **FoundRemap;
+
+						// look for a section:name remap
+						int32 ColonLoc;
+						if (FoundRemap->FindChar(':', ColonLoc))
+						{
+							// find or create a section for name before the :
+							CurrentSection = FindOrAddSection(*FoundRemap->Mid(0, ColonLoc));
+							// the name can still point right into the FString, but right after the :
+							KeyName = **FoundRemap + ColonLoc + 1;
+						}
 					}
 				}
 				
@@ -921,6 +932,9 @@ void FConfigFile::CombineFromBuffer(const FString& Buffer, const FString& FileHi
 						}
 					}
 				}
+				
+				// restore the current section, in case it was overridden
+				CurrentSection = OriginalCurrentSection;
 
 				// Mark as dirty so "Write" will actually save the changes.
 				Dirty = true;
