- Fixed legacy Xcode project commandlines

diff --git a/Engine/Source/Programs/UnrealBuildTool/ProjectFiles/Xcode/XcodeProjectLegacy.cs b/Engine/Source/Programs/UnrealBuildTool/ProjectFiles/Xcode/XcodeProjectLegacy.cs
--- a/Engine/Source/Programs/UnrealBuildTool/ProjectFiles/Xcode/XcodeProjectLegacy.cs
+++ b/Engine/Source/Programs/UnrealBuildTool/ProjectFiles/Xcode/XcodeProjectLegacy.cs
@@ -1662,7 +1662,7 @@ namespace UnrealBuildTool.XcodeProjectLegacy
 			Content.Append("      </BuildableProductRunnable>" + ProjectFileGenerator.NewLine);
 			if (string.IsNullOrEmpty(OldCommandLineArguments))
 			{
-				if (bHasEditorConfiguration && TargetName != "UE5")
+				if (bHasEditorConfiguration && GameProjectPath != "")
 				{
 					Content.Append("      <CommandLineArguments>" + ProjectFileGenerator.NewLine);
 					if (IsForeignProject)
@@ -1675,7 +1675,7 @@ namespace UnrealBuildTool.XcodeProjectLegacy
 					else
 					{
 						Content.Append("         <CommandLineArgument" + ProjectFileGenerator.NewLine);
-						Content.Append("            argument = \"" + TargetName + "\"" + ProjectFileGenerator.NewLine);
+						Content.Append("            argument = \"" + Path.GetFileNameWithoutExtension(GameProjectPath) + "\"" + ProjectFileGenerator.NewLine);
 						Content.Append("            isEnabled = \"YES\">" + ProjectFileGenerator.NewLine);
 						Content.Append("         </CommandLineArgument>" + ProjectFileGenerator.NewLine);
 					}
