UnrealBuildTool: Attempt to auto-recover some additional errors when building with XGE

#rnx

diff --git a/Engine/Source/Programs/UnrealBuildTool/Executors/XGE.cs b/Engine/Source/Programs/UnrealBuildTool/Executors/XGE.cs
--- a/Engine/Source/Programs/UnrealBuildTool/Executors/XGE.cs
+++ b/Engine/Source/Programs/UnrealBuildTool/Executors/XGE.cs
@@ -81,6 +81,18 @@ namespace UnrealBuildTool
 
 		private const string ProgressMarkupPrefix = "@action";
 
+		private static List<string> CompileAutoRecover = new List<string> {
+			"C1060", // C1060: compiler is out of heap space
+			"C1076", // C1076: compiler limit: internal heap limit reached
+			"C2855", // C2855: command-line option 'X' inconsistent with precompiled header
+			"C3435", // C3435: character set 'X' is not supported
+			"C3859", // C3859: Failed to create virtual memory for PCH
+		};
+
+		private static List<string> LinkAutoRecover = new List<string> {
+			"Unexpected PDB error; OK (0)"
+		};
+
 		public XGE()
 		{
 			XmlConfig.ApplyTo(this);
@@ -529,14 +541,11 @@ namespace UnrealBuildTool
 
 				if (Action.ActionType == ActionType.Compile)
 				{
-					// C1060: compiler is out of heap space
-					// C1076: compiler limit: internal heap limit reached
-					// C3859: Failed to create virtual memory for PCH
-					ToolElement.SetAttribute("AutoRecover", "C1060,C1076,C3859");
+					ToolElement.SetAttribute("AutoRecover", string.Join(',', CompileAutoRecover));
 				}
 				else if (Action.ActionType == ActionType.Link)
 				{
-					ToolElement.SetAttribute("AutoRecover", "Unexpected PDB error; OK (0)");
+					ToolElement.SetAttribute("AutoRecover", string.Join(',', LinkAutoRecover));
 				}
 			}
 
