[UnrealBuildTool]
* Added so xge path is showing command description (compile/link etc) in log

Example:

5% [199/4110] Compile Module.UI.12_of_40.cpp (1:06.00 at +0:30)
5% [200/4110] Compile Module.UI.38_of_40.cpp (0:54.86 at +0:42)
5% [201/4110] Link UnrealEditor-TargetPlatform.dll (0:01.42 at +1:35)
5% [202/4110] Link UnrealEditor-CommonInput.dll (0:01.46 at +1:35)

diff --git a/Engine/Source/Programs/UnrealBuildTool/Executors/XGE.cs b/Engine/Source/Programs/UnrealBuildTool/Executors/XGE.cs
--- a/Engine/Source/Programs/UnrealBuildTool/Executors/XGE.cs
+++ b/Engine/Source/Programs/UnrealBuildTool/Executors/XGE.cs
@@ -405,7 +405,7 @@ namespace UnrealBuildTool
 			string XGETaskFilePath = FileReference.Combine(Unreal.EngineDirectory, "Intermediate", "Build", "XGETasks.xml").FullName;
 			WriteTaskFile(Actions, XGETaskFilePath, true, false, Logger);
 
-			return ExecuteTaskFileWithProgressMarkup(XGETaskFilePath, Actions.Count(), Logger);
+			return ExecuteTaskFileWithProgressMarkup(XGETaskFilePath, Actions.ToArray(), Logger);
 		}
 
 		/// <summary>
@@ -702,12 +702,14 @@ namespace UnrealBuildTool
 		/// <summary>
 		/// Executes the tasks in the specified file, parsing progress markup as part of the output.
 		/// </summary>
-		bool ExecuteTaskFileWithProgressMarkup(string TaskFilePath, int NumActions, ILogger Logger)
+		bool ExecuteTaskFileWithProgressMarkup(string TaskFilePath, LinkedAction[] Actions, ILogger Logger)
 		{
+			int NumActions = Actions.Length;
 			using (ProgressWriter Writer = new ProgressWriter("Compiling C++ source files...", false, Logger))
 			{
 				int NumCompletedActions = 0;
 				string ProgressText = string.Empty;
+				string CommandDescription = string.Empty;
 				HashSet<int> ReportedActionIndices = new();
 
 				// Create a wrapper delegate that will parse the output actions
@@ -719,12 +721,12 @@ namespace UnrealBuildTool
 						if (Text.StartsWith(ProgressMarkupPrefix))
 						{
 							// Code below should not need to be tested for success but if some logging from XGE is wrong we just gracefully ignore it and accept that counting might end up wrong
+							int ActionIndex = -1;
 							int MarkupLength = ProgressMarkupPrefix.Length;
 							int EndOfMarkupPrefix = Text.IndexOf(' ', MarkupLength);
 							if (EndOfMarkupPrefix != -1)
 							{
 								MarkupLength = EndOfMarkupPrefix + 1;
-								int ActionIndex;
 								if (int.TryParse(Text.Substring(ProgressMarkupPrefix.Length + 1, EndOfMarkupPrefix - ProgressMarkupPrefix.Length - 1), out ActionIndex))
 								{
 									// We keep track of the actions that we have already reported so NumCompletedActions match up with NumActions
@@ -741,6 +743,8 @@ namespace UnrealBuildTool
 								ProgressText = string.Empty;
 							}
 
+							CommandDescription = ActionIndex != -1 ? Actions[ActionIndex].CommandDescription + " " : string.Empty;
+
 							// Strip out anything that is just an XGE timer. Some programs don't output anything except the progress text.
 							Text = Args.Data.Substring(MarkupLength);
 							if(Text.StartsWith(" (") && Text.EndsWith(")"))
@@ -749,13 +753,15 @@ namespace UnrealBuildTool
 								ProgressText = Text.Trim();
 								return;
 							}
-							Logger.LogInformation("[{NumCompletedActions}/{NumActions}] {Text}", NumCompletedActions, NumActions, Text);
+
+							Logger.LogInformation("[{NumCompletedActions}/{NumActions}] {CommandDescription}{Text}", NumCompletedActions, NumActions, CommandDescription, Text);
 							return;
 						}
 						if (!string.IsNullOrEmpty(ProgressText))
 						{
-							Logger.LogInformation("[{NumCompletedActions}/{NumActions}] {Text} {ProgressText}", NumCompletedActions, NumActions, Text, ProgressText);
+							Logger.LogInformation("[{NumCompletedActions}/{NumActions}] {CommandDescription}{Text} {ProgressText}", NumCompletedActions, NumActions, CommandDescription, Text, ProgressText);
 							ProgressText = string.Empty;
+							CommandDescription = string.Empty;
 							return;
 						}
 						Log.TraceInformation("{0}", Text); // Using old log function to pick up registered event parsers
