Revise Windows platforms support for specifying of a platform header override

- previously WindowsPlatform set the platform header name override define to "Windows" and a derived class then removed it from the preprocessor definitions list
  - the removal was incorrect but this still resulted in the derived platform having a platform header name override which matched the patforms name so it still functioned correct
  - however, it would be cleaner to not set the platform header override name at all for that platform
  - setting of this preprocessor definition has been revised so that Windows platforms specify the override via a separate GetOverridePlatformHeaderName() method which may return null to indicate no override should be used

#jria none
#rnx

diff --git a/Engine/Source/Programs/UnrealBuildTool/Platform/Windows/UEBuildWindows.cs b/Engine/Source/Programs/UnrealBuildTool/Platform/Windows/UEBuildWindows.cs
--- a/Engine/Source/Programs/UnrealBuildTool/Platform/Windows/UEBuildWindows.cs
+++ b/Engine/Source/Programs/UnrealBuildTool/Platform/Windows/UEBuildWindows.cs
@@ -1496,8 +1496,11 @@ namespace UnrealBuildTool
 			CompileEnvironment.Definitions.Add("PLATFORM_WINDOWS=1");
 			CompileEnvironment.Definitions.Add("PLATFORM_MICROSOFT=1");
 
-			// both Win32 and Win64 use Windows headers, so we enforce that here
-			CompileEnvironment.Definitions.Add(String.Format("OVERRIDE_PLATFORM_HEADER_NAME={0}", GetPlatformName()));
+			string? OverridePlatformHeaderName = GetOverridePlatformHeaderName();
+			if (!String.IsNullOrEmpty(OverridePlatformHeaderName))
+			{
+				CompileEnvironment.Definitions.Add(String.Format("OVERRIDE_PLATFORM_HEADER_NAME={0}", OverridePlatformHeaderName));
+			}
 
 			// Ray tracing only supported on 64-bit Windows.
 			if (Target.Platform == UnrealTargetPlatform.Win64 && Target.WindowsPlatform.bEnableRayTracing)
@@ -1691,6 +1694,15 @@ namespace UnrealBuildTool
 		{
 			new UEDeployWindows(Logger).PrepTargetForDeployment(Receipt);
 		}
+
+		/// <summary>
+		/// Allows the platform header name to be overridden to differ from the platform name.
+		/// </summary>
+		protected virtual string? GetOverridePlatformHeaderName()
+		{
+			// Both Win32 and Win64 use Windows headers, so we enforce that here.
+			return GetPlatformName();
+		}
 	}
 
 	class UEDeployWindows : UEBuildDeploy
