[UnrealBuildTool]
* Change so SingleFile will not build dependents by default. This can be enabled with -SingleFileBuildDependents

diff --git a/Engine/Source/Programs/UnrealBuildTool/Configuration/TargetDescriptor.cs b/Engine/Source/Programs/UnrealBuildTool/Configuration/TargetDescriptor.cs
--- a/Engine/Source/Programs/UnrealBuildTool/Configuration/TargetDescriptor.cs
+++ b/Engine/Source/Programs/UnrealBuildTool/Configuration/TargetDescriptor.cs
@@ -61,6 +61,12 @@ namespace UnrealBuildTool
 		[CommandLine("-SingleFile=")]
 		public List<FileReference> SpecificFilesToCompile = new List<FileReference>();
 
+		/// <summary>
+		/// Will build all files that directly include any of the files provided in -SingleFile
+		/// </summary>
+		[CommandLine("-SingleFileBuildDependents")]
+		public bool bSingleFileBuildDependents;
+
 		/// <summary>
 		/// Whether to perform hot reload for this target
 		/// </summary>
diff --git a/Engine/Source/Programs/UnrealBuildTool/Modes/BuildMode.cs b/Engine/Source/Programs/UnrealBuildTool/Modes/BuildMode.cs
--- a/Engine/Source/Programs/UnrealBuildTool/Modes/BuildMode.cs
+++ b/Engine/Source/Programs/UnrealBuildTool/Modes/BuildMode.cs
@@ -454,7 +454,8 @@ namespace UnrealBuildTool
 					if (FilesToBuild.Count != 0)
 					{
 						// If there are headers in the list, expand the FilesToBuild to also include all files that include those headers
-						FilesToBuild = GetAllSourceFilesIncludingHeader(FilesToBuild, TargetDescriptor.ProjectFile, Makefiles[Idx].Actions, Logger);
+						if (TargetDescriptor.bSingleFileBuildDependents)
+							FilesToBuild = GetAllSourceFilesIncludingHeader(FilesToBuild, TargetDescriptor.ProjectFile, Makefiles[Idx].Actions, Logger);
 
 						// We have specific files to compile so we will only queue up those files.
 						List<FileItem> ProducedItems = CreateLinkedActionsFromFileList(TargetDescriptor, BuildConfiguration, FilesToBuild, QueuedActions[Idx], Logger);
