[UnrealBuildTool]
* Fixed bug where every change to a header would cause invalidation of makefile due to workset changes.

diff --git a/Engine/Source/Programs/UnrealBuildTool/System/TargetMakefile.cs b/Engine/Source/Programs/UnrealBuildTool/System/TargetMakefile.cs
--- a/Engine/Source/Programs/UnrealBuildTool/System/TargetMakefile.cs
+++ b/Engine/Source/Programs/UnrealBuildTool/System/TargetMakefile.cs
@@ -872,7 +872,7 @@ namespace UnrealBuildTool
 				{
 					if (WorkingSet.Contains(SourceFile) && SourceFile.LastWriteTimeUtc > Makefile.CreateTimeUtc)
 					{
-						Logger.LogDebug("{File} was part of source working set and now is not", SourceFile.Location);
+						Logger.LogDebug("{File} was not part of source working set and should be", SourceFile.Location);
 						ReasonNotLoaded = string.Format("working set of source files changed");
 						return false;
 					}
diff --git a/Engine/Source/Programs/UnrealBuildTool/System/Unity.cs b/Engine/Source/Programs/UnrealBuildTool/System/Unity.cs
--- a/Engine/Source/Programs/UnrealBuildTool/System/Unity.cs
+++ b/Engine/Source/Programs/UnrealBuildTool/System/Unity.cs
@@ -360,6 +360,8 @@ namespace UnrealBuildTool
 				NameToHeaderFileInWorkingSet.Remove(Name);
 			}
 
+			HashSet<FileItem> UnhandledHeaderFilesInWorkingSet = new(HeaderFilesInWorkingSet);
+
 			// Add source files to the adaptive set if they or their first included header are in the working set.
 			foreach (FileItem CPPFile in CPPFiles)
 			{
@@ -368,7 +370,7 @@ namespace UnrealBuildTool
 					NameToHeaderFileInWorkingSet.TryGetValue(Path.GetFileName(FirstInclude), out FileItem? HeaderFile))
 				{
 					bHeaderInWorkingSet = true;
-					HeaderFilesInWorkingSet.Remove(HeaderFile);
+					UnhandledHeaderFilesInWorkingSet.Remove(HeaderFile);
 				}
 
 				bool bAdaptive = bHeaderInWorkingSet || WorkingSet.Contains(CPPFile);
@@ -382,12 +384,16 @@ namespace UnrealBuildTool
 				Graph.AddFileToWorkingSet(File);
 			}
 
-			List<FileItem> AdaptiveHeaderFiles = new List<FileItem>();
+			// We also need to add the headers since we don't want to invalidate makefile if they are changing
+			foreach (FileItem File in HeaderFilesInWorkingSet)
+			{
+				Graph.AddFileToWorkingSet(File);
+			}
 
 			// Add header files in the working set to the adaptive files if they are not the first include of a source file.
 			if (Target.bAdaptiveUnityCompilesHeaderFiles)
 			{
-				foreach (FileItem HeaderFile in HeaderFilesInWorkingSet)
+				foreach (FileItem HeaderFile in UnhandledHeaderFilesInWorkingSet)
 				{
 					StringWriter OutputHeaderCPPWriter = new StringWriter();
 					OutputHeaderCPPWriter.WriteLine("// This file is automatically generated at compile-time to include a modified header file.");
@@ -397,10 +403,6 @@ namespace UnrealBuildTool
 					FileReference HeaderCPPFilePath = FileReference.Combine(IntermediateDirectory, HeaderCPPFileName);
 
 					AdaptiveFiles.Add(Graph.CreateIntermediateTextFile(HeaderCPPFilePath, OutputHeaderCPPWriter.ToString()));
-
-					// Add the header file to the working set that will invalidate the makefile if it changes.
-					Graph.AddFileToWorkingSet(HeaderFile);
-					AdaptiveHeaderFiles.Add(HeaderFile);
 				}
 			}
 
@@ -408,7 +410,7 @@ namespace UnrealBuildTool
 			CandidateAdaptiveFiles.UnionWith(CPPFiles);
 			CandidateAdaptiveFiles.UnionWith(HeaderFiles);
 			CandidateAdaptiveFiles.ExceptWith(AdaptiveFiles);
-			CandidateAdaptiveFiles.ExceptWith(AdaptiveHeaderFiles);
+			CandidateAdaptiveFiles.ExceptWith(HeaderFilesInWorkingSet);
 
 			foreach (FileItem File in CandidateAdaptiveFiles)
 			{
