- Removed warning that the Sign Framework script would always run, with legacy xcode IOS projects

diff --git a/Engine/Source/Programs/UnrealBuildTool/ProjectFiles/Xcode/XcodeProjectLegacy.cs b/Engine/Source/Programs/UnrealBuildTool/ProjectFiles/Xcode/XcodeProjectLegacy.cs
--- a/Engine/Source/Programs/UnrealBuildTool/ProjectFiles/Xcode/XcodeProjectLegacy.cs
+++ b/Engine/Source/Programs/UnrealBuildTool/ProjectFiles/Xcode/XcodeProjectLegacy.cs
@@ -770,6 +770,7 @@ namespace UnrealBuildTool.XcodeProjectLegacy
 			Content.Append("/* Begin PBXShellScriptBuildPhase section */" + ProjectFileGenerator.NewLine);
 			Content.Append(string.Format("\t\t{0} /* Sign Frameworks */ = {{" + ProjectFileGenerator.NewLine, ShellScriptGuid));
 			Content.Append("\t\t\tisa = PBXShellScriptBuildPhase;" + ProjectFileGenerator.NewLine);
+			Content.Append("\t\t\talwaysOutOfDate = 1;" + ProjectFileGenerator.NewLine);
 			Content.Append("\t\t\tbuildActionMask = 2147483647;" + ProjectFileGenerator.NewLine);
 			Content.Append("\t\t\tfiles = (" + ProjectFileGenerator.NewLine);
 			Content.Append("\t\t\t);" + ProjectFileGenerator.NewLine);
