UBT: Disable chaining PCH when static analysis is enabled.

* Chaining PCH + Static Analysis results in crashing, so disable it.

diff --git a/Engine/Source/Programs/UnrealBuildTool/Platform/Mac/UEBuildMac.cs b/Engine/Source/Programs/UnrealBuildTool/Platform/Mac/UEBuildMac.cs
--- a/Engine/Source/Programs/UnrealBuildTool/Platform/Mac/UEBuildMac.cs
+++ b/Engine/Source/Programs/UnrealBuildTool/Platform/Mac/UEBuildMac.cs
@@ -297,6 +297,8 @@ namespace UnrealBuildTool
 			{
 				Target.bDisableLinking = true;
 				Target.bIgnoreBuildOutputs = true;
+				// Disable chaining PCHs for the moment because it is crashing clang
+				Target.bChainPCHs = false;
 			}
 
 			if (BuildHostPlatform.Current.Platform != UnrealTargetPlatform.Mac)
