Fixing ConfigCache key comparisons.
Changing ConfigCache to save if a file wasn't found.

Reduced -ProjectFiles time from 5m to 2m.

diff --git a/Engine/Source/Programs/UnrealBuildTool/System/ConfigCache.cs b/Engine/Source/Programs/UnrealBuildTool/System/ConfigCache.cs
--- a/Engine/Source/Programs/UnrealBuildTool/System/ConfigCache.cs
+++ b/Engine/Source/Programs/UnrealBuildTool/System/ConfigCache.cs
@@ -205,7 +205,7 @@ namespace UnrealBuildTool
 					OtherKey.ProjectDir == ProjectDir &&
 					OtherKey.Platform == Platform &&
 					OtherKey.CustomConfig == CustomConfig &&
-					OtherKey.OverrideStrings != OverrideStrings;
+					OtherKey.OverrideStrings == OverrideStrings;
 			}
 
 			/// <summary>
@@ -215,6 +215,7 @@ namespace UnrealBuildTool
 			public override int GetHashCode()
 			{
 				return HashCode.Combine(
+					Type.GetHashCode,
 					(ProjectDir != null) ? ProjectDir.GetHashCode() : 0,
 					Platform.GetHashCode(),
 					CustomConfig.GetHashCode(),
@@ -225,7 +226,7 @@ namespace UnrealBuildTool
 		/// <summary>
 		/// Cache of individual config files
 		/// </summary>
-		static Dictionary<FileReference, ConfigFile> LocationToConfigFile = new Dictionary<FileReference, ConfigFile>();
+		static Dictionary<FileReference, ConfigFile?> LocationToConfigFile = new Dictionary<FileReference, ConfigFile?>();
 
 		/// <summary>
 		/// Cache of config hierarchies by project
@@ -253,11 +254,13 @@ namespace UnrealBuildTool
 					{
 						ConfigFile = new ConfigFile(Location);
 					}
-
-					if (ConfigFile != null)
+					else
 					{
-						LocationToConfigFile.Add(Location, ConfigFile);
+						ConfigFile = null;
 					}
+
+					// Note: it's important to keep the null value here so we don't keep checking for files that don't exist
+					LocationToConfigFile.Add(Location, ConfigFile);
 				}
 			}
 
