Modern xcode changes
- .app is now fully made by Xcode, not partially by UBT/UAT (exe is copied into .app via shell script, and "cookeddata" is a folder reference in the Xcode proejct - for IOS/TVOS only for now)
- Updated how metadata files are found in the Metadata class, and template plist files are now called Info.Template.plst
- A prebuild script does some prep work to help with Xcode and UBT interaction, and the way our legacy build locations mess with Build after Archive
- Fixed some ".app" issues where the PRODUCT_NAME was missing

diff --git a/Engine/Config/BaseEngine.ini b/Engine/Config/BaseEngine.ini
--- a/Engine/Config/BaseEngine.ini
+++ b/Engine/Config/BaseEngine.ini
@@ -2825,6 +2825,9 @@ bUseModernXcode=false
 bUseModernCodeSigning=true
 ModernSigningTeam=
 ModernSigningPrefix=com.YourCompany
+TemplateMacPlist=/Engine/Build/Mac/Resources/Info.Template.plist
+TemplateMacEditorPlist=/Engine/Build/Mac/Resources/Info-Editor.Template.plist
+TemplateIOSPlist=/Game/Build/IOS/UBTGenerated/Info.Template.plist
 
 ; These are the defaults for Android settings, and they need to be in the .ini since UBT reads the .ini settings, without instantiating the class
 [/Script/AndroidRuntimeSettings.AndroidRuntimeSettings]
diff --git a/Engine/Source/Programs/AutomationTool/IOS/IOSPlatform.Automation.cs b/Engine/Source/Programs/AutomationTool/IOS/IOSPlatform.Automation.cs
--- a/Engine/Source/Programs/AutomationTool/IOS/IOSPlatform.Automation.cs
+++ b/Engine/Source/Programs/AutomationTool/IOS/IOSPlatform.Automation.cs
@@ -769,11 +769,30 @@ public class IOSPlatform : ApplePlatform
 
 		var TargetConfiguration = SC.StageTargetConfigurations[0];
 
+		ConfigHierarchy Ini = ConfigCache.ReadHierarchy(ConfigHierarchyType.Engine, Params.RawProjectPath.Directory!, UnrealTargetPlatform.IOS);
+		bool bUseModernXcode;
+		Ini.TryGetValue("XcodeConfiguration", "bUseModernXcode", out bUseModernXcode);
+
 		string MobileProvision;
 		string SigningCertificate;
 		string TeamUUID;
 		bool bAutomaticSigning;
-		GetProvisioningData(Params.RawProjectPath, Params.Distribution, out MobileProvision, out SigningCertificate, out TeamUUID, out bAutomaticSigning);
+		if (bUseModernXcode)
+		{
+			Ini.TryGetValue("XcodeConfiguration", "bUseModernCodeSigning", out bAutomaticSigning);
+			Ini.TryGetValue("XcodeConfiguration", "ModernSigningTeam", out TeamUUID);
+			//			Ini.TryGetValue("XcodeConfiguration", "ModernSigningPrefix", out SigningPrefix);
+			if (!bAutomaticSigning)
+			{
+				throw new AutomationException("Currently only Modern (automatic) codesigning is supported with Modern xcode");
+			}
+			SigningCertificate = null;
+			MobileProvision = null;
+		}
+		else
+		{
+			GetProvisioningData(Params.RawProjectPath, Params.Distribution, out MobileProvision, out SigningCertificate, out TeamUUID, out bAutomaticSigning);
+		}
 
 		//@TODO: We should be able to use this code on both platforms, when the following issues are sorted:
 		//   - Raw executable is unsigned & unstripped (need to investigate adding stripping to IPP)
@@ -1014,9 +1033,14 @@ public class IOSPlatform : ApplePlatform
 
 				// now generate the ipa
 				PackageIPA(Params, ProjectGameExeFilename, SC);
+
+				// generate a .app that can be run directly on an Apple Silicon Mac (the Mac needs to be in the mobileprovision, just like an IOS device)
+				// @todo IOSOnMac
+				// CreateIOSOnMacApp(ProjectGameExeFilename, Params, SC);
 			}
 		}
 
+
 		PrintRunTime();
 	}
 
@@ -1309,6 +1333,26 @@ public class IOSPlatform : ApplePlatform
 		}
 	}
 
+	// Generate a .app that can be run directly on an Apple Silicon Mac (the Mac needs to be in the mobileprovision, just like an IOS device)
+	private void CreateIOSOnMacApp(string ProjectGameExeFilename, ProjectParams Params, DeploymentContext SC)
+	{
+		// make a version that can run on the Mac directly
+		string ExeName = Params.IsCodeBasedProject ? SC.StageExecutables[0] : Path.GetFileNameWithoutExtension(ProjectGameExeFilename);
+		string GameName = ExeName.Split("-".ToCharArray())[0];
+		string SourceApp = Path.Combine(Path.GetDirectoryName(ProjectGameExeFilename), "Payload", $"{GameName}.app");
+		string DestOuterApp = Path.Combine(Path.GetDirectoryName(ProjectGameExeFilename), "IOSOnMac", $"{GameName}.app");
+		string DestSymlink = Path.Combine(DestOuterApp, "WrappedBundle");
+		string DestInnerApp = Path.Combine(DestOuterApp, "Wrapper", $"{GameName}.app");
+
+		LogInformation($"Creating IOSOnMac app {DestOuterApp}");
+
+		DeleteDirectory_NoExceptions(DestOuterApp);
+		CopyDirectory_NoExceptions(SourceApp, DestInnerApp);
+		// the link uses path relative to the symlink itself
+		File.CreateSymbolicLink(DestSymlink, $"Wrapper/{GameName}.app");
+	}
+
+
 	public override void GetFilesToDeployOrStage(ProjectParams Params, DeploymentContext SC)
 	{
 		//		if (UnrealBuildTool.BuildHostPlatform.Current.Platform != UnrealTargetPlatform.Mac)
diff --git a/Engine/Source/Programs/UnrealBuildTool/Platform/IOS/IOSToolChain.cs b/Engine/Source/Programs/UnrealBuildTool/Platform/IOS/IOSToolChain.cs
--- a/Engine/Source/Programs/UnrealBuildTool/Platform/IOS/IOSToolChain.cs
+++ b/Engine/Source/Programs/UnrealBuildTool/Platform/IOS/IOSToolChain.cs
@@ -1299,8 +1299,17 @@ namespace UnrealBuildTool
 				PostBuildSyncAction.WorkingDirectory = Unreal.EngineSourceDirectory;
 				PostBuildSyncAction.PrerequisiteItems.Add(Executable);
 				PostBuildSyncAction.PrerequisiteItems.UnionWith(OutputFiles);
-				PostBuildSyncAction.ProducedItems.Add(FileItem.GetItemByFileReference(GetStagedExecutablePath(Executable.Location, Target.Name)));
-				PostBuildSyncAction.DeleteItems.UnionWith(PostBuildSyncAction.ProducedItems);
+				if (bUseModernXcode)
+				{
+					// @todo: do we need one per Target for IOS? Client? I dont think so for Modern
+					FileReference PlistFile = FileReference.Combine(GetActualProjectDirectory(ProjectFile), "Build/IOS/UBTGenerated/Info.Template.plist");
+					PostBuildSyncAction.ProducedItems.Add(FileItem.GetItemByFileReference(PlistFile));
+				}
+				else
+				{
+					PostBuildSyncAction.ProducedItems.Add(FileItem.GetItemByFileReference(GetStagedExecutablePath(Executable.Location, Target.Name)));
+					PostBuildSyncAction.DeleteItems.UnionWith(PostBuildSyncAction.ProducedItems);
+				}
 				PostBuildSyncAction.StatusDescription = "Executing PostBuildSync";
 				PostBuildSyncAction.bCanExecuteRemotely = false;
 
diff --git a/Engine/Source/Programs/UnrealBuildTool/Platform/IOS/UEDeployIOS.cs b/Engine/Source/Programs/UnrealBuildTool/Platform/IOS/UEDeployIOS.cs
--- a/Engine/Source/Programs/UnrealBuildTool/Platform/IOS/UEDeployIOS.cs
+++ b/Engine/Source/Programs/UnrealBuildTool/Platform/IOS/UEDeployIOS.cs
@@ -457,6 +457,8 @@ namespace UnrealBuildTool
 				Text = new StringBuilder(result);
 			}
 
+			Text = Text.Replace("[PROJECT_NAME]", "$(UE_PROJECT_NAME)");
+
 			File.WriteAllText(PlistFile.FullName, Text.ToString());
 		}
 
@@ -490,7 +492,7 @@ namespace UnrealBuildTool
 				if (bUseModernXcode)
 				{
 					// @todo: do we need one per Target for IOS? Client? I dont think so for Modern
-					FinalPlistFile = new FileReference($"{ProjectDirectory}/Build/IOS/UBTGenerated/Info.plist");
+					FinalPlistFile = new FileReference($"{ProjectDirectory}/Build/IOS/UBTGenerated/Info.Template.plist");
 				}
 				else
 				{
@@ -1088,14 +1090,6 @@ namespace UnrealBuildTool
 
 		public bool PrepForUATPackageOrDeploy(UnrealTargetConfiguration Config, FileReference? ProjectFile, string InProjectName, string InProjectDirectory, FileReference Executable, string InEngineDir, bool bForDistribution, string CookFlavor, bool bIsDataDeploy, bool bCreateStubIPA, List<string> UPLScripts, string? BundleID, bool bBuildAsFramework)
 		{
-			ConfigHierarchy Ini = ConfigCache.ReadHierarchy(ConfigHierarchyType.Engine, ProjectFile?.Directory, UnrealTargetPlatform.IOS);
-			bool bUseModernXcode;
-			if (Ini.TryGetValue("XcodeConfiguration", "bUseModernXcode", out bUseModernXcode) && bUseModernXcode)
-			{
-				Logger.LogInformation("Modern Xcode doens't need to Deploy");
-				// none of this is needed with modern xcode
-				return false;
-			}
 
 			if (BuildHostPlatform.Current.Platform != UnrealTargetPlatform.Mac)
 			{
@@ -1121,6 +1115,26 @@ namespace UnrealBuildTool
 			string BuildDirectory_NFL = InProjectDirectory + "/Restricted/NotForLicensees/Build/" + SubDir;
 			string IntermediateDirectory = (bIsUnrealGame ? InEngineDir : InProjectDirectory) + "/Intermediate/" + SubDir;
 
+			ConfigHierarchy Ini = ConfigCache.ReadHierarchy(ConfigHierarchyType.Engine, ProjectFile?.Directory, UnrealTargetPlatform.IOS);
+			bool bUseModernXcode;
+			if (Ini.TryGetValue("XcodeConfiguration", "bUseModernXcode", out bUseModernXcode) && bUseModernXcode)
+			{
+				Logger.LogInformation("Generating plist (only step needed when deploying with Modern Xcode)");
+				GeneratePList(ProjectFile, Config, InProjectDirectory, bIsUnrealGame, GameExeName, false, InProjectName, InEngineDir, AppDirectory, UPLScripts, BundleID, bBuildAsFramework);
+
+				//// for now, copy the executable into the .app
+				//if (File.Exists(AppDirectory + "/" + GameName))
+				//{
+				//	FileInfo GameFileInfo = new FileInfo(AppDirectory + "/" + GameName);
+				//	GameFileInfo.Attributes = GameFileInfo.Attributes & ~FileAttributes.ReadOnly;
+				//}
+				//// copy the GameName binary
+				//File.Copy(BinaryPath + "/" + GameExeName, AppDirectory + "/" + GameName, true);
+
+				// none of this is needed with modern xcode
+				return false;
+			}
+
 			DirectoryReference.CreateDirectory(BinaryPath);
 			Directory.CreateDirectory(PayloadDirectory);
 			Directory.CreateDirectory(AppDirectory);
diff --git a/Engine/Source/Programs/UnrealBuildTool/ProjectFiles/Xcode/XcodeProject.cs b/Engine/Source/Programs/UnrealBuildTool/ProjectFiles/Xcode/XcodeProject.cs
--- a/Engine/Source/Programs/UnrealBuildTool/ProjectFiles/Xcode/XcodeProject.cs
+++ b/Engine/Source/Programs/UnrealBuildTool/ProjectFiles/Xcode/XcodeProject.cs
@@ -1,4 +1,4 @@
-// Copyright Epic Games, Inc. All Rights Reserved.
+﻿// Copyright Epic Games, Inc. All Rights Reserved.
 
 using System;
 using System.Collections.Generic;
@@ -189,31 +189,28 @@ namespace UnrealBuildTool.XcodeProjectXcconfig
 					string? PremadeLocation;
 					if (Ini.TryGetValue("XcodeConfiguration", Tokens[1], out PremadeLocation) && PremadeLocation != "")
 					{
-						// if the file exists (or it's a 'default' type which we use even if it doesn't exist), we are done!
-						FileReference TestFile = FileReference.Combine(ProductDirectory, PremadeLocation);
-						if (Tokens[0] == "default" || FileReference.Exists(TestFile))
+						FileReference TestFile;
+						if (PremadeLocation.StartsWith("/Engine/"))
+						{
+							TestFile = FileReference.Combine(Unreal.EngineDirectory, PremadeLocation.Substring(8));
+						}
+						else if (PremadeLocation.StartsWith("/Game/"))
+						{
+							TestFile = FileReference.Combine(ProductDirectory, PremadeLocation.Substring(6));
+						}
+						else
+						{
+							TestFile = new FileReference(Tokens[1]);
+						}
+
+						// make sure the file (if it's a generated file, then it may not exist yet, so allow it through)
+						if (TestFile.ContainsName("UBTGenerated", 0) || FileReference.Exists(TestFile))
 						{
 							File = TestFile;
 							break;
 						}
-						Log.TraceWarningOnce($"Metadata file {Tokens[1]} was specified in project settings, but does not exist. Ignoring...");
-					}
-				}
-				// otherwise it's a sub-path
-				else
-				{
-					FileReference EngineTestFile = FileReference.Combine(Unreal.EngineDirectory, Tokens[1]);
-					FileReference ProductTestFile = FileReference.Combine(ProductDirectory, Tokens[1]);
-					// if the file exists (or it's a 'default' type which we use even if it doesn't exist), we are done!
-					if (Tokens[0] == "default" || FileReference.Exists(ProductTestFile))
-					{
-						File = ProductTestFile;
-						break;
-					}
-					if (FileReference.Exists(EngineTestFile))
-					{
-						File = EngineTestFile;
-						break;
+
+						throw new BuildException($"Metadata file {Tokens[1]} (resolved to {TestFile}) was specified in project settings, but does not exist. Ignoring...");
 					}
 				}
 			}
@@ -226,7 +223,7 @@ namespace UnrealBuildTool.XcodeProjectXcconfig
 			else
 			{
 				// make it relative if it exists
-				XcodeProjectRelative = File?.MakeRelativeTo(XcodeProject.ParentDirectory!);
+				XcodeProjectRelative = File.MakeRelativeTo(XcodeProject.ParentDirectory!);
 			}
 			
 			// Log.TraceInformation($"   Metadata {Locations[0]} found {Mode} File {File}");
@@ -244,21 +241,15 @@ namespace UnrealBuildTool.XcodeProjectXcconfig
 			Log.TraceInformation($"making metadata for {ProductDirectory} / {XcodeProject}");
 			if (bSupportsMac)
 			{
-				PlistFiles[MetadataPlatform.MacEditor] = new MetadataItem(ProductDirectory, XcodeProject, Ini, "premade:PremadeMacEditorPlist", "template:TemplateMacEditorPlist", 
-					"premade:Build/Mac/Resources/Info-Editor.plist", "template:Build/Mac/Resources/Info-Editor.Template.plist");
-
-				PlistFiles[MetadataPlatform.Mac] = new MetadataItem(ProductDirectory, XcodeProject, Ini, "premade:PremadeMacPlist", "template:TemplateMacPlist", 
-					"premade:Build/Mac/Resources/Info.plist", "template:Build/Mac/Resources/Info.Template.plist");
-
-				EntitlementsFiles[MetadataPlatform.MacEditor] = new MetadataItem(ProductDirectory, XcodeProject, Ini, "premade:PremadeMacEditorEntitlements", "premade:Build/Mac/Resources/MacEditor.entitlements");
-				EntitlementsFiles[MetadataPlatform.Mac] = new MetadataItem(ProductDirectory, XcodeProject, Ini, "premade:PremadeMacEntitlements", "premade:Build/Mac/Resources/Mac.entitlements");
+				PlistFiles[MetadataPlatform.MacEditor] = new MetadataItem(ProductDirectory, XcodeProject, Ini, "premade:PremadeMacEditorPlist", "template:TemplateMacEditorPlist");
+				PlistFiles[MetadataPlatform.Mac] = new MetadataItem(ProductDirectory, XcodeProject, Ini, "premade:PremadeMacPlist", "template:TemplateMacPlist");
+				EntitlementsFiles[MetadataPlatform.MacEditor] = new MetadataItem(ProductDirectory, XcodeProject, Ini, "premade:PremadeMacEditorEntitlements");
+				EntitlementsFiles[MetadataPlatform.Mac] = new MetadataItem(ProductDirectory, XcodeProject, Ini, "premade:PremadeMacEntitlements");
 			}
 			if (bSupportsIOSOrTVOS)
 			{
-				PlistFiles[MetadataPlatform.IOS] = new MetadataItem(ProductDirectory, XcodeProject, Ini, 
-					"premade:PremadeIOSPlist", "template:TemplateIOSPlist", "premade:Build/IOS/Resources/Info.plist", "template:Build/IOS/Resources/Info.Template.plist", "default:Build/IOS/UBTGenerated/Info.plist");
-
-				EntitlementsFiles[MetadataPlatform.IOS] = new MetadataItem(ProductDirectory, XcodeProject, Ini, "premade:PremadeIOSEntitlements", "premade:Build/IOS/Resources/IOS.entitlements");
+				PlistFiles[MetadataPlatform.IOS] = new MetadataItem(ProductDirectory, XcodeProject, Ini, "premade:PremadeIOSPlist", "template:TemplateIOSPlist");
+				EntitlementsFiles[MetadataPlatform.IOS] = new MetadataItem(ProductDirectory, XcodeProject, Ini, "premade:PremadeIOSEntitlements");
 			}
 		}
 	}
@@ -1043,7 +1034,8 @@ namespace UnrealBuildTool.XcodeProjectXcconfig
 		public string Name;
 		public string IsAType;
 		public string Guid = XcodeProjectFileGenerator.MakeXcodeGuid();
-		protected List<XcodeSourceFile> Items = new();
+		protected List<XcodeSourceFile> FileItems = new();
+		protected List<string> MiscItems = new();
 
 		public XcodeBuildPhase(string Name, string IsAType)
 		{
@@ -1058,11 +1050,17 @@ namespace UnrealBuildTool.XcodeProjectXcconfig
 			Content.WriteLine(3,	$"isa = {IsAType};");
 			Content.WriteLine(3,	"buildActionMask = 2147483647;");
 			Content.WriteLine(3,	"files = (");
-			foreach (XcodeSourceFile File in Items)
+			foreach (XcodeSourceFile File in FileItems)
 			{
 				Content.WriteLine(4,		$"{File.FileGuid} /* {File.Reference.GetFileName()} in {Name} */,");
 			}
 			Content.WriteLine(3, ");");
+
+			foreach (string Line in MiscItems)
+			{
+				Content.WriteLine(3, Line);
+			}
+
 			Content.WriteLine(3, "runOnlyForDeploymentPostprocessing = 0;");
 			Content.WriteLine(2, "};");
 			Content.WriteLine($"/* End {IsAType} section */");
@@ -1078,24 +1076,7 @@ namespace UnrealBuildTool.XcodeProjectXcconfig
 
 		public void AddFile(XcodeSourceFile File)
 		{
-			Items.Add(File);
-		}
-
-		public override void Write(StringBuilder Content)
-		{
-			Content.WriteLine("/* Begin PBXSourcesBuildPhase section */");
-			Content.WriteLine($"\t\t{Guid} = {{");
-			Content.WriteLine("\t\t\tisa = PBXSourcesBuildPhase;");
-			Content.WriteLine("\t\t\tbuildActionMask = 2147483647;");
-			Content.WriteLine("\t\t\tfiles = (");
-			foreach (XcodeSourceFile File in Items)
-			{
-				Content.WriteLine($"\t\t\t\t{File.FileGuid} /* {File.Reference.GetFileName()} in {Name} */,");
-			}
-			Content.WriteLine("\t\t\t);");
-			Content.WriteLine("\t\t\trunOnlyForDeploymentPostprocessing = 0;");
-			Content.WriteLine("\t\t};");
-			Content.WriteLine("/* End PBXSourcesBuildPhase section */");
+			FileItems.Add(File);
 		}
 	}
 
@@ -1112,9 +1093,69 @@ namespace UnrealBuildTool.XcodeProjectXcconfig
 		public void AddResource(FileReference Resource)
 		{
 			XcodeSourceFile ResourceSource = new XcodeSourceFile(Resource, null);
-			FileCollection.ProcessFile(ResourceSource, true);
+			FileCollection.ProcessFile(ResourceSource, true, false);
+
+			FileItems.Add(ResourceSource);
+		}
+
+		public void AddFolderResource(DirectoryReference Resource, string GroupName)
+		{
+			XcodeSourceFile ResourceSource = new XcodeSourceFile(new FileReference(Resource.FullName), null);
+			FileCollection.ProcessFile(ResourceSource, true, true, GroupName);
+
+			//Project.FileCollection.AddFolderReference(CookedData.MakeRelativeTo(UnrealData.XcodeProjectFileLocation.ParentDirectory!), "CookedData_Game");
+			FileItems.Add(ResourceSource);
+		}
+	}
+
+	class XcodeCopyFilesBuildPhase : XcodeBuildPhase
+	{
+		private XcodeFileCollection FileCollection;
 
-			Items.Add(ResourceSource);
+		public XcodeCopyFilesBuildPhase(XcodeFileCollection FileCollection)
+			: base("Embed Frameworks", "PBXCopyFilesBuildPhase")
+		{
+			this.FileCollection = FileCollection;
+
+			MiscItems.Add($"dstPath = \"\";");
+			MiscItems.Add($"dstSubFolderSpec = 10;");
+			MiscItems.Add($"name = \"{Name}\";");
+		}
+
+		public void AddFramework(DirectoryReference Framework)
+		{
+			XcodeSourceFile ResourceSource = new XcodeSourceFile(new FileReference(Framework.FullName), null);
+			FileCollection.ProcessFile(ResourceSource, true, true, "Frameworks");
+
+			FileItems.Add(ResourceSource);
+		}
+	}
+
+	class XcodeShellScriptBuildPhase : XcodeBuildPhase
+	{
+		public XcodeShellScriptBuildPhase(string Name, IEnumerable<string> ScriptLines, IEnumerable<string> Inputs, IEnumerable<string> Outputs)
+			: base("Resources", "PBXShellScriptBuildPhase")
+		{
+			MiscItems.Add($"name = \"{Name}\";");
+
+			MiscItems.Add($"inputPaths = (");
+			foreach (string Input in Inputs)
+			{
+				MiscItems.Add($"\t\"{Input}\"");
+			}
+			MiscItems.Add($");");
+
+			MiscItems.Add($"outputPaths = (");
+			foreach (string Output in Outputs)
+			{
+				MiscItems.Add($"\t\"{Output}\"");
+			}
+			MiscItems.Add($");");
+
+			//			string Script = string.Join("&#10", ScriptLines);
+			string Script = string.Join("\\n", ScriptLines);
+			MiscItems.Add($"shellPath = /bin/sh;");
+			MiscItems.Add($"shellScript = \"{Script}\";");
 		}
 	}
 
@@ -1161,7 +1202,7 @@ namespace UnrealBuildTool.XcodeProjectXcconfig
 			return this.Platform == Platform || (this.Platform == null && BuildConfigs.Any(x => x.Info.Supports(Platform)));
 		}
 
-		private UnrealTargetPlatform? Platform;
+		public UnrealTargetPlatform? Platform;
 		public XcodeBuildConfigList(UnrealTargetPlatform? Platform, string TargetName, UnrealData UnrealData)
 		{
 			this.Platform = Platform;
@@ -1360,7 +1401,30 @@ namespace UnrealBuildTool.XcodeProjectXcconfig
 				BuildPhases.Add(ResourcesBuildPhase);
 				References.Add(ResourcesBuildPhase);
 
-				List<string>StoryboardPaths = new List<string>()
+				if (Project.Platform == UnrealTargetPlatform.IOS || Project.Platform == UnrealTargetPlatform.TVOS)
+				{
+					// UBT no longer copies the executable into the .app directory in PostBuild, so we do it here
+					string[] Script = {
+						"cp \\\"${CONFIGURATION_BUILD_DIR}/../${UE_IOS_EXECUTABLE_NAME}\\\" \\\"${CONFIGURATION_BUILD_DIR}/${PRODUCT_NAME}.app/${PRODUCT_NAME}\\\"",
+					};
+					string ScriptInput = "$(CONFIGURATION_BUILD_DIR)/../$(UE_IOS_EXECUTABLE_NAME)";
+					string ScriptOutput = "$(CONFIGURATION_BUILD_DIR)/$(PRODUCT_NAME).app/$(PRODUCT_NAME)";
+					XcodeShellScriptBuildPhase ScriptPhase = new("Copy UE Executable into .app", Script, new string[] { ScriptInput }, new string[] { ScriptOutput });
+					BuildPhases.Add(ScriptPhase);
+					References.Add(ScriptPhase);
+
+					// cookeddata needs to be brought in from the Staged directory, (IOS for Game, IOSClient for Client, etc)
+					string TargetPlatformName = Project.Platform.ToString()!;
+					if (TargetType != TargetType.Game)
+					{
+						TargetPlatformName += TargetType.ToString();
+					}
+
+					DirectoryReference CookedData = DirectoryReference.Combine(UnrealData.ProductDirectory, "Saved", "StagedBuilds",TargetPlatformName, "cookeddata");
+					ResourcesBuildPhase.AddFolderResource(CookedData, "ProjectData");
+				}
+
+				List<string> StoryboardPaths = new List<string>()
 				{
 					"$(Project)/Build/$(Platform)/Resources/Interface/LaunchScreen.storyboardc",
 					"$(Project)/Build/$(Platform)/Resources/Interface/LaunchScreen.storyboard",
@@ -1644,20 +1708,20 @@ namespace UnrealBuildTool.XcodeProjectXcconfig
 			return Orientations;
 		}
 
-		void FindBundleID(MetadataItem PlistItem, ref string BundleId)
+		void FindPlistId(MetadataItem PlistItem, string Key, ref string BundleId)
 		{
-			if (PlistItem.File == null)
+			if (PlistItem.File == null || !FileReference.Exists(PlistItem.File))
 			{
 				return;
 			}
 
-			string Identifier = XcodeProjectFile.Plist("Print :CFBundleIdentifier", PlistItem.File.FullName);
+			string Identifier = XcodeProjectFile.Plist($"Print :{Key}", PlistItem.File.FullName);
 			// handle error
 			if (Identifier == "" || Identifier.StartsWith("Print:"))
 			{
 				if (PlistItem.Mode == MetadataMode.UsePremade)
 				{
-					Log.TraceErrorOnce($"Premade .plist file '{PlistItem.File}' was found, but it did not contain a BundleIdentifier (CFBundleIdentifier key is missing or empty)");
+					Log.TraceErrorOnce($"Premade .plist file '{PlistItem.File}' was found, but it did not contain {Key} (Key is missing or value is empty)");
 				}
 			}
 			else
@@ -1868,12 +1932,12 @@ namespace UnrealBuildTool.XcodeProjectXcconfig
 			string IOSBundleID = UnrealData.BundleIdentifier;
 			if (UnrealData.bSupportsMac)
 			{
-				FindBundleID(UnrealData.Metadata!.PlistFiles[MetadataPlatform.MacEditor], ref MacEditorBundleID);
-				FindBundleID(UnrealData.Metadata.PlistFiles[MetadataPlatform.Mac], ref MacGameBundleID);
+				FindPlistId(UnrealData.Metadata!.PlistFiles[MetadataPlatform.MacEditor], "CFBundleIdentifier", ref MacEditorBundleID);
+				FindPlistId(UnrealData.Metadata.PlistFiles[MetadataPlatform.Mac], "CFBundleIdentifier", ref MacGameBundleID);
 			}
 			if (UnrealData.bSupportsIOS || UnrealData.bSupportsTVOS)
 			{
-				FindBundleID(UnrealData.Metadata!.PlistFiles[MetadataPlatform.IOS], ref IOSBundleID);
+				FindPlistId(UnrealData.Metadata!.PlistFiles[MetadataPlatform.IOS], "CFBundleIdentifier", ref IOSBundleID);
 			}
 
 			// Get Mac architectures supported by this project
@@ -1899,8 +1963,24 @@ namespace UnrealBuildTool.XcodeProjectXcconfig
 				FileReference MacExecutablePath = Config.MacExecutablePath!;
 				string MacExecutableDir = XcodeFileCollection.ConvertPath(MacExecutablePath.Directory.FullName);
 				string MacExecutableFileName = MacExecutablePath.GetFileName();
+				string IOSExecutableFilename = Config.IOSExecutablePath!.GetFileName();
+				string TVOSExecutableFilename = Config.TVOSExecutablePath!.GetFileName();
 				string SupportedMacArchitectures = string.Join(" ", UnrealData.GetSupportedMacArchitectures(Config, UnrealData.UProjectFileLocation).Architectures.Select(x => x.AppleName));
 
+				// get defeault PRODUCT_NAME values (if Premade plist, then pull it, like we do with BundleId)
+				string MacProductName = MacExecutableFileName;
+				string IOSProductName = Config.BuildTarget; // @todo: change to Path.GetFileName(Config.IOSExecutablePath) when we stop using payload
+
+				//if (UnrealData.bSupportsMac)
+				//{
+				//	FindBundleID(UnrealData.Metadata!.PlistFiles[MetadataPlatform.Mac], ref MacProductName);
+				//}
+				//if (UnrealData.bSupportsIOS || UnrealData.bSupportsTVOS)
+				//{
+				//	FindBundleID(UnrealData.Metadata!.PlistFiles[MetadataPlatform.IOS], ref IOSProductName);
+				//}
+
+
 				string SupportedPlatforms = "";
 				if (Config.ProjectTarget!.TargetRules != null)
 				{
@@ -1930,12 +2010,28 @@ namespace UnrealBuildTool.XcodeProjectXcconfig
 				{
 					ConfigXcconfig.AppendLine($"UE_MAC_EXECUTABLE_NAME = {MacExecutableFileName}");
 				}
-				if (ProjectBuildConfigs.bSupportsIOS || ProjectBuildConfigs.bSupportsTVOS)
+				if (ProjectBuildConfigs.bSupportsIOS)
 				{
-					ConfigXcconfig.AppendLine($"UE_IOS_EXECUTABLE_NAME = {Config.BuildTarget}");
+					ConfigXcconfig.AppendLine($"UE_IOS_EXECUTABLE_NAME = {IOSExecutableFilename}");
+				}
+				if (ProjectBuildConfigs.bSupportsTVOS)
+				{
+					ConfigXcconfig.AppendLine($"UE_TVOS_EXECUTABLE_NAME = {TVOSExecutableFilename}");
 				}
 				ConfigXcconfig.AppendLine($"UE_TARGET_NAME = {Config.BuildTarget}");
 				ConfigXcconfig.AppendLine($"UE_TARGET_CONFIG = {Config.BuildConfig}");
+				
+				
+				ConfigXcconfig.AppendLine($"UE_PLATFORM_NAME = {ProjectBuildConfigs.Platform}");
+
+				string TargetPlatformName = ProjectBuildConfigs.Platform.ToString()!;
+				if (ProjectBuildConfigs.BuildConfigs[0].Info.ProjectTarget!.TargetRules!.Type != TargetType.Game)
+				{
+					TargetPlatformName += ProjectBuildConfigs.BuildConfigs[0].Info.ProjectTarget!.TargetRules!.Type.ToString();
+				}
+				ConfigXcconfig.AppendLine($"UE_TARGET_PLATFORM_NAME = {TargetPlatformName}");
+
+
 				ConfigXcconfig.AppendLine("");
 
 				// debug settings
@@ -1963,13 +2059,13 @@ namespace UnrealBuildTool.XcodeProjectXcconfig
 					ConfigXcconfig.AppendLine($"CONFIGURATION_BUILD_DIR{SettingFilter} = {MacExecutableDir}");
 					ConfigXcconfig.AppendLine($"INFOPLIST_FILE{SettingFilter} = {PlistData.XcodeProjectRelative}");
 					ConfigXcconfig.AppendLine($"PRODUCT_BUNDLE_IDENTIFIER{SettingFilter} = " + (bIsEditor ? MacEditorBundleID : MacGameBundleID));
+					ConfigXcconfig.AppendLine($"PRODUCT_NAME{SettingFilter} = {MacProductName}");
 
 					if (PlistData.Mode == MetadataMode.UpdateTemplate)
 					{
-						ConfigXcconfig.AppendLine("GENERATE_INFOPLIST_FILE = YES");
-						ConfigXcconfig.AppendLine($"ASSETCATALOG_COMPILER_APPICON_NAME = AppIcon");
+						ConfigXcconfig.AppendLine($"GENERATE_INFOPLIST_FILE{SettingFilter} = YES");
+						ConfigXcconfig.AppendLine($"ASSETCATALOG_COMPILER_APPICON_NAME{SettingFilter} = AppIcon");
 
-						ConfigXcconfig.AppendLine($"PRODUCT_NAME{SettingFilter} = {MacExecutableFileName}");
 						// @todo: get a version for  games, like IOS has
 						ConfigXcconfig.AppendLine($"MARKETING_VERSION{SettingFilter} = {MacToolChain.LoadEngineDisplayVersion()}");
 						ConfigXcconfig.AppendLine($"CURRENT_PROJECT_VERSION{SettingFilter} = $(UE_MAC_BUILD_VERSION)");
@@ -1989,9 +2085,10 @@ namespace UnrealBuildTool.XcodeProjectXcconfig
 					ConfigXcconfig.AppendLine("// IOS setup");
 					ConfigXcconfig.AppendLine($"INFOPLIST_FILE{SettingFilter} = {UnrealData.Metadata!.PlistFiles[MetadataPlatform.IOS].XcodeProjectRelative}");
 					ConfigXcconfig.AppendLine($"PRODUCT_BUNDLE_IDENTIFIER{SettingFilter} = {IOSBundleID}");
+					ConfigXcconfig.AppendLine($"PRODUCT_NAME{SettingFilter} = {IOSProductName}");
+
 					if (UnrealData.Metadata!.PlistFiles[MetadataPlatform.IOS].Mode == MetadataMode.UpdateTemplate)
 					{
-						ConfigXcconfig.AppendLine($"PRODUCT_NAME{SettingFilter} = {Config.BuildTarget}"); // @todo: change to Path.GetFileName(Config.IOSExecutablePath) when we stop using payload
 					}
 					if (UnrealData.Metadata.EntitlementsFiles[MetadataPlatform.IOS].Mode == MetadataMode.UsePremade)
 					{
@@ -2119,7 +2216,7 @@ namespace UnrealBuildTool.XcodeProjectXcconfig
 			RootDirectory = Other.RootDirectory;
 			MainGroupGuid = Other.MainGroupGuid;
 			AllFiles.AddRange(Other.AllFiles);
-			BuildableResourceFiles.AddRange(Other.AllFiles);
+			BuildableResourceFiles.AddRange(Other.BuildableResourceFiles);
 			ManualFiles.AddRange(Other.ManualFiles);
 
 			BuildableFilesToResponseFile = new(Other.BuildableFilesToResponseFile);
@@ -2137,7 +2234,7 @@ namespace UnrealBuildTool.XcodeProjectXcconfig
 			RootDirectory = UProjectLocation == null || UProjectLocation.ContainsName("Programs", 0) ? (Unreal.EngineDirectory) : UProjectLocation!.Directory;
 		}
 
-		public void ProcessFile(XcodeSourceFile File, bool bIsForBuild)
+		public void ProcessFile(XcodeSourceFile File, bool bIsForBuild, bool bIsFolder, string? GroupName=null)
 		{
 			// remember all buildable files
 			if (bIsForBuild)
@@ -2161,7 +2258,7 @@ namespace UnrealBuildTool.XcodeProjectXcconfig
 			
 
 			// group the files by path
-			XcodeFileGroup? Group = FindGroupByAbsolutePath(File.Reference.Directory);
+			XcodeFileGroup? Group = GroupName == null ? FindGroupByAbsolutePath(File.Reference.Directory) : null;
 			if (Group != null)
 			{
 				AllFiles.Add(File);
@@ -2169,8 +2266,15 @@ namespace UnrealBuildTool.XcodeProjectXcconfig
 			}
 			else
 			{
-				string GroupName = File.Reference.IsUnderDirectory(Unreal.EngineDirectory) ? "EngineReferences" : "ExternalReferences";
-				AddFileReference(File.FileRefGuid, File.Reference.FullName, GetFileKey(File.Reference.GetExtension()), GetFileType(File.Reference.GetExtension()), "<absolute>", "EngineReferences");
+				GroupName = GroupName ?? (File.Reference.IsUnderDirectory(Unreal.EngineDirectory) ? "EngineReferences" : "ExternalReferences");
+				if (bIsFolder)
+				{
+					AddFolderReference(File.FileRefGuid, File.Reference.MakeRelativeTo(ProjectDirectory), GroupName);
+				}
+				else
+				{
+					AddFileReference(File.FileRefGuid, File.Reference.FullName, GetFileKey(File.Reference.GetExtension()), GetFileType(File.Reference.GetExtension()), "<absolute>", GroupName);
+				}
 			}
 		}
 
@@ -2202,6 +2306,25 @@ namespace UnrealBuildTool.XcodeProjectXcconfig
 			}
 		}
 
+		public void AddFolderReference(string Guid, string RelativePath, string GroupName)
+		{
+			ManualFiles.Add(new ManualFileReference()
+			{
+				Guid = Guid,
+				RelativePath = RelativePath,
+				FileKey = "lastKnownFileType",
+				FileType = "folder",
+				SourceTree = "\"<group>\"",
+				GroupName = GroupName,
+			});
+
+			// make sure each unique group has a guid
+			if (!GuidsForGroups.ContainsKey(GroupName))
+			{
+				GuidsForGroups[GroupName] = XcodeProjectFileGenerator.MakeXcodeGuid();
+			}
+		}
+
 
 		public void Write(StringBuilder Content)
 		{
@@ -2429,7 +2552,7 @@ namespace UnrealBuildTool.XcodeProjectXcconfig
 		/// </summary>
 		internal static bool IsResourceFile(string Extension)
 		{
-			return Extension == ".xcassets" || Extension == ".storyboardc" || Extension == ".storyboard";
+			return Extension == ".xcassets" || Extension == ".storyboardc" || Extension == ".storyboard" || Extension == "";
 		}
 
 		/// <summary>
@@ -2607,7 +2730,7 @@ namespace UnrealBuildTool.XcodeProjectXcconfig
 
 			foreach (XcodeSourceFile SourceFile in SourceFiles.OfType<XcodeSourceFile>())
 			{
-				SharedFileCollection.ProcessFile(SourceFile, bIsForBuild:IsGeneratedProject);
+				SharedFileCollection.ProcessFile(SourceFile, bIsForBuild:IsGeneratedProject, false);
 			}
 
 			// cache the main group
@@ -3439,6 +3562,28 @@ namespace UnrealBuildTool.XcodeProjectXcconfig
 					}
 				}
 
+				string[] ScriptLines = new string[]
+				{
+					"# if a UBTGenerated plist doesn&apos;t exist, create it now for Xcode&apos;s dependency needs",
+					"if [ ! -f &quot;${PROJECT_DIR}/${INFOPLIST_FILE}&quot; ]",
+					"then",
+					"  echo Creating ${PROJECT_DIR}/${INFOPLIST_FILE}...",
+					"  touch &quot;${PROJECT_DIR}/${INFOPLIST_FILE}&quot;",
+					"fi",
+					"",
+					"# if we had previously archived the build, then Xcode would have made a soft-link, but it will fail when building normally, so delete a soft-link",
+					"if [ -L &quot;${CONFIGURATION_BUILD_DIR}/${PRODUCT_NAME}.app&quot; ]",
+					"then",
+					"  rm &quot;${CONFIGURATION_BUILD_DIR}/${PRODUCT_NAME}.app&quot;",
+					"fi",
+					"",
+					"# make sure we have a cookeddata directory if needed",
+					"if [[ &quot;${PLATFORM_NAME}&quot; == &quot;iphoneos&quot; || &quot;${PLATFORM_NAME}&quot; == &quot;appletvos&quot; ]]",
+					"then",
+					"  mkdir -p &quot;${CONFIGURATION_BUILD_DIR}/../../../Saved/StagedBuilds/${UE_TARGET_PLATFORM_NAME}/cookeddata&quot;",
+					"fi",
+					"",
+				};
 
 				Content.WriteLine("<?xml version=\"1.0\" encoding=\"UTF-8\"?>");
 				Content.WriteLine("<Scheme");
@@ -3447,6 +3592,24 @@ namespace UnrealBuildTool.XcodeProjectXcconfig
 				Content.WriteLine("   <BuildAction");
 				Content.WriteLine("      parallelizeBuildables = \"YES\"");
 				Content.WriteLine("      buildImplicitDependencies = \"YES\">");
+				Content.WriteLine("      <PreActions>");
+				Content.WriteLine("	        <ExecutionAction");
+				Content.WriteLine("            ActionType = \"Xcode.IDEStandardExecutionActionsCore.ExecutionActionType.ShellScriptAction\">");
+				Content.WriteLine("            <ActionContent");
+				Content.WriteLine("               title = \"Touch UBT generated tiles\"");
+				Content.WriteLine("               scriptText = \"" + string.Join("&#10;", ScriptLines) + "\">");
+				Content.WriteLine("               <EnvironmentBuildable>");
+				Content.WriteLine("                  <BuildableReference");
+				Content.WriteLine("                     BuildableIdentifier = \"primary\"");
+				Content.WriteLine("                     BlueprintIdentifier = \"" + TargetGuid + "\"");
+				Content.WriteLine("                     BuildableName = \"" + ProductName + ".app\"");
+				Content.WriteLine("                     BlueprintName = \"" + TargetName + "\"");
+				Content.WriteLine("                     ReferencedContainer = \"container:" + ProjectName + ".xcodeproj\">");
+				Content.WriteLine("                  </BuildableReference>");
+				Content.WriteLine("               </EnvironmentBuildable>");
+				Content.WriteLine("            </ActionContent>");
+				Content.WriteLine("         </ExecutionAction>");
+				Content.WriteLine("      </PreActions>");
 				Content.WriteLine("      <BuildActionEntries>");
 				Content.WriteLine("         <BuildActionEntry");
 				Content.WriteLine("            buildForTesting = \"YES\"");
@@ -3503,7 +3666,7 @@ namespace UnrealBuildTool.XcodeProjectXcconfig
 				Content.WriteLine("      </BuildableProductRunnable>");
 				if (string.IsNullOrEmpty(OldCommandLineArguments))
 				{
-					if (bHasEditorConfiguration && TargetName != "UE5")
+					if (bHasEditorConfiguration && GameProjectPath != "")
 					{
 						Content.WriteLine("      <CommandLineArguments>");
 						if (IsForeignProject)
@@ -3516,7 +3679,7 @@ namespace UnrealBuildTool.XcodeProjectXcconfig
 						else
 						{
 							Content.WriteLine("         <CommandLineArgument");
-							Content.WriteLine("            argument = \"" + TargetName + "\"");
+							Content.WriteLine("            argument = \"" + Path.GetFileNameWithoutExtension(GameProjectPath) + "\"");
 							Content.WriteLine("            isEnabled = \"YES\">");
 							Content.WriteLine("         </CommandLineArgument>");
 						}
