Module include path validation. Disabled until outstanding issues are fixed

#rnx

diff --git a/Engine/Source/Programs/UnrealBuildTool/Configuration/ModuleRules.cs b/Engine/Source/Programs/UnrealBuildTool/Configuration/ModuleRules.cs
--- a/Engine/Source/Programs/UnrealBuildTool/Configuration/ModuleRules.cs
+++ b/Engine/Source/Programs/UnrealBuildTool/Configuration/ModuleRules.cs
@@ -818,6 +818,42 @@ namespace UnrealBuildTool
 		/// </summary>
 		public bool bEnableUndefinedIdentifierWarnings = true;
 
+		/// <summary>
+		/// How to treat general module include path validation messages
+		/// </summary>
+		public WarningLevel ModuleIncludePathWarningLevel
+		{
+			get { return (ModuleIncludePathWarningLevelPrivate == WarningLevel.Default) ? Target.ModuleIncludePathWarningLevel : ModuleIncludePathWarningLevelPrivate; }
+			set { ModuleIncludePathWarningLevelPrivate = value; }
+		}
+
+		/// <inheritdoc cref="ModuleIncludePathWarningLevel"/>
+		private WarningLevel ModuleIncludePathWarningLevelPrivate;
+
+		/// <summary>
+		/// How to treat private module include path validation messages, where a module is adding an include path that exposes private headers
+		/// </summary>
+		public WarningLevel ModuleIncludePrivateWarningLevel
+		{
+			get { return (ModuleIncludePrivateWarningLevelPrivate == WarningLevel.Default) ? Target.ModuleIncludePrivateWarningLevel : ModuleIncludePrivateWarningLevelPrivate; }
+			set { ModuleIncludePrivateWarningLevelPrivate = value; }
+		}
+
+		/// <inheritdoc cref="ModuleIncludePrivateWarningLevel"/>
+		private WarningLevel ModuleIncludePrivateWarningLevelPrivate;
+
+		/// <summary>
+		/// How to treat unnecessary module sub-directory include path validation messages
+		/// </summary>
+		public WarningLevel ModuleIncludeSubdirectoryWarningLevel
+		{
+			get { return (ModuleIncludeSubdirectoryWarningLevelPrivate == WarningLevel.Default) ? Target.ModuleIncludeSubdirectoryWarningLevel : ModuleIncludeSubdirectoryWarningLevelPrivate; }
+			set { ModuleIncludeSubdirectoryWarningLevelPrivate = value; }
+		}
+
+		/// <inheritdoc cref="ModuleIncludeSubdirectoryWarningLevel"/>
+		private WarningLevel ModuleIncludeSubdirectoryWarningLevelPrivate;
+
 		/// <summary>
 		/// Disable all static analysis - clang, msvc, pvs-studio.
 		/// </summary>
diff --git a/Engine/Source/Programs/UnrealBuildTool/Configuration/TargetRules.cs b/Engine/Source/Programs/UnrealBuildTool/Configuration/TargetRules.cs
--- a/Engine/Source/Programs/UnrealBuildTool/Configuration/TargetRules.cs
+++ b/Engine/Source/Programs/UnrealBuildTool/Configuration/TargetRules.cs
@@ -1532,6 +1532,27 @@ namespace UnrealBuildTool
 		[RequiresUniqueBuildEnvironment]
 		public bool bUndefinedIdentifierErrors = true;
 
+		/// <summary>
+		/// How to treat general module include path validation messages
+		/// </summary>
+		[XmlConfigFile(Category = "BuildConfiguration")]
+		[CommandLine("-ModuleIncludePathWarningLevel=")]
+		public WarningLevel ModuleIncludePathWarningLevel = WarningLevel.Off;
+
+		/// <summary>
+		/// How to treat private module include path validation messages, where a module is adding an include path that exposes private headers
+		/// </summary>
+		[XmlConfigFile(Category = "BuildConfiguration")]
+		[CommandLine("-ModuleIncludePrivateWarningLevel=")]
+		public WarningLevel ModuleIncludePrivateWarningLevel = WarningLevel.Off;
+
+		/// <summary>
+		/// How to treat unnecessary module sub-directory include path validation messages
+		/// </summary>
+		[XmlConfigFile(Category = "BuildConfiguration")]
+		[CommandLine("-ModuleIncludeSubdirectoryWarningLevel=")]
+		public WarningLevel ModuleIncludeSubdirectoryWarningLevel = WarningLevel.Off;
+
 		/// <summary>
 		/// Forces frame pointers to be retained this is usually required when you want reliable callstacks e.g. mallocframeprofiler
 		/// </summary>
@@ -3375,6 +3396,21 @@ namespace UnrealBuildTool
 			get { return Inner.bUndefinedIdentifierErrors; }
 		}
 
+		public WarningLevel ModuleIncludePathWarningLevel
+		{
+			get { return Inner.ModuleIncludePathWarningLevel; }
+		}
+
+		public WarningLevel ModuleIncludePrivateWarningLevel
+		{
+			get { return Inner.ModuleIncludePrivateWarningLevel; }
+		}
+
+		public WarningLevel ModuleIncludeSubdirectoryWarningLevel
+		{
+			get { return Inner.ModuleIncludeSubdirectoryWarningLevel; }
+		}
+
 		public bool bWarningsAsErrors
 		{
 			get { return Inner.bWarningsAsErrors; }
diff --git a/Engine/Source/Programs/UnrealBuildTool/Configuration/UEBuildPlatform.cs b/Engine/Source/Programs/UnrealBuildTool/Configuration/UEBuildPlatform.cs
--- a/Engine/Source/Programs/UnrealBuildTool/Configuration/UEBuildPlatform.cs
+++ b/Engine/Source/Programs/UnrealBuildTool/Configuration/UEBuildPlatform.cs
@@ -738,6 +738,226 @@ namespace UnrealBuildTool
 		{
 		}
 
+		/// <summary>
+		/// Validate a UEBuildModule's include paths before it's processed
+		/// </summary>
+		/// <param name="Module">The UEBuildModule that needs to be validated</param>
+		/// <param name="Target">Options for the target being built</param>
+		/// <param name="AllModules">Other modules to validate against, if needed</param>
+		public virtual void ValidateModuleIncludePaths(UEBuildModule Module, ReadOnlyTargetRules Target, IEnumerable<UEBuildModule> AllModules)
+		{
+			if (Module.Rules.ModuleIncludePathWarningLevel <= WarningLevel.Off
+				&& Module.Rules.ModuleIncludePrivateWarningLevel <= WarningLevel.Off
+				&& Module.Rules.ModuleIncludeSubdirectoryWarningLevel <= WarningLevel.Off)
+			{
+				return;
+			}
+
+			void LoggerFunc(string? message, params object?[] args)
+			{
+				if (Module.Rules.ModuleIncludePathWarningLevel == WarningLevel.Warning)
+				{
+					Logger.LogWarning($"Warning: {message}", args);
+				}
+				else if (Module.Rules.ModuleIncludePathWarningLevel == WarningLevel.Error)
+				{
+					Logger.LogError($"Error: {message}", args);
+				}
+			}
+
+			void LoggerFuncPrivate(string? message, params object?[] args)
+			{
+				if (Module.Rules.ModuleIncludePrivateWarningLevel == WarningLevel.Warning)
+				{
+					Logger.LogWarning($"Warning: {message}", args);
+				}
+				else if (Module.Rules.ModuleIncludePrivateWarningLevel == WarningLevel.Error)
+				{
+					Logger.LogError($"Error: {message}", args);
+				}
+			}
+
+			void LoggerFuncSubDir(string? message, params object?[] args)
+			{
+				if (Module.Rules.ModuleIncludeSubdirectoryWarningLevel == WarningLevel.Warning)
+				{
+					Logger.LogWarning($"Warning: {message}", args);
+				}
+				else if (Module.Rules.ModuleIncludeSubdirectoryWarningLevel == WarningLevel.Error)
+				{
+					Logger.LogError($"Error: {message}", args);
+				}
+			}
+
+			IOrderedEnumerable<DirectoryReference> PublicIncludePaths = Module.Rules.PublicIncludePaths.Select(x => DirectoryReference.FromString(x)!).OrderBy(x => x.FullName);
+			IOrderedEnumerable<DirectoryReference> PrivateIncludePaths = Module.Rules.PrivateIncludePaths.Select(x => DirectoryReference.FromString(x)!).OrderBy(x => x.FullName);
+			IOrderedEnumerable<DirectoryReference> InternalncludePaths = Module.Rules.InternalncludePaths.Select(x => DirectoryReference.FromString(x)!).OrderBy(x => x.FullName);
+			IOrderedEnumerable<DirectoryReference> PublicSystemIncludePaths = Module.Rules.PublicSystemIncludePaths.Select(x => DirectoryReference.FromString(x)!).OrderBy(x => x.FullName);
+			IOrderedEnumerable<UEBuildModule> OtherModules = AllModules.Where(x => x != Module).OrderBy(x => x.Name);
+
+			if (Module is UEBuildModuleExternal)
+			{
+				foreach (DirectoryReference Path in PublicIncludePaths)
+				{
+					LoggerFunc("External module '{Name}' is adding '{Path}' to PublicIncludePaths. This path should be added to PublicSystemIncludePaths.", Module.Name, Path.MakeRelativeTo(Module.ModuleDirectory));
+				}
+
+				foreach (DirectoryReference Path in PublicSystemIncludePaths.Where(x => !x.IsUnderDirectory(Module.ModuleDirectory)))
+				{
+					var OtherModule = OtherModules.FirstOrDefault(x => Path.IsUnderDirectory(x.ModuleDirectory));
+					if (OtherModule is UEBuildModuleExternal)
+					{
+						LoggerFunc("External module '{Name}' is adding '{Path}' from external module '{OtherModule}' to PublicSystemIncludePaths. Did you intend to add a public reference?", Module.Name, Path, OtherModule.Name);
+					}
+					else if (OtherModule is UEBuildModuleCPP)
+					{
+						LoggerFunc("External module '{Name}' is adding '{Path}' from module '{OtherModule}' to PublicSystemIncludePaths. This is not allowed.", Module.Name, Path, OtherModule.Name);
+					}
+				}
+
+				foreach (DirectoryReference Path in PrivateIncludePaths)
+				{
+					LoggerFunc("External module '{Name}' is adding '{Path}' to PrivateIncludePaths. This path is unused.", Module.Name, Path);
+				}
+
+				foreach (DirectoryReference Path in InternalncludePaths)
+				{
+					LoggerFunc("External module '{Name}' is adding '{Path}' to InternalncludePaths. This path is unused.", Module.Name, Path);
+				}
+			}
+			else if (Module is UEBuildModuleCPP)
+			{
+				foreach (DirectoryReference Path in PublicIncludePaths)
+				{
+					if (Path.IsUnderDirectory(Module.ModuleDirectory))
+					{
+						if (Path == Module.ModuleDirectory)
+						{
+							LoggerFunc("Module '{Name}' is adding root directory to PublicIncludePaths. This is not allowed.", Module.Name);
+						}
+						else if (Path.IsUnderDirectory(DirectoryReference.Combine(Module.ModuleDirectory, "Private"))
+							|| Path.IsUnderDirectory(DirectoryReference.Combine(Module.ModuleDirectory, "Internal")))
+						{
+							LoggerFuncPrivate("Module '{Name}' is adding subdirectory '{Path}' to PublicIncludePaths. This is not allowed.", Module.Name, Path.MakeRelativeTo(Module.ModuleDirectory));
+						}
+						else if (Path.IsUnderDirectory(DirectoryReference.Combine(Module.ModuleDirectory, "Public"))
+							|| Path.IsUnderDirectory(DirectoryReference.Combine(Module.ModuleDirectory, "Classes")))
+						{
+							LoggerFuncSubDir("Module '{Name}' is adding subdirectory '{Path}' to PublicIncludePaths. This is not necessary.", Module.Name, Path.MakeRelativeTo(Module.ModuleDirectory));
+						}
+					}
+
+					var OtherModule = OtherModules.FirstOrDefault(x => Path.IsUnderDirectory(x.ModuleDirectory));
+					if (OtherModule is UEBuildModuleExternal)
+					{
+						LoggerFunc("Module '{Name}' is adding '{Path}' from external module '{OtherModule}' to PublicIncludePaths. Did you intend to add a public reference?.", Module.Name, Path, OtherModule.Name);
+					}
+					else if (OtherModule is UEBuildModuleCPP)
+					{
+						if (Path == OtherModule.ModuleDirectory)
+						{
+							LoggerFuncPrivate("Module '{Name}' is adding root directory from '{OtherModule}' to PublicIncludePaths. This is not allowed.", Module.Name, OtherModule.Name);
+						}
+						else if (Path.IsUnderDirectory(DirectoryReference.Combine(OtherModule.ModuleDirectory, "Private"))
+							| Path.IsUnderDirectory(DirectoryReference.Combine(OtherModule.ModuleDirectory, "Internal")))
+						{
+							LoggerFuncPrivate("Module '{Name}' is adding '{Path}' from module '{OtherModule}' to PublicIncludePaths. This is not allowed.", Module.Name, Path, OtherModule.Name);
+						}
+						else if (Path.IsUnderDirectory(DirectoryReference.Combine(OtherModule.ModuleDirectory, "Public"))
+							|| Path.IsUnderDirectory(DirectoryReference.Combine(OtherModule.ModuleDirectory, "Classes")))
+						{
+							LoggerFunc("Module '{Name}' is adding '{Path}' from module '{OtherModule}' to PublicIncludePaths. Did you intend to add a public reference?.", Module.Name, Path, OtherModule.Name);
+						}
+					}
+				}
+
+				foreach (DirectoryReference Path in PrivateIncludePaths)
+				{
+					if (Path.IsUnderDirectory(Module.ModuleDirectory))
+					{
+						if (Path == Module.ModuleDirectory)
+						{
+							LoggerFunc("Module '{Name}' is adding root directory to PrivateIncludePaths. This is not recommended.", Module.Name);
+						}
+						else if (Path.IsUnderDirectory(DirectoryReference.Combine(Module.ModuleDirectory, "Private"))
+							|| Path.IsUnderDirectory(DirectoryReference.Combine(Module.ModuleDirectory, "Internal"))
+							|| Path.IsUnderDirectory(DirectoryReference.Combine(Module.ModuleDirectory, "Public"))
+							|| Path.IsUnderDirectory(DirectoryReference.Combine(Module.ModuleDirectory, "Classes")))
+						{
+							LoggerFuncSubDir("Module '{Name}' is adding subdirectory '{Path}' to PrivateIncludePaths. This is not necessary.", Module.Name, Path.MakeRelativeTo(Module.ModuleDirectory));
+						}
+					}
+
+					var OtherModule = OtherModules.FirstOrDefault(x => Path.IsUnderDirectory(x.ModuleDirectory));
+					if (OtherModule is UEBuildModuleExternal)
+					{
+						LoggerFunc("Module '{Name}' is adding '{Path}' from external module '{OtherModule}' to PrivateIncludePaths. Did you intend to add a private reference?.", Module.Name, Path, OtherModule.Name);
+					}
+					else if (OtherModule is UEBuildModuleCPP)
+					{
+						if (Path == OtherModule.ModuleDirectory)
+						{
+							LoggerFuncPrivate("Module '{Name}' is adding root directory from '{OtherModule}' to PrivateIncludePaths. This is not allowed.", Module.Name, OtherModule.Name);
+						}
+						else if (Path.IsUnderDirectory(DirectoryReference.Combine(OtherModule.ModuleDirectory, "Private"))
+							| Path.IsUnderDirectory(DirectoryReference.Combine(OtherModule.ModuleDirectory, "Internal")))
+						{
+							LoggerFuncPrivate("Module '{Name}' is adding '{Path}' from module '{OtherModule}' to PrivateIncludePaths. This is not allowed.", Module.Name, Path, OtherModule.Name);
+						}
+						else if (Path.IsUnderDirectory(DirectoryReference.Combine(OtherModule.ModuleDirectory, "Public"))
+							|| Path.IsUnderDirectory(DirectoryReference.Combine(OtherModule.ModuleDirectory, "Classes")))
+						{
+							LoggerFunc("Module '{Name}' is adding '{Path}' from module '{OtherModule}' to PrivateIncludePaths. Did you intend to add a private reference?.", Module.Name, Path, OtherModule.Name);
+						}
+					}
+				}
+
+				foreach (DirectoryReference Path in InternalncludePaths)
+				{
+					if (Path.IsUnderDirectory(Module.ModuleDirectory))
+					{
+						if (Path == Module.ModuleDirectory)
+						{
+							LoggerFunc("Module '{Name}' is adding root directory to InternalncludePaths. This is not allowed.", Module.Name);
+						}
+						else if (Path.IsUnderDirectory(DirectoryReference.Combine(Module.ModuleDirectory, "Private"))
+						|| Path.IsUnderDirectory(DirectoryReference.Combine(Module.ModuleDirectory, "Internal")))
+						{
+							LoggerFunc("Module '{Name}' is adding subdirectory '{Path}' to InternalncludePaths. This is not allowed.", Module.Name, Path.MakeRelativeTo(Module.ModuleDirectory));
+						}
+						else if (Path.IsUnderDirectory(DirectoryReference.Combine(Module.ModuleDirectory, "Public"))
+							|| Path.IsUnderDirectory(DirectoryReference.Combine(Module.ModuleDirectory, "Classes")))
+						{
+							LoggerFuncSubDir("Module '{Name}' is adding subdirectory '{Path}' to InternalncludePaths. This is not necessary.", Module.Name, Path.MakeRelativeTo(Module.ModuleDirectory));
+						}
+					}
+
+					var OtherModule = OtherModules.FirstOrDefault(x => Path.IsUnderDirectory(x.ModuleDirectory));
+					if (OtherModule is UEBuildModuleExternal)
+					{
+						LoggerFunc("Module '{Name}' is adding '{Path}' from external module '{OtherModule}' to InternalncludePaths. Did you intend to add a public reference?.", Module.Name, Path, OtherModule.Name);
+					}
+					else if (OtherModule is UEBuildModuleCPP)
+					{
+						if (Path == OtherModule.ModuleDirectory)
+						{
+							LoggerFunc("Module '{Name}' is adding root directory from '{OtherModule}' to InternalncludePaths. This is not allowed.", Module.Name, OtherModule.Name);
+						}
+						else if (Path.IsUnderDirectory(DirectoryReference.Combine(OtherModule.ModuleDirectory, "Private"))
+							| Path.IsUnderDirectory(DirectoryReference.Combine(OtherModule.ModuleDirectory, "Internal")))
+						{
+							LoggerFunc("Module '{Name}' is adding '{Path}' from module '{OtherModule}' to InternalncludePaths. This is not allowed.", Module.Name, Path, OtherModule.Name);
+						}
+						else if (Path.IsUnderDirectory(DirectoryReference.Combine(OtherModule.ModuleDirectory, "Public"))
+							|| Path.IsUnderDirectory(DirectoryReference.Combine(OtherModule.ModuleDirectory, "Classes")))
+						{
+							LoggerFunc("Module '{Name}' is adding '{Path}' from module '{OtherModule}' to InternalncludePaths. Did you intend to add a public reference?.", Module.Name, Path, OtherModule.Name);
+						}
+					}
+				}
+			}
+		}
+
 		/// <summary>
 		/// Return whether the given platform requires a monolithic build
 		/// </summary>
diff --git a/Engine/Source/Programs/UnrealBuildTool/Configuration/UEBuildTarget.cs b/Engine/Source/Programs/UnrealBuildTool/Configuration/UEBuildTarget.cs
--- a/Engine/Source/Programs/UnrealBuildTool/Configuration/UEBuildTarget.cs
+++ b/Engine/Source/Programs/UnrealBuildTool/Configuration/UEBuildTarget.cs
@@ -3149,9 +3149,13 @@ namespace UnrealBuildTool
 
 			// Allow each platform to run any needed validation on the module.
 			// ie: iOS/macOS do a static library version build check for unsupported clang version builds.
-			foreach (UEBuildModule Module in Modules.Values)
+			using (GlobalTracer.Instance.BuildSpan("UEBuildTarget.PreBuildSetup().ValidateModules").StartActive())
 			{
-				UEBuildPlatform.GetBuildPlatform(Platform).ValidateModule(Module, Rules);
+				foreach (UEBuildModule Module in Modules.Values.OrderBy(x => x.Name))
+				{
+					UEBuildPlatform.GetBuildPlatform(Platform).ValidateModule(Module, Rules);
+					UEBuildPlatform.GetBuildPlatform(Platform).ValidateModuleIncludePaths(Module, Rules, Modules.Values);
+				}
 			}
 
 			if (!bCompileMonolithic)
diff --git a/Engine/Source/Programs/UnrealBuildTool/Platform/IOS/UEBuildIOS.cs b/Engine/Source/Programs/UnrealBuildTool/Platform/IOS/UEBuildIOS.cs
--- a/Engine/Source/Programs/UnrealBuildTool/Platform/IOS/UEBuildIOS.cs
+++ b/Engine/Source/Programs/UnrealBuildTool/Platform/IOS/UEBuildIOS.cs
@@ -1,4 +1,4 @@
-﻿// Copyright Epic Games, Inc. All Rights Reserved.
+// Copyright Epic Games, Inc. All Rights Reserved.
 
 using System;
 using System.Collections.Generic;
@@ -854,7 +854,9 @@ namespace UnrealBuildTool
 		}
 
 		public override void ValidateModule(UEBuildModule Module, ReadOnlyTargetRules Target)
-		{ 
+		{
+			base.ValidateModule(Module, Target);
+
 			if (BuildHostPlatform.Current.Platform == UnrealTargetPlatform.Mac && !Target.IOSPlatform.bSkipClangValidation)
 			{
 				IOSPlatformSDK SDK = (IOSPlatformSDK?)GetSDK() ?? new IOSPlatformSDK(Logger);
diff --git a/Engine/Source/Programs/UnrealBuildTool/Platform/Mac/UEBuildMac.cs b/Engine/Source/Programs/UnrealBuildTool/Platform/Mac/UEBuildMac.cs
--- a/Engine/Source/Programs/UnrealBuildTool/Platform/Mac/UEBuildMac.cs
+++ b/Engine/Source/Programs/UnrealBuildTool/Platform/Mac/UEBuildMac.cs
@@ -348,7 +348,9 @@ namespace UnrealBuildTool
 
 		static HashSet<FileReference> ValidatedLibs = new();
 		public override void ValidateModule(UEBuildModule Module, ReadOnlyTargetRules Target)
-		{ 
+		{
+			base.ValidateModule(Module, Target);
+
 			if (BuildHostPlatform.Current.Platform == UnrealTargetPlatform.Mac && !Target.MacPlatform.bSkipClangValidation)
 			{
 				ApplePlatformSDK SDK = (ApplePlatformSDK?)GetSDK() ?? new ApplePlatformSDK(Logger);
