- Fixed .precompiled file location to be architcture independent
- Added IntermediateDirectoryNoArch to UEBuildModule (we can use this in the future to put generated definitions headers, which can help reduce rebuilds going between multi- and single-arch)

diff --git a/Engine/Source/Programs/UnrealBuildTool/Configuration/UEBuildModule.cs b/Engine/Source/Programs/UnrealBuildTool/Configuration/UEBuildModule.cs
--- a/Engine/Source/Programs/UnrealBuildTool/Configuration/UEBuildModule.cs
+++ b/Engine/Source/Programs/UnrealBuildTool/Configuration/UEBuildModule.cs
@@ -30,6 +30,11 @@ namespace UnrealBuildTool
 		/// </summary>
 		public readonly DirectoryReference IntermediateDirectory;
 
+		/// <summary>
+		/// The directory for this module's generated files that are architecture independent
+		/// </summary>
+		public readonly DirectoryReference IntermediateDirectoryNoArch;
+
 		/// <summary>
 		/// The name that uniquely identifies the module.
 		/// </summary>
@@ -210,11 +215,13 @@ namespace UnrealBuildTool
 		/// </summary>
 		/// <param name="Rules">Rules for this module</param>
 		/// <param name="IntermediateDirectory">Intermediate directory for this module</param>
+		/// <param name="IntermediateDirectoryNoArch">Intermediate directory for this module for files that are architecture independent</param>
 		/// <param name="Logger">Logger for output</param>
-		public UEBuildModule(ModuleRules Rules, DirectoryReference IntermediateDirectory, ILogger Logger)
+		public UEBuildModule(ModuleRules Rules, DirectoryReference IntermediateDirectory, DirectoryReference IntermediateDirectoryNoArch, ILogger Logger)
 		{
 			this.Rules = Rules;
 			this.IntermediateDirectory = IntermediateDirectory;
+			this.IntermediateDirectoryNoArch = IntermediateDirectoryNoArch;
 
 			ModuleApiDefine = Name.ToUpperInvariant() + "_API";
 
diff --git a/Engine/Source/Programs/UnrealBuildTool/Configuration/UEBuildModuleCPP.cs b/Engine/Source/Programs/UnrealBuildTool/Configuration/UEBuildModuleCPP.cs
--- a/Engine/Source/Programs/UnrealBuildTool/Configuration/UEBuildModuleCPP.cs
+++ b/Engine/Source/Programs/UnrealBuildTool/Configuration/UEBuildModuleCPP.cs
@@ -181,8 +181,8 @@ namespace UnrealBuildTool
 			new KeyValuePair<string, string>("WebBrowser", "WebBrowserTexture"),
 		};
 
-		public UEBuildModuleCPP(ModuleRules Rules, DirectoryReference IntermediateDirectory, DirectoryReference? GeneratedCodeDirectory, ILogger Logger)
-			: base(Rules, IntermediateDirectory, Logger)
+		public UEBuildModuleCPP(ModuleRules Rules, DirectoryReference IntermediateDirectory, DirectoryReference IntermediateDirectoryNoArch, DirectoryReference? GeneratedCodeDirectory, ILogger Logger)
+			: base(Rules, IntermediateDirectory, IntermediateDirectoryNoArch, Logger)
 		{
 			this.GeneratedCodeDirectory = GeneratedCodeDirectory;
 
@@ -300,7 +300,7 @@ namespace UnrealBuildTool
 		/// </summary>
 		public virtual FileReference PrecompiledManifestLocation
 		{
-			get { return FileReference.Combine(IntermediateDirectory, String.Format("{0}.precompiled", Name)); }
+			get { return FileReference.Combine(IntermediateDirectoryNoArch, String.Format("{0}.precompiled", Name)); }
 		}
 
 		/// <summary>
diff --git a/Engine/Source/Programs/UnrealBuildTool/Configuration/UEBuildModuleExternal.cs b/Engine/Source/Programs/UnrealBuildTool/Configuration/UEBuildModuleExternal.cs
--- a/Engine/Source/Programs/UnrealBuildTool/Configuration/UEBuildModuleExternal.cs
+++ b/Engine/Source/Programs/UnrealBuildTool/Configuration/UEBuildModuleExternal.cs
@@ -15,8 +15,8 @@ namespace UnrealBuildTool
 	/// </summary>
 	class UEBuildModuleExternal : UEBuildModule
 	{
-		public UEBuildModuleExternal(ModuleRules Rules, DirectoryReference IntermediateDirectory, ILogger Logger)
-			: base(Rules, IntermediateDirectory, Logger)
+		public UEBuildModuleExternal(ModuleRules Rules, DirectoryReference IntermediateDirectory, DirectoryReference IntermediateDirectoryNoArch, ILogger Logger)
+			: base(Rules, IntermediateDirectory, IntermediateDirectoryNoArch, Logger)
 		{
 		}
 	}
diff --git a/Engine/Source/Programs/UnrealBuildTool/Configuration/UEBuildTarget.cs b/Engine/Source/Programs/UnrealBuildTool/Configuration/UEBuildTarget.cs
--- a/Engine/Source/Programs/UnrealBuildTool/Configuration/UEBuildTarget.cs
+++ b/Engine/Source/Programs/UnrealBuildTool/Configuration/UEBuildTarget.cs
@@ -3747,8 +3747,9 @@ namespace UnrealBuildTool
 		/// Finds the base output directory for a module
 		/// </summary>
 		/// <param name="ModuleRules">The rules object created for this module</param>
+		/// <param name="Architectures">The architectures (or none) to insert into the directory structure</param>
 		/// <returns>The output directory for compiled object files for this module</returns>
-		private DirectoryReference GetModuleIntermediateDirectory(ModuleRules ModuleRules)
+		private DirectoryReference GetModuleIntermediateDirectory(ModuleRules ModuleRules, UnrealArchitectures? Architectures)
 		{
 			// Get the root output directory and base name (target name/app name) for this binary
 			DirectoryReference BaseOutputDirectory = GetBaseOutputDirectory(ModuleRules);
@@ -5126,13 +5127,19 @@ namespace UnrealBuildTool
 				case ModuleRules.ModuleType.CPlusPlus:
 					return new UEBuildModuleCPP(
 							Rules: RulesObject,
-							IntermediateDirectory: GetModuleIntermediateDirectory(RulesObject),
+							IntermediateDirectory: GetModuleIntermediateDirectory(RulesObject, Architectures),
+							IntermediateDirectoryNoArch: GetModuleIntermediateDirectory(RulesObject, null),
 							GeneratedCodeDirectory: GeneratedCodeDirectory,
 							Logger
 						);
 
 				case ModuleRules.ModuleType.External:
-					return new UEBuildModuleExternal(RulesObject, GetModuleIntermediateDirectory(RulesObject), Logger);
+					return new UEBuildModuleExternal(
+							Rules: RulesObject,
+							IntermediateDirectory: GetModuleIntermediateDirectory(RulesObject, Architectures),
+							IntermediateDirectoryNoArch: GetModuleIntermediateDirectory(RulesObject, null),
+							Logger
+						);
 
 				default:
 					throw new BuildException("Unrecognized module type specified by 'Rules' object {0}", RulesObject.ToString());
