Removed the ability to disable inlining gen.cpp files. The engine code relies on this feature and can't compile without it now.

diff --git a/Engine/Source/Programs/UnrealBuildTool/Configuration/TargetRules.cs b/Engine/Source/Programs/UnrealBuildTool/Configuration/TargetRules.cs
--- a/Engine/Source/Programs/UnrealBuildTool/Configuration/TargetRules.cs
+++ b/Engine/Source/Programs/UnrealBuildTool/Configuration/TargetRules.cs
@@ -1971,13 +1971,6 @@ namespace UnrealBuildTool
 		[XmlConfigFile(Category = "BuildConfiguration")]
 		public bool bPublicSymbolsByDefault = false;
 
-		/// <summary>
-		/// Disable supports for inlining gen.cpps
-		/// </summary>
-		[XmlConfigFile(Name = "bDisableInliningGenCpps")]
-		[CommandLine("-DisableInliningGenCpps")]
-		public bool bDisableInliningGenCpps = false;
-
 		/// <summary>
 		/// Allows overriding the toolchain to be created for this target. This must match the name of a class declared in the UnrealBuildTool assembly.
 		/// </summary>
@@ -3701,11 +3694,6 @@ namespace UnrealBuildTool
 			get { return Inner.bPublicSymbolsByDefault; }
 		}
 
-		public bool bDisableInliningGenCpps
-		{
-			get { return Inner.bDisableInliningGenCpps; }
-		}
-
 		public string? ToolChainName
 		{
 			get { return Inner.ToolChainName; }
diff --git a/Engine/Source/Programs/UnrealBuildTool/Configuration/UEBuildModuleCPP.cs b/Engine/Source/Programs/UnrealBuildTool/Configuration/UEBuildModuleCPP.cs
--- a/Engine/Source/Programs/UnrealBuildTool/Configuration/UEBuildModuleCPP.cs
+++ b/Engine/Source/Programs/UnrealBuildTool/Configuration/UEBuildModuleCPP.cs
@@ -492,28 +492,25 @@ namespace UnrealBuildTool
 
 				if (GeneratedFiles.Count > 0)
 				{
-					if (!Target.bDisableInliningGenCpps)
+					// Remove any generated files from the compile list if they are inlined
+					foreach (FileItem CPPFileItem in CPPFiles)
 					{
-						// Remove any generated files from the compile list if they are inlined
-						foreach (FileItem CPPFileItem in CPPFiles)
+						var ListOfInlinedGenCpps = ModuleCompileEnvironment.MetadataCache.GetListOfInlinedGeneratedCppFiles(CPPFileItem);
+						foreach (string ListOfInlinedGenCppsItem in ListOfInlinedGenCpps)
 						{
-							var ListOfInlinedGenCpps = ModuleCompileEnvironment.MetadataCache.GetListOfInlinedGeneratedCppFiles(CPPFileItem);
-							foreach (string ListOfInlinedGenCppsItem in ListOfInlinedGenCpps)
+							string Prefix = "UHT/";
+							string Key = Prefix + ListOfInlinedGenCppsItem;
+							if (GeneratedFiles.Remove(Key, out FileItem? FoundGenCppFile))
 							{
-								string Prefix = "UHT/";
-								string Key = Prefix + ListOfInlinedGenCppsItem;
-								if (GeneratedFiles.Remove(Key, out FileItem? FoundGenCppFile))
+								if (!CompileEnvironment.FileInlineGenCPPMap.ContainsKey(CPPFileItem))
 								{
-									if (!CompileEnvironment.FileInlineGenCPPMap.ContainsKey(CPPFileItem))
-									{
-										CompileEnvironment.FileInlineGenCPPMap[CPPFileItem] = new List<FileItem>();
-									}
-									CompileEnvironment.FileInlineGenCPPMap[CPPFileItem].Add(FoundGenCppFile);
-								}
-								else
-								{
-									Logger.LogError("'{CPPFileItem}' is looking for a generated cpp with named '{HeaderFile}.gen.cpp'", CPPFileItem.AbsolutePath, ListOfInlinedGenCppsItem);
+									CompileEnvironment.FileInlineGenCPPMap[CPPFileItem] = new List<FileItem>();
 								}
+								CompileEnvironment.FileInlineGenCPPMap[CPPFileItem].Add(FoundGenCppFile);
+							}
+							else
+							{
+								Logger.LogError("'{CPPFileItem}' is looking for a generated cpp with named '{HeaderFile}.gen.cpp'", CPPFileItem.AbsolutePath, ListOfInlinedGenCppsItem);
 							}
 						}
 					}
@@ -1745,15 +1742,6 @@ namespace UnrealBuildTool
 				Result.Definitions.Add("UE_IS_ENGINE_MODULE=0");
 			}
 
-			if (Target.bDisableInliningGenCpps)
-			{
-				Result.Definitions.Add("UE_DISABLE_INLINE_GEN_CPP=1");
-			}
-			else
-			{
-				Result.Definitions.Add("UE_DISABLE_INLINE_GEN_CPP=0");
-			}
-
 			Result.Definitions.AddRange(EngineIncludeOrderHelper.GetDeprecationDefines(Rules.IncludeOrderVersion));
 
 			// For game modules, set the define for the project and target names, which will be used by the IMPLEMENT_PRIMARY_GAME_MODULE macro.
@@ -1812,15 +1800,6 @@ namespace UnrealBuildTool
 				CompileEnvironment.Definitions.Add("UE_IS_ENGINE_MODULE=0");
 			}
 
-			if (Rules.Target.bDisableInliningGenCpps)
-			{
-				CompileEnvironment.Definitions.Add("UE_DISABLE_INLINE_GEN_CPP=1");
-			}
-			else
-			{
-				CompileEnvironment.Definitions.Add("UE_DISABLE_INLINE_GEN_CPP=0");
-			}
-
 			CompileEnvironment.Definitions.AddRange(EngineIncludeOrderHelper.GetDeprecationDefines(Rules.IncludeOrderVersion));
 
 			// Add the module's private definitions.
diff --git a/Engine/Source/Runtime/CoreUObject/Public/UObject/ObjectMacros.h b/Engine/Source/Runtime/CoreUObject/Public/UObject/ObjectMacros.h
--- a/Engine/Source/Runtime/CoreUObject/Public/UObject/ObjectMacros.h
+++ b/Engine/Source/Runtime/CoreUObject/Public/UObject/ObjectMacros.h
@@ -680,11 +680,7 @@ struct COREUOBJECT_API FReferencerInformationList
 
 // Used to inline generated cpp files from UObject headers
 #define UE_INLINE_STRINGIFY(name) #name
-#if UE_DISABLE_INLINE_GEN_CPP
-	#define UE_INLINE_GENERATED_CPP_BY_NAME(name) UE_INLINE_STRINGIFY(CoreTypes.h)
-#else
-	#define UE_INLINE_GENERATED_CPP_BY_NAME(name) UE_INLINE_STRINGIFY(name.gen.cpp)
-#endif
+#define UE_INLINE_GENERATED_CPP_BY_NAME(name) UE_INLINE_STRINGIFY(name.gen.cpp)
 
 // This pair of macros is used to help implement GENERATED_BODY() and GENERATED_USTRUCT_BODY()
 #define BODY_MACRO_COMBINE_INNER(A,B,C,D) A##B##C##D
