Added to UBT an alternative module file for live coding.  When in use, UBT will only generate "not enabled" errors if explicitly listed in the lazy load modules.

#rnx

diff --git a/Engine/Source/Programs/UnrealBuildTool/System/HotReload.cs b/Engine/Source/Programs/UnrealBuildTool/System/HotReload.cs
--- a/Engine/Source/Programs/UnrealBuildTool/System/HotReload.cs
+++ b/Engine/Source/Programs/UnrealBuildTool/System/HotReload.cs
@@ -7,9 +7,9 @@ using System.IO;
 using System.Linq;
 using System.Text;
 using System.Threading;
-using System.Threading.Tasks;
 using EpicGames.Core;
 using Microsoft.Extensions.Logging;
+using System.Text.Json;
 using OpenTracing.Util;
 using UnrealBuildBase;
 
@@ -133,6 +133,23 @@ namespace UnrealBuildTool
 		}
 	}
 
+	/// <summary>
+	/// Contents of the JSON version of the live coding modules file
+	/// </summary>
+	class LiveCodingModules
+	{
+
+		/// <summary>
+		/// These modules have been loaded by a process and are enabled for patching
+		/// </summary>
+		public List<string> EnabledModules { get; set; } = new();
+
+		/// <summary>
+		/// These modules have been loaded by a process, but not explicitly enabled
+		/// </summary>
+		public List<string> LazyLoadModules { get; set; } = new();
+	}
+
 	static class HotReload
 	{
 		/// <summary>
@@ -531,6 +548,27 @@ namespace UnrealBuildTool
 			}
 		}
 
+		/// <summary>
+		/// Given a collection of strings which are file paths, create a hash set from the file name and extension.
+		/// Empty strings are eliminated.
+		/// </summary>
+		/// <param name="Collection">Source collection</param>
+		/// <returns>Trimmed and unique collection</returns>
+		private static HashSet<string> CreateHashSetFromFileList(IEnumerable<string> Collection)
+		{
+			// Parse it out into a set of filenames
+			HashSet<string> Out = new HashSet<string>(FileReference.Comparer);
+			foreach (string Line in Collection)
+			{
+				string TrimLine = Line.Trim();
+				if (TrimLine.Length > 0)
+				{
+					Out.Add(Path.GetFileName(TrimLine));
+				}
+			}
+			return Out;
+		}
+
 		/// <summary>
 		/// Determine what needs to be built for a target
 		/// </summary>
@@ -561,18 +599,27 @@ namespace UnrealBuildTool
 				// Make sure we're not overwriting any lazy-loaded modules
 				if (TargetDescriptor.LiveCodingModules != null)
 				{
-					// Read the list of modules that we're allowed to build
-					string[] Lines = FileReference.ReadAllLines(TargetDescriptor.LiveCodingModules);
 
-					// Parse it out into a set of filenames
-					HashSet<string> AllowedOutputFileNames = new HashSet<string>(FileReference.Comparer);
-					foreach (string Line in Lines)
+					// In the old style module list, which was just a text file, we allow only modules found in the known list of enabled modules.
+					//		All other modules are assumed to be lazy loaded.
+					// In the new style module list, which is a json file, we disallow modules found in list of lazy loaded modules and allow
+					//		all other modules. The enabled module list is not used in the new format, but is there for diagnostics or future expansion.
+					HashSet<string>? AllowedOutputFileNames = null;
+					HashSet<string>? DisallowedOutputFileNames = null;
+					if (TargetDescriptor.LiveCodingModules.GetExtension() == ".json")
 					{
-						string TrimLine = Line.Trim();
-						if (TrimLine.Length > 0)
+						LiveCodingModules? Modules = JsonSerializer.Deserialize<LiveCodingModules>(File.OpenRead(TargetDescriptor.LiveCodingModules.FullName));
+						if (Modules == null)
 						{
-							AllowedOutputFileNames.Add(Path.GetFileName(TrimLine));
+							throw new BuildException("Unable to load live coding modules file '{0}'", TargetDescriptor.LiveCodingModules.FullName);
 						}
+						DisallowedOutputFileNames = CreateHashSetFromFileList(Modules.LazyLoadModules);
+					}
+					else
+					{
+						// Read the list of modules that we're allowed to build
+						string[] Lines = FileReference.ReadAllLines(TargetDescriptor.LiveCodingModules);
+						AllowedOutputFileNames = CreateHashSetFromFileList(Lines);
 					}
 
 					// Find all the binaries that we're actually going to build
@@ -586,7 +633,12 @@ namespace UnrealBuildTool
 					}
 
 					// Find all the files that will be built that aren't allowed
-					List<FileReference> ProtectedOutputFiles = OutputFiles.Where(x => !AllowedOutputFileNames.Contains(x.GetFileName())).ToList();
+					List<FileReference> ProtectedOutputFiles = OutputFiles.Where(x =>
+							(AllowedOutputFileNames != null && !AllowedOutputFileNames.Contains(x.GetFileName())) ||
+							(DisallowedOutputFileNames != null && DisallowedOutputFileNames.Contains(x.GetFileName()))
+						).ToList();
+
+					// Generate the error messages
 					if (ProtectedOutputFiles.Count > 0)
 					{
 						FileReference.WriteAllLines(new FileReference(TargetDescriptor.LiveCodingModules.FullName + ".out"), ProtectedOutputFiles.Select(x => x.ToString()));
