UnrealBuildTool: Remove CanUseParallelExecutor as it is no longer needed as the UAT ParallelExecutor has been removed and is obsolete.

#rnx

diff --git a/Engine/Source/Programs/AutomationTool/AutomationUtils/UnrealBuild.cs b/Engine/Source/Programs/AutomationTool/AutomationUtils/UnrealBuild.cs
--- a/Engine/Source/Programs/AutomationTool/AutomationUtils/UnrealBuild.cs
+++ b/Engine/Source/Programs/AutomationTool/AutomationUtils/UnrealBuild.cs
@@ -494,16 +494,6 @@ namespace AutomationTool
 			BuildProductFiles.Clear();
 		}
 
-		public bool CanUseXGE(UnrealTargetPlatform Platform)
-		{
-			return PlatformExports.CanUseXGE(Platform, Log.Logger);
-		}
-
-		public bool CanUseParallelExecutor(UnrealTargetPlatform Platform)
-		{
-			return PlatformExports.CanUseParallelExecutor(Platform, Log.Logger);
-		}
-
 		private bool ParseParam(string Name)
 		{
 			return OwnerCommand != null && OwnerCommand.ParseParam(Name);
diff --git a/Engine/Source/Programs/UnrealBuildTool/Configuration/UEBuildPlatform.cs b/Engine/Source/Programs/UnrealBuildTool/Configuration/UEBuildPlatform.cs
--- a/Engine/Source/Programs/UnrealBuildTool/Configuration/UEBuildPlatform.cs
+++ b/Engine/Source/Programs/UnrealBuildTool/Configuration/UEBuildPlatform.cs
@@ -691,14 +691,6 @@ namespace UnrealBuildTool
 			return true;
 		}
 
-		/// <summary>
-		/// If this platform can be compiled with the parallel executor
-		/// </summary>
-		public virtual bool CanUseParallelExecutor()
-		{
-			return CanUseXGE();
-		}
-
 		/// <summary>
 		/// If this platform can be compiled with FASTBuild
 		/// </summary>
diff --git a/Engine/Source/Programs/UnrealBuildTool/Platform/Linux/UEBuildLinux.cs b/Engine/Source/Programs/UnrealBuildTool/Platform/Linux/UEBuildLinux.cs
--- a/Engine/Source/Programs/UnrealBuildTool/Platform/Linux/UEBuildLinux.cs
+++ b/Engine/Source/Programs/UnrealBuildTool/Platform/Linux/UEBuildLinux.cs
@@ -285,12 +285,6 @@ namespace UnrealBuildTool
 			return BuildHostPlatform.Current.Platform == UnrealTargetPlatform.Win64;
 		}
 
-		public override bool CanUseParallelExecutor()
-		{
-			// No known problems with parallel executor, always use for build machines
-			return true;
-		}
-
 		/// <summary>
 		/// Determines if the given name is a build product for a target.
 		/// </summary>
diff --git a/Engine/Source/Programs/UnrealBuildTool/System/PlatformExports.cs b/Engine/Source/Programs/UnrealBuildTool/System/PlatformExports.cs
--- a/Engine/Source/Programs/UnrealBuildTool/System/PlatformExports.cs
+++ b/Engine/Source/Programs/UnrealBuildTool/System/PlatformExports.cs
@@ -107,9 +107,10 @@ namespace UnrealBuildTool
 		/// <param name="Platform">Platform to check</param>
 		/// <param name="Logger">Logger for output</param>
 		/// <returns>True if the platform supports the parallel executor in UAT</returns>
+		[Obsolete]
 		public static bool CanUseParallelExecutor(UnrealTargetPlatform Platform, ILogger Logger)
 		{
-			return UEBuildPlatform.IsPlatformAvailable(Platform) && UEBuildPlatform.GetBuildPlatform(Platform).CanUseParallelExecutor();
+			return true;
 		}
 
 		/// <summary>
