UnrealBuildTool: If module valdiation is set to Error, raise an exception if any errors are found

#rnx

diff --git a/Engine/Source/Programs/UnrealBuildTool/Configuration/UEBuildPlatform.cs b/Engine/Source/Programs/UnrealBuildTool/Configuration/UEBuildPlatform.cs
--- a/Engine/Source/Programs/UnrealBuildTool/Configuration/UEBuildPlatform.cs
+++ b/Engine/Source/Programs/UnrealBuildTool/Configuration/UEBuildPlatform.cs
@@ -736,15 +736,16 @@ namespace UnrealBuildTool
 		/// <param name="Module">The UEBuildModule that needs to be validated</param>
 		/// <param name="Target">Options for the target being built</param>
 		/// <param name="AllModules">Other modules to validate against, if needed</param>
-		public virtual void ValidateModuleIncludePaths(UEBuildModule Module, ReadOnlyTargetRules Target, IEnumerable<UEBuildModule> AllModules)
+		public virtual bool ValidateModuleIncludePaths(UEBuildModule Module, ReadOnlyTargetRules Target, IEnumerable<UEBuildModule> AllModules)
 		{
 			if (Module.Rules.ModuleIncludePathWarningLevel <= WarningLevel.Off
 				&& Module.Rules.ModuleIncludePrivateWarningLevel <= WarningLevel.Off
 				&& Module.Rules.ModuleIncludeSubdirectoryWarningLevel <= WarningLevel.Off)
 			{
-				return;
+				return false;
 			}
 
+			bool AnyErrors = false;
 			void LoggerFunc(string? message, params object?[] args)
 			{
 				if (Module.Rules.ModuleIncludePathWarningLevel == WarningLevel.Warning)
@@ -754,6 +755,7 @@ namespace UnrealBuildTool
 				else if (Module.Rules.ModuleIncludePathWarningLevel == WarningLevel.Error)
 				{
 					Logger.LogError($"Error: {message}", args);
+					AnyErrors = true;
 				}
 			}
 
@@ -766,6 +768,7 @@ namespace UnrealBuildTool
 				else if (Module.Rules.ModuleIncludePrivateWarningLevel == WarningLevel.Error)
 				{
 					Logger.LogError($"Error: {message}", args);
+					AnyErrors = true;
 				}
 			}
 
@@ -778,6 +781,7 @@ namespace UnrealBuildTool
 				else if (Module.Rules.ModuleIncludeSubdirectoryWarningLevel == WarningLevel.Error)
 				{
 					Logger.LogError($"Error: {message}", args);
+					AnyErrors = true;
 				}
 			}
 
@@ -948,6 +952,8 @@ namespace UnrealBuildTool
 					}
 				}
 			}
+
+			return AnyErrors;
 		}
 
 		/// <summary>
diff --git a/Engine/Source/Programs/UnrealBuildTool/Configuration/UEBuildTarget.cs b/Engine/Source/Programs/UnrealBuildTool/Configuration/UEBuildTarget.cs
--- a/Engine/Source/Programs/UnrealBuildTool/Configuration/UEBuildTarget.cs
+++ b/Engine/Source/Programs/UnrealBuildTool/Configuration/UEBuildTarget.cs
@@ -3151,10 +3151,16 @@ namespace UnrealBuildTool
 			// ie: iOS/macOS do a static library version build check for unsupported clang version builds.
 			using (GlobalTracer.Instance.BuildSpan("UEBuildTarget.PreBuildSetup().ValidateModules").StartActive())
 			{
+				bool AnyErrors = false;
 				foreach (UEBuildModule Module in Modules.Values.OrderBy(x => x.Name))
 				{
 					UEBuildPlatform.GetBuildPlatform(Platform).ValidateModule(Module, Rules);
-					UEBuildPlatform.GetBuildPlatform(Platform).ValidateModuleIncludePaths(Module, Rules, Modules.Values);
+					AnyErrors |= UEBuildPlatform.GetBuildPlatform(Platform).ValidateModuleIncludePaths(Module, Rules, Modules.Values);
+				}
+
+				if (AnyErrors)
+				{
+					throw new BuildException("Errors validating modules.");
 				}
 			}
 
