Replaced the target rule bool "bMergeModuleAndGeneratedUnityFiles" with a list of modules that disable the functionality of merging the module and generated cpp files(DisableMergingModuleAndGeneratedFilesInUnityFiles). The engine modules will depend on merging the module and generated cpp files being in the same unity files very soon so we can't allow turning off the functionality across all the code. This change allows licensees to disable the merging of a specific module if it is needed.

diff --git a/Engine/Source/Programs/UnrealBuildTool/Configuration/TargetRules.cs b/Engine/Source/Programs/UnrealBuildTool/Configuration/TargetRules.cs
--- a/Engine/Source/Programs/UnrealBuildTool/Configuration/TargetRules.cs
+++ b/Engine/Source/Programs/UnrealBuildTool/Configuration/TargetRules.cs
@@ -1408,11 +1408,10 @@ namespace UnrealBuildTool
 		public bool bForceUnityBuild = false;
 
 		/// <summary>
-		/// Whether to merge module and generated unity files for faster compilation.
+		/// List of modules that disables merging module and generated cpp files in the same unity files.
 		/// </summary>
-		[CommandLine("-DisableMergingUnityFiles", Value = "false")]
-		[XmlConfigFile(Category = "BuildConfiguration")]
-		public bool bMergeModuleAndGeneratedUnityFiles = true;
+		[XmlConfigFile(Category = "ModuleConfiguration", Name = "DisableMergingModuleAndGeneratedFilesInUnityFiles")]
+		public string[]? DisableMergingModuleAndGeneratedFilesInUnityFiles = null;
 
 		/// <summary>
 		/// Use a heuristic to determine which files are currently being iterated on and exclude them from unity blobs, result in faster
@@ -3325,10 +3324,11 @@ namespace UnrealBuildTool
 			get { return Inner.bForceUnityBuild; }
 		}
 
-		public bool bMergeModuleAndGeneratedUnityFiles
+		public IReadOnlyList<string>? DisableMergingModuleAndGeneratedFilesInUnityFiles
 		{
-			get { return Inner.bMergeModuleAndGeneratedUnityFiles; }
+			get { return Inner.DisableMergingModuleAndGeneratedFilesInUnityFiles; }
 		}
+
 		public bool bAdaptiveUnityDisablesOptimizations
 		{
 			get { return Inner.bAdaptiveUnityDisablesOptimizations; }
diff --git a/Engine/Source/Programs/UnrealBuildTool/Configuration/UEBuildModuleCPP.cs b/Engine/Source/Programs/UnrealBuildTool/Configuration/UEBuildModuleCPP.cs
--- a/Engine/Source/Programs/UnrealBuildTool/Configuration/UEBuildModuleCPP.cs
+++ b/Engine/Source/Programs/UnrealBuildTool/Configuration/UEBuildModuleCPP.cs
@@ -463,7 +463,7 @@ namespace UnrealBuildTool
 			List<FileItem> CPPFiles = new List<FileItem>(InputFiles.CPPFiles);
 			List<FileItem> GeneratedFileItems = new List<FileItem>();
 			CppCompileEnvironment GeneratedCPPCompileEnvironment = CompileEnvironment;
-			bool bMergeUnityFiles = Target.bMergeModuleAndGeneratedUnityFiles && Rules.bMergeUnityFiles;
+			bool bMergeUnityFiles = !(Target.DisableMergingModuleAndGeneratedFilesInUnityFiles?.Contains(Name) ?? false) && Rules.bMergeUnityFiles;
 
 			// Compile all the generated CPP files
 			if (GeneratedCppDirectories != null && !CompileEnvironment.bHackHeaderGenerator)
