Added another possible fix for SNDBS when using PCH chains.

diff --git a/Engine/Source/Programs/UnrealBuildTool/ToolChain/ClangToolChain.cs b/Engine/Source/Programs/UnrealBuildTool/ToolChain/ClangToolChain.cs
--- a/Engine/Source/Programs/UnrealBuildTool/ToolChain/ClangToolChain.cs
+++ b/Engine/Source/Programs/UnrealBuildTool/ToolChain/ClangToolChain.cs
@@ -944,6 +944,7 @@ namespace UnrealBuildTool
 				while (PCHInstance != null)
 				{
 					CompileAction.PrerequisiteItems.Add(PCHInstance.Output.GetPrecompiledHeaderFile(CompileEnvironment.Architecture)!);
+					CompileAction.PrerequisiteItems.Add(PCHInstance.HeaderFile!);
 					PCHInstance = PCHInstance.ParentPCHInstance;
 				}
 			}
@@ -965,6 +966,7 @@ namespace UnrealBuildTool
 				while (ParentPCHInstance != null)
 				{
 					CompileAction.PrerequisiteItems.Add(ParentPCHInstance.Output.GetPrecompiledHeaderFile(CompileEnvironment.Architecture)!);
+					CompileAction.PrerequisiteItems.Add(ParentPCHInstance.HeaderFile!);
 					ParentPCHInstance = ParentPCHInstance.ParentPCHInstance;
 				}
 			}
