Fix iOS BP only builds on Windows

diff --git a/Engine/Source/Programs/UnrealBuildTool/Platform/IOS/IOSToolChain.cs b/Engine/Source/Programs/UnrealBuildTool/Platform/IOS/IOSToolChain.cs
--- a/Engine/Source/Programs/UnrealBuildTool/Platform/IOS/IOSToolChain.cs
+++ b/Engine/Source/Programs/UnrealBuildTool/Platform/IOS/IOSToolChain.cs
@@ -1761,14 +1761,8 @@ namespace UnrealBuildTool
 					string SchemeName;
 					if (AppName == "UnrealGame" || AppName == "UnrealClient")
 					{
-						if (Target.bBuildAsFramework)
-						{
-							SchemeName = "UnrealGame";
-						}
-						else
-						{
-							SchemeName = "UE5";
-						}
+
+						SchemeName = "UnrealGame";
 					}
 					else
 					{
