UnrealBuildTool: Allow overriding bWarningsAsErrors on a per module basis

diff --git a/Engine/Source/Programs/UnrealBuildTool/Configuration/ModuleRules.cs b/Engine/Source/Programs/UnrealBuildTool/Configuration/ModuleRules.cs
--- a/Engine/Source/Programs/UnrealBuildTool/Configuration/ModuleRules.cs
+++ b/Engine/Source/Programs/UnrealBuildTool/Configuration/ModuleRules.cs
@@ -801,6 +801,11 @@ namespace UnrealBuildTool
 		/// <inheritdoc cref="ShadowVariableWarningLevelPrivate"/>
 		private WarningLevel ShadowVariableWarningLevelPrivate;
 
+		/// <summary>
+		/// Whether to enable all warnings as errors. UE enables most warnings as errors already, but disables a few (such as deprecation warnings).
+		/// </summary>`
+		public bool bWarningsAsErrors = false;
+
 		/// <summary>
 		/// How to treat unsafe implicit type cast warnings (e.g., double->float or int64->int32)
 		/// </summary>
diff --git a/Engine/Source/Programs/UnrealBuildTool/Configuration/UEBuildModuleCPP.cs b/Engine/Source/Programs/UnrealBuildTool/Configuration/UEBuildModuleCPP.cs
--- a/Engine/Source/Programs/UnrealBuildTool/Configuration/UEBuildModuleCPP.cs
+++ b/Engine/Source/Programs/UnrealBuildTool/Configuration/UEBuildModuleCPP.cs
@@ -1678,6 +1678,7 @@ namespace UnrealBuildTool
 			Result.bEnableExceptions |= Rules.bEnableExceptions;
 			Result.bEnableObjCExceptions |= Rules.bEnableObjCExceptions;
 			Result.bEnableObjCAutomaticReferenceCounting = Rules.bEnableObjCAutomaticReferenceCounting;
+			Result.bWarningsAsErrors |= Rules.bWarningsAsErrors;
 			Result.ShadowVariableWarningLevel = Rules.ShadowVariableWarningLevel;
 			Result.UnsafeTypeCastWarningLevel = Rules.UnsafeTypeCastWarningLevel;
 			Result.bDisableStaticAnalysis = Rules.bDisableStaticAnalysis;
