Re-added a few properties that were removed and added obsolete tags to them.

diff --git a/Engine/Source/Programs/UnrealBuildTool/Configuration/TargetRules.cs b/Engine/Source/Programs/UnrealBuildTool/Configuration/TargetRules.cs
--- a/Engine/Source/Programs/UnrealBuildTool/Configuration/TargetRules.cs
+++ b/Engine/Source/Programs/UnrealBuildTool/Configuration/TargetRules.cs
@@ -1407,6 +1407,14 @@ namespace UnrealBuildTool
 		[XmlConfigFile(Category = "BuildConfiguration")]
 		public bool bForceUnityBuild = false;
 
+		/// <summary>
+		/// Whether to merge module and generated unity files for faster compilation.
+		/// </summary>
+		[Obsolete("Deprecated in UE5.3; use DisableMergingModuleAndGeneratedFilesInUnityFiles instead")]
+		[CommandLine("-DisableMergingUnityFiles", Value = "false")]
+		[XmlConfigFile(Category = "BuildConfiguration")]
+		public bool bMergeModuleAndGeneratedUnityFiles = true;
+
 		/// <summary>
 		/// List of modules that disables merging module and generated cpp files in the same unity files.
 		/// </summary>
@@ -1970,6 +1978,14 @@ namespace UnrealBuildTool
 		[XmlConfigFile(Category = "BuildConfiguration")]
 		public bool bPublicSymbolsByDefault = false;
 
+		/// <summary>
+		/// Disable supports for inlining gen.cpps
+		/// </summary>
+		[Obsolete("Deprecated in UE5.3; the engine code relies on inlining generated cpp files to compile")]
+		[XmlConfigFile(Name = "bDisableInliningGenCpps")]
+		[CommandLine("-DisableInliningGenCpps")]
+		public bool bDisableInliningGenCpps = false;
+
 		/// <summary>
 		/// Allows overriding the toolchain to be created for this target. This must match the name of a class declared in the UnrealBuildTool assembly.
 		/// </summary>
