PR #10191: Allow usage of incredibuild on macOS (Contributed by [REDACTED_PERSON])

diff --git a/Engine/Source/Programs/UnrealBuildTool/Platform/IOS/UEBuildIOS.cs b/Engine/Source/Programs/UnrealBuildTool/Platform/IOS/UEBuildIOS.cs
--- a/Engine/Source/Programs/UnrealBuildTool/Platform/IOS/UEBuildIOS.cs
+++ b/Engine/Source/Programs/UnrealBuildTool/Platform/IOS/UEBuildIOS.cs
@@ -1017,7 +1017,7 @@ namespace UnrealBuildTool
 
 		public override bool CanUseXGE()
 		{
-			return false;
+			return BuildHostPlatform.Current.Platform == UnrealTargetPlatform.Mac;
 		}
 
 		public override bool CanUseFASTBuild()
diff --git a/Engine/Source/Programs/UnrealBuildTool/Platform/Mac/UEBuildMac.cs b/Engine/Source/Programs/UnrealBuildTool/Platform/Mac/UEBuildMac.cs
--- a/Engine/Source/Programs/UnrealBuildTool/Platform/Mac/UEBuildMac.cs
+++ b/Engine/Source/Programs/UnrealBuildTool/Platform/Mac/UEBuildMac.cs
@@ -265,11 +265,6 @@ namespace UnrealBuildTool
 		{
 		}
 
-		public override bool CanUseXGE()
-		{
-			return false;
-		}
-
 		public override bool CanUseFASTBuild()
 		{
 			return true;
