- Fixed warning about always running a script phase in Legacy XcodeProjects

diff --git a/Engine/Source/Programs/UnrealBuildTool/ProjectFiles/Xcode/XcodeProjectLegacy.cs b/Engine/Source/Programs/UnrealBuildTool/ProjectFiles/Xcode/XcodeProjectLegacy.cs
--- a/Engine/Source/Programs/UnrealBuildTool/ProjectFiles/Xcode/XcodeProjectLegacy.cs
+++ b/Engine/Source/Programs/UnrealBuildTool/ProjectFiles/Xcode/XcodeProjectLegacy.cs
@@ -770,13 +770,13 @@ namespace UnrealBuildTool.XcodeProjectLegacy
 			Content.Append("/* Begin PBXShellScriptBuildPhase section */" + ProjectFileGenerator.NewLine);
 			Content.Append(string.Format("\t\t{0} /* Sign Frameworks */ = {{" + ProjectFileGenerator.NewLine, ShellScriptGuid));
 			Content.Append("\t\t\tisa = PBXShellScriptBuildPhase;" + ProjectFileGenerator.NewLine);
-			Content.Append("\t\t\talwaysOutOfDate = 1;" + ProjectFileGenerator.NewLine);
 			Content.Append("\t\t\tbuildActionMask = 2147483647;" + ProjectFileGenerator.NewLine);
 			Content.Append("\t\t\tfiles = (" + ProjectFileGenerator.NewLine);
 			Content.Append("\t\t\t);" + ProjectFileGenerator.NewLine);
 			Content.Append("\t\t\tinputPaths = (" + ProjectFileGenerator.NewLine);
 			Content.Append("\t\t\t);" + ProjectFileGenerator.NewLine);
 			Content.Append("\t\t\toutputPaths = (" + ProjectFileGenerator.NewLine);
+			Content.Append("\t\t\t\t/dev/null" + ProjectFileGenerator.NewLine);
 			Content.Append("\t\t\t);" + ProjectFileGenerator.NewLine);
 			Content.Append("\t\t\tname = \"Sign Manual Frameworks\";" + ProjectFileGenerator.NewLine);
 			Content.Append("\t\t\trunOnlyForDeploymentPostprocessing = 0;" + ProjectFileGenerator.NewLine);
