UnrealBuildTool: Remove obsolete HordeExecutor prototype

#rnx

diff --git a/Engine/Source/Programs/UnrealBuildTool/Configuration/BuildConfiguration.cs b/Engine/Source/Programs/UnrealBuildTool/Configuration/BuildConfiguration.cs
--- a/Engine/Source/Programs/UnrealBuildTool/Configuration/BuildConfiguration.cs
+++ b/Engine/Source/Programs/UnrealBuildTool/Configuration/BuildConfiguration.cs
@@ -42,13 +42,6 @@ namespace UnrealBuildTool
 		[XmlConfigFile]
 		public bool bAllowHybridExecutor = false;
 
-		/// <summary>
-		/// Whether Horde remote compute may be used. Highly experimental, disabled by default.
-		/// </summary>
-		[XmlConfigFile]
-		[CommandLine("-HordeCompute", Value = "true")]
-		public bool bAllowHordeCompute = false;
-
 		/// <summary>
 		/// Whether XGE may be used.
 		/// </summary>
diff --git a/Engine/Source/Programs/UnrealBuildTool/Executors/Experimental/HordeExecutor.cs b/Engine/Source/Programs/UnrealBuildTool/Executors/Experimental/HordeExecutor.cs
deleted file mode 100644
diff --git a/Engine/Source/Programs/UnrealBuildTool/System/ActionGraph.cs b/Engine/Source/Programs/UnrealBuildTool/System/ActionGraph.cs
--- a/Engine/Source/Programs/UnrealBuildTool/System/ActionGraph.cs
+++ b/Engine/Source/Programs/UnrealBuildTool/System/ActionGraph.cs
@@ -371,10 +371,6 @@ namespace UnrealBuildTool
 				{
 					return new SNDBS(TargetDescriptors);
 				}
-				else if (BuildConfiguration.bAllowHordeCompute && HordeExecutor.IsAvailable())
-				{
-					return new HordeExecutor(BuildConfiguration.MaxParallelActions, BuildConfiguration.bAllCores, BuildConfiguration.bCompactOutput, Logger);
-				}
 			}
 
 			return new ParallelExecutor(BuildConfiguration.MaxParallelActions, BuildConfiguration.bAllCores, BuildConfiguration.bCompactOutput, Logger);
