UnrealBuildTool: Stash shared paths in the CppCompileEnvironment after creating a shared response file so those paths can still be used for validation

#rnx

diff --git a/Engine/Source/Programs/UnrealBuildTool/Platform/Windows/VCToolChain.cs b/Engine/Source/Programs/UnrealBuildTool/Platform/Windows/VCToolChain.cs
--- a/Engine/Source/Programs/UnrealBuildTool/Platform/Windows/VCToolChain.cs
+++ b/Engine/Source/Programs/UnrealBuildTool/Platform/Windows/VCToolChain.cs
@@ -1378,8 +1378,7 @@ namespace UnrealBuildTool
 			BaseCompileAction.IncludePaths.AddRange(CompileEnvironment.UserIncludePaths);
 			BaseCompileAction.SystemIncludePaths.AddRange(CompileEnvironment.SystemIncludePaths);
 			
-			// Maybe not the nicest way of checking but if we don't have any user include paths it means we have added them to a shared rsp instead
-			if (CompileEnvironment.UserIncludePaths.Count != 0)
+			if (!CompileEnvironment.SharedSystemIncludePaths.Any())
 			{
 				BaseCompileAction.SystemIncludePaths.AddRange(EnvVars.IncludePaths);
 			}
@@ -2070,6 +2069,11 @@ namespace UnrealBuildTool
 				AddSystemIncludePath(Arguments, IncludePath, Target.WindowsPlatform.Compiler);
 			}
 
+			// Stash shared include paths for validation purposes
+			NewCompileEnvironment.SharedUserIncludePaths = new(NewCompileEnvironment.UserIncludePaths);
+			NewCompileEnvironment.SharedSystemIncludePaths = new(NewCompileEnvironment.SystemIncludePaths);
+			NewCompileEnvironment.SharedSystemIncludePaths.UnionWith(EnvVars.IncludePaths);
+
 			NewCompileEnvironment.UserIncludePaths.Clear();
 			NewCompileEnvironment.SystemIncludePaths.Clear();
 
diff --git a/Engine/Source/Programs/UnrealBuildTool/System/CppCompileEnvironment.cs b/Engine/Source/Programs/UnrealBuildTool/System/CppCompileEnvironment.cs
--- a/Engine/Source/Programs/UnrealBuildTool/System/CppCompileEnvironment.cs
+++ b/Engine/Source/Programs/UnrealBuildTool/System/CppCompileEnvironment.cs
@@ -456,6 +456,18 @@ namespace UnrealBuildTool
 		/// </summary>
 		public HashSet<DirectoryReference> SystemIncludePaths;
 
+		/// <summary>
+		/// The include paths which were previously in UserIncludePaths, but are now in a shared response file, persisted in the environment for validation.
+		/// Do not add to this set unless a shared response is in use, and only when removing those headers from UserIncludePaths.
+		/// </summary>
+		public HashSet<DirectoryReference> SharedUserIncludePaths;
+
+		/// <summary>
+		/// The include paths which were previously in SystemIncludePaths, but are now in a shared response file, persisted in the environment for validation.
+		/// Do not add to this set unless a shared response is in use, and only when removing those headers from SystemIncludePaths.
+		/// </summary>
+		public HashSet<DirectoryReference> SharedSystemIncludePaths;
+
 		/// <summary>
 		/// List of paths to search for compiled module interface (*.ifc) files
 		/// </summary>
@@ -607,6 +619,8 @@ namespace UnrealBuildTool
 			this.SharedPCHs = new List<PrecompiledHeaderTemplate>();
 			this.UserIncludePaths = new HashSet<DirectoryReference>();
 			this.SystemIncludePaths = new HashSet<DirectoryReference>();
+			this.SharedUserIncludePaths = new HashSet<DirectoryReference>();
+			this.SharedSystemIncludePaths = new HashSet<DirectoryReference>();
 			this.ModuleInterfacePaths = new HashSet<DirectoryReference>();
 		}
 
@@ -671,6 +685,8 @@ namespace UnrealBuildTool
 			bAllowRemotelyCompiledPCHs = Other.bAllowRemotelyCompiledPCHs;
 			UserIncludePaths = new HashSet<DirectoryReference>(Other.UserIncludePaths);
 			SystemIncludePaths = new HashSet<DirectoryReference>(Other.SystemIncludePaths);
+			SharedUserIncludePaths = new HashSet<DirectoryReference>(Other.SharedUserIncludePaths);
+			SharedSystemIncludePaths = new HashSet<DirectoryReference>(Other.SharedSystemIncludePaths);
 			ModuleInterfacePaths = new HashSet<DirectoryReference>(Other.ModuleInterfacePaths);
 			bCheckSystemHeadersForModification = Other.bCheckSystemHeadersForModification;
 			ForceIncludeFiles.AddRange(Other.ForceIncludeFiles);
diff --git a/Engine/Source/Programs/UnrealBuildTool/ToolChain/ClangToolChain.cs b/Engine/Source/Programs/UnrealBuildTool/ToolChain/ClangToolChain.cs
--- a/Engine/Source/Programs/UnrealBuildTool/ToolChain/ClangToolChain.cs
+++ b/Engine/Source/Programs/UnrealBuildTool/ToolChain/ClangToolChain.cs
@@ -1056,6 +1056,11 @@ namespace UnrealBuildTool
 			List<string> Arguments = new List<string>();
 
 			GetCompileArguments_Global(CompileEnvironment, Arguments);
+
+			// Stash shared include paths for validation purposes
+			NewCompileEnvironment.SharedUserIncludePaths = new(NewCompileEnvironment.UserIncludePaths);
+			NewCompileEnvironment.SharedSystemIncludePaths = new(NewCompileEnvironment.SystemIncludePaths);
+
 			NewCompileEnvironment.UserIncludePaths.Clear();
 			NewCompileEnvironment.SystemIncludePaths.Clear();
 
