UnrealBuildTool: Override the temp directory to a subfolder, and allow specifying a directory to use instead of the default %TMP% in BuildConfiguration.xml

  <BuildConfiguration>
    <TempDirectory>C:\Temp</TempDirectory>
  </BuildConfiguration>

There have been reports of large amounts of temp files in the root standard temp directory causing slowdown when compiling and affecting other programs

diff --git a/Engine/Source/Programs/UnrealBuildTool/UnrealBuildTool.cs b/Engine/Source/Programs/UnrealBuildTool/UnrealBuildTool.cs
--- a/Engine/Source/Programs/UnrealBuildTool/UnrealBuildTool.cs
+++ b/Engine/Source/Programs/UnrealBuildTool/UnrealBuildTool.cs
@@ -308,7 +308,13 @@ namespace UnrealBuildTool
 			/// </summary>
 			[XmlConfigFile(Category = "Log")]
 			public int LogFileBackupCount = Log.LogFileBackupCount;
-			
+
+			/// <summary>
+			/// If set TMP\TEMP will be overidden to this directory, each process will create a unique subdirectory in this folder 
+			/// </summary>
+			[XmlConfigFile(Category = "BuildConfiguration")]
+			public string? TempDirectory = null;
+
 			/// <summary>
 			/// Initialize the options with the given command line arguments
 			/// </summary>
@@ -404,6 +410,7 @@ namespace UnrealBuildTool
 		private static int Main(string[] ArgumentsArray)
 		{
 			FileReference? RunFile = null;
+			DirectoryReference? TempDirectory = null;
 			SingleInstanceMutex? Mutex = null;
 			JsonTracer? Tracer = null;
 
@@ -562,6 +569,43 @@ namespace UnrealBuildTool
 				{
 				}
 
+				// Override the temp directory
+				try
+				{
+					// If the temp directory is already overridden from a parent process, do not override again
+					DirectoryReference? OverrideTempDirectory = DirectoryReference.FromString(Environment.GetEnvironmentVariable("UBT_TEMP_OVERRIDE"));
+					if (OverrideTempDirectory == null)
+					{
+						OverrideTempDirectory = new DirectoryReference(Path.Combine(Path.GetTempPath(), "UnrealBuildTool"));
+						if (Options.TempDirectory != null)
+						{
+							if (Directory.Exists(Options.TempDirectory))
+							{
+								OverrideTempDirectory = new DirectoryReference(Options.TempDirectory);
+							}
+							else
+							{
+								Logger.LogWarning("Warning: TempDirectory override '{Override}' does not exist, using '{Temp}'", Options.TempDirectory, OverrideTempDirectory.FullName);
+							}
+						}
+
+						// Create a subdirectory for the current process
+						OverrideTempDirectory = DirectoryReference.Combine(OverrideTempDirectory, RunFile?.GetFileName() ?? Process.GetCurrentProcess().Id.ToString());
+						DirectoryReference.CreateDirectory(OverrideTempDirectory);
+
+						// Only save the directory overriding temp so it can be cleaned up the first time it is overridden
+						Environment.SetEnvironmentVariable("UBT_TEMP_OVERRIDE", OverrideTempDirectory.FullName);
+						TempDirectory = OverrideTempDirectory;
+					}
+
+					Logger.LogDebug("Setting temp directory to '{Path}'", OverrideTempDirectory);
+					Environment.SetEnvironmentVariable("TMP", OverrideTempDirectory.FullName);
+					Environment.SetEnvironmentVariable("TEMP", OverrideTempDirectory.FullName);
+				}
+				catch
+				{
+				}
+
 				// Acquire a lock for this branch
 				if((ModeOptions & ToolModeOptions.SingleInstance) != 0 && !Options.bNoMutex)
 				{
@@ -676,6 +720,18 @@ namespace UnrealBuildTool
 					{
 					}
 				}
+
+				// Delete the temp directory if overridden. This is a unique subdirectory so it is safe to remove
+				if (TempDirectory != null)
+				{
+					try
+					{
+						DirectoryReference.Delete(TempDirectory, true);
+					}
+					catch
+					{
+					}
+				}
 			}
 		}
 	}
