UnrealBuildTool: Shorten temp path to just process id

#rnx

diff --git a/Engine/Source/Programs/UnrealBuildTool/UnrealBuildTool.cs b/Engine/Source/Programs/UnrealBuildTool/UnrealBuildTool.cs
--- a/Engine/Source/Programs/UnrealBuildTool/UnrealBuildTool.cs
+++ b/Engine/Source/Programs/UnrealBuildTool/UnrealBuildTool.cs
@@ -590,7 +590,7 @@ namespace UnrealBuildTool
 						}
 
 						// Create a subdirectory for the current process
-						OverrideTempDirectory = DirectoryReference.Combine(OverrideTempDirectory, RunFile?.GetFileName() ?? Process.GetCurrentProcess().Id.ToString());
+						OverrideTempDirectory = DirectoryReference.Combine(OverrideTempDirectory, Process.GetCurrentProcess().Id.ToString());
 						DirectoryReference.CreateDirectory(OverrideTempDirectory);
 
 						// Only save the directory overriding temp so it can be cleaned up the first time it is overridden
