UnrealBuildTool: Allow overriding windows application icon with TargetRules.WindowsPlatform.ApplicationIcon

diff --git a/Engine/Source/Programs/UnrealBuildTool/Platform/Windows/UEBuildWindows.cs b/Engine/Source/Programs/UnrealBuildTool/Platform/Windows/UEBuildWindows.cs
--- a/Engine/Source/Programs/UnrealBuildTool/Platform/Windows/UEBuildWindows.cs
+++ b/Engine/Source/Programs/UnrealBuildTool/Platform/Windows/UEBuildWindows.cs
@@ -370,6 +370,11 @@ namespace UnrealBuildTool
 		/// </summary>
 		public string? ManifestFile;
 
+		/// <summary>
+		/// Specifies the path to an .ico file to use as the appliction icon. Can be assigned to null and will default to Engine/Build/Windows/Resources/Default.ico for engine targets or Build/Windows/Application.ico for projects.
+		/// </summary>
+		public string? ApplicationIcon;
+
 		/// <summary>
 		/// Enables strict standard conformance mode (/permissive-).
 		/// </summary>
@@ -749,6 +754,11 @@ namespace UnrealBuildTool
 			get { return Inner.ManifestFile; }
 		}
 
+		public string? ApplicationIcon
+		{
+			get { return Inner.ApplicationIcon; }
+		}
+
 		public bool bNeedsLegacyStdioDefinitionsLib
 		{
 			get { return Inner.bNeedsLegacyStdioDefinitionsLib; }
diff --git a/Engine/Source/Programs/UnrealBuildTool/Platform/Windows/VCToolChain.cs b/Engine/Source/Programs/UnrealBuildTool/Platform/Windows/VCToolChain.cs
--- a/Engine/Source/Programs/UnrealBuildTool/Platform/Windows/VCToolChain.cs
+++ b/Engine/Source/Programs/UnrealBuildTool/Platform/Windows/VCToolChain.cs
@@ -1832,6 +1832,10 @@ namespace UnrealBuildTool
 
 		public virtual FileReference GetApplicationIcon(FileReference? ProjectFile)
 		{
+			if (Target.WindowsPlatform.ApplicationIcon != null)
+			{
+				return new FileReference(Target.WindowsPlatform.ApplicationIcon);
+			}
 			return WindowsPlatform.GetWindowsApplicationIcon(ProjectFile);
 		}
 
