UnrealBuildTool: Override the temp directory to a subfolder, and allow specifying a directory to use instead of the default %TMP% in BuildConfiguration.xml

  <BuildConfiguration>
    <TempDirectory>C:\Temp</TempDirectory>
  </BuildConfiguration>

You can also optionally delete the temp directory on exit with the following:

  <BuildConfiguration>
    <bDeleteTempDirectory>true</bDeleteTempDirectory>
  </BuildConfiguration>

There have been reports of large amounts of temp files in the root standard temp directory causing slowdown when compiling and affecting other programs

diff --git a/Engine/Source/Programs/UnrealBuildTool/UnrealBuildTool.cs b/Engine/Source/Programs/UnrealBuildTool/UnrealBuildTool.cs
--- a/Engine/Source/Programs/UnrealBuildTool/UnrealBuildTool.cs
+++ b/Engine/Source/Programs/UnrealBuildTool/UnrealBuildTool.cs
@@ -308,7 +308,19 @@ namespace UnrealBuildTool
 			/// </summary>
 			[XmlConfigFile(Category = "Log")]
 			public int LogFileBackupCount = Log.LogFileBackupCount;
-			
+
+			/// <summary>
+			/// If set TMP\TEMP will be overidden to this directory, each process will create a unique subdirectory in this folder.
+			/// </summary>
+			[XmlConfigFile(Category = "BuildConfiguration")]
+			public string? TempDirectory = null;
+
+			/// <summary>
+			/// If set the application temp directory will be deleted on exit, only when running with a single instance mutex.
+			/// </summary>
+			[XmlConfigFile(Category = "BuildConfiguration")]
+			public bool bDeleteTempDirectory = false;
+
 			/// <summary>
 			/// Initialize the options with the given command line arguments
 			/// </summary>
@@ -404,6 +416,7 @@ namespace UnrealBuildTool
 		private static int Main(string[] ArgumentsArray)
 		{
 			FileReference? RunFile = null;
+			DirectoryReference? TempDirectory = null;
 			SingleInstanceMutex? Mutex = null;
 			JsonTracer? Tracer = null;
 
@@ -562,6 +575,49 @@ namespace UnrealBuildTool
 				{
 				}
 
+				// Override the temp directory
+				try
+				{
+					// If the temp directory is already overridden from a parent process, do not override again
+					if (string.IsNullOrEmpty(Environment.GetEnvironmentVariable("UnrealBuildTool_TMP")))
+					{
+						DirectoryReference OverrideTempDirectory = new DirectoryReference(Path.Combine(Path.GetTempPath(), "UnrealBuildTool"));
+						if (Options.TempDirectory != null)
+						{
+							if (Directory.Exists(Options.TempDirectory))
+							{
+								OverrideTempDirectory = new DirectoryReference(Options.TempDirectory);
+								if (OverrideTempDirectory.GetDirectoryName() != "UnrealBuildTool")
+								{
+									OverrideTempDirectory = DirectoryReference.Combine(OverrideTempDirectory, "UnrealBuildTool");
+								}
+							}
+							else
+							{
+								Logger.LogWarning("Warning: TempDirectory override '{Override}' does not exist, using '{Temp}'", Options.TempDirectory, OverrideTempDirectory.FullName);
+							}
+						}
+
+						OverrideTempDirectory = DirectoryReference.Combine(OverrideTempDirectory, ContentHash.MD5(Encoding.UTF8.GetBytes(Unreal.UnrealBuildToolDllPath.FullName)).ToString().Substring(0, 8));
+						DirectoryReference.CreateDirectory(OverrideTempDirectory);
+
+						Logger.LogDebug("Setting temp directory to '{Path}'", OverrideTempDirectory);
+						Environment.SetEnvironmentVariable("UnrealBuildTool_TMP", OverrideTempDirectory.FullName);
+						Environment.SetEnvironmentVariable("TMP", OverrideTempDirectory.FullName);
+						Environment.SetEnvironmentVariable("TEMP", OverrideTempDirectory.FullName);
+
+						// Deleting the directory is only safe in single instance mode, and only if requested
+						if ((ModeOptions & ToolModeOptions.SingleInstance) != 0 && !Options.bNoMutex && Options.bDeleteTempDirectory)
+						{
+							Logger.LogDebug("Temp directory '{Path}' will be deleted on exit", OverrideTempDirectory);
+							TempDirectory = OverrideTempDirectory;
+						}
+					}
+				}
+				catch
+				{
+				}
+
 				// Acquire a lock for this branch
 				if((ModeOptions & ToolModeOptions.SingleInstance) != 0 && !Options.bNoMutex)
 				{
@@ -659,12 +715,6 @@ namespace UnrealBuildTool
 					Tracer.Flush();
 				}
 
-				// Dispose of the mutex. Must be done last to ensure that another process does not startup and start trying to write to the same log file.
-				if (Mutex != null)
-				{
-					Mutex.Dispose();
-				}
-
 				// Delete the ubt run file
 				if (RunFile != null)
 				{
@@ -676,6 +726,24 @@ namespace UnrealBuildTool
 					{
 					}
 				}
+
+				// Remove the the temp subdirectory. TempDirectory will only be set if running in single instance mode when Options.DeleteTempDirectory is enabled
+				if (TempDirectory != null)
+				{
+					try
+					{
+						DirectoryReference.Delete(TempDirectory, true);
+					}
+					catch
+					{
+					}
+				}
+
+				// Dispose of the mutex. Must be done last to ensure that another process does not startup and start trying to write to the same log file.
+				if (Mutex != null)
+				{
+					Mutex.Dispose();
+				}
 			}
 		}
 	}
