UnrealBuildTool: Put /fastfail on a setting and disable by default

#rnx

diff --git a/Engine/Source/Programs/UnrealBuildTool/Platform/Windows/UEBuildWindows.cs b/Engine/Source/Programs/UnrealBuildTool/Platform/Windows/UEBuildWindows.cs
--- a/Engine/Source/Programs/UnrealBuildTool/Platform/Windows/UEBuildWindows.cs
+++ b/Engine/Source/Programs/UnrealBuildTool/Platform/Windows/UEBuildWindows.cs
@@ -266,6 +266,14 @@ namespace UnrealBuildTool
 		[CommandLine("-CompilerVersion")]
 		public string? CompilerVersion = null;
 
+		/// <summary>
+		/// True if /fastfail should be passed to the msvc compiler and linker
+		/// </summary>
+		[ConfigFile(ConfigHierarchyType.Engine, "/Script/WindowsTargetPlatform.WindowsTargetSettings", "bVCFastFail")]
+		[XmlConfigFile(Category = "WindowsPlatform")]
+		[CommandLine("-VCFastFail")]
+		public bool bVCFastFail = false;
+
 		/// <summary>
 		/// True if we should use the Clang linker (LLD) when we are compiling with Clang, or Intel linker (xilink\xilib) when we are compiling with Intel oneAPI, otherwise we use the MSVC linker.
 		/// </summary>
@@ -697,6 +705,11 @@ namespace UnrealBuildTool
 			get { return Inner.bUseCPPWinRT; }
 		}
 
+		public bool bVCFastFail
+		{
+			get { return Inner.bVCFastFail; }
+		}
+
 		public bool bAllowClangLinker
 		{
 			get { return Inner.bAllowClangLinker; }
diff --git a/Engine/Source/Programs/UnrealBuildTool/Platform/Windows/VCToolChain.cs b/Engine/Source/Programs/UnrealBuildTool/Platform/Windows/VCToolChain.cs
--- a/Engine/Source/Programs/UnrealBuildTool/Platform/Windows/VCToolChain.cs
+++ b/Engine/Source/Programs/UnrealBuildTool/Platform/Windows/VCToolChain.cs
@@ -542,7 +542,7 @@ namespace UnrealBuildTool
 				}
 			}
 
-			if (Target.WindowsPlatform.Compiler.IsMSVC())
+			if (Target.WindowsPlatform.Compiler.IsMSVC() && Target.WindowsPlatform.bVCFastFail)
 			{
 				Arguments.Add("/fastfail");
 			}
@@ -1228,7 +1228,7 @@ namespace UnrealBuildTool
 				}
 			}
 
-			if (Target.WindowsPlatform.Compiler.IsMSVC() && !LinkEnvironment.bPGOOptimize && !LinkEnvironment.bPGOProfile)
+			if (Target.WindowsPlatform.Compiler.IsMSVC() && Target.WindowsPlatform.bVCFastFail && !LinkEnvironment.bPGOOptimize && !LinkEnvironment.bPGOProfile)
 			{
 				Arguments.Add("/fastfail");
 			}
