UnrealBuildTool: Fix issues with Modified Files Only compiles

#rnx

diff --git a/Engine/Source/Programs/UnrealBuildTool/Configuration/TargetDescriptor.cs b/Engine/Source/Programs/UnrealBuildTool/Configuration/TargetDescriptor.cs
--- a/Engine/Source/Programs/UnrealBuildTool/Configuration/TargetDescriptor.cs
+++ b/Engine/Source/Programs/UnrealBuildTool/Configuration/TargetDescriptor.cs
@@ -145,12 +145,9 @@ namespace UnrealBuildTool
 				foreach (FileReference FileList in FileLists)
 				{
 					string[] Files = FileReference.ReadAllLines(FileList);
-					foreach (string File in Files)
+					foreach (string File in Files.Where(x => !string.IsNullOrWhiteSpace(x)))
 					{
-						if (!String.IsNullOrWhiteSpace(File))
-						{
-							SpecificFilesToCompile.Add(FileReference.Combine(Unreal.RootDirectory, File));
-						}
+						SpecificFilesToCompile.Add(FileReference.Combine(Unreal.RootDirectory, File));
 					}
 				}
 
diff --git a/Engine/Source/Programs/UnrealBuildTool/Configuration/UEBuildTarget.cs b/Engine/Source/Programs/UnrealBuildTool/Configuration/UEBuildTarget.cs
--- a/Engine/Source/Programs/UnrealBuildTool/Configuration/UEBuildTarget.cs
+++ b/Engine/Source/Programs/UnrealBuildTool/Configuration/UEBuildTarget.cs
@@ -1241,6 +1241,13 @@ namespace UnrealBuildTool
 				RulesObject.bDisableLinking = true;
 			}
 
+			// Don't link if we're compiling specific files and ignore build outputs
+			if (Descriptor.SpecificFilesToCompile.Any())
+			{
+				RulesObject.bDisableLinking = true;
+				RulesObject.bIgnoreBuildOutputs = true;
+			}
+
 			// Generate a build target from this rules module
 			UEBuildTarget Target;
 			using (GlobalTracer.Instance.BuildSpan("UEBuildTarget constructor").StartActive())
diff --git a/Engine/Source/Programs/UnrealBuildTool/Modes/BuildMode.cs b/Engine/Source/Programs/UnrealBuildTool/Modes/BuildMode.cs
--- a/Engine/Source/Programs/UnrealBuildTool/Modes/BuildMode.cs
+++ b/Engine/Source/Programs/UnrealBuildTool/Modes/BuildMode.cs
@@ -174,10 +174,13 @@ namespace UnrealBuildTool
 				using (GlobalTracer.Instance.BuildSpan("TargetDescriptor.ParseCommandLine()").StartActive())
 				{
 					TargetDescriptors = TargetDescriptor.ParseCommandLine(Arguments, BuildConfiguration.bUsePrecompiled, BuildConfiguration.bSkipRulesCompile, BuildConfiguration.bForceRulesCompile, Logger);
+
+					// Handle BuildConfiguration arguments nested inside -Target= args
+					TargetDescriptors.ForEach(x => x.AdditionalArguments.ApplyTo(BuildConfiguration));
 				}
 
 				// Hack for specific files compile; don't build the ShaderCompileWorker target that's added to the command line for generated project files
-				if(TargetDescriptors.Count >= 2)
+				if (TargetDescriptors.Count > 1)
 				{
 					TargetDescriptors.RemoveAll(x => (x.Name == "ShaderCompileWorker" || x.Name == "LiveCodingConsole" || x.Name == "InterchangeWorker") && x.SpecificFilesToCompile.Count > 0);
 				}
@@ -455,7 +458,9 @@ namespace UnrealBuildTool
 					{
 						// If there are headers in the list, expand the FilesToBuild to also include all files that include those headers
 						if (TargetDescriptor.bSingleFileBuildDependents)
+						{
 							FilesToBuild = GetAllSourceFilesIncludingHeader(FilesToBuild, TargetDescriptor.ProjectFile, Makefiles[Idx].Actions, Logger);
+						}
 
 						// We have specific files to compile so we will only queue up those files.
 						List<FileItem> ProducedItems = CreateLinkedActionsFromFileList(TargetDescriptor, BuildConfiguration, FilesToBuild, QueuedActions[Idx], Logger);
@@ -865,7 +870,9 @@ namespace UnrealBuildTool
 			Func<FileReference, bool> IsCode = x => IsHeader(x)|| x.HasExtension(".cpp") || x.HasExtension(".c");
 			List<FileReference> SpecificHeaderFiles = SpecificFilesToCompile.Where(IsHeader).ToList();
 			if (SpecificHeaderFiles.Count == 0)
+			{
 				return SpecificFilesToCompile;
+			}
 
 			// TODO: This code below works very similar to the code before this changelist but is not a great solution for figuring out which files to compile.
 			// Suggested approach is to:
@@ -952,7 +959,7 @@ namespace UnrealBuildTool
 			}
 
 			NewFiles.UnionWith(SpecificFilesToCompile);
-			return NewFiles.ToList(); ;
+			return NewFiles.ToList();
 		}
 
 		/// <summary>
