- Fixed location of ini files that are Restricted + Program + PlatformExtension ini files
- Previously it would look in, say, "/Engine/Restricted/NotForLicensees/Programs/UnrealHeaderTool/Platforms/Mac/Config/MacEngine.ini", which breaks the structure of Restricted and Platforms. Now it will look in "/Engine/Restricted/NotForLicensees/Platforms/Mac/Programs/UnrealHeaderTool/Config/MacEngine.ini" which is correct

diff --git a/Engine/Source/Programs/UnrealBuildTool/System/ConfigHierarchy.cs b/Engine/Source/Programs/UnrealBuildTool/System/ConfigHierarchy.cs
--- a/Engine/Source/Programs/UnrealBuildTool/System/ConfigHierarchy.cs
+++ b/Engine/Source/Programs/UnrealBuildTool/System/ConfigHierarchy.cs
@@ -956,13 +956,13 @@ namespace UnrealBuildTool
 			// The base expansion (ie, no expansion)
 			new ConfigLayerExpansion { }, 
 			// Restricted Locations
-			new ConfigLayerExpansion { Before1 = "{ENGINE}/", After1 = "{ENGINE}/Restricted/NotForLicensees/",	Before2 = "{PROJECT}/Config/", After2 = "{RESTRICTEDPROJECT_NFL}/Config/" },
-			new ConfigLayerExpansion { Before1 = "{ENGINE}/", After1 = "{ENGINE}/Restricted/NoRedist/",			Before2 = "{PROJECT}/Config/", After2 = "{RESTRICTEDPROJECT_NR}/Config/" },
+			new ConfigLayerExpansion { Before1 = "{ENGINE}/", After1 = "{ENGINE}/Restricted/NotForLicensees/",	Before2 = "{PROJECT}/Config/", After2 = "{RESTRICTEDPROJECT_NFL}/{OPT_SUBDIR}Config/" },
+			new ConfigLayerExpansion { Before1 = "{ENGINE}/", After1 = "{ENGINE}/Restricted/NoRedist/",			Before2 = "{PROJECT}/Config/", After2 = "{RESTRICTEDPROJECT_NR}/{OPT_SUBDIR}Config/" },
 			// Platform Extensions
 			new ConfigLayerExpansion { Before1 = "{ENGINE}/Config/{PLATFORM}/", After1 = "{EXTENGINE}/Config/",	Before2 = "{PROJECT}/Config/{PLATFORM}/", After2 = "{EXTPROJECT}/Config/" },
 			// Platform Extensions in Restricted Locations
-			new ConfigLayerExpansion { Before1 = "{ENGINE}/Config/{PLATFORM}/", After1 = "{ENGINE}/Restricted/NotForLicensees/Platforms/{PLATFORM}/Config/",   Before2 = "{PROJECT}/Config/{PLATFORM}/", After2 = "{RESTRICTEDPROJECT_NFL}/Platforms/{PLATFORM}/Config/" },
-			new ConfigLayerExpansion { Before1 = "{ENGINE}/Config/{PLATFORM}/", After1 = "{ENGINE}/Restricted/NoRedist/Platforms/{PLATFORM}/Config/",          Before2 = "{PROJECT}/Config/{PLATFORM}/", After2 = "{RESTRICTEDPROJECT_NR}/Platforms/{PLATFORM}/Config/" },
+			new ConfigLayerExpansion { Before1 = "{ENGINE}/Config/{PLATFORM}/", After1 = "{ENGINE}/Restricted/NotForLicensees/Platforms/{PLATFORM}/Config/",   Before2 = "{PROJECT}/Config/{PLATFORM}/", After2 = "{RESTRICTEDPROJECT_NFL}/Platforms/{PLATFORM}/{OPT_SUBDIR}Config/" },
+			new ConfigLayerExpansion { Before1 = "{ENGINE}/Config/{PLATFORM}/", After1 = "{ENGINE}/Restricted/NoRedist/Platforms/{PLATFORM}/Config/",          Before2 = "{PROJECT}/Config/{PLATFORM}/", After2 = "{RESTRICTEDPROJECT_NR}/Platforms/{PLATFORM}/{OPT_SUBDIR}Config/" },
 		};
 
 		// Match FPlatformProcess::UserDir()
@@ -1036,11 +1036,12 @@ namespace UnrealBuildTool
 			{
 				DirectoryReference NFLDir;
 				DirectoryReference NRDir;
+				string OptionalSubDir = "";
 				if (ProjectDir.IsUnderDirectory(Unreal.EngineDirectory))
 				{
-					string RelativeDir = ProjectDir.MakeRelativeTo(Unreal.EngineDirectory);
-					NFLDir = DirectoryReference.Combine(Unreal.EngineDirectory, "Restricted/NotForLicensees", RelativeDir);
-					NRDir = DirectoryReference.Combine(Unreal.EngineDirectory, "Restricted/NoRedist", RelativeDir);
+					OptionalSubDir = ProjectDir.MakeRelativeTo(Unreal.EngineDirectory) + "/";
+					NFLDir = DirectoryReference.Combine(Unreal.EngineDirectory, "Restricted/NotForLicensees");
+					NRDir = DirectoryReference.Combine(Unreal.EngineDirectory, "Restricted/NoRedist");
 				}
 				else
 				{
@@ -1053,6 +1054,8 @@ namespace UnrealBuildTool
 				OutString = OutString.Replace("{EXTPROJECT}", PlatformExtensionProjectConfigDir);
 				OutString = OutString.Replace("{RESTRICTEDPROJECT_NFL}", NFLDir.FullName);
 				OutString = OutString.Replace("{RESTRICTEDPROJECT_NR}", NRDir.FullName);
+				OutString = OutString.Replace("{OPT_SUBDIR}", OptionalSubDir);
+
 			}
 
 			return OutString;
diff --git a/Engine/Source/Runtime/Core/Private/Misc/ConfigContext.cpp b/Engine/Source/Runtime/Core/Private/Misc/ConfigContext.cpp
--- a/Engine/Source/Runtime/Core/Private/Misc/ConfigContext.cpp
+++ b/Engine/Source/Runtime/Core/Private/Misc/ConfigContext.cpp
@@ -434,6 +434,18 @@ FString FConfigContext::PerformFinalExpansions(const FString& InString, const FS
 	OutString = OutString.Replace(TEXT("{RESTRICTEDPROJECT_NFL}"), *ProjectNotForLicenseesDir);
 	OutString = OutString.Replace(TEXT("{RESTRICTEDPROJECT_NR}"), *ProjectNoRedistDir);
 
+	if (FPaths::IsUnderDirectory(ProjectRootDir, EngineRootDir))
+	{
+		FString RelativeDir = ProjectRootDir;
+		FPaths::MakePathRelativeTo(RelativeDir, *EngineRootDir);
+		
+		OutString = OutString.Replace(TEXT("{OPT_SUBDIR}"), *RelativeDir);
+	}
+	else
+	{
+		OutString = OutString.Replace(TEXT("{OPT_SUBDIR}"), TEXT(""));
+	}
+	
 	if (Platform.Len() > 0)
 	{
 		OutString = OutString.Replace(TEXT("{EXTENGINE}"), *GetPerPlatformDirs(InPlatform).PlatformExtensionEngineDir);
diff --git a/Engine/Source/Runtime/Core/Public/Misc/ConfigHierarchy.h b/Engine/Source/Runtime/Core/Public/Misc/ConfigHierarchy.h
--- a/Engine/Source/Runtime/Core/Public/Misc/ConfigHierarchy.h
+++ b/Engine/Source/Runtime/Core/Public/Misc/ConfigHierarchy.h
@@ -78,12 +78,12 @@ inline FConfigLayerExpansion GConfigExpansions[] =
 	// Restricted Locations
 	{ 
 		TEXT("{ENGINE}/"),						TEXT("{ENGINE}/Restricted/NotForLicensees/"),	
-		TEXT("{PROJECT}/Config/"),				TEXT("{RESTRICTEDPROJECT_NFL}/Config/"), 
+		TEXT("{PROJECT}/Config/"),				TEXT("{RESTRICTEDPROJECT_NFL}/{OPT_SUBDIR}Config/"),
 		EConfigExpansionFlags::ForUncooked | EConfigExpansionFlags::ForCooked
 	},
 	{ 
 		TEXT("{ENGINE}/"),						TEXT("{ENGINE}/Restricted/NoRedist/"),			
-		TEXT("{PROJECT}/Config/"),				TEXT("{RESTRICTEDPROJECT_NR}/Config/"), 
+		TEXT("{PROJECT}/Config/"),				TEXT("{RESTRICTEDPROJECT_NR}/{OPT_SUBDIR}Config/"),
 		EConfigExpansionFlags::ForUncooked 
 	},
 
@@ -105,8 +105,8 @@ inline FConfigLayerExpansion GConfigExpansions[] =
 		EConfigExpansionFlags::ForUncooked | EConfigExpansionFlags::ForCooked // | EConfigExpansionFlags::ForPlugin 
 	},
 	{
-		TEXT("{ENGINE}/Config/{PLATFORM}/"),	TEXT("{ENGINE}/Restricted/NoRedist/Platforms/{PLATFORM}/Config/"),			
-		TEXT("{PROJECT}/Config/{PLATFORM}/"),	TEXT("{RESTRICTEDPROJECT_NR}/Platforms/{PLATFORM}/Config/"), 
+		TEXT("{ENGINE}/Config/{PLATFORM}/"),	TEXT("{ENGINE}/Restricted/NoRedist/Platforms/{PLATFORM}/{OPT_SUBDIR}Config/"),
+		TEXT("{PROJECT}/Config/{PLATFORM}/"),	TEXT("{RESTRICTEDPROJECT_NR}/Platforms/{PLATFORM}/{OPT_SUBDIR}Config/"), 
 		EConfigExpansionFlags::ForUncooked // | EConfigExpansionFlags::ForPlugin
 	},
 };
