- Skip hidden source files

diff --git a/Engine/Source/Programs/UnrealBuildTool/Configuration/UEBuildModuleCPP.cs b/Engine/Source/Programs/UnrealBuildTool/Configuration/UEBuildModuleCPP.cs
--- a/Engine/Source/Programs/UnrealBuildTool/Configuration/UEBuildModuleCPP.cs
+++ b/Engine/Source/Programs/UnrealBuildTool/Configuration/UEBuildModuleCPP.cs
@@ -381,7 +381,7 @@ namespace UnrealBuildTool
 
 			// Find all the input files
 			Dictionary<DirectoryItem, FileItem[]> DirectoryToSourceFiles = new Dictionary<DirectoryItem, FileItem[]>();
-			InputFileCollection InputFiles = FindInputFiles(Target.Platform, DirectoryToSourceFiles);
+			InputFileCollection InputFiles = FindInputFiles(Target.Platform, DirectoryToSourceFiles, Logger);
 
 			foreach (KeyValuePair<DirectoryItem, FileItem[]> Pair in DirectoryToSourceFiles)
 			{
@@ -1823,8 +1823,9 @@ namespace UnrealBuildTool
 		/// </summary>
 		/// <param name="Platform">The platform the module is being built for</param>
 		/// <param name="DirectoryToSourceFiles">Map of directory to source files inside it</param>
+		/// <param name="Logger">Logger for output</param>
 		/// <returns>Set of source files that should be built</returns>
-		public InputFileCollection FindInputFiles(UnrealTargetPlatform Platform, Dictionary<DirectoryItem, FileItem[]> DirectoryToSourceFiles)
+		public InputFileCollection FindInputFiles(UnrealTargetPlatform Platform, Dictionary<DirectoryItem, FileItem[]> DirectoryToSourceFiles, ILogger Logger)
 		{
 			ReadOnlyHashSet<string> ExcludedNames = UEBuildPlatform.GetBuildPlatform(Platform).GetExcludedFolderNames();
 
@@ -1834,7 +1835,7 @@ namespace UnrealBuildTool
 			foreach (DirectoryReference Dir in ModuleDirectories)
 			{
 				DirectoryItem ModuleDirectoryItem = DirectoryItem.GetItemByDirectoryReference(Dir);
-				FindInputFilesFromDirectoryRecursive(ModuleDirectoryItem, ExcludedNames, SourceDirectories, DirectoryToSourceFiles, InputFiles);
+				FindInputFilesFromDirectoryRecursive(ModuleDirectoryItem, ExcludedNames, SourceDirectories, DirectoryToSourceFiles, InputFiles, Logger);
 			}
 
 			return InputFiles;
@@ -1848,10 +1849,11 @@ namespace UnrealBuildTool
 		/// <param name="SourceDirectories">Set of all non-empty source directories.</param>
 		/// <param name="DirectoryToSourceFiles">Map from directory to source files inside it</param>
 		/// <param name="InputFiles">Collection of source files, categorized by type</param>
-		static void FindInputFilesFromDirectoryRecursive(DirectoryItem BaseDirectory, ReadOnlyHashSet<string> ExcludedNames, HashSet<DirectoryReference> SourceDirectories, Dictionary<DirectoryItem, FileItem[]> DirectoryToSourceFiles, InputFileCollection InputFiles)
+		/// <param name="Logger">Logger for output</param>
+		static void FindInputFilesFromDirectoryRecursive(DirectoryItem BaseDirectory, ReadOnlyHashSet<string> ExcludedNames, HashSet<DirectoryReference> SourceDirectories, Dictionary<DirectoryItem, FileItem[]> DirectoryToSourceFiles, InputFileCollection InputFiles, ILogger Logger)
 		{
 			bool bIgnoreFileFound;
-			FileItem[] SourceFiles = FindInputFilesFromDirectory(BaseDirectory, InputFiles, out bIgnoreFileFound);
+			FileItem[] SourceFiles = FindInputFilesFromDirectory(BaseDirectory, InputFiles, out bIgnoreFileFound, Logger);
 
 			if (bIgnoreFileFound)
 			{
@@ -1862,7 +1864,7 @@ namespace UnrealBuildTool
 			{
 				if (!ExcludedNames.Contains(SubDirectory.Name))
 				{
-					FindInputFilesFromDirectoryRecursive(SubDirectory, ExcludedNames, SourceDirectories, DirectoryToSourceFiles, InputFiles);
+					FindInputFilesFromDirectoryRecursive(SubDirectory, ExcludedNames, SourceDirectories, DirectoryToSourceFiles, InputFiles, Logger);
 				}
 			}
 
@@ -1873,20 +1875,44 @@ namespace UnrealBuildTool
 			DirectoryToSourceFiles.Add(BaseDirectory, SourceFiles);
 		}
 
+		// List of known extensions that are checked in FindInputFilesFromDirectory, so we can log out only if a hidden file
+		// is skipped, only if it has an extension that would normally be compiled
+		static string[] KnownInputFileExtensions =
+		{
+			".h", ".ipsh",
+			".cpp", ".ixx", ".c", ".cc",
+			".m", ".mm", 
+			".rc", ".ispc",
+		};
+
 		/// <summary>
 		/// Finds the input files that should be built for this module, from a given directory
 		/// </summary>
 		/// <param name="BaseDirectory"></param>
 		/// <param name="InputFiles"></param>
 		/// <param name="bIgnoreFileFound"></param>
+		/// <param name="Logger">Logger for output</param>
 		/// <returns>Array of source files</returns>
-		static FileItem[] FindInputFilesFromDirectory(DirectoryItem BaseDirectory, InputFileCollection InputFiles, out bool bIgnoreFileFound)
+		static FileItem[] FindInputFilesFromDirectory(DirectoryItem BaseDirectory, InputFileCollection InputFiles, out bool bIgnoreFileFound, ILogger Logger)
 		{
 			bIgnoreFileFound = false;
 			List<FileItem> SourceFiles = new List<FileItem>();
 			foreach(FileItem InputFile in BaseDirectory.EnumerateFiles())
 			{
-				if (InputFile.HasExtension(".h"))
+				if (InputFile.Name == ".ubtignore")
+				{
+					bIgnoreFileFound = true;
+				}
+				// some text editors will leave temp files that start with . next to a, say, .cpp file, and we do not want to compile them,
+				// but skip files like .DS_Store that don't have an extension (so, no more .'s after the first)
+				else if (InputFile.Name.StartsWith(".") && InputFile.Name.IndexOf('.', 1) != -1)
+				{
+					if (KnownInputFileExtensions.Contains(InputFile.Location.GetExtension()))
+					{
+						Logger.LogInformation("Hidden file '{HiddenFile}' found, skipping.", InputFile.FullName);
+					}
+				}
+				else if (InputFile.HasExtension(".h"))
 				{
 					InputFiles.HeaderFiles.Add(InputFile);
 				}
@@ -1929,10 +1955,7 @@ namespace UnrealBuildTool
 					SourceFiles.Add(InputFile);
 					InputFiles.ISPCFiles.Add(InputFile);
 				}
-				else if (InputFile.Name == ".ubtignore")
-				{
-					bIgnoreFileFound = true;
-				}
+				// if you add any extension checks here, make sure to update KnownInputFileExtensions
 			}
 			return SourceFiles.ToArray();
 		}
@@ -1941,11 +1964,12 @@ namespace UnrealBuildTool
 		/// Gets a set of source files for the given directory. Used to detect when the makefile is out of date.
 		/// </summary>
 		/// <param name="Directory"></param>
+		/// <param name="Logger">Logger for output</param>
 		/// <returns>Array of source files</returns>
-		public static FileItem[] GetSourceFiles(DirectoryItem Directory)
+		public static FileItem[] GetSourceFiles(DirectoryItem Directory, ILogger Logger)
 		{
 			bool bIgnoreFileFound;
-			FileItem[] Files = FindInputFilesFromDirectory(Directory, new InputFileCollection(), out bIgnoreFileFound);
+			FileItem[] Files = FindInputFilesFromDirectory(Directory, new InputFileCollection(), out bIgnoreFileFound, Logger);
 			if (bIgnoreFileFound)
 			{
 				return Array.Empty<FileItem>();
diff --git a/Engine/Source/Programs/UnrealBuildTool/Modes/FixIncludePathsMode.cs b/Engine/Source/Programs/UnrealBuildTool/Modes/FixIncludePathsMode.cs
--- a/Engine/Source/Programs/UnrealBuildTool/Modes/FixIncludePathsMode.cs
+++ b/Engine/Source/Programs/UnrealBuildTool/Modes/FixIncludePathsMode.cs
@@ -103,7 +103,7 @@ namespace UnrealBuildTool
 					{
 						foreach (UEBuildModuleCPP Module in Binary.Modules.OfType<UEBuildModuleCPP>())
 						{
-							UEBuildModuleCPP.InputFileCollection InputFileCollection = Module.FindInputFiles(Target.Platform, new Dictionary<DirectoryItem, FileItem[]>());
+							UEBuildModuleCPP.InputFileCollection InputFileCollection = Module.FindInputFiles(Target.Platform, new Dictionary<DirectoryItem, FileItem[]>(), Logger);
 							List<FileItem> InputFiles = new List<FileItem>();
 							InputFiles.AddRange(InputFileCollection.HeaderFiles);
 							InputFiles.AddRange(InputFileCollection.CPPFiles);
diff --git a/Engine/Source/Programs/UnrealBuildTool/Modes/GenerateClangDatabase.cs b/Engine/Source/Programs/UnrealBuildTool/Modes/GenerateClangDatabase.cs
--- a/Engine/Source/Programs/UnrealBuildTool/Modes/GenerateClangDatabase.cs
+++ b/Engine/Source/Programs/UnrealBuildTool/Modes/GenerateClangDatabase.cs
@@ -175,7 +175,7 @@ namespace UnrealBuildTool
 							if (!Module.Rules.bUsePrecompiled)
 							{
 								// Gather all the files we care about
-								UEBuildModuleCPP.InputFileCollection InputFileCollection = Module.FindInputFiles(Target.Platform, new Dictionary<DirectoryItem, FileItem[]>());
+								UEBuildModuleCPP.InputFileCollection InputFileCollection = Module.FindInputFiles(Target.Platform, new Dictionary<DirectoryItem, FileItem[]>(), Logger);
 								List<FileItem> InputFiles = new List<FileItem>();
 								InputFiles.AddRange(InputFileCollection.CPPFiles);
 								InputFiles.AddRange(InputFileCollection.CCFiles);
diff --git a/Engine/Source/Programs/UnrealBuildTool/Modes/InlineGeneratedCpps.cs b/Engine/Source/Programs/UnrealBuildTool/Modes/InlineGeneratedCpps.cs
--- a/Engine/Source/Programs/UnrealBuildTool/Modes/InlineGeneratedCpps.cs
+++ b/Engine/Source/Programs/UnrealBuildTool/Modes/InlineGeneratedCpps.cs
@@ -98,7 +98,7 @@ namespace UnrealBuildTool
 
 						foreach (UEBuildModuleCPP Module in Binary.Modules.OfType<UEBuildModuleCPP>())
 						{
-							UEBuildModuleCPP.InputFileCollection InputFileCollection = Module.FindInputFiles(Target.Platform, new Dictionary<DirectoryItem, FileItem[]>());
+							UEBuildModuleCPP.InputFileCollection InputFileCollection = Module.FindInputFiles(Target.Platform, new Dictionary<DirectoryItem, FileItem[]>(), Logger);
 
 							var FileList = new List<FileItem>();
 							foreach (FileItem InputFile in InputFileCollection.CPPFiles)
diff --git a/Engine/Source/Programs/UnrealBuildTool/System/TargetMakefile.cs b/Engine/Source/Programs/UnrealBuildTool/System/TargetMakefile.cs
--- a/Engine/Source/Programs/UnrealBuildTool/System/TargetMakefile.cs
+++ b/Engine/Source/Programs/UnrealBuildTool/System/TargetMakefile.cs
@@ -746,7 +746,7 @@ namespace UnrealBuildTool
 				{
 					if(!InputDirectory.Exists || InputDirectory.LastWriteTimeUtc > Makefile.CreateTimeUtc)
 					{
-						FileItem[] SourceFiles = UEBuildModuleCPP.GetSourceFiles(InputDirectory);
+						FileItem[] SourceFiles = UEBuildModuleCPP.GetSourceFiles(InputDirectory, Logger);
 						if(SourceFiles.Length < SourceFileInfos.Length)
 						{
 							ReasonNotLoaded = "source file removed";
@@ -765,7 +765,7 @@ namespace UnrealBuildTool
 
 						foreach(DirectoryItem Directory in InputDirectory.EnumerateDirectories())
 						{
-							if(!Makefile.DirectoryToSourceFiles.ContainsKey(Directory) && ContainsSourceFiles(Directory, ExcludedFolderNames))
+							if(!Makefile.DirectoryToSourceFiles.ContainsKey(Directory) && ContainsSourceFiles(Directory, ExcludedFolderNames, Logger))
 							{
 								ReasonNotLoaded = "directory added";
 								return false;
@@ -888,14 +888,15 @@ namespace UnrealBuildTool
 		/// </summary>
 		/// <param name="Directory">Directory to search through</param>
 		/// <param name="ExcludedFolderNames">Set of directory names to exclude</param>
+		/// <param name="Logger">Logger for output diagnostics</param>
 		/// <returns>True if the directory contains any source files</returns>
-		static bool ContainsSourceFiles(DirectoryItem Directory, ReadOnlyHashSet<string> ExcludedFolderNames)
+		static bool ContainsSourceFiles(DirectoryItem Directory, ReadOnlyHashSet<string> ExcludedFolderNames, ILogger Logger)
 		{
 			// Check this directory isn't ignored
 			if(!ExcludedFolderNames.Contains(Directory.Name))
 			{
 				// Check for any source files in this actual directory
-				FileItem[] SourceFiles = UEBuildModuleCPP.GetSourceFiles(Directory);
+				FileItem[] SourceFiles = UEBuildModuleCPP.GetSourceFiles(Directory, Logger);
 				if(SourceFiles.Length > 0)
 				{
 					return true;
@@ -904,7 +905,7 @@ namespace UnrealBuildTool
 				// Check for any source files in a subdirectory
 				foreach(DirectoryItem SubDirectory in Directory.EnumerateDirectories())
 				{
-					if(ContainsSourceFiles(SubDirectory, ExcludedFolderNames))
+					if(ContainsSourceFiles(SubDirectory, ExcludedFolderNames, Logger))
 					{
 						return true;
 					}
