Revert 5.1 QFE, make sure universal build works on installed build

diff --git a/Engine/Source/Programs/UnrealBuildTool/ProjectFiles/Xcode/XcodeProjectLegacy.cs b/Engine/Source/Programs/UnrealBuildTool/ProjectFiles/Xcode/XcodeProjectLegacy.cs
--- a/Engine/Source/Programs/UnrealBuildTool/ProjectFiles/Xcode/XcodeProjectLegacy.cs
+++ b/Engine/Source/Programs/UnrealBuildTool/ProjectFiles/Xcode/XcodeProjectLegacy.cs
@@ -921,10 +921,6 @@ namespace UnrealBuildTool.XcodeProjectLegacy
 				Content.Append("\t\t\t\tENABLE_TESTABILITY = YES;" + ProjectFileGenerator.NewLine);
 			}
 			Content.Append("\t\t\t\tALWAYS_SEARCH_USER_PATHS = NO;" + ProjectFileGenerator.NewLine);
-			if (Unreal.IsEngineInstalled())
-			{
-				Content.Append("\t\t\t\t\"ARCHS[sdk=macosx*]\" = x86_64;" + ProjectFileGenerator.NewLine);
-			}
 			Content.Append("\t\t\t\tCLANG_CXX_LANGUAGE_STANDARD = \"c++14\";" + ProjectFileGenerator.NewLine);
 			Content.Append("\t\t\t\tGCC_ENABLE_CPP_RTTI = NO;" + ProjectFileGenerator.NewLine);
 			Content.Append("\t\t\t\tGCC_WARN_CHECK_SWITCH_STATEMENTS = NO;" + ProjectFileGenerator.NewLine);
