Update preferred MSVC to 14.35.32215 (VS2022 17.5.x)

diff --git a/Engine/Source/Programs/UnrealBuildTool/Platform/Windows/MicrosoftPlatformSDK.Versions.cs b/Engine/Source/Programs/UnrealBuildTool/Platform/Windows/MicrosoftPlatformSDK.Versions.cs
--- a/Engine/Source/Programs/UnrealBuildTool/Platform/Windows/MicrosoftPlatformSDK.Versions.cs
+++ b/Engine/Source/Programs/UnrealBuildTool/Platform/Windows/MicrosoftPlatformSDK.Versions.cs
@@ -43,6 +43,7 @@ namespace UnrealBuildTool
 		/// </summary>
 		static readonly VersionNumberRange[] PreferredVisualCppVersions = new VersionNumberRange[]
 		{
+			VersionNumberRange.Parse("14.35.32215", "14.35.99999"), // VS2022 17.5.x
 			VersionNumberRange.Parse("14.34.31933", "14.34.99999"), // VS2022 17.4.x
 			VersionNumberRange.Parse("14.29.30133", "14.29.99999"), // VS2019 16.11.x
 		};
diff --git a/Engine/Source/Programs/UnrealBuildTool/Platform/Windows/MicrosoftPlatformSDK.VisualStudio.cs b/Engine/Source/Programs/UnrealBuildTool/Platform/Windows/MicrosoftPlatformSDK.VisualStudio.cs
--- a/Engine/Source/Programs/UnrealBuildTool/Platform/Windows/MicrosoftPlatformSDK.VisualStudio.cs
+++ b/Engine/Source/Programs/UnrealBuildTool/Platform/Windows/MicrosoftPlatformSDK.VisualStudio.cs
@@ -56,7 +56,7 @@ namespace UnrealBuildTool
 		static readonly string[] VisualStudio2022SuggestedComponents = new string[]
 		{
 			"Microsoft.VisualStudio.Workload.ManagedDesktop",
-			"Microsoft.VisualStudio.Component.VC.14.34.17.4.x86.x64",
+			"Microsoft.VisualStudio.Component.VC.14.35.17.5.x86.x64",
 			"Microsoft.Net.Component.4.6.2.TargetingPack",
 		};
 
@@ -66,7 +66,7 @@ namespace UnrealBuildTool
 		/// </summary>
 		static readonly string[] VisualStudio2022SuggestedHololensComponents = new string[]
 		{
-			"Microsoft.VisualStudio.Component.VC.14.34.17.4.ARM64",
+			"Microsoft.VisualStudio.Component.VC.14.35.17.5.ARM64",
 		};
 
 		/// <summary>
