UnrealBuildTool: Module CppStandard needs to be overridden if it doesn't match the target, rather than if it doesn't match the default

diff --git a/Engine/Source/Programs/UnrealBuildTool/Configuration/UEBuildModuleCPP.cs b/Engine/Source/Programs/UnrealBuildTool/Configuration/UEBuildModuleCPP.cs
--- a/Engine/Source/Programs/UnrealBuildTool/Configuration/UEBuildModuleCPP.cs
+++ b/Engine/Source/Programs/UnrealBuildTool/Configuration/UEBuildModuleCPP.cs
@@ -1710,19 +1710,18 @@ namespace UnrealBuildTool
 				Result.OptimizationLevel = Rules.OptimizationLevel;
 			}
 
-
 			// If the module overrides the C++ language version, override it on the compile environment
-			if (Rules.CppStandard != CppStandardVersion.Default)
+			if (Rules.CppStandard != BaseCompileEnvironment.CppStandard)
 			{
 				Result.CppStandard = Rules.CppStandard;
 			}
-			if (Target.bEnableCppModules && Result.CppStandard == CppStandardVersion.Default)
+			if (Target.bEnableCppModules && Result.CppStandard < CppStandardVersion.Cpp20)
 			{
 				Result.CppStandard = CppStandardVersion.Cpp20;
 			}
 
 			// If the module overrides the C language version, override it on the compile environment
-			if (Rules.CStandard != CStandardVersion.Default)
+			if (Rules.CStandard != BaseCompileEnvironment.CStandard)
 			{
 				Result.CStandard = Rules.CStandard;
 			}
