[Backout] - CL24530832
Original CL Desc
-----------------------------------------------------------------
UnrealBuildTool: Module CppStandard needs to be overridden if it doesn't match the target, rather than if it doesn't match the default

diff --git a/Engine/Source/Programs/UnrealBuildTool/Configuration/UEBuildModuleCPP.cs b/Engine/Source/Programs/UnrealBuildTool/Configuration/UEBuildModuleCPP.cs
--- a/Engine/Source/Programs/UnrealBuildTool/Configuration/UEBuildModuleCPP.cs
+++ b/Engine/Source/Programs/UnrealBuildTool/Configuration/UEBuildModuleCPP.cs
@@ -1710,18 +1710,19 @@ namespace UnrealBuildTool
 				Result.OptimizationLevel = Rules.OptimizationLevel;
 			}
 
+
 			// If the module overrides the C++ language version, override it on the compile environment
-			if (Rules.CppStandard != BaseCompileEnvironment.CppStandard)
+			if (Rules.CppStandard != CppStandardVersion.Default)
 			{
 				Result.CppStandard = Rules.CppStandard;
 			}
-			if (Target.bEnableCppModules && Result.CppStandard < CppStandardVersion.Cpp20)
+			if (Target.bEnableCppModules && Result.CppStandard == CppStandardVersion.Default)
 			{
 				Result.CppStandard = CppStandardVersion.Cpp20;
 			}
 
 			// If the module overrides the C language version, override it on the compile environment
-			if (Rules.CStandard != BaseCompileEnvironment.CStandard)
+			if (Rules.CStandard != CStandardVersion.Default)
 			{
 				Result.CStandard = Rules.CStandard;
 			}
