UnrealBuildTool: ModuleRules CppStandardVersion should be nullable so we can detect when it's actually overridden

diff --git a/Engine/Source/Programs/UnrealBuildTool/Configuration/ModuleRules.cs b/Engine/Source/Programs/UnrealBuildTool/Configuration/ModuleRules.cs
--- a/Engine/Source/Programs/UnrealBuildTool/Configuration/ModuleRules.cs
+++ b/Engine/Source/Programs/UnrealBuildTool/Configuration/ModuleRules.cs
@@ -1261,12 +1261,12 @@ namespace UnrealBuildTool
 		/// <summary>
 		/// Which stanard to use for compiling this module
 		/// </summary>
-		public CppStandardVersion CppStandard = CppStandardVersion.Default;
+		public CppStandardVersion? CppStandard;
 		
 		/// <summary>
 		/// Which standard to use for compiling this module
 		/// </summary>
-		public CStandardVersion CStandard = CStandardVersion.Default;
+		public CStandardVersion? CStandard;
 
 
 		/// <summary>
diff --git a/Engine/Source/Programs/UnrealBuildTool/Configuration/UEBuildModuleCPP.cs b/Engine/Source/Programs/UnrealBuildTool/Configuration/UEBuildModuleCPP.cs
--- a/Engine/Source/Programs/UnrealBuildTool/Configuration/UEBuildModuleCPP.cs
+++ b/Engine/Source/Programs/UnrealBuildTool/Configuration/UEBuildModuleCPP.cs
@@ -1710,21 +1710,20 @@ namespace UnrealBuildTool
 				Result.OptimizationLevel = Rules.OptimizationLevel;
 			}
 
-
 			// If the module overrides the C++ language version, override it on the compile environment
-			if (Rules.CppStandard != CppStandardVersion.Default)
+			if (Rules.CppStandard != null)
 			{
-				Result.CppStandard = Rules.CppStandard;
+				Result.CppStandard = (CppStandardVersion)Rules.CppStandard;
 			}
-			if (Target.bEnableCppModules && Result.CppStandard == CppStandardVersion.Default)
+			if (Target.bEnableCppModules && Result.CppStandard < CppStandardVersion.Cpp20)
 			{
 				Result.CppStandard = CppStandardVersion.Cpp20;
 			}
 
 			// If the module overrides the C language version, override it on the compile environment
-			if (Rules.CStandard != CStandardVersion.Default)
+			if (Rules.CStandard != null)
 			{
-				Result.CStandard = Rules.CStandard;
+				Result.CStandard = (CStandardVersion)Rules.CStandard;
 			}
 
 			// Set the macro used to check whether monolithic headers can be used
diff --git a/Engine/Source/Programs/UnrealBuildTool/ProjectFiles/ProjectFile.cs b/Engine/Source/Programs/UnrealBuildTool/ProjectFiles/ProjectFile.cs
--- a/Engine/Source/Programs/UnrealBuildTool/ProjectFiles/ProjectFile.cs
+++ b/Engine/Source/Programs/UnrealBuildTool/ProjectFiles/ProjectFile.cs
@@ -257,7 +257,7 @@ namespace UnrealBuildTool
 			}
 
 			SetIntelliSenseEnableCoroutines(CompileEnvironment.bEnableCoroutines);
-			SetIntelliSenseCppVersion(Module.Rules.CppStandard);
+			SetIntelliSenseCppVersion(Module.Rules.CppStandard ?? CompileEnvironment.CppStandard);
 		}
 
 		/// <summary>
