- Fixed a Multi-arch error when using -clean in UBT

diff --git a/Engine/Source/Programs/UnrealBuildTool/Modes/CleanMode.cs b/Engine/Source/Programs/UnrealBuildTool/Modes/CleanMode.cs
--- a/Engine/Source/Programs/UnrealBuildTool/Modes/CleanMode.cs
+++ b/Engine/Source/Programs/UnrealBuildTool/Modes/CleanMode.cs
@@ -137,7 +137,7 @@ namespace UnrealBuildTool
 						NameSuffixes.Add("");
 					}
 					NameSuffixes.Add(String.Format("-{0}-{1}", Target.Platform.ToString(), Target.Configuration.ToString()));
-					if (!UnrealArchitectureConfig.ForPlatform(Target.Platform).RequiresArchitectureFilenames(Target.Architectures))
+					if (UnrealArchitectureConfig.ForPlatform(Target.Platform).RequiresArchitectureFilenames(Target.Architectures))
 					{
 						NameSuffixes.AddRange(NameSuffixes.ToArray().Select(x => x + Target.Architecture.ToString()));
 					}
