Return all platforms from target query in QueryMode, not just desktop ones.

diff --git a/Engine/Source/Programs/UnrealBuildTool/Modes/QueryMode.cs b/Engine/Source/Programs/UnrealBuildTool/Modes/QueryMode.cs
--- a/Engine/Source/Programs/UnrealBuildTool/Modes/QueryMode.cs
+++ b/Engine/Source/Programs/UnrealBuildTool/Modes/QueryMode.cs
@@ -162,17 +162,13 @@ namespace UnrealBuildTool
 				List<UnrealTargetPlatform> Platforms = new();
 				foreach (UnrealTargetPlatform Platform in UnrealTargetPlatform.GetValidPlatforms())
 				{
-					// Only include desktop platforms for now 
-					if (ProjectFileGenerator.IsValidDesktopPlatform(Platform))
+					// If there is a build platform present, add it to the SupportedPlatforms list
+					UEBuildPlatform? BuildPlatform;
+					if (UEBuildPlatform.TryGetBuildPlatform(Platform, out BuildPlatform))
 					{
-						// If there is a build platform present, add it to the SupportedPlatforms list
-						UEBuildPlatform? BuildPlatform;
-						if (UEBuildPlatform.TryGetBuildPlatform(Platform, out BuildPlatform))
+						if (InstalledPlatformInfo.IsValidPlatform(Platform, EProjectType.Code))
 						{
-							if (InstalledPlatformInfo.IsValidPlatform(Platform, EProjectType.Code))
-							{
-								Platforms.Add(Platform);
-							}
+							Platforms.Add(Platform);
 						}
 					}
 				}
