Use RunUBT batch file instead of Build batch file.
Add default file exclusions to vs code workspace when building for use with extension.

diff --git a/Engine/Source/Programs/UnrealBuildTool/ProjectFiles/VisualStudioCode/VSCodeProjectFileGenerator.cs b/Engine/Source/Programs/UnrealBuildTool/ProjectFiles/VisualStudioCode/VSCodeProjectFileGenerator.cs
--- a/Engine/Source/Programs/UnrealBuildTool/ProjectFiles/VisualStudioCode/VSCodeProjectFileGenerator.cs
+++ b/Engine/Source/Programs/UnrealBuildTool/ProjectFiles/VisualStudioCode/VSCodeProjectFileGenerator.cs
@@ -1650,12 +1650,28 @@ namespace UnrealBuildTool
 					WorkspaceFile.AddField("UE.loggingLevel", "Warning");
 					if (HostPlatform == UnrealTargetPlatform.Win64)
 					{
-						WorkspaceFile.AddField("UE.UBTScriptPath", MakePathString(FileReference.Combine(ProjectRoot, "Engine", "Build", "BatchFiles", "Build.bat"), true));
+						WorkspaceFile.AddField("UE.UBTScriptPath", MakePathString(FileReference.Combine(ProjectRoot, "Engine", "Build", "BatchFiles", "RunUBT.bat"), true));
 					}
 					else
 					{
-						WorkspaceFile.AddField("UE.UBTScriptPath", MakePathString(FileReference.Combine(ProjectRoot, "Engine", "Build", "BatchFiles", HostPlatform.ToString(), "Build.sh"), true));
+						WorkspaceFile.AddField("UE.UBTScriptPath", MakePathString(FileReference.Combine(ProjectRoot, "Engine", "Build", "BatchFiles", HostPlatform.ToString(), "RunUBT.sh"), true));
 					}
+
+					// Exclude some large directories/filetypes by default
+					WorkspaceFile.BeginObject("files.exclude");
+					WorkspaceFile.AddField("**/Intermediate/", true);
+					WorkspaceFile.AddField("**/Binaries/", true);
+					WorkspaceFile.AddField("Engine/DerivedDataCache", true);
+					WorkspaceFile.AddField("**/*.uasset", true);
+					WorkspaceFile.AddField("**/*.umap", true);
+					WorkspaceFile.AddField("**/*.uexp", true);
+					WorkspaceFile.AddField("**/*.upayload", true);
+					WorkspaceFile.AddField("**/*.ubulk", true);
+					WorkspaceFile.AddField("**/*.m.ubulk", true);
+					WorkspaceFile.AddField("**/*.uptnl", true);
+					WorkspaceFile.EndObject();
+					WorkspaceFile.BeginObject("search.exclude");
+					WorkspaceFile.EndObject();
 				}
 			}
 			WorkspaceFile.EndObject();
