UHT: Fix typo in

diff --git a/Engine/Programs/UnrealHeaderTool/Config/DefaultEngine.ini b/Engine/Programs/UnrealHeaderTool/Config/DefaultEngine.ini
--- a/Engine/Programs/UnrealHeaderTool/Config/DefaultEngine.ini
+++ b/Engine/Programs/UnrealHeaderTool/Config/DefaultEngine.ini
@@ -168,7 +168,7 @@ EnginePluginObjectPtrMemberBehavior=AllowSilently
 NonEngineNativePointerMemberBehavior=AllowSilently
 NonEngineObjectPtrMemberBehavior=AllowSilently
 AutomaticSettersAndGetters=False
-AreRigVMUObjectProeprtiesEnabled=False
+AreRigVMUObjectPropertiesEnabled=False
 AreRigVMUInterfaceProeprtiesEnabled=True
 ShowDeprecations=True
 
diff --git a/Engine/Source/Programs/Shared/EpicGames.UHT/Types/Properties/UhtObjectPropertyBase.cs b/Engine/Source/Programs/Shared/EpicGames.UHT/Types/Properties/UhtObjectPropertyBase.cs
--- a/Engine/Source/Programs/Shared/EpicGames.UHT/Types/Properties/UhtObjectPropertyBase.cs
+++ b/Engine/Source/Programs/Shared/EpicGames.UHT/Types/Properties/UhtObjectPropertyBase.cs
@@ -61,7 +61,7 @@ namespace EpicGames.UHT.Types
 			}
 
 			PropertyCaps &= ~(UhtPropertyCaps.CanHaveConfig);
-			if (Session.Config!.AreRigVMUObjectProeprtiesEnabled)
+			if (Session.Config!.AreRigVMUObjectPropertiesEnabled)
 			{
 				PropertyCaps |= UhtPropertyCaps.SupportsRigVM;
 			}
diff --git a/Engine/Source/Programs/Shared/EpicGames.UHT/Utils/UhtConfig.cs b/Engine/Source/Programs/Shared/EpicGames.UHT/Utils/UhtConfig.cs
--- a/Engine/Source/Programs/Shared/EpicGames.UHT/Utils/UhtConfig.cs
+++ b/Engine/Source/Programs/Shared/EpicGames.UHT/Utils/UhtConfig.cs
@@ -78,7 +78,7 @@ namespace EpicGames.UHT.Utils
 		/// <summary>
 		/// If true, UObject properties are enabled in RigVM
 		/// </summary>
-		public bool AreRigVMUObjectProeprtiesEnabled { get; }
+		public bool AreRigVMUObjectPropertiesEnabled { get; }
 
 		/// <summary>
 		/// If true, UInterface properties are enabled in RigVM
diff --git a/Engine/Source/Programs/UnrealBuildTool/Modes/UnrealHeaderToolMode.cs b/Engine/Source/Programs/UnrealBuildTool/Modes/UnrealHeaderToolMode.cs
--- a/Engine/Source/Programs/UnrealBuildTool/Modes/UnrealHeaderToolMode.cs
+++ b/Engine/Source/Programs/UnrealBuildTool/Modes/UnrealHeaderToolMode.cs
@@ -99,7 +99,7 @@ namespace UnrealBuildTool.Modes
 		/// <summary>
 		/// If true, UObject properties are enabled in RigVM
 		/// </summary>
-		private readonly bool _areRigVMUObjectProeprtiesEnabled = false;
+		private readonly bool _areRigVMUObjectPropertiesEnabled = false;
 
 		/// <summary>
 		/// If true, UInterface properties are enabled in RigVM
@@ -161,7 +161,7 @@ namespace UnrealBuildTool.Modes
 		/// <summary>
 		/// If true, UObject properties are enabled in RigVM
 		/// </summary>
-		public bool AreRigVMUObjectProeprtiesEnabled => this._areRigVMUObjectProeprtiesEnabled;
+		public bool AreRigVMUObjectPropertiesEnabled => this._areRigVMUObjectPropertiesEnabled;
 
 		/// <summary>
 		/// If true, UInterface properties are enabled in RigVM
@@ -281,7 +281,7 @@ namespace UnrealBuildTool.Modes
 			this._enginePluginObjectPtrMemberBehavior = GetIssueBehavior("UnrealHeaderTool", "EnginePluginObjectPtrMemberBehavior", UhtIssueBehavior.AllowSilently);
 			this._nonEngineNativePointerMemberBehavior = GetIssueBehavior("UnrealHeaderTool", "NonEngineNativePointerMemberBehavior", UhtIssueBehavior.AllowSilently);
 			this._nonEngineObjectPtrMemberBehavior = GetIssueBehavior("UnrealHeaderTool", "NonEngineObjectPtrMemberBehavior", UhtIssueBehavior.AllowSilently);
-			this._areRigVMUObjectProeprtiesEnabled = GetBoolean("UnrealHeaderTool", "AreRigVMUObjectProeprtiesEnabled", false);
+			this._areRigVMUObjectPropertiesEnabled = GetBoolean("UnrealHeaderTool", "AreRigVMUObjectPropertiesEnabled", false);
 			this._areRigVMUInterfaceProeprtiesEnabled = GetBoolean("UnrealHeaderTool", "AreRigVMUInterfaceProeprtiesEnabled", false);
 			this._showDeprecations = GetBoolean("UnrealHeaderTool", "ShowDeprecations", true);
 			this._engineMissingGeneratedHeaderIncludeBehavior = GetIssueBehavior("UnrealHeaderTool", "EngineMissingGeneratedHeaderIncludeBehavior", UhtIssueBehavior.AllowSilently);
