UnrealBuildTool
* Added option to generate assembly files while compiling cpp files.
* Supported only on MSVC compiler for now.
* The generated files have the .asm extension and are located in the same Intermediate folder as the corresponding .obj

diff --git a/Engine/Platforms/Hololens/Source/Programs/UnrealBuildTool/HoloLensToolChain.cs b/Engine/Platforms/Hololens/Source/Programs/UnrealBuildTool/HoloLensToolChain.cs
--- a/Engine/Platforms/Hololens/Source/Programs/UnrealBuildTool/HoloLensToolChain.cs
+++ b/Engine/Platforms/Hololens/Source/Programs/UnrealBuildTool/HoloLensToolChain.cs
@@ -75,9 +75,9 @@ namespace UnrealBuildTool
 			}
 		}
 
-		protected override void AppendCLArguments_CPP(CppCompileEnvironment CompileEnvironment, List<string> Arguments)
+		protected override void AppendCLArguments_CPP(CppCompileEnvironment CompileEnvironment, List<string> Arguments, DirectoryReference OutputDir)
 		{
-			base.AppendCLArguments_CPP(CompileEnvironment, Arguments);
+			base.AppendCLArguments_CPP(CompileEnvironment, Arguments, OutputDir);
 
 			// Enable Windows Runtime extensions.  Do this even for libs (plugins) so that these too can consume WinRT APIs
 			Arguments.Add("/ZW");
diff --git a/Engine/Source/Programs/UnrealBuildTool/Configuration/TargetRules.cs b/Engine/Source/Programs/UnrealBuildTool/Configuration/TargetRules.cs
--- a/Engine/Source/Programs/UnrealBuildTool/Configuration/TargetRules.cs
+++ b/Engine/Source/Programs/UnrealBuildTool/Configuration/TargetRules.cs
@@ -1700,6 +1700,12 @@ namespace UnrealBuildTool
 		[XmlConfigFile(Category = "BuildConfiguration")]
 		public bool bPreprocessDepends = false;
 
+		/// <summary>
+		/// Whether to generate assembly data while compiling this target. Works exclusively on MSVC for now.
+		/// </summary>
+		[CommandLine("-WithAssembly")]
+		public bool bWithAssembly = false;
+
 		/// <summary>
 		/// Whether static code analysis should be enabled.
 		/// </summary>
@@ -3505,6 +3511,11 @@ namespace UnrealBuildTool
 			get { return Inner.bPreprocessDepends; }
 		}
 
+		public bool bWithAssembly
+		{
+			get { return Inner.bWithAssembly; }
+		}
+
 		public StaticAnalyzer StaticAnalyzer
 		{
 			get { return Inner.StaticAnalyzer; }
diff --git a/Engine/Source/Programs/UnrealBuildTool/Configuration/UEBuildTarget.cs b/Engine/Source/Programs/UnrealBuildTool/Configuration/UEBuildTarget.cs
--- a/Engine/Source/Programs/UnrealBuildTool/Configuration/UEBuildTarget.cs
+++ b/Engine/Source/Programs/UnrealBuildTool/Configuration/UEBuildTarget.cs
@@ -4409,6 +4409,7 @@ namespace UnrealBuildTool
 			GlobalCompileEnvironment.bUsePDBFiles = Rules.bUsePDBFiles;
 			GlobalCompileEnvironment.bSupportEditAndContinue = Rules.bSupportEditAndContinue;
 			GlobalCompileEnvironment.bPreprocessOnly = Rules.bPreprocessOnly;
+			GlobalCompileEnvironment.bWithAssembly = Rules.bWithAssembly;
 			GlobalCompileEnvironment.bUseIncrementalLinking = Rules.bUseIncrementalLinking;
 			GlobalCompileEnvironment.bAllowLTCG = Rules.bAllowLTCG;
 			GlobalCompileEnvironment.bPGOOptimize = Rules.bPGOOptimize;
diff --git a/Engine/Source/Programs/UnrealBuildTool/Platform/Windows/VCToolChain.cs b/Engine/Source/Programs/UnrealBuildTool/Platform/Windows/VCToolChain.cs
--- a/Engine/Source/Programs/UnrealBuildTool/Platform/Windows/VCToolChain.cs
+++ b/Engine/Source/Programs/UnrealBuildTool/Platform/Windows/VCToolChain.cs
@@ -847,12 +847,20 @@ namespace UnrealBuildTool
 			}
 		}
 
-		protected virtual void AppendCLArguments_CPP(CppCompileEnvironment CompileEnvironment, List<string> Arguments)
+		protected virtual void AppendCLArguments_CPP(CppCompileEnvironment CompileEnvironment, List<string> Arguments, DirectoryReference OutputDir)
 		{
 			if (Target.WindowsPlatform.Compiler.IsMSVC())
 			{
 				// Explicitly compile the file as C++.
 				Arguments.Add("/TP");
+
+				if (CompileEnvironment.bWithAssembly)
+				{
+					// Write out an assembly file (.asm) with the c++ code embedded in it via comments
+					Arguments.Add("/FAs");
+					// Set the output directory for the asm files
+					Arguments.Add("/Fa" + OutputDir.FullName + "\\");
+				}
 			}
 			else
 			{
@@ -1478,6 +1486,14 @@ namespace UnrealBuildTool
 						}
 					}
 
+					if (CompileEnvironment.bWithAssembly && SourceFile.HasExtension(".cpp"))
+					{
+						if (Target.WindowsPlatform.Compiler.IsMSVC())
+						{
+							CompileAction.AdditionalProducedItems.Add(FileItem.GetItemByFileReference(FileReference.Combine(OutputDir, SourceFile.Location.GetFileName() + ".asm")));
+						}
+					}
+
 					// Experimental: support for JSON output of timing data
 					if (Target.WindowsPlatform.Compiler.IsClang() && (Target.bPrintToolChainTimingInfo || Target.WindowsPlatform.bClangTimeTrace))
 					{
@@ -1538,7 +1554,7 @@ namespace UnrealBuildTool
 				}
 				else
 				{
-					AppendCLArguments_CPP(CompileEnvironment, CompileAction.Arguments);
+					AppendCLArguments_CPP(CompileEnvironment, CompileAction.Arguments, OutputDir);
 				}
 
 				List<FileItem>? InlinedFiles;
@@ -2097,7 +2113,7 @@ namespace UnrealBuildTool
 				return;
 
 			VCCompileAction BaseCompileAction = CreateBaseCompileAction(CompileEnvironment);
-			AppendCLArguments_CPP(CompileEnvironment, BaseCompileAction.Arguments);
+			AppendCLArguments_CPP(CompileEnvironment, BaseCompileAction.Arguments, OutputDir);
 			BaseCompileAction.AdditionalPrerequisiteItems.AddRange(CompileEnvironment.AdditionalPrerequisites); // Primarily for ispc.generated.h files
 			Graph.AddAction(new VcSpecificFileAction(SourceDir, OutputDir, BaseCompileAction));
 		}
diff --git a/Engine/Source/Programs/UnrealBuildTool/System/CppCompileEnvironment.cs b/Engine/Source/Programs/UnrealBuildTool/System/CppCompileEnvironment.cs
--- a/Engine/Source/Programs/UnrealBuildTool/System/CppCompileEnvironment.cs
+++ b/Engine/Source/Programs/UnrealBuildTool/System/CppCompileEnvironment.cs
@@ -400,6 +400,11 @@ namespace UnrealBuildTool
 		/// </summary>
 		public bool bPreprocessOnly = false;
 
+		/// <summary>
+		/// Should an assembly file be generated while compiling. Works exclusively on MSVC compilers for now.
+		/// </summary>
+		public bool bWithAssembly = false;
+
 		/// <summary>
 		/// Whether to support edit and continue.  Only works on Microsoft compilers in 32-bit compiles.
 		/// </summary>
@@ -674,6 +679,7 @@ namespace UnrealBuildTool
 			bEnableOSX109Support = Other.bEnableOSX109Support;
 			bUsePDBFiles = Other.bUsePDBFiles;
 			bPreprocessOnly = Other.bPreprocessOnly;
+			bWithAssembly = Other.bWithAssembly;
 			bSupportEditAndContinue = Other.bSupportEditAndContinue;
 			bUseIncrementalLinking = Other.bUseIncrementalLinking;
 			bAllowLTCG = Other.bAllowLTCG;
