UnrealBuildTool: Write manifest for Analyze mode, suppress warning

diff --git a/Engine/Source/Programs/UnrealBuildTool/Modes/AnalyzeMode.cs b/Engine/Source/Programs/UnrealBuildTool/Modes/AnalyzeMode.cs
--- a/Engine/Source/Programs/UnrealBuildTool/Modes/AnalyzeMode.cs
+++ b/Engine/Source/Programs/UnrealBuildTool/Modes/AnalyzeMode.cs
@@ -8,6 +8,7 @@ using System.Linq;
 using System.Text;
 using System.Xml;
 using Microsoft.Extensions.Logging;
+using UnrealBuildBase;
 
 namespace UnrealBuildTool.Modes
 {
@@ -214,18 +215,22 @@ namespace UnrealBuildTool.Modes
 			// Warn about any missing modules
 			foreach (UEBuildModule MissingModule in MissingModules.OrderBy(x => x.Name))
 			{
-				Logger.LogWarning("Missing module '{MissingModuleName}'", MissingModule.Name);
+				Logger.LogInformation("Missing module '{MissingModuleName}'", MissingModule.Name);
 			}
 
+			List<KeyValuePair<FileReference, BuildProductType>> AnalyzeProducts = new();
+
 			// Generate the dependency graph between modules
 			FileReference DependencyGraphFile = Target.ReceiptFileName.ChangeExtension(".Dependencies.gexf");
 			Logger.LogInformation("Writing dependency graph to {DependencyGraphFile}...", DependencyGraphFile);
 			WriteDependencyGraph(Target, ModuleToInfo, DependencyGraphFile);
+			AnalyzeProducts.Add(new(DependencyGraphFile, BuildProductType.BuildResource));
 
 			// Generate the dependency graph between modules
 			FileReference ShortestPathGraphFile = Target.ReceiptFileName.ChangeExtension(".ShortestPath.gexf");
 			Logger.LogInformation("Writing shortest-path graph to {ShortestPathGraphFile}...", ShortestPathGraphFile);
 			WriteShortestPathGraph(Target, ModuleToInfo, ShortestPathGraphFile);
+			AnalyzeProducts.Add(new(ShortestPathGraphFile, BuildProductType.BuildResource));
 
 			// Write all the target stats as a text file
 			FileReference TextFile = Target.ReceiptFileName.ChangeExtension(".txt");
@@ -249,6 +254,7 @@ namespace UnrealBuildTool.Modes
 					Writer.WriteLine("Binary files:            {0}", String.Join(", ", ModuleInfo.BinaryFiles.Select(x => x.GetFileName())));
 				}
 			}
+			AnalyzeProducts.Add(new(TextFile, BuildProductType.BuildResource));
 
 			// Write all the target stats as a CSV file
 			FileReference CsvFile = Target.ReceiptFileName.ChangeExtension(".csv");
@@ -292,6 +298,12 @@ namespace UnrealBuildTool.Modes
 					Writer.WriteLine(String.Join(",", Columns));
 				}
 			}
+			AnalyzeProducts.Add(new(CsvFile, BuildProductType.BuildResource));
+
+			foreach (FileReference ManifestFileName in Target.Rules.ManifestFileNames)
+			{
+				Target.GenerateManifest(ManifestFileName, AnalyzeProducts, Logger);
+			}
 		}
 
 		private void WriteDependencyList(TextWriter Writer, string Prefix, HashSet<UEBuildModule> Modules)
