- Added Framework support in Modern Xcode projects - this will copy and code sign frameworks directly via Xcode, instead of UBT doing it

diff --git a/Engine/Source/Programs/UnrealBuildTool/Configuration/UEBuildFramework.cs b/Engine/Source/Programs/UnrealBuildTool/Configuration/UEBuildFramework.cs
--- a/Engine/Source/Programs/UnrealBuildTool/Configuration/UEBuildFramework.cs
+++ b/Engine/Source/Programs/UnrealBuildTool/Configuration/UEBuildFramework.cs
@@ -47,9 +47,17 @@ namespace UnrealBuildTool
 		public readonly bool bCopyFramework = false;
 
 		/// <summary>
-		/// File created after the framework has been extracted. Used to add dependencies into the action graph.
+		/// File created after the framework has been extracted. Used to add dependencies into the action graph, used by Modern Xcode as well
 		/// </summary>
-		public FileItem? ExtractedTokenFile;
+		public FileItem? ExtractedTokenFile
+		{
+			get { return ZipOutputDirectory == null ? null : FileItem.GetItemByFileReference(new FileReference(ZipOutputDirectory!.FullName + ".extracted")); }
+		}
+		/// <summary>
+		/// For legacy xcode projects, we unzip in UBT (isntead of xcode), so we track if we've made an action for this framework yet (if two
+		/// modules use the same framework, we only want to unzip it once. Other than time waste, there could be conflicts
+		/// </summary>
+		public bool bHasMadeUnzipAction;
 
 		/// <summary>
 		/// List of variants contained in this XCFramework. Only non null for XCFrameworks
@@ -75,7 +83,7 @@ namespace UnrealBuildTool
 		/// <param name="OutputDirectory">Path for the extracted zip file</param>
 		/// <param name="CopyBundledAssets"></param>
 		/// <param name="bCopyFramework">Copy the framework to the target's Framework directory</param>
-		public UEBuildFramework(string Name, FileReference? ZipFile, DirectoryReference? OutputDirectory, string? CopyBundledAssets, bool bCopyFramework)
+		public UEBuildFramework(string Name, FileReference? ZipFile, DirectoryReference OutputDirectory, string? CopyBundledAssets, bool bCopyFramework)
 		{
 			this.Name = Name;
 			this.ZipFile = ZipFile;
@@ -92,13 +100,13 @@ namespace UnrealBuildTool
 		/// <param name="FrameworkDirectory">Path for the framework on disk</param>
 		/// <param name="CopyBundledAssets"></param>
 		/// <param name="bCopyFramework">Copy the framework to the target's Framework directory</param>
-		public UEBuildFramework(String Name, DirectoryReference? FrameworkDirectory, string? CopyBundledAssets, bool bCopyFramework)
+		public UEBuildFramework(String Name, DirectoryReference FrameworkDirectory, string? CopyBundledAssets, bool bCopyFramework)
 		{
 			this.Name = Name;
 			this.FrameworkDirectory = FrameworkDirectory;
 			this.CopyBundledAssets = CopyBundledAssets;
 			this.bCopyFramework = bCopyFramework;
-			if( this.FrameworkDirectory != null && this.FrameworkDirectory.FullName.EndsWith(".xcframework"))
+			if(this.FrameworkDirectory.FullName.EndsWith(".xcframework"))
 			{
 				this.XCFrameworkVariants = LoadXCFrameworkVariants();
 			}
@@ -109,11 +117,11 @@ namespace UnrealBuildTool
 		/// </summary>
 		/// <param name="Platform"></param>
 		/// <param name="Architecture"></param>
-		public DirectoryReference? GetFrameworkDirectory(UnrealTargetPlatform Platform, UnrealArch Architecture)
+		public DirectoryReference? GetFrameworkDirectory(UnrealTargetPlatform? Platform, UnrealArch? Architecture)
 		{
-			if( XCFrameworkVariants != null)
+			if( XCFrameworkVariants != null && Platform != null && Architecture != null)
 			{
-				XCFrameworkVariantEntry? FrameworkVariant = XCFrameworkVariants.Find(x => x.Matches(Platform, Architecture));
+				XCFrameworkVariantEntry? FrameworkVariant = XCFrameworkVariants.Find(x => x.Matches(Platform.Value, Architecture.Value));
 				if( FrameworkVariant != null )
 				{
 					return DirectoryReference.Combine(FrameworkDirectory!, FrameworkVariant.LibraryIdentifier);
diff --git a/Engine/Source/Programs/UnrealBuildTool/Platform/IOS/IOSToolChain.cs b/Engine/Source/Programs/UnrealBuildTool/Platform/IOS/IOSToolChain.cs
--- a/Engine/Source/Programs/UnrealBuildTool/Platform/IOS/IOSToolChain.cs
+++ b/Engine/Source/Programs/UnrealBuildTool/Platform/IOS/IOSToolChain.cs
@@ -474,8 +474,12 @@ namespace UnrealBuildTool
 			{
 				if (Framework.ZipFile != null)
 				{
-					FileItem ExtractedTokenFile = ExtractFramework(Framework, Graph, Logger);
-					FrameworkTokenFiles.Add(ExtractedTokenFile);
+					// modern xcode manages all this in xcode script
+					if (!bUseModernXcode)
+					{
+						ExtractFramework(Framework, Graph, Logger);
+					}
+					FrameworkTokenFiles.Add(Framework.ExtractedTokenFile!);
 				}
 			}
 
@@ -938,18 +942,19 @@ namespace UnrealBuildTool
 			{
 				throw new BuildException("Unable to extract framework '{0}' - no zip file specified", Framework.Name);
 			}
-			if(Framework.ExtractedTokenFile == null)
+			if (!Framework.bHasMadeUnzipAction)
 			{
+				Framework.bHasMadeUnzipAction = true;
+
 				FileItem InputFile = FileItem.GetItemByFileReference(Framework.ZipFile);
-				Framework.ExtractedTokenFile = FileItem.GetItemByFileReference(new FileReference(Framework.ZipOutputDirectory!.FullName + ".extracted"));
 
 				StringBuilder ExtractScript = new StringBuilder();
 				ExtractScript.AppendLine("#!/bin/sh");
 				ExtractScript.AppendLine("set -e");
 				// ExtractScript.AppendLine("set -x"); // For debugging
-				ExtractScript.AppendLine(String.Format("[ -d {0} ] && rm -rf {0}", Utils.MakePathSafeToUseWithCommandLine(Framework.ZipOutputDirectory.FullName)));
+				ExtractScript.AppendLine(String.Format("[ -d {0} ] && rm -rf {0}", Utils.MakePathSafeToUseWithCommandLine(Framework.ZipOutputDirectory!.FullName)));
 				ExtractScript.AppendLine(String.Format("unzip -q -o {0} -d {1}", Utils.MakePathSafeToUseWithCommandLine(Framework.ZipFile.FullName), Utils.MakePathSafeToUseWithCommandLine(Framework.ZipOutputDirectory.ParentDirectory!.FullName))); // Zip contains folder with the same name, hence ParentDirectory
-				ExtractScript.AppendLine(String.Format("touch {0}", Utils.MakePathSafeToUseWithCommandLine(Framework.ExtractedTokenFile.AbsolutePath)));
+				ExtractScript.AppendLine(String.Format("touch {0}", Utils.MakePathSafeToUseWithCommandLine(Framework.ExtractedTokenFile!.AbsolutePath)));
 
 				FileItem ExtractScriptFileItem = Graph.CreateIntermediateTextFile(new FileReference(Framework.ZipOutputDirectory.FullName + ".sh"), ExtractScript.ToString());
 
@@ -964,7 +969,7 @@ namespace UnrealBuildTool
 				UnzipAction.StatusDescription = String.Format("Unzipping : {0} -> {1}", Framework.ZipFile, Framework.ZipOutputDirectory);
 				UnzipAction.bCanExecuteRemotely = false;
 			}
-			return Framework.ExtractedTokenFile;
+			return Framework.ExtractedTokenFile!;
 		}
 
 		public static DirectoryReference GenerateAssetCatalog(FileReference? ProjectFile, UnrealTargetPlatform Platform, ref bool bUserImagesExist)
diff --git a/Engine/Source/Programs/UnrealBuildTool/ProjectFiles/ProjectFileGenerator.cs b/Engine/Source/Programs/UnrealBuildTool/ProjectFiles/ProjectFileGenerator.cs
--- a/Engine/Source/Programs/UnrealBuildTool/ProjectFiles/ProjectFileGenerator.cs
+++ b/Engine/Source/Programs/UnrealBuildTool/ProjectFiles/ProjectFileGenerator.cs
@@ -1169,6 +1169,7 @@ namespace UnrealBuildTool
 				// Figure out which targets we need about IntelliSense for.  We only need to worry about targets for projects
 				// that we're actually generating in this session.
 				List<Tuple<ProjectFile, ProjectTarget>> IntelliSenseTargetFiles = new List<Tuple<ProjectFile, ProjectTarget>>();
+				List<Tuple<ProjectFile, ProjectTarget>> AllProjectTargetFiles = new List<Tuple<ProjectFile, ProjectTarget>>();
 				{
 					// Engine targets
 					foreach (ProjectFile EngineProject in EngineProjects)
@@ -1182,6 +1183,7 @@ namespace UnrealBuildTool
 							{
 								IntelliSenseTargetFiles.Add(Tuple.Create(EngineProject, EngineTarget));
 							}
+							AllProjectTargetFiles.Add(Tuple.Create(EngineProject, EngineTarget));
 						}
 					}
 
@@ -1194,6 +1196,7 @@ namespace UnrealBuildTool
 							{
 								IntelliSenseTargetFiles.Add(Tuple.Create(ProgramProject, ProjectTarget));
 							}
+							AllProjectTargetFiles.Add(Tuple.Create(ProgramProject, ProjectTarget));
 						}
 					}
 
@@ -1208,6 +1211,7 @@ namespace UnrealBuildTool
 							{
 								IntelliSenseTargetFiles.Add(Tuple.Create(GameProject, ProjectTarget));
 							}
+							AllProjectTargetFiles.Add(Tuple.Create(GameProject, ProjectTarget));
 						}
 					}
 				}
@@ -1215,6 +1219,9 @@ namespace UnrealBuildTool
 				// write out any additional debug information for the solution (such as UnrealVS configuration)
 				WriteDebugSolutionFiles(PlatformProjectGenerators, IntermediateProjectFilesPath, Logger);
 
+				// add any native project information for all Targets we are generating
+				AddAdditionalNativeTargetInformation(PlatformProjectGenerators, AllProjectTargetFiles, Logger);
+
 				// Generate IntelliSense data if we need to.  This involves having UBT simulate the action compilation of
 				// the targets so that we can extra the compiler defines, include paths, etc.
 				GenerateIntelliSenseData(Arguments, IntelliSenseTargetFiles, Logger);
@@ -2933,6 +2940,17 @@ namespace UnrealBuildTool
 		{
 		}
 
+		/// <summary>
+		/// Allows the generator to XXXXXXXxXXXXXXXxXXXXXXXxXXXXXXXx
+		/// </summary>
+		/// <param name="PlatformProjectGenerators"></param>
+		/// <param name="Targets"></param>
+		/// <param name="Logger"></param>
+		protected virtual void AddAdditionalNativeTargetInformation(PlatformProjectGeneratorCollection PlatformProjectGenerators, List<Tuple<ProjectFile, ProjectTarget>> Targets, ILogger Logger)
+		{
+		}
+		
+
 
 		/// <summary>
 		/// Writes the specified string content to a file.  Before writing to the file, it loads the existing file (if present) to see if the contents have changed
diff --git a/Engine/Source/Programs/UnrealBuildTool/ProjectFiles/Xcode/XcodeProject.cs b/Engine/Source/Programs/UnrealBuildTool/ProjectFiles/Xcode/XcodeProject.cs
--- a/Engine/Source/Programs/UnrealBuildTool/ProjectFiles/Xcode/XcodeProject.cs
+++ b/Engine/Source/Programs/UnrealBuildTool/ProjectFiles/Xcode/XcodeProject.cs
@@ -14,6 +14,7 @@ using Microsoft.Extensions.Logging;
 using System.Text.RegularExpressions;
 using System.Diagnostics.CodeAnalysis;
 using System.Security.Cryptography;
+using static UnrealBuildTool.Unity;
 
 
 /*****
@@ -271,7 +272,9 @@ namespace UnrealBuildTool.XcodeProjectXcconfig
 		public List<UnrealBuildConfig> AllConfigs = new();
 
 		public List<UnrealBatchedFiles> BatchedFiles = new();
- 
+
+		public List<string> ExtraPreBuildScriptLines = new();
+
 		public FileReference? UProjectFileLocation = null;
 		public DirectoryReference XcodeProjectFileLocation;
 		public DirectoryReference ProductDirectory;
@@ -806,11 +809,11 @@ namespace UnrealBuildTool.XcodeProjectXcconfig
 		/// <summary>
 		/// Constructor
 		/// </summary>
-		public XcodeSourceFile(FileReference InitFilePath, DirectoryReference? InitRelativeBaseFolder)
+		public XcodeSourceFile(FileReference InitFilePath, DirectoryReference? InitRelativeBaseFolder, string? FileRefGuid=null)
 			: base(InitFilePath, InitRelativeBaseFolder)
 		{
 			FileGuid = XcodeProjectFileGenerator.MakeXcodeGuid();
-			FileRefGuid = XcodeProjectFileGenerator.MakeXcodeGuid();
+			this.FileRefGuid = FileRefGuid ?? XcodeProjectFileGenerator.MakeXcodeGuid();
 		}
 
 		/// <summary>
@@ -1093,7 +1096,7 @@ namespace UnrealBuildTool.XcodeProjectXcconfig
 		public void AddResource(FileReference Resource)
 		{
 			XcodeSourceFile ResourceSource = new XcodeSourceFile(Resource, null);
-			FileCollection.ProcessFile(ResourceSource, true, false);
+			FileCollection.ProcessFile(ResourceSource, true, false, "Resources");
 
 			FileItems.Add(ResourceSource);
 		}
@@ -1101,13 +1104,32 @@ namespace UnrealBuildTool.XcodeProjectXcconfig
 		public void AddFolderResource(DirectoryReference Resource, string GroupName)
 		{
 			XcodeSourceFile ResourceSource = new XcodeSourceFile(new FileReference(Resource.FullName), null);
-			FileCollection.ProcessFile(ResourceSource, true, true, GroupName);
+			FileCollection.ProcessFile(ResourceSource, true, false, GroupName);
 
 			//Project.FileCollection.AddFolderReference(CookedData.MakeRelativeTo(UnrealData.XcodeProjectFileLocation.ParentDirectory!), "CookedData_Game");
 			FileItems.Add(ResourceSource);
 		}
 	}
 
+	class XcodeFrameworkBuildPhase : XcodeBuildPhase
+	{
+		private XcodeFileCollection FileCollection;
+
+		public XcodeFrameworkBuildPhase(XcodeFileCollection FileCollection)
+			: base("Frameworks", "PBXFrameworksBuildPhase")
+		{
+			this.FileCollection = FileCollection;
+		}
+
+		public void AddFramework(DirectoryReference Framework, string FileRefGuid)
+		{
+			XcodeSourceFile FrameworkSource = new XcodeSourceFile(new FileReference(Framework.FullName), null, FileRefGuid);
+			FileCollection.ProcessFile(FrameworkSource, true, false, "Frameworks", ""); ;
+			FileItems.Add(FrameworkSource);
+
+		}
+	}
+
 	class XcodeCopyFilesBuildPhase : XcodeBuildPhase
 	{
 		private XcodeFileCollection FileCollection;
@@ -1118,23 +1140,22 @@ namespace UnrealBuildTool.XcodeProjectXcconfig
 			this.FileCollection = FileCollection;
 
 			MiscItems.Add($"dstPath = \"\";");
-			MiscItems.Add($"dstSubFolderSpec = 10;");
+			MiscItems.Add($"dstSubfolderSpec = 10;");
 			MiscItems.Add($"name = \"{Name}\";");
 		}
 
-		public void AddFramework(DirectoryReference Framework)
+		public void AddFramework(DirectoryReference Framework, string FileRefGuid)
 		{
-			XcodeSourceFile ResourceSource = new XcodeSourceFile(new FileReference(Framework.FullName), null);
-			FileCollection.ProcessFile(ResourceSource, true, true, "Frameworks");
-
-			FileItems.Add(ResourceSource);
+			XcodeSourceFile FrameworkSource = new XcodeSourceFile(new FileReference(Framework.FullName), null, FileRefGuid);
+			FileCollection.ProcessFile(FrameworkSource, true, false, "", "settings = {ATTRIBUTES = (CodeSignOnCopy, RemoveHeadersOnCopy, ); };");
+			FileItems.Add(FrameworkSource);
 		}
 	}
 
 	class XcodeShellScriptBuildPhase : XcodeBuildPhase
 	{
 		public XcodeShellScriptBuildPhase(string Name, IEnumerable<string> ScriptLines, IEnumerable<string> Inputs, IEnumerable<string> Outputs)
-			: base("Resources", "PBXShellScriptBuildPhase")
+			: base(Name, "PBXShellScriptBuildPhase")
 		{
 			MiscItems.Add($"name = \"{Name}\";");
 
@@ -1382,7 +1403,7 @@ namespace UnrealBuildTool.XcodeProjectXcconfig
 		public UnrealTargetPlatform Platform;
 		public TargetType TargetType;
 
-		public XcodeRunTarget(XcodeProject Project, string TargetName, TargetType TargetType, UnrealTargetPlatform Platform, XcodeBuildTarget? BuildTarget)
+		public XcodeRunTarget(XcodeProject Project, string TargetName, TargetType TargetType, UnrealTargetPlatform Platform, XcodeBuildTarget? BuildTarget, XcodeProjectFile ProjectFile)
 			: base(Project.UnrealData.bIsAppBundle ? XcodeTarget.Type.Run_App : XcodeTarget.Type.Run_Tool, Project.UnrealData,
 				  TargetName + (XcodeProjectFileGenerator.PerPlatformMode == XcodePerPlatformMode.OneWorkspacePerPlatform ? "" : $"_{Platform}"))
 		{
@@ -1401,6 +1422,8 @@ namespace UnrealBuildTool.XcodeProjectXcconfig
 				BuildPhases.Add(ResourcesBuildPhase);
 				References.Add(ResourcesBuildPhase);
 
+				ProcessFrameworks(ResourcesBuildPhase, Project, ProjectFile);
+
 				if (Project.Platform == UnrealTargetPlatform.IOS || Project.Platform == UnrealTargetPlatform.TVOS)
 				{
 					// UBT no longer copies the executable into the .app directory in PostBuild, so we do it here
@@ -1421,7 +1444,7 @@ namespace UnrealBuildTool.XcodeProjectXcconfig
 					}
 
 					DirectoryReference CookedData = DirectoryReference.Combine(UnrealData.ProductDirectory, "Saved", "StagedBuilds",TargetPlatformName, "cookeddata");
-					ResourcesBuildPhase.AddFolderResource(CookedData, "ProjectData");
+					ResourcesBuildPhase.AddFolderResource(CookedData, "Resources");
 				}
 
 				List<string> StoryboardPaths = new List<string>()
@@ -1468,6 +1491,104 @@ namespace UnrealBuildTool.XcodeProjectXcconfig
 			BuildConfigList!.BuildConfigs.ForEach(x => x.Xcconfig = Xcconfig);
 		}
 
+		protected void ProcessFrameworks(XcodeResourcesBuildPhase ResourcesBuildPhase, XcodeProject Project, XcodeProjectFile ProjectFile)
+		{
+			// look up to see if we had cached any Frameworks
+			Tuple<ProjectFile, UnrealTargetPlatform> FrameworkKey = Tuple.Create((ProjectFile)ProjectFile, Platform);
+			List<UEBuildFramework>? Frameworks;
+			if (XcodeProjectFileGenerator.TargetFrameworks.TryGetValue(FrameworkKey, out Frameworks))
+			{
+				XcodeCopyFilesBuildPhase EmbedFrameworks = new XcodeCopyFilesBuildPhase(Project.FileCollection);
+				XcodeFrameworkBuildPhase FrameworkPhase = new XcodeFrameworkBuildPhase(Project.FileCollection);
+
+				// filter frameworks that need to installed into the .app (either the framework or a bundle inside a .zip)
+				IEnumerable<UEBuildFramework> InstalledFrameworks = Frameworks.Where(x => x.bCopyFramework || !string.IsNullOrEmpty(x.CopyBundledAssets));
+				// filter frameworks that need to be unzipped before we compile
+				IEnumerable<UEBuildFramework> ZippedFrameworks = Frameworks.Where(x => x.ZipFile != null);
+
+				bool bHasEmbeddedFrameworks = false;
+				// only look at frameworks that need anything copied into 
+				foreach (UEBuildFramework Framework in InstalledFrameworks)
+				{
+					DirectoryReference? FinalFrameworkDir = Framework.GetFrameworkDirectory(null, null);
+					if (FinalFrameworkDir == null)
+					{
+						continue;
+					}
+					// the framework may come with FrameworkDir being parent of a .framework with name of Framework.Name
+					if (!FinalFrameworkDir.HasExtension(".framework") && !FinalFrameworkDir.HasExtension(".xcframework"))
+					{
+						FinalFrameworkDir = DirectoryReference.Combine(FinalFrameworkDir, Framework.Name + ".framework");
+					}
+
+					DirectoryReference BundleRootDir = Framework.GetFrameworkDirectory(null, null)!;
+
+					if (Framework.ZipFile != null)
+					{
+						if (Framework.ZipFile.FullName.EndsWith(".embeddedframework.zip"))
+						{
+							// foo.embeddedframework.zip would have foo.framework inside it, which is what we want to install
+							FinalFrameworkDir = DirectoryReference.Combine(Framework.ZipOutputDirectory!, Framework.ZipFile.GetFileNameWithoutAnyExtensions() + ".framework");
+							BundleRootDir = Framework.ZipOutputDirectory!;
+						}
+						else
+						{
+							FinalFrameworkDir = Framework.ZipOutputDirectory!;
+						}
+					}
+
+					// set up the CopyBundle to be copied
+					if (!string.IsNullOrEmpty(Framework.CopyBundledAssets))
+					{
+						ResourcesBuildPhase.AddFolderResource(DirectoryReference.Combine(BundleRootDir, Framework.CopyBundledAssets), "Resources");
+					}
+
+					if (Framework.bCopyFramework)
+					{
+						string FileRefGuid = XcodeProjectFileGenerator.MakeXcodeGuid();
+						EmbedFrameworks.AddFramework(FinalFrameworkDir, FileRefGuid);
+						FrameworkPhase.AddFramework(FinalFrameworkDir, FileRefGuid);
+
+						bHasEmbeddedFrameworks = true;
+					}
+				}
+
+				if (bHasEmbeddedFrameworks)
+				{
+					BuildPhases.Add(EmbedFrameworks);
+					References.Add(EmbedFrameworks);
+					BuildPhases.Add(FrameworkPhase);
+					References.Add(FrameworkPhase);
+				}
+
+				// each zipped framework needs to be unzipped in case C++ code needs to compile/link against it - this will add unzip commands
+				// to the PreBuild script that is run before anything else happens - note that the ZipDependToken is shared with UBT, so that
+				// if a new framework is unzipped, it will dirty any source files in modules that use this framework
+				foreach (UEBuildFramework Framework in ZippedFrameworks)
+				{
+					string ZipIn = Utils.MakePathSafeToUseWithCommandLine(Framework.ZipFile!.FullName);
+					string ZipOut = Utils.MakePathSafeToUseWithCommandLine(Framework.ZipOutputDirectory!.FullName);
+					// Zip contains folder with the same name, hence ParentDirectory
+					string ZipOutParent = Utils.MakePathSafeToUseWithCommandLine(Framework.ZipOutputDirectory.ParentDirectory!.FullName);
+					string ZipDependToken = Utils.MakePathSafeToUseWithCommandLine(Framework.ExtractedTokenFile!.FullName);
+
+					UnrealData.ExtraPreBuildScriptLines.AddRange(new[]
+					{
+						// delete any output and make sure parent dir exists
+						$"if [ {ZipIn} -nt {ZipDependToken} ] ",
+						$"then",
+						$"  [ -d {ZipOut} ] &amp;&amp; rm -rf {ZipOut}",
+						$"  mkdir -p {ZipOutParent}",
+						// unzip the framework and maybe extra data
+						$"  unzip -q -o {ZipIn} -d {ZipOutParent}",
+						$"  touch {ZipDependToken}",
+						$"fi",
+						$"",
+					});
+				}
+			}
+		}
+
 		protected override void WriteExtraTargetProperties(StringBuilder Content)
 		{
 			Content.WriteLine($"\t\t\tproductReference = {ProductGuid};");
@@ -1585,7 +1706,7 @@ namespace UnrealBuildTool.XcodeProjectXcconfig
 
 		public XcodeBuildConfigList ProjectBuildConfigs;
 
-		public XcodeProject(UnrealTargetPlatform? Platform, UnrealData UnrealData, XcodeFileCollection FileCollection)
+		public XcodeProject(UnrealTargetPlatform? Platform, UnrealData UnrealData, XcodeFileCollection FileCollection, XcodeProjectFile ProjectFile)
 		{
 			this.Platform = Platform;
 			this.UnrealData = UnrealData;
@@ -1604,7 +1725,7 @@ namespace UnrealBuildTool.XcodeProjectXcconfig
 			List<UnrealTargetPlatform> TargetPlatforms = Platform == null ? XcodeProjectFileGenerator.XcodePlatforms : new() { Platform.Value };
 			foreach (UnrealTargetPlatform TargetPlatform in TargetPlatforms)
 			{
-				XcodeRunTarget RunTarget = new XcodeRunTarget(this, UnrealData.ProductName, UnrealData.AllConfigs[0].ProjectTarget!.TargetRules!.Type, TargetPlatform, BuildTarget);
+				XcodeRunTarget RunTarget = new XcodeRunTarget(this, UnrealData.ProductName, UnrealData.AllConfigs[0].ProjectTarget!.TargetRules!.Type, TargetPlatform, BuildTarget, ProjectFile);
 				RunTargets.Add(RunTarget);
 				References.Add(RunTarget);
 			}
@@ -2175,6 +2296,51 @@ namespace UnrealBuildTool.XcodeProjectXcconfig
 
 	class XcodeFileCollection
 	{
+		static readonly string[] ResourceExtensions =
+		{
+			".xcassets",
+			".storyboardc",
+			".storyboard",
+			".xcframework",
+			".framework",
+			".lproj",
+			// no extensions will be folders, like "cookeddata" that we want to copy as a resource
+			"",
+		};
+
+		static readonly string[] SourceCodeExtensions =
+		{
+			".c",
+			".cc",
+			".cpp",
+			".m",
+			".mm",
+		};
+
+		static readonly Dictionary<string, string> FileTypeMap = new()
+		{
+			{ ".c", "sourcecode.c.objc" },
+			{ ".m", "sourcecode.c.objc" },
+			{ ".cc", "sourcecode.cpp.objcpp" },
+			{ ".cpp", "sourcecode.cpp.objcpp" },
+			{ ".mm", "sourcecode.cpp.objcpp" },
+			{ ".h", "sourcecode.c.h" },
+			{ ".inl", "sourcecode.c.h" },
+			{ ".pch", "sourcecode.c.h" },
+			{ ".framework", "wrapper.framework" },
+			{ ".plist", "text.plist.xml" },
+			{ ".png", "image.png" },
+			{ ".icns", "image.icns" },
+			{ ".xcassets", "folder.assetcatalog" },
+			{ ".storyboardc", "wrapper.storyboardc" },
+			{ ".storyboard", "file.storyboard" },
+			{ ".xcframework", "wrapper.xcframework" },
+			{ ".bundle", "\"wrapper.plug-in\"" }
+		};
+
+		const string FileTypeDefault = "file.txt";
+
+
 		class ManualFileReference
 		{
 			public string Guid = XcodeProjectFileGenerator.MakeXcodeGuid();
@@ -2196,7 +2362,7 @@ namespace UnrealBuildTool.XcodeProjectXcconfig
 
 		internal Dictionary<string, XcodeFileGroup> Groups = new();
 		internal Dictionary<XcodeSourceFile, FileReference?> BuildableFilesToResponseFile = new();
-		internal List<XcodeSourceFile> BuildableResourceFiles = new();
+		internal List<Tuple<XcodeSourceFile, string>> BuildableResourceFiles = new();
 		internal List<XcodeSourceFile> AllFiles = new();
 
 		private List<ManualFileReference> ManualFiles = new();
@@ -2234,7 +2400,7 @@ namespace UnrealBuildTool.XcodeProjectXcconfig
 			RootDirectory = UProjectLocation == null || UProjectLocation.ContainsName("Programs", 0) ? (Unreal.EngineDirectory) : UProjectLocation!.Directory;
 		}
 
-		public void ProcessFile(XcodeSourceFile File, bool bIsForBuild, bool bIsFolder, string? GroupName=null)
+		public void ProcessFile(XcodeSourceFile File, bool bIsForBuild, bool bIsFolder, string? GroupName=null, string ExtraResourceSettings="")
 		{
 			// remember all buildable files
 			if (bIsForBuild)
@@ -2252,28 +2418,30 @@ namespace UnrealBuildTool.XcodeProjectXcconfig
 				}
 				else if (IsResourceFile(FileExtension))
 				{
-					BuildableResourceFiles.Add(File);
+					BuildableResourceFiles.Add(Tuple.Create(File, ExtraResourceSettings));
 				}
 			}
 			
-
-			// group the files by path
-			XcodeFileGroup? Group = GroupName == null ? FindGroupByAbsolutePath(File.Reference.Directory) : null;
-			if (Group != null)
+			if (GroupName != "")
 			{
-				AllFiles.Add(File);
-				Group.Files.Add(File);
-			}
-			else
-			{
-				GroupName = GroupName ?? (File.Reference.IsUnderDirectory(Unreal.EngineDirectory) ? "EngineReferences" : "ExternalReferences");
-				if (bIsFolder)
+				// group the files by path
+				XcodeFileGroup? Group = GroupName == null ? FindGroupByAbsolutePath(File.Reference.Directory) : null;
+				if (Group != null)
 				{
-					AddFolderReference(File.FileRefGuid, File.Reference.MakeRelativeTo(ProjectDirectory), GroupName);
+					AllFiles.Add(File);
+					Group.Files.Add(File);
 				}
 				else
 				{
-					AddFileReference(File.FileRefGuid, File.Reference.FullName, GetFileKey(File.Reference.GetExtension()), GetFileType(File.Reference.GetExtension()), "<absolute>", GroupName);
+					GroupName = GroupName ?? (File.Reference.IsUnderDirectory(Unreal.EngineDirectory) ? "EngineReferences" : "ExternalReferences");
+					if (bIsFolder)
+					{
+						AddFolderReference(File.FileRefGuid, File.Reference.MakeRelativeTo(ProjectDirectory), GroupName);
+					}
+					else
+					{
+						AddFileReference(File.FileRefGuid, File.Reference.FullName, GetFileKey(File.Reference.GetExtension()), GetFileType(File.Reference.GetExtension()), "<absolute>", GroupName);
+					}
 				}
 			}
 		}
@@ -2338,9 +2506,11 @@ namespace UnrealBuildTool.XcodeProjectXcconfig
 				}
 				Content.WriteLine($"\t\t{Pair.Key.FileGuid} = {{isa = PBXBuildFile; fileRef = {Pair.Key.FileRefGuid};{CompileFlags} }}; /* {Pair.Key.Reference.GetFileName()} */");
 			}
-			foreach (XcodeSourceFile Resource in BuildableResourceFiles)
+			foreach (Tuple<XcodeSourceFile, string> Resource in BuildableResourceFiles)
 			{
-				Content.WriteLine($"\t\t{Resource.FileGuid} = {{isa = PBXBuildFile; fileRef = {Resource.FileRefGuid}; }}; /* {Resource.Reference.GetFileName()} */");
+				XcodeSourceFile File = Resource.Item1;
+				string ExtraSettings = Resource.Item2;
+				Content.WriteLine($"\t\t{File.FileGuid} = {{isa = PBXBuildFile; fileRef = {File.FileRefGuid}; {ExtraSettings}}}; /* {File.Reference.GetFileName()} */");
 			}
 			Content.WriteLine("/* End PBXBuildFile section */");
 			Content.WriteLine();
@@ -2475,57 +2645,17 @@ namespace UnrealBuildTool.XcodeProjectXcconfig
 		}
 
 
-		/// <summary>
-		/// Gets Xcode file category based on its extension
-		/// </summary>
-		internal static string GetFileCategory(string Extension)
-		{
-			// @todo Mac: Handle more categories
-			switch (Extension)
-			{
-				case ".framework":
-					return "Frameworks";
-				default:
-					return "Sources";
-			}
-		}
-
 		/// <summary>
 		/// Gets Xcode file type based on its extension
 		/// </summary>
 		internal static string GetFileType(string Extension)
 		{
-			// @todo Mac: Handle more file types
-			switch (Extension)
+			string? FileType;
+			if (FileTypeMap.TryGetValue(Extension, out FileType))
 			{
-				case ".c":
-				case ".m":
-					return "sourcecode.c.objc";
-				case ".cc":
-				case ".cpp":
-				case ".mm":
-					return "sourcecode.cpp.objcpp";
-				case ".h":
-				case ".inl":
-				case ".pch":
-					return "sourcecode.c.h";
-				case ".framework":
-					return "wrapper.framework";
-				case ".plist":
-					return "text.plist.xml";
-				case ".png":
-					return "image.png";
-				case ".icns":
-					return "image.icns";
-				case ".xcassets":
-					return "folder.assetcatalog";
-				case ".storyboardc":
-					return "wrapper.storyboardc";
-				case ".storyboard":
-					return "file.storyboard";
-				default:
-					return "file.text";
+				return FileType;
 			}
+			return FileTypeDefault;
 		}
 
 		// most types Xcode wants as explicitFileType, but some want lastKnownFileType
@@ -2534,6 +2664,7 @@ namespace UnrealBuildTool.XcodeProjectXcconfig
 			switch (Extension)
 			{
 				case ".storyboard":
+				case ".bundle":
 					return "lastKnownFileType";
 				default:
 					return "explicitFileType";
@@ -2544,7 +2675,7 @@ namespace UnrealBuildTool.XcodeProjectXcconfig
 		/// </summary>
 		internal static bool IsSourceCode(string Extension)
 		{
-			return Extension == ".c" || Extension == ".cc" || Extension == ".cpp" || Extension == ".m" || Extension == ".mm";
+			return SourceCodeExtensions.Contains(Extension);
 		}
 
 		/// <summary>
@@ -2552,7 +2683,7 @@ namespace UnrealBuildTool.XcodeProjectXcconfig
 		/// </summary>
 		internal static bool IsResourceFile(string Extension)
 		{
-			return Extension == ".xcassets" || Extension == ".storyboardc" || Extension == ".storyboard" || Extension == "";
+			return ResourceExtensions.Contains(Extension);
 		}
 
 		/// <summary>
@@ -2683,7 +2814,7 @@ namespace UnrealBuildTool.XcodeProjectXcconfig
 			return new XcodeSourceFile(InitFilePath, InitProjectSubFolder);
 		}
 
-		private void ConditionalCreateLegacyProject()
+		public void ConditionalCreateLegacyProject()
 		{
 			if (!bHasCheckedForLegacy)
 			{
@@ -2828,7 +2959,7 @@ namespace UnrealBuildTool.XcodeProjectXcconfig
 
 				// now create the xcodeproject elements (project -> target -> buildconfigs, etc)
 				FileCollection = new XcodeFileCollection(SharedFileCollection);
-				XcodeProject RootProject = new XcodeProject(Platform, UnrealData, FileCollection);
+				XcodeProject RootProject = new XcodeProject(Platform, UnrealData, FileCollection, this);
 				RootProjects[RootProject] = Platform;
 
 				if (FileReference.Exists(TemplateProject))
@@ -2876,7 +3007,8 @@ namespace UnrealBuildTool.XcodeProjectXcconfig
 						string ProjectName = ProjectFilePathForPlatform(Platform).GetFileNameWithoutAnyExtensions();
 						string? BuildTargetGuid = XcodeProjectNode.GetNodesOfType<XcodeBuildTarget>(RootProject).FirstOrDefault()?.Guid;
 						string? IndexTargetGuid = XcodeProjectNode.GetNodesOfType<XcodeIndexTarget>(RootProject).FirstOrDefault()?.Guid;
-						WriteSchemeFile(Platform, ProjectName, UnrealData.ProductName, RootProject.RunTargets, BuildTargetGuid, IndexTargetGuid, UnrealData.UProjectFileLocation != null ? UnrealData.UProjectFileLocation.FullName : "");
+						WriteSchemeFile(Platform, ProjectName, UnrealData.ProductName, RootProject.RunTargets, BuildTargetGuid, IndexTargetGuid, 
+							UnrealData.UProjectFileLocation != null ? UnrealData.UProjectFileLocation.FullName : "", UnrealData.ExtraPreBuildScriptLines);
 					}
 				}
 				else
@@ -3533,7 +3665,7 @@ namespace UnrealBuildTool.XcodeProjectXcconfig
 			return FileReference.Combine(GetProjectSchemeDirectory(Platform), TargetName + ".xcscheme");
 		}
 
-		private void WriteSchemeFile(UnrealTargetPlatform? Platform, string ProjectName, string ProductName, List<XcodeRunTarget> RunTargets, string? BuildTargetGuid, string? IndexTargetGuid, string GameProjectPath)
+		private void WriteSchemeFile(UnrealTargetPlatform? Platform, string ProjectName, string ProductName, List<XcodeRunTarget> RunTargets, string? BuildTargetGuid, string? IndexTargetGuid, string GameProjectPath, IEnumerable<string> ExtraPreBuildScriptLines)
 		{
 			StringBuilder Content = new StringBuilder();
 
@@ -3584,6 +3716,9 @@ namespace UnrealBuildTool.XcodeProjectXcconfig
 					"fi",
 					"",
 				};
+				// convert special characters to be xml happy
+				IEnumerable<string> Lines = ScriptLines.Concat(ExtraPreBuildScriptLines);
+				Lines = Lines.Select(x => x.Replace("\"", "&quot;").Replace("&&", "&amp;&amp;"));
 
 				Content.WriteLine("<?xml version=\"1.0\" encoding=\"UTF-8\"?>");
 				Content.WriteLine("<Scheme");
@@ -3597,7 +3732,7 @@ namespace UnrealBuildTool.XcodeProjectXcconfig
 				Content.WriteLine("            ActionType = \"Xcode.IDEStandardExecutionActionsCore.ExecutionActionType.ShellScriptAction\">");
 				Content.WriteLine("            <ActionContent");
 				Content.WriteLine("               title = \"Touch UBT generated tiles\"");
-				Content.WriteLine("               scriptText = \"" + string.Join("&#10;", ScriptLines) + "\">");
+				Content.WriteLine("               scriptText = \"" + string.Join("&#10;", Lines) + "\">");
 				Content.WriteLine("               <EnvironmentBuildable>");
 				Content.WriteLine("                  <BuildableReference");
 				Content.WriteLine("                     BuildableIdentifier = \"primary\"");
diff --git a/Engine/Source/Programs/UnrealBuildTool/ProjectFiles/Xcode/XcodeProjectFileGenerator.cs b/Engine/Source/Programs/UnrealBuildTool/ProjectFiles/Xcode/XcodeProjectFileGenerator.cs
--- a/Engine/Source/Programs/UnrealBuildTool/ProjectFiles/Xcode/XcodeProjectFileGenerator.cs
+++ b/Engine/Source/Programs/UnrealBuildTool/ProjectFiles/Xcode/XcodeProjectFileGenerator.cs
@@ -214,7 +214,7 @@ namespace UnrealBuildTool
 
 			// if we want one workspace with multiple platforms, and we have at least one modern project, then process each platform individually
 			// otherwise use null as Platform which means to merge all platforms
-			List<UnrealTargetPlatform?> PlatformsToProcess = bHasModernProjects ? WorkspacePlatforms : NullPlatformList; 
+			List<UnrealTargetPlatform?> PlatformsToProcess = bHasModernProjects ? WorkspacePlatforms : NullPlatformList;
 			foreach (UnrealTargetPlatform? Platform in PlatformsToProcess)
 			{
 				StringBuilder WorkspaceDataContent = new();
@@ -327,7 +327,7 @@ namespace UnrealBuildTool
 		/// </summary>
 		static public List<UnrealTargetPlatform> XcodePlatforms = new();
 
-		static public List<UnrealTargetPlatform?> WorkspacePlatforms = new ();
+		static public List<UnrealTargetPlatform?> WorkspacePlatforms = new();
 		static public List<UnrealTargetPlatform?> RunTargetPlatforms = new();
 		static public List<UnrealTargetPlatform?> NullPlatformList = new() { null };
 
@@ -405,5 +405,78 @@ namespace UnrealBuildTool
 				}
 			}
 		}
+
+		internal static Dictionary<Tuple<ProjectFile, UnrealTargetPlatform>, List<UEBuildFramework>> TargetFrameworks = new();
+		internal static Dictionary<Tuple<ProjectFile, UnrealTargetPlatform>, List<UEBuildBundleResource>> TargetBundles = new();
+
+		protected override void AddAdditionalNativeTargetInformation(PlatformProjectGeneratorCollection PlatformProjectGenerators, List<Tuple<ProjectFile, ProjectTarget>> Targets, ILogger Logger)
+		{
+			DateTime MainStart = DateTime.UtcNow;
+			foreach (var TargetPair in Targets)
+			{
+				ProjectFile TargetProjectFile = TargetPair.Item1;
+				// don't do this for legacy projets, for speed
+				((XcodeProjectXcconfig.XcodeProjectFile)TargetProjectFile).ConditionalCreateLegacyProject();
+				if (((XcodeProjectXcconfig.XcodeProjectFile)TargetProjectFile).bHasLegacyProject)
+				{
+					continue;
+				}
+
+				ProjectTarget CurTarget = TargetPair.Item2;
+
+				UnrealTargetPlatform Platform = UnrealTargetPlatform.IOS;
+				UnrealArch Arch = UnrealArch.Arm64;
+
+				if (!CurTarget.SupportedPlatforms.Any(x => x == Platform))
+				{
+					continue;
+				}
+				TargetDescriptor TargetDesc = new TargetDescriptor(CurTarget.UnrealProjectFilePath, CurTarget.Name, Platform, UnrealTargetConfiguration.Development,
+					new UnrealArchitectures(Arch), new CommandLineArguments(new string[] { "-skipclangvalidation" }));
+				DateTime Start = DateTime.UtcNow;
+
+				try
+				{
+					// Create the target
+					UEBuildTarget Target = UEBuildTarget.Create(TargetDesc, true, false, bUsePrecompiled, Logger);
+
+					List<UEBuildFramework> Frameworks = new();
+					List<UEBuildBundleResource> Bundles = new();
+					// Generate a compile environment for each module in the binary
+					CppCompileEnvironment GlobalCompileEnvironment = Target.CreateCompileEnvironmentForProjectFiles(Logger);
+					foreach (UEBuildBinary Binary in Target.Binaries)
+					{
+						CppCompileEnvironment BinaryCompileEnvironment = Binary.CreateBinaryCompileEnvironment(GlobalCompileEnvironment);
+						foreach (UEBuildModuleCPP Module in Binary.Modules.OfType<UEBuildModuleCPP>())
+						{
+							CppCompileEnvironment CompileEnvironment = Module.CreateModuleCompileEnvironment(Target.Rules, BinaryCompileEnvironment, Logger);
+							Frameworks.AddRange(CompileEnvironment.AdditionalFrameworks);
+						}
+					}
+
+					// track frameworks if we found any
+					if (Frameworks.Count > 0)
+					{
+						lock (TargetFrameworks)
+						{
+							TargetFrameworks.Add(Tuple.Create(TargetProjectFile, Platform), Frameworks);
+						}
+					}
+					if (Bundles.Count > 0)
+					{
+						lock (TargetBundles)
+						{
+							TargetBundles.Add(Tuple.Create(TargetProjectFile, Platform), Bundles);
+						}
+					}
+				}
+				catch (Exception)
+				{
+
+				}
+			}
+
+			Log.TraceInformation("GettingNativeInfo {0}ms overall", (DateTime.UtcNow - MainStart).TotalMilliseconds);
+		}
 	}
 }
