Move AssemblyFile output to VCCompileAction. Support compiling assemblies and preprocessing files with SingleFile

diff --git a/Engine/Platforms/Hololens/Source/Programs/UnrealBuildTool/HoloLensToolChain.cs b/Engine/Platforms/Hololens/Source/Programs/UnrealBuildTool/HoloLensToolChain.cs
--- a/Engine/Platforms/Hololens/Source/Programs/UnrealBuildTool/HoloLensToolChain.cs
+++ b/Engine/Platforms/Hololens/Source/Programs/UnrealBuildTool/HoloLensToolChain.cs
@@ -75,9 +75,9 @@ namespace UnrealBuildTool
 			}
 		}
 
-		protected override void AppendCLArguments_CPP(CppCompileEnvironment CompileEnvironment, List<string> Arguments, DirectoryReference OutputDir)
+		protected override void AppendCLArguments_CPP(CppCompileEnvironment CompileEnvironment, List<string> Arguments)
 		{
-			base.AppendCLArguments_CPP(CompileEnvironment, Arguments, OutputDir);
+			base.AppendCLArguments_CPP(CompileEnvironment, Arguments);
 
 			// Enable Windows Runtime extensions.  Do this even for libs (plugins) so that these too can consume WinRT APIs
 			Arguments.Add("/ZW");
diff --git a/Engine/Source/Programs/UnrealBuildTool/Platform/Windows/VCCompileAction.cs b/Engine/Source/Programs/UnrealBuildTool/Platform/Windows/VCCompileAction.cs
--- a/Engine/Source/Programs/UnrealBuildTool/Platform/Windows/VCCompileAction.cs
+++ b/Engine/Source/Programs/UnrealBuildTool/Platform/Windows/VCCompileAction.cs
@@ -54,6 +54,11 @@ namespace UnrealBuildTool
 		/// </summary>
 		public FileItem? ObjectFile { get; set; }
 
+		/// <summary>
+		/// The assembly file to output
+		/// </summary>
+		public FileItem? AssemblyFile { get; set; }
+
 		/// <summary>
 		/// The output preprocessed file
 		/// </summary>
@@ -301,6 +306,7 @@ namespace UnrealBuildTool
 			ToolChainVersion = InAction.ToolChainVersion;
 			SourceFile = InAction.SourceFile;
 			ObjectFile = InAction.ObjectFile;
+			AssemblyFile = InAction.AssemblyFile;
 			PreprocessedFile = InAction.PreprocessedFile;
 			DependencyListFile = InAction.DependencyListFile;
 			CompiledModuleInterfaceFile = InAction.CompiledModuleInterfaceFile;
@@ -336,6 +342,7 @@ namespace UnrealBuildTool
 			ToolChainVersion = Reader.ReadString()!;
 			SourceFile = Reader.ReadFileItem();
 			ObjectFile = Reader.ReadFileItem();
+			AssemblyFile = Reader.ReadFileItem();
 			PreprocessedFile = Reader.ReadFileItem();
 			DependencyListFile = Reader.ReadFileItem();
 			CompiledModuleInterfaceFile = Reader.ReadFileItem();
@@ -367,6 +374,7 @@ namespace UnrealBuildTool
 			Writer.WriteString(ToolChainVersion);
 			Writer.WriteFileItem(SourceFile);
 			Writer.WriteFileItem(ObjectFile);
+			Writer.WriteFileItem(AssemblyFile);
 			Writer.WriteFileItem(PreprocessedFile);
 			Writer.WriteFileItem(DependencyListFile);
 			Writer.WriteFileItem(CompiledModuleInterfaceFile);
@@ -455,6 +463,11 @@ namespace UnrealBuildTool
 				VCToolChain.AddObjectFile(Arguments, ObjectFile);
 			}
 
+			if (AssemblyFile != null)
+			{
+				VCToolChain.AddAssemblyFile(Arguments, AssemblyFile);
+			}
+
 			// A better way to express this? .json is used as output for /sourceDependencies), but .md.json is used as output for /sourceDependencies:directives)
 			if (DependencyListFile != null && DependencyListFile.HasExtension(".json") && !DependencyListFile.HasExtension(".md.json"))
 			{
diff --git a/Engine/Source/Programs/UnrealBuildTool/Platform/Windows/VCSpecificFileAction.cs b/Engine/Source/Programs/UnrealBuildTool/Platform/Windows/VCSpecificFileAction.cs
--- a/Engine/Source/Programs/UnrealBuildTool/Platform/Windows/VCSpecificFileAction.cs
+++ b/Engine/Source/Programs/UnrealBuildTool/Platform/Windows/VCSpecificFileAction.cs
@@ -13,19 +13,27 @@ namespace UnrealBuildTool
 		DirectoryReference SourceDir;
 		DirectoryReference OutputDir;
 		VCCompileAction BaseAction;
+
+		// TODO: Entire CppCompileEnvironment needs to be saved with BinaryArchiveWriter, some options may still be unused
+		bool bPreprocessOnly;
+		bool bWithAssembly;
 		int SingleFileCounter;
 
-		internal VcSpecificFileAction(DirectoryReference Source, DirectoryReference Output, VCCompileAction Action) : base(ActionType.Compile)
+		internal VcSpecificFileAction(DirectoryReference Source, DirectoryReference Output, VCCompileAction Action, CppCompileEnvironment? CompileEnvironment = null) : base(ActionType.Compile)
 		{
 			SourceDir = Source;
 			OutputDir = Output;
 			BaseAction = Action;
+			bPreprocessOnly = CompileEnvironment?.bPreprocessOnly ?? false;
+			bWithAssembly = CompileEnvironment?.bWithAssembly ?? false;
 		}
 
 		public VcSpecificFileAction(BinaryArchiveReader Reader) : base(ActionType.Compile)
 		{
 			SourceDir = Reader.ReadCompactDirectoryReference();
 			OutputDir = Reader.ReadCompactDirectoryReference();
+			bPreprocessOnly = Reader.ReadBool();
+			bWithAssembly = Reader.ReadBool();
 			BaseAction = new VCCompileAction(Reader);
 		}
 
@@ -33,6 +41,8 @@ namespace UnrealBuildTool
 		{
 			Writer.WriteCompactDirectoryReference(SourceDir);
 			Writer.WriteCompactDirectoryReference(OutputDir);
+			Writer.WriteBool(bPreprocessOnly);
+			Writer.WriteBool(bWithAssembly);
 			BaseAction.Write(Writer);
 		}
 
@@ -54,17 +64,33 @@ namespace UnrealBuildTool
 				return null;
 			}
 
-			FileItem ResponseFile = FileItem.GetItemByFileReference(FileReference.Combine(OutputDir, $"SingleFile{SingleFileCounter}.rsp"));
-			FileItem ObjectFile = FileItem.GetItemByFileReference(FileReference.Combine(OutputDir, $"SingleFile{SingleFileCounter}.obj"));
-			++SingleFileCounter;
-
 			VCCompileAction Action = new VCCompileAction(BaseAction);
-			Action.ObjectFile = ObjectFile;
 			Action.SourceFile = SourceFile;
-			FileItem TempFile = FileItem.GetItemByPath(System.IO.Path.GetTempFileName());
-			Action.ResponseFile = ResponseFile;
-			List<string> Arguments = Action.GetCompilerArguments(Logger);
-			System.IO.File.WriteAllLines(ResponseFile.FullName, Arguments);
+			if (bPreprocessOnly)
+			{
+				Action.PreprocessedFile = FileItem.GetItemByFileReference(FileReference.Combine(OutputDir, $"SingleFile{SingleFileCounter}.i"));
+			}
+			else
+			{
+				Action.ObjectFile = FileItem.GetItemByFileReference(FileReference.Combine(OutputDir, $"SingleFile{SingleFileCounter}.obj"));
+				if (BaseAction.CompilerType.IsMSVC() && bWithAssembly)
+				{
+					Action.AssemblyFile = FileItem.GetItemByFileReference(FileReference.Combine(OutputDir, $"SingleFile{SingleFileCounter}.asm"));
+				}
+			}
+
+			if (BaseAction.CompilerType.IsMSVC())
+			{
+				Action.DependencyListFile = FileItem.GetItemByFileReference(FileReference.Combine(OutputDir, $"SingleFile{SingleFileCounter}.dep.json"));
+			}
+			else if (BaseAction.CompilerType.IsMSVC())
+			{
+				Action.DependencyListFile = FileItem.GetItemByFileReference(FileReference.Combine(OutputDir, $"SingleFile{SingleFileCounter}.d"));
+			}
+			Action.ResponseFile = FileItem.GetItemByFileReference(FileReference.Combine(OutputDir, $"SingleFile{SingleFileCounter}.rsp"));
+			File.WriteAllLines(Action.ResponseFile.FullName, Action.GetCompilerArguments(Logger));
+
+			++SingleFileCounter;
 			return Action;
 		}
 	}
diff --git a/Engine/Source/Programs/UnrealBuildTool/Platform/Windows/VCToolChain.cs b/Engine/Source/Programs/UnrealBuildTool/Platform/Windows/VCToolChain.cs
--- a/Engine/Source/Programs/UnrealBuildTool/Platform/Windows/VCToolChain.cs
+++ b/Engine/Source/Programs/UnrealBuildTool/Platform/Windows/VCToolChain.cs
@@ -248,6 +248,13 @@ namespace UnrealBuildTool
 			Arguments.Add($"/Fo\"{ObjectFileString}\"");
 		}
 
+		public static void AddAssemblyFile(List<string> Arguments, FileItem AssemblyFile)
+		{
+			string AssemblyFileString = NormalizeCommandLinePath(AssemblyFile);
+			Arguments.Add("/FAs");// Write out an assembly file (.asm) with the c++ code embedded in it via comments
+			Arguments.Add($"/Fa\"{AssemblyFileString}\""); // Set the output patj for the asm file
+		}
+
 		public static void AddSourceDependenciesFile(List<string> Arguments, FileItem SourceDependenciesFile)
 		{
 			string SourceDependenciesFileString = NormalizeCommandLinePath(SourceDependenciesFile);
@@ -847,20 +854,12 @@ namespace UnrealBuildTool
 			}
 		}
 
-		protected virtual void AppendCLArguments_CPP(CppCompileEnvironment CompileEnvironment, List<string> Arguments, DirectoryReference OutputDir)
+		protected virtual void AppendCLArguments_CPP(CppCompileEnvironment CompileEnvironment, List<string> Arguments)
 		{
 			if (Target.WindowsPlatform.Compiler.IsMSVC())
 			{
 				// Explicitly compile the file as C++.
 				Arguments.Add("/TP");
-
-				if (CompileEnvironment.bWithAssembly)
-				{
-					// Write out an assembly file (.asm) with the c++ code embedded in it via comments
-					Arguments.Add("/FAs");
-					// Set the output directory for the asm files
-					Arguments.Add("/Fa" + OutputDir.FullName + "\\");
-				}
 			}
 			else
 			{
@@ -1486,12 +1485,9 @@ namespace UnrealBuildTool
 						}
 					}
 
-					if (CompileEnvironment.bWithAssembly && SourceFile.HasExtension(".cpp"))
+					if (Target.WindowsPlatform.Compiler.IsMSVC() && CompileEnvironment.bWithAssembly && SourceFile.HasExtension(".cpp"))
 					{
-						if (Target.WindowsPlatform.Compiler.IsMSVC())
-						{
-							CompileAction.AdditionalProducedItems.Add(FileItem.GetItemByFileReference(FileReference.Combine(OutputDir, SourceFile.Location.GetFileName() + ".asm")));
-						}
+						CompileAction.AssemblyFile = FileItem.GetItemByFileReference(FileReference.Combine(OutputDir, SourceFile.Location.GetFileName() + ".asm"));
 					}
 
 					// Experimental: support for JSON output of timing data
@@ -1554,7 +1550,7 @@ namespace UnrealBuildTool
 				}
 				else
 				{
-					AppendCLArguments_CPP(CompileEnvironment, CompileAction.Arguments, OutputDir);
+					AppendCLArguments_CPP(CompileEnvironment, CompileAction.Arguments);
 				}
 
 				List<FileItem>? InlinedFiles;
@@ -2113,9 +2109,9 @@ namespace UnrealBuildTool
 				return;
 
 			VCCompileAction BaseCompileAction = CreateBaseCompileAction(CompileEnvironment);
-			AppendCLArguments_CPP(CompileEnvironment, BaseCompileAction.Arguments, OutputDir);
+			AppendCLArguments_CPP(CompileEnvironment, BaseCompileAction.Arguments);
 			BaseCompileAction.AdditionalPrerequisiteItems.AddRange(CompileEnvironment.AdditionalPrerequisites); // Primarily for ispc.generated.h files
-			Graph.AddAction(new VcSpecificFileAction(SourceDir, OutputDir, BaseCompileAction));
+			Graph.AddAction(new VcSpecificFileAction(SourceDir, OutputDir, BaseCompileAction, CompileEnvironment));
 		}
 
 		/// <summary>
