UnrealBuildTool: Add missing additional produced item for AssemblyFile

diff --git a/Engine/Source/Programs/UnrealBuildTool/Platform/Windows/VCCompileAction.cs b/Engine/Source/Programs/UnrealBuildTool/Platform/Windows/VCCompileAction.cs
--- a/Engine/Source/Programs/UnrealBuildTool/Platform/Windows/VCCompileAction.cs
+++ b/Engine/Source/Programs/UnrealBuildTool/Platform/Windows/VCCompileAction.cs
@@ -213,6 +213,10 @@ namespace UnrealBuildTool
 				{
 					yield return ObjectFile;
 				}
+				if (AssemblyFile != null)
+				{
+					yield return AssemblyFile;
+				}
 				if (PreprocessedFile != null)
 				{
 					yield return PreprocessedFile;
