UnrealBuildTool: Fix output typo

diff --git a/Engine/Source/Programs/UnrealBuildTool/Configuration/UEBuildPlatform.cs b/Engine/Source/Programs/UnrealBuildTool/Configuration/UEBuildPlatform.cs
--- a/Engine/Source/Programs/UnrealBuildTool/Configuration/UEBuildPlatform.cs
+++ b/Engine/Source/Programs/UnrealBuildTool/Configuration/UEBuildPlatform.cs
@@ -846,7 +846,7 @@ namespace UnrealBuildTool
 					var OtherModule = OtherModules.FirstOrDefault(x => Path.IsUnderDirectory(x.ModuleDirectory));
 					if (OtherModule is UEBuildModuleExternal)
 					{
-						LoggerFunc("Module '{Name}' is adding '{Path}' from external module '{OtherModule}' to PublicIncludePaths. Did you intend to add a public reference?.", Module.Name, Path, OtherModule.Name);
+						LoggerFunc("Module '{Name}' is adding '{Path}' from external module '{OtherModule}' to PublicIncludePaths. Did you intend to add a public reference?", Module.Name, Path, OtherModule.Name);
 					}
 					else if (OtherModule is UEBuildModuleCPP)
 					{
@@ -862,7 +862,7 @@ namespace UnrealBuildTool
 						else if (Path.IsUnderDirectory(DirectoryReference.Combine(OtherModule.ModuleDirectory, "Public"))
 							|| Path.IsUnderDirectory(DirectoryReference.Combine(OtherModule.ModuleDirectory, "Classes")))
 						{
-							LoggerFunc("Module '{Name}' is adding '{Path}' from module '{OtherModule}' to PublicIncludePaths. Did you intend to add a public reference?.", Module.Name, Path, OtherModule.Name);
+							LoggerFunc("Module '{Name}' is adding '{Path}' from module '{OtherModule}' to PublicIncludePaths. Did you intend to add a public reference?", Module.Name, Path, OtherModule.Name);
 						}
 					}
 				}
@@ -887,7 +887,7 @@ namespace UnrealBuildTool
 					var OtherModule = OtherModules.FirstOrDefault(x => Path.IsUnderDirectory(x.ModuleDirectory));
 					if (OtherModule is UEBuildModuleExternal)
 					{
-						LoggerFunc("Module '{Name}' is adding '{Path}' from external module '{OtherModule}' to PrivateIncludePaths. Did you intend to add a private reference?.", Module.Name, Path, OtherModule.Name);
+						LoggerFunc("Module '{Name}' is adding '{Path}' from external module '{OtherModule}' to PrivateIncludePaths. Did you intend to add a private reference?", Module.Name, Path, OtherModule.Name);
 					}
 					else if (OtherModule is UEBuildModuleCPP)
 					{
@@ -903,7 +903,7 @@ namespace UnrealBuildTool
 						else if (Path.IsUnderDirectory(DirectoryReference.Combine(OtherModule.ModuleDirectory, "Public"))
 							|| Path.IsUnderDirectory(DirectoryReference.Combine(OtherModule.ModuleDirectory, "Classes")))
 						{
-							LoggerFunc("Module '{Name}' is adding '{Path}' from module '{OtherModule}' to PrivateIncludePaths. Did you intend to add a private reference?.", Module.Name, Path, OtherModule.Name);
+							LoggerFunc("Module '{Name}' is adding '{Path}' from module '{OtherModule}' to PrivateIncludePaths. Did you intend to add a private reference?", Module.Name, Path, OtherModule.Name);
 						}
 					}
 				}
@@ -931,7 +931,7 @@ namespace UnrealBuildTool
 					var OtherModule = OtherModules.FirstOrDefault(x => Path.IsUnderDirectory(x.ModuleDirectory));
 					if (OtherModule is UEBuildModuleExternal)
 					{
-						LoggerFunc("Module '{Name}' is adding '{Path}' from external module '{OtherModule}' to InternalncludePaths. Did you intend to add a public reference?.", Module.Name, Path, OtherModule.Name);
+						LoggerFunc("Module '{Name}' is adding '{Path}' from external module '{OtherModule}' to InternalncludePaths. Did you intend to add a public reference?", Module.Name, Path, OtherModule.Name);
 					}
 					else if (OtherModule is UEBuildModuleCPP)
 					{
@@ -947,7 +947,7 @@ namespace UnrealBuildTool
 						else if (Path.IsUnderDirectory(DirectoryReference.Combine(OtherModule.ModuleDirectory, "Public"))
 							|| Path.IsUnderDirectory(DirectoryReference.Combine(OtherModule.ModuleDirectory, "Classes")))
 						{
-							LoggerFunc("Module '{Name}' is adding '{Path}' from module '{OtherModule}' to InternalncludePaths. Did you intend to add a public reference?.", Module.Name, Path, OtherModule.Name);
+							LoggerFunc("Module '{Name}' is adding '{Path}' from module '{OtherModule}' to InternalncludePaths. Did you intend to add a public reference?", Module.Name, Path, OtherModule.Name);
 						}
 					}
 				}
