UnrealBuildTool: Make CStandard enum consistent with CppStandard

#rnx

diff --git a/Engine/Source/Programs/UnrealBuildTool/Platform/Windows/VCToolChain.cs b/Engine/Source/Programs/UnrealBuildTool/Platform/Windows/VCToolChain.cs
--- a/Engine/Source/Programs/UnrealBuildTool/Platform/Windows/VCToolChain.cs
+++ b/Engine/Source/Programs/UnrealBuildTool/Platform/Windows/VCToolChain.cs
@@ -1096,13 +1096,22 @@ namespace UnrealBuildTool
 			Arguments.Add("/W0");
 
 			// Select C Standard version available
-			if (CompileEnvironment.CStandard == CStandardVersion.C11)
-			{
-				Arguments.Add("/std:c11");
-			}
-			else if (CompileEnvironment.CStandard >= CStandardVersion.C17)
+			// Select C Standard version available
+			switch (CompileEnvironment.CStandard)
 			{
-				Arguments.Add("/std:c17");
+				case CStandardVersion.None:
+				case CStandardVersion.C89:
+				case CStandardVersion.C99:
+					break;
+				case CStandardVersion.C11:
+					Arguments.Add("/std:c11");
+					break;
+				case CStandardVersion.C17:
+				case CStandardVersion.Latest:
+					Arguments.Add("/std:c17");
+					break;
+				default:
+					throw new BuildException($"Unsupported C standard type set: {CompileEnvironment.CStandard}");
 			}
 		}
 
diff --git a/Engine/Source/Programs/UnrealBuildTool/System/CppCompileEnvironment.cs b/Engine/Source/Programs/UnrealBuildTool/System/CppCompileEnvironment.cs
--- a/Engine/Source/Programs/UnrealBuildTool/System/CppCompileEnvironment.cs
+++ b/Engine/Source/Programs/UnrealBuildTool/System/CppCompileEnvironment.cs
@@ -61,9 +61,9 @@ namespace UnrealBuildTool
 	public enum CStandardVersion
 	{
 		/// <summary>
-		/// Use the default standard version
+		/// Supports no additional standard version flag
 		/// </summary>
-		Default,
+		None,
 
 		/// <summary>
 		/// Supports C89
@@ -89,6 +89,11 @@ namespace UnrealBuildTool
 		/// Latest standard supported by the compiler
 		/// </summary>
 		Latest,
+
+		/// <summary>
+		/// Use the default standard version
+		/// </summary>
+		Default = None,
 	}
 
 	/// <summary>
diff --git a/Engine/Source/Programs/UnrealBuildTool/ToolChain/ClangToolChain.cs b/Engine/Source/Programs/UnrealBuildTool/ToolChain/ClangToolChain.cs
--- a/Engine/Source/Programs/UnrealBuildTool/ToolChain/ClangToolChain.cs
+++ b/Engine/Source/Programs/UnrealBuildTool/ToolChain/ClangToolChain.cs
@@ -358,7 +358,7 @@ namespace UnrealBuildTool
 		{
 			switch (CompileEnvironment.CStandard)
 			{
-				case CStandardVersion.Default:
+				case CStandardVersion.None:
 					break;
 				case CStandardVersion.C89:
 					Arguments.Add("-std=c89");
