UnrealBuildTool: Persist name for -SingleFile compiles, unless there is a filename conflict

diff --git a/Engine/Source/Programs/UnrealBuildTool/Platform/Windows/VCSpecificFileAction.cs b/Engine/Source/Programs/UnrealBuildTool/Platform/Windows/VCSpecificFileAction.cs
--- a/Engine/Source/Programs/UnrealBuildTool/Platform/Windows/VCSpecificFileAction.cs
+++ b/Engine/Source/Programs/UnrealBuildTool/Platform/Windows/VCSpecificFileAction.cs
@@ -3,6 +3,7 @@
 using EpicGames.Core;
 using Microsoft.Extensions.Logging;
 using System.Collections.Generic;
+using System.Diagnostics;
 using System.IO;
 using UnrealBuildBase;
 
@@ -17,7 +18,8 @@ namespace UnrealBuildTool
 		// TODO: Entire CppCompileEnvironment needs to be saved with BinaryArchiveWriter, some options may still be unused
 		bool bPreprocessOnly;
 		bool bWithAssembly;
-		int SingleFileCounter;
+
+		Dictionary<string, List<FileItem>> SingleFiles = new();
 
 		internal VcSpecificFileAction(DirectoryReference Source, DirectoryReference Output, VCCompileAction Action, CppCompileEnvironment? CompileEnvironment = null) : base(ActionType.Compile)
 		{
@@ -50,11 +52,25 @@ namespace UnrealBuildTool
 
 		public IExternalAction? CreateAction(FileItem SourceFile, ILogger Logger)
 		{
+			DirectoryReference IntermediateDirectory = DirectoryReference.Combine(OutputDir, "SingleFile");
+			// Keep track of all specific files, so the output file can be renamed if there's a naming conflict
+			string Filename = SourceFile.Name;
+			if (!SingleFiles.ContainsKey(SourceFile.Name))
+			{
+				SingleFiles[SourceFile.Name] = new();
+			}
+			else
+			{
+				Filename = $"{Path.GetFileNameWithoutExtension(SourceFile.Name)}{SingleFiles[Filename].Count}{Path.GetExtension(SourceFile.Name)}";
+			}
+			SingleFiles[SourceFile.Name].Add(SourceFile);
+
 			// If it is a header file we need to wrap it from another file.. because otherwise it will fail when there are circular dependencies
 			// (There are a lot of those, very few of the .h files in core can be compiled without errors)
 			if (SourceFile.HasExtension(".h"))
 			{
-				FileItem DummyFile = FileItem.GetItemByFileReference(FileReference.Combine(OutputDir, "SingleFile", SourceFile.Name));
+				DirectoryReference.CreateDirectory(IntermediateDirectory);
+				FileItem DummyFile = FileItem.GetItemByFileReference(FileReference.Combine(IntermediateDirectory, Filename));
 				Directory.CreateDirectory(DummyFile.Directory.FullName);
 				File.WriteAllText(DummyFile.FullName, $"#include \"{SourceFile.FullName.Replace('\\', '/')}\"");
 				SourceFile = DummyFile;
@@ -64,33 +80,33 @@ namespace UnrealBuildTool
 				return null;
 			}
 
+			DirectoryReference.CreateDirectory(IntermediateDirectory);
 			VCCompileAction Action = new VCCompileAction(BaseAction);
 			Action.SourceFile = SourceFile;
 			if (bPreprocessOnly)
 			{
-				Action.PreprocessedFile = FileItem.GetItemByFileReference(FileReference.Combine(OutputDir, $"SingleFile{SingleFileCounter}.i"));
+				Action.PreprocessedFile = FileItem.GetItemByFileReference(FileReference.Combine(IntermediateDirectory, $"{Filename}.i"));
 			}
 			else
 			{
-				Action.ObjectFile = FileItem.GetItemByFileReference(FileReference.Combine(OutputDir, $"SingleFile{SingleFileCounter}.obj"));
+				Action.ObjectFile = FileItem.GetItemByFileReference(FileReference.Combine(IntermediateDirectory, $"{Filename}.obj"));
 				if (BaseAction.CompilerType.IsMSVC() && bWithAssembly)
 				{
-					Action.AssemblyFile = FileItem.GetItemByFileReference(FileReference.Combine(OutputDir, $"SingleFile{SingleFileCounter}.asm"));
+					Action.AssemblyFile = FileItem.GetItemByFileReference(FileReference.Combine(IntermediateDirectory, $"{Filename}.asm"));
 				}
 			}
 
 			if (BaseAction.CompilerType.IsMSVC())
 			{
-				Action.DependencyListFile = FileItem.GetItemByFileReference(FileReference.Combine(OutputDir, $"SingleFile{SingleFileCounter}.dep.json"));
+				Action.DependencyListFile = FileItem.GetItemByFileReference(FileReference.Combine(IntermediateDirectory, $"{Filename}.dep.json"));
 			}
 			else if (BaseAction.CompilerType.IsMSVC())
 			{
-				Action.DependencyListFile = FileItem.GetItemByFileReference(FileReference.Combine(OutputDir, $"SingleFile{SingleFileCounter}.d"));
+				Action.DependencyListFile = FileItem.GetItemByFileReference(FileReference.Combine(IntermediateDirectory, $"{Filename}.d"));
 			}
-			Action.ResponseFile = FileItem.GetItemByFileReference(FileReference.Combine(OutputDir, $"SingleFile{SingleFileCounter}.rsp"));
+			Action.ResponseFile = FileItem.GetItemByFileReference(FileReference.Combine(IntermediateDirectory, $"{Filename}.rsp"));
 			File.WriteAllLines(Action.ResponseFile.FullName, Action.GetCompilerArguments(Logger));
 
-			++SingleFileCounter;
 			return Action;
 		}
 	}
