Fix ClangToolchain omitting global compiler options from the command line on platforms that don't use a shared response file.

diff --git a/Engine/Source/Programs/UnrealBuildTool/Platform/Windows/VCToolChain.cs b/Engine/Source/Programs/UnrealBuildTool/Platform/Windows/VCToolChain.cs
--- a/Engine/Source/Programs/UnrealBuildTool/Platform/Windows/VCToolChain.cs
+++ b/Engine/Source/Programs/UnrealBuildTool/Platform/Windows/VCToolChain.cs
@@ -1394,7 +1394,7 @@ namespace UnrealBuildTool
 			BaseCompileAction.IncludePaths.AddRange(CompileEnvironment.UserIncludePaths);
 			BaseCompileAction.SystemIncludePaths.AddRange(CompileEnvironment.SystemIncludePaths);
 			
-			if (!CompileEnvironment.SharedSystemIncludePaths.Any())
+			if (!CompileEnvironment.bHasSharedResponseFile)
 			{
 				BaseCompileAction.SystemIncludePaths.AddRange(EnvVars.IncludePaths);
 			}
@@ -2107,6 +2107,7 @@ namespace UnrealBuildTool
 
 			NewCompileEnvironment.AdditionalPrerequisites.Add(FileItem);
 			NewCompileEnvironment.AdditionalResponseFiles.Add(FileItem);
+			NewCompileEnvironment.bHasSharedResponseFile = true;
 
 			return NewCompileEnvironment;
 		}
diff --git a/Engine/Source/Programs/UnrealBuildTool/System/CppCompileEnvironment.cs b/Engine/Source/Programs/UnrealBuildTool/System/CppCompileEnvironment.cs
--- a/Engine/Source/Programs/UnrealBuildTool/System/CppCompileEnvironment.cs
+++ b/Engine/Source/Programs/UnrealBuildTool/System/CppCompileEnvironment.cs
@@ -518,6 +518,11 @@ namespace UnrealBuildTool
 		/// </summary>
 		public List<FileItem> AdditionalResponseFiles = new();
 
+		/// <summary>
+		/// Whether the compile environment has a response file in AdditionalResponseFiles that contains global compiler arguments.
+		/// </summary>
+		public bool bHasSharedResponseFile = false;
+
 		/// <summary>
 		/// Additional arguments to pass to the compiler.
 		/// </summary>
diff --git a/Engine/Source/Programs/UnrealBuildTool/ToolChain/ClangToolChain.cs b/Engine/Source/Programs/UnrealBuildTool/ToolChain/ClangToolChain.cs
--- a/Engine/Source/Programs/UnrealBuildTool/ToolChain/ClangToolChain.cs
+++ b/Engine/Source/Programs/UnrealBuildTool/ToolChain/ClangToolChain.cs
@@ -1071,6 +1071,7 @@ namespace UnrealBuildTool
 
 			NewCompileEnvironment.AdditionalPrerequisites.Add(FileItem);
 			NewCompileEnvironment.AdditionalResponseFiles.Add(FileItem);
+			NewCompileEnvironment.bHasSharedResponseFile = true;
 
 
 			return NewCompileEnvironment;
@@ -1083,7 +1084,7 @@ namespace UnrealBuildTool
 				return;
 
 			List<string> GlobalArguments = new();
-			if (CompileEnvironment.SystemIncludePaths.Count != 0)
+			if (!CompileEnvironment.bHasSharedResponseFile)
 			{
 				GetCompileArguments_Global(CompileEnvironment, GlobalArguments);
 			}
@@ -1217,9 +1218,8 @@ namespace UnrealBuildTool
 		protected override CPPOutput CompileCPPFiles(CppCompileEnvironment CompileEnvironment, List<FileItem> InputFiles, DirectoryReference OutputDir, string ModuleName, IActionGraphBuilder Graph)
 		{
 			List<string> GlobalArguments = new();
-			
-			// Maybe not the nicest way of checking but if we don't have any system include paths it means we are not using a shared rsp file
-			if (CompileEnvironment.SystemIncludePaths.Count != 0)
+
+			if (!CompileEnvironment.bHasSharedResponseFile)
 			{
 				GetCompileArguments_Global(CompileEnvironment, GlobalArguments);
 			}
