Added support to give actions in UBT "weight". "Weight" is an average measurement of how many cores and memory an action uses. Using one core with a normal amount of memory for an action would be a value of 1.

Why do this:
Currently when compiling a cpp file with MSVC, it compiles across multiple cores while clang does not. This means that while we support limiting the number of cores(using ProcessorCountMultiplier), MSVC will use more cores than we specify. It also means that MSVC will always be faster when compiling because clang does not support compiling a cpp over multiple cores. To get similiar results when compiling with clang, we set the weight of MSVC to 1.5 and the weight of clang to 1.0. We then set the ProcessorCountMultiplier to 1.5. This results in MSVC and clang taking roughly the same amount of CPU utilization and clang compiles to be much faster.

	                     Old Timing(secs)	Old CPU Utilization	New Timing	New CPU Utilization(secs)
PlatformA AncientGame	590.94	51	                               431.47	73
MSVC AncientGameEditor	1016.96	94	                             1026.08	95
Clang AncientGameEditor	1543.72	63	                               1270.4	84
PlatformB AncientGame	494	52	                               396.95	74

Old = without weight path
New = with weight path

diff --git a/Engine/Source/Programs/UnrealBuildTool/Configuration/TargetRules.cs b/Engine/Source/Programs/UnrealBuildTool/Configuration/TargetRules.cs
--- a/Engine/Source/Programs/UnrealBuildTool/Configuration/TargetRules.cs
+++ b/Engine/Source/Programs/UnrealBuildTool/Configuration/TargetRules.cs
@@ -1998,6 +1998,18 @@ namespace UnrealBuildTool
 		[CommandLine("-ToolChain")]
 		public string? ToolChainName = null;
 
+		/// <summary>
+		/// The weight(cpu/memory utilization) of a MSVC compile action
+		/// </summary>
+		[XmlConfigFile(Category = "BuildConfiguration")]
+		public float MSVCCompileActionWeight = 1.5f;
+
+		/// <summary>
+		/// The weight(cpu/memory utilization) of a clang compile action
+		/// </summary>
+		[XmlConfigFile(Category = "BuildConfiguration")]
+		public float ClangCompileActionWeight = 1.0f;
+
 		/// <summary>
 		/// Whether to allow engine configuration to determine if we can load unverified certificates.
 		/// </summary>
@@ -3726,6 +3738,16 @@ namespace UnrealBuildTool
 			get { return Inner.ToolChainName; }
 		}
 
+		public float MSVCCompileActionWeight
+		{
+			get { return Inner.MSVCCompileActionWeight; }
+		}
+
+		public float ClangCompileActionWeight
+		{
+			get { return Inner.ClangCompileActionWeight; }
+		}
+
 		public bool bLegacyPublicIncludePaths
 		{
 			get { return Inner.bLegacyPublicIncludePaths; }
diff --git a/Engine/Source/Programs/UnrealBuildTool/Executors/ParallelExecutor.cs b/Engine/Source/Programs/UnrealBuildTool/Executors/ParallelExecutor.cs
--- a/Engine/Source/Programs/UnrealBuildTool/Executors/ParallelExecutor.cs
+++ b/Engine/Source/Programs/UnrealBuildTool/Executors/ParallelExecutor.cs
@@ -80,6 +80,18 @@ namespace UnrealBuildTool
 		[XmlConfigFile]
 		private static bool bPrintActionTargetNames = false;
 
+		/// <summary>
+		/// Whether to take into account the Action's weight when determining to do more work or not.
+		/// </summary>
+		[XmlConfigFile]
+		private static bool bUseActionWeights = false;
+
+		/// <summary>
+		/// Whether to show CPU utilization after the work is complete.
+		/// </summary>
+		[XmlConfigFile]
+		private static bool bShowCPUUtilization = false;
+
 		/// <summary>
 		/// Collapse non-error output lines
 		/// </summary>
@@ -204,8 +216,13 @@ namespace UnrealBuildTool
 
 			var ActiveTasks = new Task<ExecuteResults>[ActualNumParallelProcesses];
 			var ActiveActionIndices = new int[ActualNumParallelProcesses];
+			float MaxParallelProcessWeight = (float)ActualNumParallelProcesses;
+			float ActiveTaskWeight = 0.0f;
 			int ActiveTaskCount = 0;
 			bool Success = true;
+			List<float> LoggedCpuUtilization = new List<float>();
+			Stopwatch CpuQueryUtilization = new Stopwatch();
+			CpuQueryUtilization.Start();
 
 			int FirstQueuedAction = 0;
 
@@ -215,14 +232,15 @@ namespace UnrealBuildTool
 			while (!CancellationTokenSource.IsCancellationRequested)
 			{
 				// Always fill up to allowed number of processes
-				while (ActiveTaskCount < ActualNumParallelProcesses)
+				while (bUseActionWeights ? 
+					(ActiveTaskWeight < MaxParallelProcessWeight && ActiveTaskCount < ActualNumParallelProcesses)
+					: ActiveTaskCount < ActualNumParallelProcesses)
 				{
-
 					// Always traverse from first not started/finished action. We always want to respect the sorting.
 					// Note that there might be multiple actions we need to search passed to find the one we can run next
 					int ActionToRunIndex = -1;
 					bool FoundFirstQueued = false;
-					for (int i= FirstQueuedAction; i!= Actions.Length; ++i)
+					for (int i = FirstQueuedAction; i!= Actions.Length; ++i)
 					{
 						if (ActionsStatus[i] != ActionStatus.Queued)
 						{
@@ -260,6 +278,16 @@ namespace UnrealBuildTool
 						break;
 					}
 
+					if (bShowCPUUtilization && CpuQueryUtilization.Elapsed.Seconds > 1)
+					{
+						if (Utils.GetTotalCpuUtilization(out float CpuUtilization))
+						{
+							LoggedCpuUtilization.Add(CpuUtilization);
+						}
+						
+						CpuQueryUtilization.Restart();
+					}
+
 					// Didn't find an action we can run. Either there are none left or all depends on currently running actions
 					if (ActionToRunIndex == -1)
 					{
@@ -274,6 +302,7 @@ namespace UnrealBuildTool
 					ActiveTasks[ActiveTaskCount] = NewTask;
 					ActionsStatus[ActionToRunIndex] = ActionStatus.Running;
 					ActiveActionIndices[ActiveTaskCount] = ActionToRunIndex;
+					ActiveTaskWeight += ActionToRun.Weight;
 					++ActiveTaskCount;
 				}
 
@@ -317,6 +346,7 @@ namespace UnrealBuildTool
 				--ActiveTaskCount;
 				ActiveTasks[ActiveTaskIndex] = ActiveTasks[ActiveTaskCount];
 				ActiveActionIndices[ActiveTaskIndex] = ActiveActionIndices[ActiveTaskCount];
+				ActiveTaskWeight -= FinishedAction.Weight;
 
 				// Update Success status
 				Success = Success && ActionSuccess;
@@ -326,7 +356,7 @@ namespace UnrealBuildTool
 			LastLoggingTask.Wait();
 
 			// Summary
-			TraceSummary(ActionToTaskLookup, ProcessGroup, Logger);
+			TraceSummary(ActionToTaskLookup, ProcessGroup, Logger, LoggedCpuUtilization);
 
 			return Success;
 		}
@@ -468,7 +498,7 @@ namespace UnrealBuildTool
 					}
 					// END TEMPORARY
 
-					// prevent overrwriting of error text
+					// prevent overwriting of error text
 					s_previousLineLength = -1;
 
 					// Cancel all other pending tasks
@@ -480,8 +510,14 @@ namespace UnrealBuildTool
 			}
 		}
 
-		protected static void TraceSummary(Dictionary<LinkedAction, Task<ExecuteResults>> Tasks, ManagedProcessGroup ProcessGroup, ILogger Logger)
+		protected static void TraceSummary(Dictionary<LinkedAction, Task<ExecuteResults>> Tasks, ManagedProcessGroup ProcessGroup, ILogger Logger, IList<float>? LoggedCpuUtilization)
 		{
+			if (bShowCPUUtilization && LoggedCpuUtilization != null && LoggedCpuUtilization.Any())
+			{
+				Logger.LogInformation("");
+				Logger.LogInformation("Average CPU Utilization: {CPUPercentage}%",(int)(LoggedCpuUtilization.Average()));
+			}
+
 			if (!bShowCompilationTimes)
 			{
 				return;
diff --git a/Engine/Source/Programs/UnrealBuildTool/Platform/Windows/VCCompileAction.cs b/Engine/Source/Programs/UnrealBuildTool/Platform/Windows/VCCompileAction.cs
--- a/Engine/Source/Programs/UnrealBuildTool/Platform/Windows/VCCompileAction.cs
+++ b/Engine/Source/Programs/UnrealBuildTool/Platform/Windows/VCCompileAction.cs
@@ -168,6 +168,8 @@ namespace UnrealBuildTool
 		/// <inheritdoc/>
 		public bool bIsHighPriority { get => CreatePchFile != null; }
 
+		/// <inheritdoc/>
+		public float Weight { get; set; } = 1.0f;
 		#endregion
 
 		#region Implementation of IAction
@@ -328,6 +330,7 @@ namespace UnrealBuildTool
 			bCanExecuteRemotely = InAction.bCanExecuteRemotely;
 			bCanExecuteRemotelyWithSNDBS = InAction.bCanExecuteRemotelyWithSNDBS;
 			Architecture = InAction.Architecture;
+			Weight = InAction.Weight;
 
 			AdditionalPrerequisiteItems = new List<FileItem>(InAction.AdditionalPrerequisiteItems);
 			AdditionalProducedItems = new List<FileItem>(InAction.AdditionalProducedItems);
diff --git a/Engine/Source/Programs/UnrealBuildTool/Platform/Windows/VCToolChain.cs b/Engine/Source/Programs/UnrealBuildTool/Platform/Windows/VCToolChain.cs
--- a/Engine/Source/Programs/UnrealBuildTool/Platform/Windows/VCToolChain.cs
+++ b/Engine/Source/Programs/UnrealBuildTool/Platform/Windows/VCToolChain.cs
@@ -1438,6 +1438,16 @@ namespace UnrealBuildTool
 		{
 			VCCompileAction BaseCompileAction = CreateBaseCompileAction(CompileEnvironment);
 
+			// MSVC uses multiple processes when compiling CPP files, so the "weight" is more than 1
+			if (Target.WindowsPlatform.Compiler.IsMSVC())
+			{
+				BaseCompileAction.Weight = Target.MSVCCompileActionWeight;
+			}
+			else if (Target.WindowsPlatform.Compiler.IsClang())
+			{
+				BaseCompileAction.Weight = Target.ClangCompileActionWeight;
+			}
+
 			// Create a compile action for each source file.
 			List<VCCompileAction> Actions = new List<VCCompileAction>();
 			foreach (FileItem SourceFile in InputFiles)
diff --git a/Engine/Source/Programs/UnrealBuildTool/System/Action.cs b/Engine/Source/Programs/UnrealBuildTool/System/Action.cs
--- a/Engine/Source/Programs/UnrealBuildTool/System/Action.cs
+++ b/Engine/Source/Programs/UnrealBuildTool/System/Action.cs
@@ -134,6 +134,11 @@ namespace UnrealBuildTool
 		/// If set to true it will propagate high priority to all its prerequisite tasks
 		/// </summary>
 		bool bIsHighPriority { get; }
+
+		/// <summary>
+		/// Used to determine how much weight(CPU and Memory work) this action is.
+		/// </summary>
+		float Weight { get; }
 	}
 
 	/// <summary>
@@ -232,6 +237,9 @@ namespace UnrealBuildTool
 		/// <inheritdoc/>
 		public bool bIsHighPriority { get; set; } = false;
 
+		/// <inheritdoc/>
+		public float Weight { get; set; } = 1.0f;
+
 		IEnumerable<FileItem> IExternalAction.PrerequisiteItems => PrerequisiteItems;
 		IEnumerable<FileItem> IExternalAction.ProducedItems => ProducedItems;
 		IEnumerable<FileItem> IExternalAction.DeleteItems => DeleteItems;
@@ -596,6 +604,7 @@ namespace UnrealBuildTool
 		public bool bProducesImportLibrary => Inner.bProducesImportLibrary;
 		public bool bUseActionHistory => Inner.bUseActionHistory;
 		public bool bIsHighPriority => IsHighPriority != 0;
+		public float Weight => Inner.Weight;
 
 		#endregion
 
diff --git a/Engine/Source/Programs/UnrealBuildTool/System/Utils.cs b/Engine/Source/Programs/UnrealBuildTool/System/Utils.cs
--- a/Engine/Source/Programs/UnrealBuildTool/System/Utils.cs
+++ b/Engine/Source/Programs/UnrealBuildTool/System/Utils.cs
@@ -1263,6 +1263,64 @@ namespace UnrealBuildTool
 			return FreeMemoryBytes;
 		}
 
+		[DllImport("pdh.dll", SetLastError = true, CharSet = CharSet.Auto)]
+		static extern int PdhOpenQueryW([MarshalAs(UnmanagedType.LPWStr)] string? szDataSource, UIntPtr dwUserData, out IntPtr phQuery);
+
+		[DllImport("pdh.dll", SetLastError = true, CharSet = CharSet.Auto)]
+		static extern UInt32 PdhAddCounter(IntPtr hQuery, string szFullCounterPath, IntPtr dwUserData, out IntPtr phCounter);
+
+		[DllImport("pdh.dll", SetLastError = true)]
+		static extern UInt32 PdhCollectQueryData(IntPtr phQuery);
+
+		struct PDH_FMT_COUNTERVALUE
+		{
+			public uint CStatus;
+			public double doubleValue;
+		};
+
+		[DllImport("pdh.dll", SetLastError = true)]
+		static extern UInt32 PdhGetFormattedCounterValue(IntPtr phCounter, uint dwFormat, IntPtr lpdwType, out PDH_FMT_COUNTERVALUE pValue);
+
+		static IntPtr CpuQuery;
+		static IntPtr CpuTotal;
+		static bool CpuInitialized = false;
+
+		/// <summary>
+		/// Gives the current CPU utilization.
+		/// </summary>
+		/// <param name="Utilization">Percentage of CPU utilization currently.</param>
+		/// <returns>Whether or not it was successful in getting the CPU utilization.</returns>
+		public static bool GetTotalCpuUtilization(out float Utilization)
+		{
+			Utilization = 0.0f;
+			if (RuntimePlatform.IsWindows)
+			{
+				const UInt32 ERROR_SUCCESS = 0;
+				if (!CpuInitialized)
+				{
+					if (PdhOpenQueryW(null, UIntPtr.Zero, out CpuQuery) == ERROR_SUCCESS)
+					{
+						PdhAddCounter(CpuQuery, "\\Processor(_Total)\\% Processor Time", IntPtr.Zero, out CpuTotal);
+						PdhCollectQueryData(CpuQuery);
+						CpuInitialized = true;
+					}
+				}
+
+				if (CpuInitialized)
+				{
+					const uint PDH_FMT_DOUBLE = 0x00000200;
+
+					PDH_FMT_COUNTERVALUE counterVal;
+					PdhCollectQueryData(CpuQuery);
+					PdhGetFormattedCounterValue(CpuTotal, PDH_FMT_DOUBLE, IntPtr.Zero, out counterVal);
+					Utilization = (float)counterVal.doubleValue;
+					return true;
+				}
+			}
+
+			return false;
+		}
+
 		/// <summary>
 		/// Determines the maximum number of actions to execute in parallel, taking into account the resources available on this machine.
 		/// </summary>
diff --git a/Engine/Source/Programs/UnrealBuildTool/ToolChain/ClangToolChain.cs b/Engine/Source/Programs/UnrealBuildTool/ToolChain/ClangToolChain.cs
--- a/Engine/Source/Programs/UnrealBuildTool/ToolChain/ClangToolChain.cs
+++ b/Engine/Source/Programs/UnrealBuildTool/ToolChain/ClangToolChain.cs
@@ -182,6 +182,7 @@ namespace UnrealBuildTool
 		protected StaticAnalyzer StaticAnalyzer = StaticAnalyzer.None;
 		protected StaticAnalyzerMode StaticAnalyzerMode = StaticAnalyzerMode.Deep;
 		protected StaticAnalyzerOutputType StaticAnalyzerOutputType = StaticAnalyzerOutputType.Text;
+		protected float CompileActionWeight = 1.0f;
 
 		static ClangToolChain()
 		{
@@ -208,6 +209,7 @@ namespace UnrealBuildTool
 			StaticAnalyzer = Target.StaticAnalyzer;
 			StaticAnalyzerMode = Target.StaticAnalyzerMode;
 			StaticAnalyzerOutputType = Target.StaticAnalyzerOutputType;
+			CompileActionWeight = Target.ClangCompileActionWeight;
 		}
 
 		public override void FinalizeOutput(ReadOnlyTargetRules Target, TargetMakefileBuilder MakefileBuilder)
@@ -1102,6 +1104,7 @@ namespace UnrealBuildTool
 		protected virtual Action CompileCPPFile(CppCompileEnvironment CompileEnvironment, FileItem SourceFile, DirectoryReference OutputDir, string ModuleName, IActionGraphBuilder Graph, IReadOnlyCollection<string> GlobalArguments, CPPOutput Result)
 		{
 			Action CompileAction = Graph.CreateAction(ActionType.Compile);
+			CompileAction.Weight = CompileActionWeight;
 
 			// copy the global arguments into the file arguments, so GetCompileArguments_FileType can remove entries if needed (special case but can be important)
 			List<string> FileArguments = new(GlobalArguments);
