Provide option to disable dump_syms where iteration speed is desirable. (-NoDumpSyms or set bDisableDumpSyms)

diff --git a/Engine/Source/Programs/UnrealBuildTool/Platform/Linux/LinuxToolChain.cs b/Engine/Source/Programs/UnrealBuildTool/Platform/Linux/LinuxToolChain.cs
--- a/Engine/Source/Programs/UnrealBuildTool/Platform/Linux/LinuxToolChain.cs
+++ b/Engine/Source/Programs/UnrealBuildTool/Platform/Linux/LinuxToolChain.cs
@@ -161,19 +161,29 @@ namespace UnrealBuildTool
 			StringWriter Out = new StringWriter();
 			Out.NewLine = bUseCmdExe ? "\r\n" : "\n";
 
-			// dump_syms
-			Out.WriteLine(DumpCommand,
-				LinuxInfo.DumpSyms,
-				OutputFile.AbsolutePath,
-				SymbolsFile.AbsolutePath
-			);
-
-			// encode breakpad symbols
-			Out.WriteLine("\"{0}\" \"{1}\" \"{2}\"",
-				LinuxInfo.BreakpadEncoder,
-				SymbolsFile.AbsolutePath,
-				EncodedBinarySymbolsFile.AbsolutePath
-			);
+
+			if (!Options.HasFlag(ClangToolChainOptions.DisableDumpSyms) || Options.HasFlag(ClangToolChainOptions.PreservePSYM))
+			{
+				// dump_syms
+				Out.WriteLine(DumpCommand,
+					LinuxInfo.DumpSyms,
+					OutputFile.AbsolutePath,
+					SymbolsFile.AbsolutePath
+				);
+
+				// encode breakpad symbols
+				Out.WriteLine("\"{0}\" \"{1}\" \"{2}\"",
+					LinuxInfo.BreakpadEncoder,
+					SymbolsFile.AbsolutePath,
+					EncodedBinarySymbolsFile.AbsolutePath
+				);
+			}
+			else
+			{
+				// we have to create dummy files to prevent packaging errors
+				Out.WriteLine("echo DummyPSym >> \"{0}\"", SymbolsFile.AbsolutePath);
+				Out.WriteLine("echo DummySyms>> \"{0}\"", EncodedBinarySymbolsFile.AbsolutePath);
+			}
 
 			if (!Options.HasFlag(ClangToolChainOptions.DisableSplitDebugInfoWithObjCopy) && LinkEnvironment.bCreateDebugInfo)
 			{
diff --git a/Engine/Source/Programs/UnrealBuildTool/Platform/Linux/UEBuildLinux.cs b/Engine/Source/Programs/UnrealBuildTool/Platform/Linux/UEBuildLinux.cs
--- a/Engine/Source/Programs/UnrealBuildTool/Platform/Linux/UEBuildLinux.cs
+++ b/Engine/Source/Programs/UnrealBuildTool/Platform/Linux/UEBuildLinux.cs
@@ -88,6 +88,13 @@ namespace UnrealBuildTool
 		[XmlConfigFile(Category = "BuildConfiguration", Name = "bTuneDebugInfoForLLDB")]
 		public bool bTuneDebugInfoForLLDB = false;
 
+		/// <summary>
+		/// Whether to globally disable calling dump_syms
+		/// </summary>
+		[CommandLine("-NoDumpSyms")]
+		[XmlConfigFile(Category = "BuildConfiguration", Name = "bDisableDumpSyms")]
+		public bool bDisableDumpSyms = false;
+
 		/// <summary>
 		/// Enables runtime ray tracing support.
 		/// </summary>
@@ -147,6 +154,11 @@ namespace UnrealBuildTool
 			get { return Inner.bTuneDebugInfoForLLDB; }
 		}
 
+		public bool bDisableDumpSyms
+		{
+			get { return Inner.bDisableDumpSyms; }
+		}
+
 		public bool bEnableRayTracing
 		{
 			get { return Inner.bEnableRayTracing; }
@@ -584,6 +596,10 @@ namespace UnrealBuildTool
 					throw new BuildException("Memory Sanitizer (MSan) unsupported for non-monolithic builds");
 				}
 			}
+			if (Target.LinuxPlatform.bDisableDumpSyms)
+			{
+				Options |= ClangToolChainOptions.DisableDumpSyms;
+			}
 			if (Target.bAllowLTCG && Target.bPreferThinLTO)
 			{
 				Options |= ClangToolChainOptions.EnableThinLTO;
diff --git a/Engine/Source/Programs/UnrealBuildTool/ToolChain/ClangToolChain.cs b/Engine/Source/Programs/UnrealBuildTool/ToolChain/ClangToolChain.cs
--- a/Engine/Source/Programs/UnrealBuildTool/ToolChain/ClangToolChain.cs
+++ b/Engine/Source/Programs/UnrealBuildTool/ToolChain/ClangToolChain.cs
@@ -103,6 +103,11 @@ namespace UnrealBuildTool
 		/// Indicates that the target is a moduler build i.e. Target.LinkType == TargetLinkType.Modular
 		/// </summary>
 		ModularBuild = 1 << 15,
+
+		/// <summary>
+		/// Disable Dump Syms step for faster iteration
+		/// </summary>
+		DisableDumpSyms = 1 << 16,
 	}
 
 	abstract class ClangToolChain : ISPCToolChain
