UnrealBuildTool: Fix issue generating visual studio solution because of incorrect check if a PCH is in use which was causing modules to use an incorrect SharedPCH when a Private PCH should have been used

diff --git a/Engine/Source/Programs/UnrealBuildTool/System/CppCompileEnvironment.cs b/Engine/Source/Programs/UnrealBuildTool/System/CppCompileEnvironment.cs
--- a/Engine/Source/Programs/UnrealBuildTool/System/CppCompileEnvironment.cs
+++ b/Engine/Source/Programs/UnrealBuildTool/System/CppCompileEnvironment.cs
@@ -573,7 +573,7 @@ namespace UnrealBuildTool
 		/// <summary>
 		/// True if a single PRecompiledHeader exists, or at least one PerArchPrecompiledHeaderFile exists
 		/// </summary>
-		public bool bHasPrecompiledHeader => (PerArchPrecompiledHeaderFiles != null && PerArchPrecompiledHeaderFiles.Count > 0) || PrecompiledHeaderFile != null;
+		public bool bHasPrecompiledHeader => PrecompiledHeaderAction == PrecompiledHeaderAction.Include;
 
 		/// <summary>
 		/// Whether or not UHT is being built
