UnrealBuildTool: Sort Visual Studio installs so 2022 is  always prioritized over 2019, even if a community or preview version

diff --git a/Engine/Source/Programs/UnrealBuildTool/Platform/Windows/MicrosoftPlatformSDK.cs b/Engine/Source/Programs/UnrealBuildTool/Platform/Windows/MicrosoftPlatformSDK.cs
--- a/Engine/Source/Programs/UnrealBuildTool/Platform/Windows/MicrosoftPlatformSDK.cs
+++ b/Engine/Source/Programs/UnrealBuildTool/Platform/Windows/MicrosoftPlatformSDK.cs
@@ -828,7 +828,11 @@ namespace UnrealBuildTool
 						Logger.LogDebug("Found Visual Studio installation: {BaseDir} (Product={ProductId}, Version={Version})", BaseDir, ProductId, Version);
 					}
 
-					Installations = Installations.OrderBy(x => x.bCommunity).ThenBy(x => x.bPreview).ThenByDescending(x => x.Version).ToList();
+					Installations = Installations.OrderByDescending(x => x.Compiler)
+						.ThenBy(x => x.bCommunity)
+						.ThenBy(x => x.bPreview)
+						.ThenByDescending(x => x.Version)
+						.ToList();
 				}
 				catch (Exception Ex)
 				{
