Update UBT preprocessor code to much stricter coding standards

#rnx

diff --git a/Engine/Source/Programs/UnrealBuildTool.Tests/PreprocessorStateTests.cs b/Engine/Source/Programs/UnrealBuildTool.Tests/PreprocessorStateTests.cs
--- a/Engine/Source/Programs/UnrealBuildTool.Tests/PreprocessorStateTests.cs
+++ b/Engine/Source/Programs/UnrealBuildTool.Tests/PreprocessorStateTests.cs
@@ -87,9 +87,9 @@ namespace UnrealBuildToolTests
 		static string FormatTransform(PreprocessorTransform Transform)
 		{
 			StringBuilder Result = new StringBuilder();
-			if (Transform.bRequireTopmostActive.HasValue)
+			if (Transform.RequireTopmostActive.HasValue)
 			{
-				Result.AppendFormat("=b {0}\n", Transform.bRequireTopmostActive.Value ? "Active" : "0");
+				Result.AppendFormat("=b {0}\n", Transform.RequireTopmostActive.Value ? "Active" : "0");
 			}
 			foreach (PreprocessorBranch Branch in Transform.RequiredBranches)
 			{
diff --git a/Engine/Source/Programs/UnrealBuildTool.Tests/SourceFileTests.cs b/Engine/Source/Programs/UnrealBuildTool.Tests/SourceFileTests.cs
--- a/Engine/Source/Programs/UnrealBuildTool.Tests/SourceFileTests.cs
+++ b/Engine/Source/Programs/UnrealBuildTool.Tests/SourceFileTests.cs
@@ -46,8 +46,6 @@ namespace UnrealBuildToolTests
 				}
 			}
 
-			Log.TraceLog("Read {0} source files", SourceFiles.Count);
-
 			FileReference TempDataFile = FileReference.Combine(Unreal.EngineDirectory, "Intermediate", "Temp", "SourceFileTests.bin");
 			DirectoryReference.CreateDirectory(TempDataFile.Directory);
 
diff --git a/Engine/Source/Programs/UnrealBuildTool/Preprocessor/.editorconfig b/Engine/Source/Programs/UnrealBuildTool/Preprocessor/.editorconfig
new file mode 100644
--- /dev/null
+++ b/Engine/Source/Programs/UnrealBuildTool/Preprocessor/.editorconfig
@@ -0,0 +1,257 @@
+# NOTE: Requires **VS2019 16.3** or later
+
+# EditorConfig is awesome:http://EditorConfig.org
+
+# Dotnet code style settings:
+[*.cs]
+
+# Sort using and Import directives with System.* appearing first
+dotnet_sort_system_directives_first = true
+dotnet_separate_import_directive_groups = false
+
+# Avoid "this." and "Me." if not necessary
+dotnet_style_qualification_for_field = false:error
+dotnet_style_qualification_for_property = false:error
+dotnet_style_qualification_for_method = false:error
+dotnet_style_qualification_for_event = false:error
+
+# Use language keywords instead of framework type names for type references
+dotnet_style_predefined_type_for_locals_parameters_members = true:error
+dotnet_style_predefined_type_for_member_access = false:error
+
+# Suggest more modern language features when available
+dotnet_style_object_initializer = false:suggestion
+dotnet_style_collection_initializer = false:suggestion
+dotnet_style_coalesce_expression = true:suggestion
+dotnet_style_null_propagation = true:suggestion
+dotnet_style_explicit_tuple_names = true:suggestion
+
+# Whitespace options
+dotnet_style_allow_multiple_blank_lines_experimental = false
+dotnet_style_allow_statement_immediately_after_block_experimental = true
+
+# Non-private static fields are PascalCase
+dotnet_naming_rule.non_private_static_fields_should_be_pascal_case.severity = suggestion
+dotnet_naming_rule.non_private_static_fields_should_be_pascal_case.symbols = non_private_static_fields
+dotnet_naming_rule.non_private_static_fields_should_be_pascal_case.style = non_private_static_field_style
+
+dotnet_naming_symbols.non_private_static_fields.applicable_kinds = field
+dotnet_naming_symbols.non_private_static_fields.applicable_accessibilities = public, protected, internal, protected_internal, private_protected
+dotnet_naming_symbols.non_private_static_fields.required_modifiers = static
+
+dotnet_naming_style.non_private_static_field_style.capitalization = pascal_case
+
+# Non-private readonly fields are PascalCase
+dotnet_naming_rule.non_private_readonly_fields_should_be_pascal_case.severity = suggestion
+dotnet_naming_rule.non_private_readonly_fields_should_be_pascal_case.symbols = non_private_readonly_fields
+dotnet_naming_rule.non_private_readonly_fields_should_be_pascal_case.style = non_private_readonly_field_style
+
+dotnet_naming_symbols.non_private_readonly_fields.applicable_kinds = field
+dotnet_naming_symbols.non_private_readonly_fields.applicable_accessibilities = public, protected, internal, protected_internal, private_protected
+dotnet_naming_symbols.non_private_readonly_fields.required_modifiers = readonly
+
+dotnet_naming_style.non_private_readonly_field_style.capitalization = pascal_case
+
+# Constants are PascalCase
+dotnet_naming_rule.constants_should_be_pascal_case.severity = suggestion
+dotnet_naming_rule.constants_should_be_pascal_case.symbols = constants
+dotnet_naming_rule.constants_should_be_pascal_case.style = constant_style
+
+dotnet_naming_symbols.constants.applicable_kinds = field, local
+dotnet_naming_symbols.constants.required_modifiers = const
+
+dotnet_naming_style.constant_style.capitalization = pascal_case
+
+# Static fields are camelCase and start with s_
+dotnet_naming_rule.static_fields_should_be_camel_case.severity = suggestion
+dotnet_naming_rule.static_fields_should_be_camel_case.symbols = static_fields
+dotnet_naming_rule.static_fields_should_be_camel_case.style = static_field_style
+
+dotnet_naming_symbols.static_fields.applicable_kinds = field
+dotnet_naming_symbols.static_fields.required_modifiers = static
+
+dotnet_naming_style.static_field_style.capitalization = camel_case
+dotnet_naming_style.static_field_style.required_prefix = s_
+
+# Instance fields are camelCase and start with _
+dotnet_naming_rule.instance_fields_should_be_camel_case.severity = suggestion
+dotnet_naming_rule.instance_fields_should_be_camel_case.symbols = instance_fields
+dotnet_naming_rule.instance_fields_should_be_camel_case.style = instance_field_style
+
+dotnet_naming_symbols.instance_fields.applicable_kinds = field
+
+dotnet_naming_style.instance_field_style.capitalization = camel_case
+dotnet_naming_style.instance_field_style.required_prefix = _
+
+# Locals and parameters are camelCase
+dotnet_naming_rule.locals_should_be_camel_case.severity = suggestion
+dotnet_naming_rule.locals_should_be_camel_case.symbols = locals_and_parameters
+dotnet_naming_rule.locals_should_be_camel_case.style = camel_case_style
+
+dotnet_naming_symbols.locals_and_parameters.applicable_kinds = parameter, local
+
+dotnet_naming_style.camel_case_style.capitalization = camel_case
+
+# Local functions are PascalCase
+dotnet_naming_rule.local_functions_should_be_pascal_case.severity = suggestion
+dotnet_naming_rule.local_functions_should_be_pascal_case.symbols = local_functions
+dotnet_naming_rule.local_functions_should_be_pascal_case.style = local_function_style
+
+dotnet_naming_symbols.local_functions.applicable_kinds = local_function
+
+dotnet_naming_style.local_function_style.capitalization = pascal_case
+
+# By default, name items with PascalCase
+dotnet_naming_rule.members_should_be_pascal_case.severity = suggestion
+dotnet_naming_rule.members_should_be_pascal_case.symbols = all_members
+dotnet_naming_rule.members_should_be_pascal_case.style = pascal_case_style
+
+dotnet_naming_symbols.all_members.applicable_kinds = *
+
+dotnet_naming_style.pascal_case_style.capitalization = pascal_case
+
+# IDE0005: Remove unnecessary import
+dotnet_diagnostic.IDE0005.severity = none
+
+# IDE0073: File header
+dotnet_diagnostic.IDE0073.severity = warning
+file_header_template = Copyright Epic Games, Inc. All Rights Reserved.
+
+# IDE0035: Remove unreachable code
+dotnet_diagnostic.IDE0035.severity = warning
+
+# IDE0036: Order modifiers
+dotnet_diagnostic.IDE0036.severity = warning
+
+# IDE0043: Format string contains invalid placeholder
+dotnet_diagnostic.IDE0043.severity = warning
+
+# IDE0044: Make field readonly
+dotnet_diagnostic.IDE0044.severity = warning
+
+# IDE0011: Add braces
+csharp_prefer_braces = when_multiline:warning
+# NOTE: We need the below severity entry for Add Braces due to [REDACTED_PROFILE_URL]/roslyn/issues/44201
+dotnet_diagnostic.IDE0011.severity = warning
+
+# IDE0040: Add accessibility modifiers
+dotnet_diagnostic.IDE0040.severity = warning
+
+# CONSIDER: Are IDE0051 and IDE0052 too noisy to be warnings for IDE editing scenarios? Should they be made build-only warnings?
+# IDE0051: Remove unused private member
+dotnet_diagnostic.IDE0051.severity = warning
+
+# IDE0052: Remove unread private member
+dotnet_diagnostic.IDE0052.severity = warning
+
+# IDE0059: Unnecessary assignment to a value
+dotnet_diagnostic.IDE0059.severity = warning
+
+# IDE0060: Remove unused parameter
+dotnet_diagnostic.IDE0060.severity = warning
+
+# CA1012: Abstract types should not have public constructors
+dotnet_diagnostic.CA1012.severity = warning
+
+# CA1822: Make member static
+dotnet_diagnostic.CA1822.severity = warning
+
+# IDE0005: Using directive is unnecessary
+dotnet_diagnostic.IDE0005.severity = warning
+
+# dotnet_style_allow_multiple_blank_lines_experimental
+dotnet_diagnostic.IDE2000.severity = warning
+
+# csharp_style_allow_embedded_statements_on_same_line_experimental
+dotnet_diagnostic.IDE2001.severity = warning
+
+# csharp_style_allow_blank_lines_between_consecutive_braces_experimental
+dotnet_diagnostic.IDE2002.severity = warning
+
+# dotnet_style_allow_statement_immediately_after_block_experimental
+dotnet_diagnostic.IDE2003.severity = warning
+
+# csharp_style_allow_blank_line_after_colon_in_constructor_initializer_experimental
+dotnet_diagnostic.IDE2004.severity = warning
+
+dotnet_style_require_accessibility_modifiers = false
+
+# CSharp code style settings:
+[*.cs]
+csharp_prefer_simple_using_statement = true
+csharp_style_prefer_range_operator = true
+csharp_style_implicit_object_creation_when_type_is_apparent = true
+
+# Newline settings
+csharp_new_line_before_open_brace = all
+csharp_new_line_before_else = true
+csharp_new_line_before_catch = true
+csharp_new_line_before_finally = true
+csharp_new_line_before_members_in_object_initializers = true
+csharp_new_line_before_members_in_anonymous_types = true
+csharp_new_line_between_query_expression_clauses = true
+
+# Indentation preferences
+csharp_indent_block_contents = true
+csharp_indent_braces = false
+csharp_indent_case_contents = true
+csharp_indent_case_contents_when_block = true
+csharp_indent_switch_labels = true
+csharp_indent_labels = flush_left
+
+# Whitespace options
+csharp_style_allow_embedded_statements_on_same_line_experimental = false
+csharp_style_allow_blank_lines_between_consecutive_braces_experimental = false
+csharp_style_allow_blank_line_after_colon_in_constructor_initializer_experimental = false
+
+# Prefer "var" everywhere
+dotnet_diagnostic.IDE0007.severity = error
+csharp_style_var_for_built_in_types = false:error
+csharp_style_var_when_type_is_apparent = false:error
+csharp_style_var_elsewhere = false:error
+
+# Prefer method-like constructs to have a block body
+csharp_style_expression_bodied_methods = false:none
+csharp_style_expression_bodied_constructors = false:none
+csharp_style_expression_bodied_operators = false:none
+
+# Prefer property-like constructs to have an expression-body
+csharp_style_expression_bodied_properties = true:error
+csharp_style_expression_bodied_indexers = true:error
+csharp_style_expression_bodied_accessors = true:error
+
+# Suggest more modern language features when available
+csharp_style_pattern_matching_over_is_with_cast_check = true:error
+csharp_style_pattern_matching_over_as_with_null_check = false:error
+csharp_style_inlined_variable_declaration = true:suggestion
+csharp_style_throw_expression = true:error
+csharp_style_conditional_delegate_call = true:suggestion
+
+# Spacing
+csharp_space_after_cast = false
+csharp_space_after_colon_in_inheritance_clause = true
+csharp_space_after_comma = true
+csharp_space_after_dot = false
+csharp_space_after_keywords_in_control_flow_statements = true
+csharp_space_after_semicolon_in_for_statement = true
+csharp_space_around_binary_operators = before_and_after
+csharp_space_around_declaration_statements = do_not_ignore
+csharp_space_before_colon_in_inheritance_clause = true
+csharp_space_before_comma = false
+csharp_space_before_dot = false
+csharp_space_before_open_square_brackets = false
+csharp_space_before_semicolon_in_for_statement = false
+csharp_space_between_empty_square_brackets = false
+csharp_space_between_method_call_empty_parameter_list_parentheses = false
+csharp_space_between_method_call_name_and_opening_parenthesis = false
+csharp_space_between_method_call_parameter_list_parentheses = false
+csharp_space_between_method_declaration_empty_parameter_list_parentheses = false
+csharp_space_between_method_declaration_name_and_open_parenthesis = false
+csharp_space_between_method_declaration_parameter_list_parentheses = false
+csharp_space_between_parentheses = false
+csharp_space_between_square_brackets = false
+
+# Blocks are allowed
+csharp_prefer_braces = true:silent
+csharp_preserve_single_line_blocks = true
+csharp_preserve_single_line_statements = true
diff --git a/Engine/Source/Programs/UnrealBuildTool/Preprocessor/Identifier.cs b/Engine/Source/Programs/UnrealBuildTool/Preprocessor/Identifier.cs
--- a/Engine/Source/Programs/UnrealBuildTool/Preprocessor/Identifier.cs
+++ b/Engine/Source/Programs/UnrealBuildTool/Preprocessor/Identifier.cs
@@ -1,12 +1,8 @@
 // Copyright Epic Games, Inc. All Rights Reserved.
 
+using EpicGames.Core;
 using System;
 using System.Collections.Concurrent;
-using System.Collections.Generic;
-using System.Linq;
-using System.Text;
-using System.Threading.Tasks;
-using EpicGames.Core;
 
 namespace UnrealBuildTool
 {
@@ -18,68 +14,68 @@ namespace UnrealBuildTool
 		/// <summary>
 		/// The name of this identifier
 		/// </summary>
-		string Name;
+		readonly string _name;
 
 		/// <summary>
 		/// Global map of name to identifier instance
 		/// </summary>
-		static ConcurrentDictionary<string, Identifier> NameToIdentifier = new ConcurrentDictionary<string, Identifier>(StringComparer.Ordinal);
+		static readonly ConcurrentDictionary<string, Identifier> s_nameToIdentifier = new(StringComparer.Ordinal);
 
 		/// <summary>
 		/// Constructor
 		/// </summary>
-		/// <param name="Name">Name of this identifier</param>
-		private Identifier(string Name)
+		/// <param name="name">Name of this identifier</param>
+		private Identifier(string name)
 		{
-			this.Name = Name;
+			_name = name;
 		}
 
 		/// <summary>
-		/// Finds or adds an identifer with the given name
+		/// Finds or adds an identifier with the given name
 		/// </summary>
-		/// <param name="Name">Name of the identifier</param>
+		/// <param name="name">Name of the identifier</param>
 		/// <returns>New Identifier instance</returns>
-		public static Identifier FindOrAdd(string Name)
+		public static Identifier FindOrAdd(string name)
 		{
-			Identifier? Result;
-			if(!NameToIdentifier.TryGetValue(Name, out Result))
+			if (!s_nameToIdentifier.TryGetValue(name, out Identifier? result))
 			{
-				Identifier NewIdentifier = new Identifier(Name);
-				if(NameToIdentifier.TryAdd(Name, NewIdentifier))
+				Identifier newIdentifier = new(name);
+				if (s_nameToIdentifier.TryAdd(name, newIdentifier))
 				{
-					Result = NewIdentifier;
+					result = newIdentifier;
 				}
 				else
 				{
-					Result = NameToIdentifier[Name];
+					result = s_nameToIdentifier[name];
 				}
 			}
-			return Result;
+			return result;
 		}
 
 		/// <summary>
 		/// Compares this identifier to another identifier
 		/// </summary>
-		/// <param name="Other">Identifier to compare to</param>
+		/// <param name="other">Identifier to compare to</param>
 		/// <returns>Value indicating which identifier should sort first</returns>
-		public int CompareTo(Identifier? Other)
+		public int CompareTo(Identifier? other)
 		{
-			return ReferenceEquals(Other, null)? 1 : Name.CompareTo(Other.Name);
+			return other is null ? 1 : _name.CompareTo(other._name);
 		}
 
 		/// <summary>
-		/// Formats this identifer as a string for debugging
+		/// Formats this identifier as a string for debugging
 		/// </summary>
 		/// <returns>Name of this identifier</returns>
 		public override string ToString()
 		{
-			return Name;
+			return _name;
 		}
 	}
 
 	/// <summary>
 	/// Well known predefined identifiers
 	/// </summary>
+	[System.Diagnostics.CodeAnalysis.SuppressMessage("Style", "IDE1006:Naming Styles", Justification = "Keep the names closer to definitions")]
 	class Identifiers
 	{
 		public static readonly Identifier Include = Identifier.FindOrAdd("include");
@@ -129,21 +125,21 @@ namespace UnrealBuildTool
 		/// <summary>
 		/// Read an identifier from a binary archive
 		/// </summary>
-		/// <param name="Reader">Reader to serialize data from</param>
+		/// <param name="reader">Reader to serialize data from</param>
 		/// <returns>Instance of the serialized identifier</returns>
-		public static Identifier ReadIdentifier(this BinaryArchiveReader Reader)
+		public static Identifier ReadIdentifier(this BinaryArchiveReader reader)
 		{
-			return Reader.ReadObjectReference<Identifier>(() => Identifier.FindOrAdd(Reader.ReadString()!))!;
+			return reader.ReadObjectReference<Identifier>(() => Identifier.FindOrAdd(reader.ReadString()!))!;
 		}
 
 		/// <summary>
 		/// Write an identifier to a binary archive
 		/// </summary>
-		/// <param name="Writer">Writer to serialize data to</param>
-		/// <param name="Identifier">Identifier to write</param>
-		public static void WriteIdentifier(this BinaryArchiveWriter Writer, Identifier? Identifier)
+		/// <param name="writer">Writer to serialize data to</param>
+		/// <param name="identifier">Identifier to write</param>
+		public static void WriteIdentifier(this BinaryArchiveWriter writer, Identifier? identifier)
 		{
-			Writer.WriteObjectReference<Identifier?>(Identifier, () => Writer.WriteString(Identifier!.ToString()));
+			writer.WriteObjectReference<Identifier?>(identifier, () => writer.WriteString(identifier!.ToString()));
 		}
 	}
 }
diff --git a/Engine/Source/Programs/UnrealBuildTool/Preprocessor/Preprocessor.cs b/Engine/Source/Programs/UnrealBuildTool/Preprocessor/Preprocessor.cs
--- a/Engine/Source/Programs/UnrealBuildTool/Preprocessor/Preprocessor.cs
+++ b/Engine/Source/Programs/UnrealBuildTool/Preprocessor/Preprocessor.cs
@@ -1,18 +1,15 @@
 // Copyright Epic Games, Inc. All Rights Reserved.
 
+using EpicGames.Core;
+using Microsoft.Extensions.Logging;
 using System;
 using System.Collections.Generic;
-using System.Linq;
-using System.Text;
-using System.Threading.Tasks;
 using System.Diagnostics;
-using System.IO;
-using System.Collections;
-using System.Collections.Concurrent;
-using EpicGames.Core;
 using System.Diagnostics.CodeAnalysis;
+using System.IO;
+using System.Linq;
+using System.Text;
 using UnrealBuildBase;
-using Microsoft.Extensions.Logging;
 
 namespace UnrealBuildTool
 {
@@ -24,18 +21,18 @@ namespace UnrealBuildTool
 		/// <summary>
 		/// The context when the error was encountered
 		/// </summary>
-		PreprocessorContext? Context;
+		public readonly PreprocessorContext? Context;
 
 		/// <summary>
 		/// Constructor
 		/// </summary>
-		/// <param name="Context">The current preprocesor context</param>
-		/// <param name="Format">Format string, to be passed to String.Format</param>
-		/// <param name="Args">Optional argument list for the format string</param>
-		public PreprocessorException(PreprocessorContext? Context, string Format, params object[] Args)
-			: base(String.Format(Format, Args))
+		/// <param name="context">The current preprocesor context</param>
+		/// <param name="format">Format string, to be passed to String.Format</param>
+		/// <param name="args">Optional argument list for the format string</param>
+		public PreprocessorException(PreprocessorContext? context, string format, params object[] args)
+			: base(String.Format(format, args))
 		{
-			this.Context = Context;
+			Context = context;
 		}
 	}
 
@@ -63,42 +60,42 @@ namespace UnrealBuildTool
 		/// <summary>
 		/// Include paths to look in
 		/// </summary>
-		List<DirectoryItem> IncludeDirectories = new List<DirectoryItem>();
+		readonly List<DirectoryItem> _includeDirectories = new();
 
 		/// <summary>
-        /// Framework paths to look in
-        /// </summary>
-		List<DirectoryItem> FrameworkDirectories = new List<DirectoryItem>();
+		/// Framework paths to look in
+		/// </summary>
+		readonly List<DirectoryItem> _frameworkDirectories = new();
 
 		/// <summary>
 		/// Set of all included files with the #pragma once directive
 		/// </summary>
-		HashSet<FileItem> PragmaOnceFiles = new HashSet<FileItem>();
+		readonly HashSet<FileItem> _pragmaOnceFiles = new();
 
 		/// <summary>
 		/// Set of any files that has been processed
 		/// </summary>
-		HashSet<FileItem> ProcessedFiles = new HashSet<FileItem>();
+		readonly HashSet<FileItem> _processedFiles = new();
 
 		/// <summary>
 		/// The current state of the preprocessor
 		/// </summary>
-		PreprocessorState State = new PreprocessorState();
+		readonly PreprocessorState _state = new();
 
 		/// <summary>
 		/// Predefined token containing the constant "0"
 		/// </summary>
-		static readonly byte[] ZeroLiteral = Encoding.UTF8.GetBytes("0");
+		static readonly byte[] s_zeroLiteral = Encoding.UTF8.GetBytes("0");
 
 		/// <summary>
 		/// Predefined token containing the constant "1"
 		/// </summary>
-		static readonly byte[] OneLiteral = Encoding.UTF8.GetBytes("1");
+		static readonly byte[] s_oneLiteral = Encoding.UTF8.GetBytes("1");
 
 		/// <summary>
 		/// Value of the __COUNTER__ variable
 		/// </summary>
-		int Counter;
+		int _counter;
 
 		/// <summary>
 		/// List of files included by the preprocessor
@@ -106,7 +103,7 @@ namespace UnrealBuildTool
 		/// <returns>Enumerable of processed files</returns>
 		public IEnumerable<FileItem> GetProcessedFiles()
 		{
-			return ProcessedFiles.AsEnumerable();
+			return _processedFiles.AsEnumerable();
 		}
 
 		/// <summary>
@@ -114,9 +111,9 @@ namespace UnrealBuildTool
 		/// </summary>
 		public Preprocessor()
 		{
-			DateTime Now = DateTime.Now;
-			AddLiteralMacro("__DATE__", TokenType.String, String.Format("\"{0} {1,2} {2}\"", Now.ToString("MMM"), Now.Day, Now.Year));
-			AddLiteralMacro("__TIME__", TokenType.String, "\"" + Now.ToString("HH:mm:ss") + "\"");
+			DateTime now = DateTime.Now;
+			AddLiteralMacro("__DATE__", TokenType.String, String.Format("\"{0} {1,2} {2}\"", now.ToString("MMM"), now.Day, now.Year));
+			AddLiteralMacro("__TIME__", TokenType.String, "\"" + now.ToString("HH:mm:ss") + "\"");
 			AddLiteralMacro("__FILE__", TokenType.String, "\"<unknown>\"");
 			AddLiteralMacro("__LINE__", TokenType.Number, "-1");
 			AddLiteralMacro("__COUNTER__", TokenType.Number, "-1");
@@ -129,150 +126,150 @@ namespace UnrealBuildTool
 		/// <returns>True if the current branch is active</returns>
 		public bool IsCurrentBranchActive()
 		{
-			return State.IsCurrentBranchActive();
+			return _state.IsCurrentBranchActive();
 		}
 
 		/// <summary>
 		/// Defines a macro. May have an optional '=Value' suffix.
 		/// </summary>
-		/// <param name="Definition">Macro to define</param>
-		public void AddDefinition(string Definition)
+		/// <param name="definition">Macro to define</param>
+		public void AddDefinition(string definition)
 		{
-			List<Token> Tokens = TokenReader.Tokenize(Definition);
-			if(Tokens.Count == 0)
+			List<Token> tokens = TokenReader.Tokenize(definition);
+			if (tokens.Count == 0)
 			{
 				throw new PreprocessorException(null, "Missing macro name");
 			}
-			if(Tokens[0].Type != TokenType.Identifier)
+			if (tokens[0].Type != TokenType.Identifier)
 			{
-				throw new PreprocessorException(null, "'{0}' is not a valid macro name", Tokens[0].ToString()!);
+				throw new PreprocessorException(null, "'{0}' is not a valid macro name", tokens[0].ToString()!);
 			}
 
-			List<Token> ValueTokens = new List<Token>();
-			if(Tokens.Count == 1)
+			List<Token> valueTokens = new();
+			if (tokens.Count == 1)
 			{
-				ValueTokens.Add(new Token(TokenType.Number, TokenFlags.None, OneLiteral));
+				valueTokens.Add(new Token(TokenType.Number, TokenFlags.None, s_oneLiteral));
 			}
-			else if(Tokens[1].Type != TokenType.Equals)
+			else if (tokens[1].Type != TokenType.Equals)
 			{
-				throw new PreprocessorException(null, "Unable to parse macro definition '{0}'", Definition);
+				throw new PreprocessorException(null, "Unable to parse macro definition '{0}'", definition);
 			}
 			else
 			{
-				ValueTokens.AddRange(Tokens.Skip(2));
+				valueTokens.AddRange(tokens.Skip(2));
 			}
 
-			PreprocessorMacro Macro = new PreprocessorMacro(Tokens[0].Identifier!, null, ValueTokens);
-			State.DefineMacro(Macro);
+			PreprocessorMacro macro = new(tokens[0].Identifier!, null, valueTokens);
+			_state.DefineMacro(macro);
 		}
 
 		/// <summary>
 		/// Defines a macro
 		/// </summary>
-		/// <param name="Name">Name of the macro</param>
-		/// <param name="Value">String to be parsed for the macro's value</param>
-		public void AddDefinition(string Name, string Value)
+		/// <param name="name">Name of the macro</param>
+		/// <param name="value">String to be parsed for the macro's value</param>
+		public void AddDefinition(string name, string value)
 		{
-			List<Token> Tokens = new List<Token>();
+			List<Token> tokens = new();
 
-			TokenReader Reader = new TokenReader(Value);
-			while(Reader.MoveNext())
+			TokenReader reader = new(value);
+			while (reader.MoveNext())
 			{
-				Tokens.Add(Reader.Current);
+				tokens.Add(reader.Current);
 			}
 
-			PreprocessorMacro Macro = new PreprocessorMacro(Identifier.FindOrAdd(Name), null, Tokens);
-			State.DefineMacro(Macro);
+			PreprocessorMacro macro = new(Identifier.FindOrAdd(name), null, tokens);
+			_state.DefineMacro(macro);
 		}
 
 		/// <summary>
 		/// Defines a macro
 		/// </summary>
-		/// <param name="Macro">The macro definition</param>
-		public void AddDefinition(PreprocessorMacro Macro)
+		/// <param name="macro">The macro definition</param>
+		public void AddDefinition(PreprocessorMacro macro)
 		{
-			State.DefineMacro(Macro);
+			_state.DefineMacro(macro);
 		}
 
 		/// <summary>
 		/// Adds an include path to the preprocessor
 		/// </summary>
-		/// <param name="Directory">The include path</param>
-		public void AddIncludePath(DirectoryItem Directory)
+		/// <param name="directory">The include path</param>
+		public void AddIncludePath(DirectoryItem directory)
 		{
-			if(!IncludeDirectories.Contains(Directory))
+			if (!_includeDirectories.Contains(directory))
 			{
-				IncludeDirectories.Add(Directory);
+				_includeDirectories.Add(directory);
 			}
 		}
 
 		/// <summary>
 		/// Adds an include path to the preprocessor
 		/// </summary>
-		/// <param name="Location">The include path</param>
-		public void AddIncludePath(DirectoryReference Location)
+		/// <param name="location">The include path</param>
+		public void AddIncludePath(DirectoryReference location)
 		{
-			DirectoryItem Directory = DirectoryItem.GetItemByDirectoryReference(Location);
-			if(!Directory.Exists)
+			DirectoryItem directory = DirectoryItem.GetItemByDirectoryReference(location);
+			if (!directory.Exists)
 			{
-				throw new FileNotFoundException("Unable to find " + Location.FullName);
+				throw new FileNotFoundException("Unable to find " + location.FullName);
 			}
-			AddIncludePath(Directory);
+			AddIncludePath(directory);
 		}
 
 		/// <summary>
 		/// Adds an include path to the preprocessor
 		/// </summary>
-		/// <param name="DirectoryName">The include path</param>
-		public void AddIncludePath(string DirectoryName)
+		/// <param name="directoryName">The include path</param>
+		public void AddIncludePath(string directoryName)
 		{
-			AddIncludePath(new DirectoryReference(DirectoryName));
+			AddIncludePath(new DirectoryReference(directoryName));
 		}
 
 		/// <summary>
 		/// Adds a framework path to the preprocessor
 		/// </summary>
-		/// <param name="Directory">The framework path</param>
-		public void AddFrameworkPath(DirectoryItem Directory)
+		/// <param name="directory">The framework path</param>
+		public void AddFrameworkPath(DirectoryItem directory)
 		{
-			if (!FrameworkDirectories.Contains(Directory))
+			if (!_frameworkDirectories.Contains(directory))
 			{
-				FrameworkDirectories.Add(Directory);
+				_frameworkDirectories.Add(directory);
 			}
 		}
 
 		/// <summary>
 		/// Adds a framework path to the preprocessor
 		/// </summary>
-		/// <param name="Location">The framework path</param>
-		public void AddFrameworkPath(DirectoryReference Location)
+		/// <param name="location">The framework path</param>
+		public void AddFrameworkPath(DirectoryReference location)
 		{
-			DirectoryItem Directory = DirectoryItem.GetItemByDirectoryReference(Location);
-			if (!Directory.Exists)
+			DirectoryItem directory = DirectoryItem.GetItemByDirectoryReference(location);
+			if (!directory.Exists)
 			{
-				throw new FileNotFoundException("Unable to find " + Location.FullName);
+				throw new FileNotFoundException("Unable to find " + location.FullName);
 			}
-			AddFrameworkPath(Directory);
+			AddFrameworkPath(directory);
 		}
 
 		/// <summary>
 		/// Adds a framework path to the preprocessor
 		/// </summary>
-		/// <param name="DirectoryName">The framework path</param>
-		public void AddFrameworkPath(string DirectoryName)
+		/// <param name="directoryName">The framework path</param>
+		public void AddFrameworkPath(string directoryName)
 		{
-			AddFrameworkPath(new DirectoryReference(DirectoryName));
+			AddFrameworkPath(new DirectoryReference(directoryName));
 		}
 
 		/// <summary>
 		/// Try to resolve an quoted include against the list of include directories. Uses search order described by https://msdn.microsoft.com/en-us/library/36k2cdd4.aspx.
 		/// </summary>
-		/// <param name="Context">The current preprocessor context</param>
-		/// <param name="IncludePath">The path appearing in an #include directive</param>
-		/// <param name="Type">Specifies rules for how to resolve the include path (normal/system)</param>
-		/// <param name="File">If found, receives the resolved file</param>
+		/// <param name="context">The current preprocessor context</param>
+		/// <param name="includePath">The path appearing in an #include directive</param>
+		/// <param name="type">Specifies rules for how to resolve the include path (normal/system)</param>
+		/// <param name="file">If found, receives the resolved file</param>
 		/// <returns>True if the The resolved file</returns>
-		public bool TryResolveIncludePath(PreprocessorContext Context, string IncludePath, IncludePathType Type, [NotNullWhen(true)] out FileItem? File)
+		public bool TryResolveIncludePath(PreprocessorContext context, string includePath, IncludePathType type, [NotNullWhen(true)] out FileItem? file)
 		{
 			// From MSDN (https://msdn.microsoft.com/en-us/library/36k2cdd4.aspx?f=255&MSPPError=-2147217396)
 			//
@@ -290,36 +287,34 @@ namespace UnrealBuildTool
 			//   2) Along the paths that are specified by the INCLUDE environment variable.
 
 			// If it's an absolute path, return it immediately
-			if(Path.IsPathRooted(IncludePath))
+			if (Path.IsPathRooted(includePath))
 			{
-				FileItem FileItem = FileItem.GetItemByPath(IncludePath);
-				if(FileItem.Exists)
+				FileItem fileItem = FileItem.GetItemByPath(includePath);
+				if (fileItem.Exists)
 				{
-					File = FileItem;
+					file = fileItem;
 					return true;
 				}
 				else
 				{
-					File = null;
+					file = null;
 					return false;
 				}
 			}
 
 			// Split the path into fragments
-			string[] Fragments = IncludePath.Split(Path.DirectorySeparatorChar, Path.AltDirectorySeparatorChar);
+			string[] fragments = includePath.Split(Path.DirectorySeparatorChar, Path.AltDirectorySeparatorChar);
 
 			// Try to match the include path against any of the included directories
-			if(Type == IncludePathType.Normal)
+			if (type == IncludePathType.Normal)
 			{
-				for(PreprocessorContext? OuterContext = Context; OuterContext != null; OuterContext = OuterContext.Outer)
+				for (PreprocessorContext? outerContext = context; outerContext != null; outerContext = outerContext.Outer)
 				{
-					PreprocessorFileContext? OuterFileContext = OuterContext as PreprocessorFileContext;
-					if(OuterFileContext != null)
+					if (outerContext is PreprocessorFileContext outerFileContext)
 					{
-						FileItem? ResolvedFile;
-						if(TryResolveRelativeIncludePath(OuterFileContext.Directory, Fragments, out ResolvedFile))
+						if (TryResolveRelativeIncludePath(outerFileContext.Directory, fragments, out FileItem? resolvedFile))
 						{
-							File = ResolvedFile;
+							file = resolvedFile;
 							return true;
 						}
 					}
@@ -327,127 +322,123 @@ namespace UnrealBuildTool
 			}
 
 			// Try to match the include path against any of the system directories
-			foreach(DirectoryItem BaseDirectory in IncludeDirectories)
+			foreach (DirectoryItem baseDirectory in _includeDirectories)
 			{
-				FileItem? ResolvedFile;
-				if(TryResolveRelativeIncludePath(BaseDirectory, Fragments, out ResolvedFile))
+				if (TryResolveRelativeIncludePath(baseDirectory, fragments, out FileItem? resolvedFile))
 				{
-					File = ResolvedFile;
+					file = resolvedFile;
 					return true;
 				}
 			}
 
 			// Try to match the include path against any of the MacOS framework Header paths
-			if (Fragments.Length > 1)
+			if (fragments.Length > 1)
 			{
-				foreach (DirectoryItem BaseDirectory in FrameworkDirectories)
+				foreach (DirectoryItem baseDirectory in _frameworkDirectories)
 				{
-					if (BaseDirectory.TryGetDirectory($"{Fragments[0]}.framework", out DirectoryItem? FrameworkBaseDirectory) &&
-						FrameworkBaseDirectory.TryGetDirectory("Headers", out DirectoryItem? HeaderDirectory) &&
-						TryResolveRelativeIncludePath(HeaderDirectory, Fragments.Skip(1).ToArray(), out FileItem? ResolvedFile))
+					if (baseDirectory.TryGetDirectory($"{fragments[0]}.framework", out DirectoryItem? frameworkBaseDirectory) &&
+						frameworkBaseDirectory.TryGetDirectory("Headers", out DirectoryItem? headerDirectory) &&
+						TryResolveRelativeIncludePath(headerDirectory, fragments.Skip(1).ToArray(), out FileItem? resolvedFile))
 					{
-						File = ResolvedFile;
+						file = resolvedFile;
 						return true;
 					}
 				}
 			}
 
 			// Failed to find the file
-			File = null;
+			file = null;
 			return false;
 		}
 
 		/// <summary>
 		/// Try to resolve an quoted include against the list of include directories. Uses search order described by https://msdn.microsoft.com/en-us/library/36k2cdd4.aspx.
 		/// </summary>
-		/// <param name="BaseDirectory">The base directory to search from</param>
-		/// <param name="Fragments">Fragments of the relative path to follow</param>
-		/// <param name="File">The file that was found, if successful</param>
+		/// <param name="baseDirectory">The base directory to search from</param>
+		/// <param name="fragments">Fragments of the relative path to follow</param>
+		/// <param name="file">The file that was found, if successful</param>
 		/// <returns>True if the The resolved file</returns>
-		public bool TryResolveRelativeIncludePath(DirectoryItem BaseDirectory, string[] Fragments, [NotNullWhen(true)] out FileItem? File)
+		public static bool TryResolveRelativeIncludePath(DirectoryItem baseDirectory, string[] fragments, [NotNullWhen(true)] out FileItem? file)
 		{
-			DirectoryItem? Directory = BaseDirectory;
-			for(int Idx = 0; Idx < Fragments.Length - 1; Idx++)
+			DirectoryItem? directory = baseDirectory;
+			for (int idx = 0; idx < fragments.Length - 1; idx++)
 			{
-				if(!Directory.TryGetDirectory(Fragments[Idx], out Directory))
+				if (!directory.TryGetDirectory(fragments[idx], out directory))
 				{
-					File = null;
+					file = null;
 					return false;
 				}
 			}
-			return Directory.TryGetFile(Fragments[Fragments.Length - 1], out File);
+			return directory.TryGetFile(fragments[^1], out file);
 		}
 
 		/// <summary>
 		/// Parses a file recursively 
 		/// </summary>
-		/// <param name="File">File to parse</param>
-		/// <param name="Fragments">Lists of fragments that are parsed</param>
-		/// <param name="OuterContext">Outer context information, for error messages</param>
-		/// <param name="SourceFileCache">Cache for source files</param>
-		/// <param name="Logger">Logger for output</param>
-		/// <param name="bShowIncludes">Show all the included files, in order</param>
-		/// <param name="bIgnoreMissingIncludes">Suppress exceptions if an include path can not be resolved</param>
-		public void ParseFile(FileItem File, List<SourceFileFragment> Fragments, PreprocessorContext? OuterContext, SourceFileMetadataCache SourceFileCache, ILogger Logger, bool bShowIncludes = false, bool bIgnoreMissingIncludes = false)
+		/// <param name="file">File to parse</param>
+		/// <param name="fragments">Lists of fragments that are parsed</param>
+		/// <param name="outerContext">Outer context information, for error messages</param>
+		/// <param name="sourceFileCache">Cache for source files</param>
+		/// <param name="logger">Logger for output</param>
+		/// <param name="showIncludes">Show all the included files, in order</param>
+		/// <param name="ignoreMissingIncludes">Suppress exceptions if an include path can not be resolved</param>
+		public void ParseFile(FileItem file, List<SourceFileFragment> fragments, PreprocessorContext? outerContext, SourceFileMetadataCache sourceFileCache, ILogger logger, bool showIncludes = false, bool ignoreMissingIncludes = false)
 		{
 			// If the file has already been included and had a #pragma once directive, don't include it again
-			if(PragmaOnceFiles.Contains(File))
+			if (_pragmaOnceFiles.Contains(file))
 			{
 				return;
 			}
 
-			if (!ProcessedFiles.Contains(File))
-			{
-				ProcessedFiles.Add(File);
-			}
+			_processedFiles.Add(file);
 
 			// Output a trace of the included files
-			if(bShowIncludes)
+			if (showIncludes)
 			{
-				Logger.LogInformation("Note: including file: {FileLocation}", File.Location);
+				logger.LogInformation("Note: including file: {FileLocation}", file.Location);
 			}
 
 			// If the file had a header guard, and the macro is still defined, don't include it again
-			SourceFile SourceFile = SourceFileCache.GetSourceFile(File);
-			if(SourceFile.HeaderGuardMacro != null && State.IsMacroDefined(SourceFile.HeaderGuardMacro))
+			SourceFile sourceFile = sourceFileCache.GetSourceFile(file);
+			if (sourceFile.HeaderGuardMacro != null && _state.IsMacroDefined(sourceFile.HeaderGuardMacro))
 			{
 				return;
 			}
 
 			// Create a context for this file
-			PreprocessorFileContext Context = new PreprocessorFileContext(SourceFile, OuterContext);
+			PreprocessorFileContext context = new(sourceFile, outerContext);
 
 			// Parse the markup for this file
-			while(Context.MarkupIdx < SourceFile.Markup.Length)
+			while (context.MarkupIdx < sourceFile.Markup.Length)
 			{
-				SourceFileMarkup Markup = SourceFile.Markup[Context.MarkupIdx];
-				if(Markup.Type == SourceFileMarkupType.Include)
+				SourceFileMarkup markup = sourceFile.Markup[context.MarkupIdx];
+				if (markup.Type == SourceFileMarkupType.Include)
 				{
-					if(State.IsCurrentBranchActive())
+					if (_state.IsCurrentBranchActive())
 					{
 						// Parse the directive
-						FileItem? IncludedFile = ParseIncludeDirective(Markup, Context, bIgnoreMissingIncludes);
+						FileItem? includedFile = ParseIncludeDirective(markup, context, ignoreMissingIncludes);
 
 						// Parse the included file
-						if (IncludedFile != null)
+						if (includedFile != null)
 						{
-							ParseFile(IncludedFile, Fragments, Context, SourceFileCache, Logger, bShowIncludes, bIgnoreMissingIncludes);
+							ParseFile(includedFile, fragments, context, sourceFileCache, logger, showIncludes, ignoreMissingIncludes);
 						}
 					}
-					Context.MarkupIdx++;
+					context.MarkupIdx++;
 				}
 				else
 				{
 					// Get the next fragment
-					SourceFileFragment Fragment = SourceFile.Fragments[Context.FragmentIdx];
-					Debug.Assert(Fragment.MarkupMin == Context.MarkupIdx);
+					SourceFileFragment fragment = sourceFile.Fragments[context.FragmentIdx];
+					Debug.Assert(fragment.MarkupMin == context.MarkupIdx);
 
 					// Parse this fragment
-					ParseFragment(SourceFile, Fragment, Context);
+					ParseFragment(sourceFile, fragment, context);
 
 					// Add this fragment to the list
-					Fragments.Add(Fragment);
-					Context.FragmentIdx++;
+					fragments.Add(fragment);
+					context.FragmentIdx++;
 				}
 			}
 		}
@@ -455,107 +446,106 @@ namespace UnrealBuildTool
 		/// <summary>
 		/// Parse an include directive and resolve the file it references
 		/// </summary>
-		/// <param name="Markup">Markup for the include directive</param>
-		/// <param name="Context">Current preprocessor context</param>
-		/// <param name="bIgnoreMissingIncludes">Suppress exceptions if an include path can not be resolved</param>
+		/// <param name="markup">Markup for the include directive</param>
+		/// <param name="context">Current preprocessor context</param>
+		/// <param name="ignoreMissingIncludes">Suppress exceptions if an include path can not be resolved</param>
 		/// <returns>Included file</returns>
-		FileItem? ParseIncludeDirective(SourceFileMarkup Markup, PreprocessorFileContext Context, bool bIgnoreMissingIncludes = false)
+		FileItem? ParseIncludeDirective(SourceFileMarkup markup, PreprocessorFileContext context, bool ignoreMissingIncludes = false)
 		{
 			// Expand macros in the given tokens
-			List<Token> ExpandedTokens = new List<Token>();
-			ExpandMacros(Markup.Tokens!, ExpandedTokens, false, Context);
-						
+			List<Token> expandedTokens = new();
+			ExpandMacros(markup.Tokens!, expandedTokens, false, context);
+
 			// Convert the string to a single token
-			string IncludeToken = Token.Format(ExpandedTokens);
+			string includeToken = Token.Format(expandedTokens);
 
 			// Expand any macros in them and resolve it
-			IncludePathType Type;
-			if(IncludeToken.Length >= 2 && IncludeToken[0] == '\"' && IncludeToken[IncludeToken.Length - 1] == '\"')
+			IncludePathType type;
+			if (includeToken.Length >= 2 && includeToken[0] == '\"' && includeToken[^1] == '\"')
 			{
-				Type = IncludePathType.Normal;
+				type = IncludePathType.Normal;
 			}
-			else if(IncludeToken.Length >= 2 && IncludeToken[0] == '<' && IncludeToken[IncludeToken.Length - 1] == '>')
+			else if (includeToken.Length >= 2 && includeToken[0] == '<' && includeToken[^1] == '>')
 			{
-				Type = IncludePathType.System;
+				type = IncludePathType.System;
 			}
 			else
 			{
-				throw new PreprocessorException(Context, "Couldn't resolve include '{0}'", IncludeToken);
+				throw new PreprocessorException(context, "Couldn't resolve include '{0}'", includeToken);
 			}
 
 			// Get the include path
-			string IncludePath = IncludeToken.Substring(1, IncludeToken.Length - 2);
+			string includePath = includeToken[1..^1];
 
 			// Resolve the included file
-			FileItem? IncludedFile;
-			if(!TryResolveIncludePath(Context, IncludePath, Type, out IncludedFile))
+			if (!TryResolveIncludePath(context, includePath, type, out FileItem? includedFile))
 			{
-				if (bIgnoreMissingIncludes)
+				if (ignoreMissingIncludes)
 				{
-					Log.TraceWarningOnce("Couldn't resolve include '{0}' ({1})", IncludePath, Context.SourceFile.Location);
+					Log.TraceWarningOnce("Couldn't resolve include '{0}' ({1})", includePath, context.SourceFile.Location);
 				}
 				else
 				{
-					throw new PreprocessorException(Context, "Couldn't resolve include '{0}' ({1})", IncludePath, Context.SourceFile.Location);
+					throw new PreprocessorException(context, "Couldn't resolve include '{0}' ({1})", includePath, context.SourceFile.Location);
 				}
 			}
-			return IncludedFile;
+			return includedFile;
 		}
 
 		/// <summary>
 		/// Parse a source file fragment, using cached transforms if possible
 		/// </summary>
-		/// <param name="SourceFile">The source file being parsed</param>
-		/// <param name="Fragment">Fragment to parse</param>
-		/// <param name="Context">Current preprocessor context</param>
-		void ParseFragment(SourceFile SourceFile, SourceFileFragment Fragment, PreprocessorFileContext Context)
+		/// <param name="sourceFile">The source file being parsed</param>
+		/// <param name="fragment">Fragment to parse</param>
+		/// <param name="context">Current preprocessor context</param>
+		void ParseFragment(SourceFile sourceFile, SourceFileFragment fragment, PreprocessorFileContext context)
 		{
 			// Check if there's a valid transform that matches the current state
-			int TransformIdx = 0;
-			for(;;)
+			int transformIdx = 0;
+			for (; ; )
 			{
-				PreprocessorTransform[] Transforms;
-				lock(Fragment)
+				PreprocessorTransform[] transforms;
+				lock (fragment)
 				{
-					Transforms = Fragment.Transforms;
-					if(TransformIdx == Transforms.Length)
+					transforms = fragment.Transforms;
+					if (transformIdx == transforms.Length)
 					{
 						// Attach a new transform to the current state
-						PreprocessorTransform Transform = State.BeginCapture();
-						for(; Context.MarkupIdx < Fragment.MarkupMax; Context.MarkupIdx++)
+						PreprocessorTransform transform = _state.BeginCapture();
+						for (; context.MarkupIdx < fragment.MarkupMax; context.MarkupIdx++)
 						{
-							SourceFileMarkup Markup = SourceFile.Markup[Context.MarkupIdx];
-							ParseMarkup(Markup.Type, Markup.Tokens!, Context);
+							SourceFileMarkup markup = sourceFile.Markup[context.MarkupIdx];
+							ParseMarkup(markup.Type, markup.Tokens!, context);
 						}
-						Transform = State.EndCapture()!;
+						transform = _state.EndCapture()!;
 
 						// Add it to the fragment for future fragments
-						PreprocessorTransform[] NewTransforms = new PreprocessorTransform[Fragment.Transforms.Length + 1];
-						for(int Idx = 0; Idx < Transforms.Length; Idx++)
+						PreprocessorTransform[] newTransforms = new PreprocessorTransform[fragment.Transforms.Length + 1];
+						for (int idx = 0; idx < transforms.Length; idx++)
 						{
-							NewTransforms[Idx] = Transforms[Idx];
+							newTransforms[idx] = transforms[idx];
 						}
-						NewTransforms[Transforms.Length] = Transform;
+						newTransforms[transforms.Length] = transform;
 
 						// Assign it to the fragment
-						Fragment.Transforms = NewTransforms;
+						fragment.Transforms = newTransforms;
 						return;
 					}
 				}
 
-				for(; TransformIdx < Transforms.Length; TransformIdx++)
+				for (; transformIdx < transforms.Length; transformIdx++)
 				{
-					PreprocessorTransform Transform = Transforms[TransformIdx];
-					if(State.TryToApply(Transform))
+					PreprocessorTransform transform = transforms[transformIdx];
+					if (_state.TryToApply(transform))
 					{
 						// Update the pragma once state
-						if(Transform.bHasPragmaOnce)
+						if (transform.HasPragmaOnce)
 						{
-							PragmaOnceFiles.Add(SourceFile.File);
+							_pragmaOnceFiles.Add(sourceFile.File);
 						}
 
 						// Move to the end of the fragment
-						Context.MarkupIdx = Fragment.MarkupMax;
+						context.MarkupIdx = fragment.MarkupMax;
 						return;
 					}
 				}
@@ -565,89 +555,89 @@ namespace UnrealBuildTool
 		/// <summary>
 		/// Validate and add a macro using the given parameter and token list
 		/// </summary>
-		/// <param name="Context">The current preprocessor context</param>
-		/// <param name="Name">Name of the macro</param>
-		/// <param name="Parameters">Parameter list for the macro</param>
-		/// <param name="Tokens">List of tokens</param>
-		void AddMacro(PreprocessorContext Context, Identifier Name, List<Identifier>? Parameters, List<Token> Tokens)
+		/// <param name="context">The current preprocessor context</param>
+		/// <param name="name">Name of the macro</param>
+		/// <param name="parameters">Parameter list for the macro</param>
+		/// <param name="tokens">List of tokens</param>
+		void AddMacro(PreprocessorContext context, Identifier name, List<Identifier>? parameters, List<Token> tokens)
 		{
-			if(Tokens.Count == 0)
+			if (tokens.Count == 0)
 			{
-				Tokens.Add(new Token(TokenType.Placemarker, TokenFlags.None));
+				tokens.Add(new Token(TokenType.Placemarker, TokenFlags.None));
 			}
 			else
 			{
-				if(Tokens[0].HasLeadingSpace)
+				if (tokens[0].HasLeadingSpace)
 				{
-					Tokens[0] = Tokens[0].RemoveFlags(TokenFlags.HasLeadingSpace);
+					tokens[0] = tokens[0].RemoveFlags(TokenFlags.HasLeadingSpace);
 				}
-				if (Tokens[0].Type == TokenType.HashHash || Tokens[Tokens.Count - 1].Type == TokenType.HashHash)
+				if (tokens[0].Type == TokenType.HashHash || tokens[^1].Type == TokenType.HashHash)
 				{
-					throw new PreprocessorException(Context, "Invalid use of concatenation at start or end of token sequence");
+					throw new PreprocessorException(context, "Invalid use of concatenation at start or end of token sequence");
 				}
-				if (Parameters == null || Parameters.Count == 0 || Parameters[Parameters.Count - 1] != Identifiers.__VA_ARGS__)
+				if (parameters == null || parameters.Count == 0 || parameters[^1] != Identifiers.__VA_ARGS__)
 				{
-					if(Tokens.Any(x => x.Identifier == Identifiers.__VA_ARGS__))
+					if (tokens.Any(x => x.Identifier == Identifiers.__VA_ARGS__))
 					{
-						throw new PreprocessorException(Context, "Invalid reference to {0}", Identifiers.__VA_ARGS__);
+						throw new PreprocessorException(context, "Invalid reference to {0}", Identifiers.__VA_ARGS__);
 					}
 				}
 			}
-			State.DefineMacro(new PreprocessorMacro(Name, Parameters, Tokens));
+			_state.DefineMacro(new PreprocessorMacro(name, parameters, tokens));
 		}
-		
+
 		/// <summary>
 		/// Set a predefined macro to a given value
 		/// </summary>
-		/// <param name="Name">Name of the macro</param>
-		/// <param name="Type">Type of the macro token</param>
-		/// <param name="Value">Value of the macro</param>
+		/// <param name="name">Name of the macro</param>
+		/// <param name="type">Type of the macro token</param>
+		/// <param name="value">Value of the macro</param>
 		/// <returns>The created macro</returns>
-		void AddLiteralMacro(string Name, TokenType Type, string Value)
+		void AddLiteralMacro(string name, TokenType type, string value)
 		{
-			Token Token = new Token(Type, TokenFlags.None, Value);
-			PreprocessorMacro Macro = new PreprocessorMacro(Identifier.FindOrAdd(Name), null, new List<Token> { Token });
-			State.DefineMacro(Macro);
+			Token token = new(type, TokenFlags.None, value);
+			PreprocessorMacro macro = new(Identifier.FindOrAdd(name), null, new List<Token> { token });
+			_state.DefineMacro(macro);
 		}
 
 		/// <summary>
 		/// Parse a marked up directive from a file
 		/// </summary>
-		/// <param name="Type">The markup type</param>
-		/// <param name="Tokens">Tokens for the directive</param>
-		/// <param name="Context">The context that this markup is being parsed in</param>
-		public void ParseMarkup(SourceFileMarkupType Type, List<Token> Tokens, PreprocessorContext Context)
+		/// <param name="type">The markup type</param>
+		/// <param name="tokens">Tokens for the directive</param>
+		/// <param name="context">The context that this markup is being parsed in</param>
+		public void ParseMarkup(SourceFileMarkupType type, List<Token> tokens, PreprocessorContext context)
 		{
-			switch(Type)
+			switch (type)
 			{
 				case SourceFileMarkupType.Include:
-					throw new PreprocessorException(Context, "Include directives should be handled by the caller.");
+					throw new PreprocessorException(context, "Include directives should be handled by the caller.");
 				case SourceFileMarkupType.Define:
-					ParseDefineDirective(Tokens, Context);
+					ParseDefineDirective(tokens, context);
 					break;
 				case SourceFileMarkupType.Undef:
-					ParseUndefDirective(Tokens, Context);
+					ParseUndefDirective(tokens, context);
 					break;
 				case SourceFileMarkupType.If:
-					ParseIfDirective(Tokens, Context);
+					ParseIfDirective(tokens, context);
 					break;
 				case SourceFileMarkupType.Ifdef:
-					ParseIfdefDirective(Tokens, Context);
+					ParseIfdefDirective(tokens, context);
 					break;
 				case SourceFileMarkupType.Ifndef:
-					ParseIfndefDirective(Tokens, Context);
+					ParseIfndefDirective(tokens, context);
 					break;
 				case SourceFileMarkupType.Elif:
-					ParseElifDirective(Tokens, Context);
+					ParseElifDirective(tokens, context);
 					break;
 				case SourceFileMarkupType.Else:
-					ParseElseDirective(Tokens, Context);
+					ParseElseDirective(tokens, context);
 					break;
 				case SourceFileMarkupType.Endif:
-					ParseEndifDirective(Tokens, Context);
+					ParseEndifDirective(tokens, context);
 					break;
 				case SourceFileMarkupType.Pragma:
-					ParsePragmaDirective(Tokens, Context);
+					ParsePragmaDirective(tokens, context);
 					break;
 			}
 		}
@@ -655,286 +645,284 @@ namespace UnrealBuildTool
 		/// <summary>
 		/// Read a macro definition
 		/// </summary>
-		/// <param name="Tokens">List of tokens in the directive</param>
-		/// <param name="Context">The context that this directive is being parsed in</param>
-		public void ParseDefineDirective(List<Token> Tokens, PreprocessorContext Context)
+		/// <param name="tokens">List of tokens in the directive</param>
+		/// <param name="context">The context that this directive is being parsed in</param>
+		public void ParseDefineDirective(List<Token> tokens, PreprocessorContext context)
 		{
-			if(State.IsCurrentBranchActive())
+			if (_state.IsCurrentBranchActive())
 			{
 				// Check there's a name token
-				if(Tokens.Count < 1 || Tokens[0].Type != TokenType.Identifier || Tokens[0].Identifier == Identifiers.Defined)
+				if (tokens.Count < 1 || tokens[0].Type != TokenType.Identifier || tokens[0].Identifier == Identifiers.Defined)
 				{
-					throw new PreprocessorException(Context, "Invalid macro name");
+					throw new PreprocessorException(context, "Invalid macro name");
 				}
 
 				// Read the macro name
-				Identifier Name = Tokens[0].Identifier!;
-				int TokenIdx = 1;
+				Identifier name = tokens[0].Identifier!;
+				int tokenIdx = 1;
 
 				// Read the macro parameter list, if there is one
-				List<Identifier>? Parameters = null;
-				if (TokenIdx < Tokens.Count && !Tokens[TokenIdx].HasLeadingSpace && Tokens[TokenIdx].Type == TokenType.LeftParen)
+				List<Identifier>? parameters = null;
+				if (tokenIdx < tokens.Count && !tokens[tokenIdx].HasLeadingSpace && tokens[tokenIdx].Type == TokenType.LeftParen)
 				{
-					Parameters = new List<Identifier>();
-					if(++TokenIdx == Tokens.Count)
+					parameters = new List<Identifier>();
+					if (++tokenIdx == tokens.Count)
 					{
-						throw new PreprocessorException(Context, "Unexpected end of macro parameter list");
+						throw new PreprocessorException(context, "Unexpected end of macro parameter list");
 					}
-					if(Tokens[TokenIdx].Type != TokenType.RightParen)
+					if (tokens[tokenIdx].Type != TokenType.RightParen)
 					{
-						for(;;TokenIdx++)
+						for (; ; tokenIdx++)
 						{
 							// Check there's enough tokens left for a parameter name, plus ',' or ')'
-							if(TokenIdx + 2 > Tokens.Count)
+							if (tokenIdx + 2 > tokens.Count)
 							{
-								throw new PreprocessorException(Context, "Unexpected end of macro parameter list");
+								throw new PreprocessorException(context, "Unexpected end of macro parameter list");
 							}
 
 							// Check it's a valid name, and add it to the list
-							Token NameToken = Tokens[TokenIdx++];
-							if(NameToken.Type == TokenType.Ellipsis)
+							Token nameToken = tokens[tokenIdx++];
+							if (nameToken.Type == TokenType.Ellipsis)
 							{
-								if(Tokens[TokenIdx].Type != TokenType.RightParen)
+								if (tokens[tokenIdx].Type != TokenType.RightParen)
 								{
-									throw new PreprocessorException(Context, "Variadic macro arguments must be last in list");
+									throw new PreprocessorException(context, "Variadic macro arguments must be last in list");
 								}
 								else
 								{
-									NameToken = new Token(Identifiers.__VA_ARGS__, NameToken.Flags & TokenFlags.HasLeadingSpace);
+									nameToken = new Token(Identifiers.__VA_ARGS__, nameToken.Flags & TokenFlags.HasLeadingSpace);
 								}
 							}
 							else
 							{
-								if (NameToken.Type != TokenType.Identifier || NameToken.Identifier == Identifiers.__VA_ARGS__)
+								if (nameToken.Type != TokenType.Identifier || nameToken.Identifier == Identifiers.__VA_ARGS__)
 								{
-									throw new PreprocessorException(Context, "Invalid preprocessor token: {0}", NameToken);
+									throw new PreprocessorException(context, "Invalid preprocessor token: {0}", nameToken);
 								}
-								if (Parameters.Contains(NameToken.Identifier!))
+								if (parameters.Contains(nameToken.Identifier!))
 								{
-									throw new PreprocessorException(Context, "'{0}' has already been used as an argument name", NameToken);
+									throw new PreprocessorException(context, "'{0}' has already been used as an argument name", nameToken);
 								}
 							}
-							Parameters.Add(NameToken.Identifier!);
+							parameters.Add(nameToken.Identifier!);
 
 							// Read the separator
-							Token SeparatorToken = Tokens[TokenIdx];
-							if (SeparatorToken.Type == TokenType.RightParen)
+							Token separatorToken = tokens[tokenIdx];
+							if (separatorToken.Type == TokenType.RightParen)
 							{
 								break;
 							}
-							if (SeparatorToken.Type != TokenType.Comma)
+							if (separatorToken.Type != TokenType.Comma)
 							{
-								throw new PreprocessorException(Context, "Expected ',' or ')'");
+								throw new PreprocessorException(context, "Expected ',' or ')'");
 							}
 						}
 					}
-					TokenIdx++;
+					tokenIdx++;
 				}
 
 				// Read the macro tokens
-				AddMacro(Context, Name, Parameters, Tokens.Skip(TokenIdx).ToList());
+				AddMacro(context, name, parameters, tokens.Skip(tokenIdx).ToList());
 			}
 		}
 
 		/// <summary>
 		/// Parse an #undef directive
 		/// </summary>
-		/// <param name="Tokens">List of tokens in the directive</param>
-		/// <param name="Context">The context that this directive is being parsed in</param>
-		public void ParseUndefDirective(List<Token> Tokens, PreprocessorContext Context)
+		/// <param name="tokens">List of tokens in the directive</param>
+		/// <param name="context">The context that this directive is being parsed in</param>
+		public void ParseUndefDirective(List<Token> tokens, PreprocessorContext context)
 		{
-			if(State.IsCurrentBranchActive())
+			if (_state.IsCurrentBranchActive())
 			{
 				// Check there's a name token
-				if (Tokens.Count != 1)
+				if (tokens.Count != 1)
 				{
-					throw new PreprocessorException(Context, "Expected a single token after #undef");
+					throw new PreprocessorException(context, "Expected a single token after #undef");
 				}
-				if (Tokens[0].Type != TokenType.Identifier)
+				if (tokens[0].Type != TokenType.Identifier)
 				{
-					throw new PreprocessorException(Context, "Invalid macro name '{0}'", Tokens[0]);
+					throw new PreprocessorException(context, "Invalid macro name '{0}'", tokens[0]);
 				}
 
 				// Remove the macro from the list of definitions
-				State.UndefMacro(Tokens[0].Identifier!);
+				_state.UndefMacro(tokens[0].Identifier!);
 			}
 		}
 
 		/// <summary>
 		/// Parse an #if directive
 		/// </summary>
-		/// <param name="Tokens">List of tokens in the directive</param>
-		/// <param name="Context">The context that this directive is being parsed in</param>
-		public void ParseIfDirective(List<Token> Tokens, PreprocessorContext Context)
+		/// <param name="tokens">List of tokens in the directive</param>
+		/// <param name="context">The context that this directive is being parsed in</param>
+		public void ParseIfDirective(List<Token> tokens, PreprocessorContext context)
 		{
-			PreprocessorBranch Branch = PreprocessorBranch.HasIfDirective;
-			if (State.IsCurrentBranchActive())
+			PreprocessorBranch branch = PreprocessorBranch.HasIfDirective;
+			if (_state.IsCurrentBranchActive())
 			{
 				// Read a line into the buffer and expand the macros in it
-				List<Token> ExpandedTokens = new List<Token>();
-				ExpandMacros(Tokens, ExpandedTokens, true, Context);
+				List<Token> expandedTokens = new();
+				ExpandMacros(tokens, expandedTokens, true, context);
 
 				// Evaluate the condition
-				long Result = PreprocessorExpression.Evaluate(Context, ExpandedTokens);
-				if (Result != 0)
+				long result = PreprocessorExpression.Evaluate(context, expandedTokens);
+				if (result != 0)
 				{
-					Branch |= PreprocessorBranch.Active | PreprocessorBranch.Taken;
+					branch |= PreprocessorBranch.Active | PreprocessorBranch.Taken;
 				}
 			}
-			State.PushBranch(Branch);
+			_state.PushBranch(branch);
 		}
 
 		/// <summary>
 		/// Parse an #ifdef directive
 		/// </summary>
-		/// <param name="Tokens">List of tokens in the directive</param>
-		/// <param name="Context">The context that this directive is being parsed in</param>
-		public void ParseIfdefDirective(List<Token> Tokens, PreprocessorContext Context)
+		/// <param name="tokens">List of tokens in the directive</param>
+		/// <param name="context">The context that this directive is being parsed in</param>
+		public void ParseIfdefDirective(List<Token> tokens, PreprocessorContext context)
 		{
-			PreprocessorBranch Branch = PreprocessorBranch.HasIfdefDirective;
-			if (State.IsCurrentBranchActive())
+			PreprocessorBranch branch = PreprocessorBranch.HasIfdefDirective;
+			if (_state.IsCurrentBranchActive())
 			{
 				// Make sure there's only one token
-				if(Tokens.Count != 1 || Tokens[0].Type != TokenType.Identifier)
+				if (tokens.Count != 1 || tokens[0].Type != TokenType.Identifier)
 				{
-					throw new PreprocessorException(Context, "Missing or invalid macro name for #ifdef directive");
+					throw new PreprocessorException(context, "Missing or invalid macro name for #ifdef directive");
 				}
 
 				// Check if the macro is defined
-				if(State.IsMacroDefined(Tokens[0].Identifier!))
+				if (_state.IsMacroDefined(tokens[0].Identifier!))
 				{
-					Branch |= PreprocessorBranch.Active | PreprocessorBranch.Taken;
+					branch |= PreprocessorBranch.Active | PreprocessorBranch.Taken;
 				}
 			}
-			State.PushBranch(Branch);
+			_state.PushBranch(branch);
 		}
 
 		/// <summary>
 		/// Parse an #ifndef directive
 		/// </summary>
-		/// <param name="Tokens">List of tokens for this directive</param>
-		/// <param name="Context">The context that this directive is being parsed in</param>
-		public void ParseIfndefDirective(List<Token> Tokens, PreprocessorContext Context)
+		/// <param name="tokens">List of tokens for this directive</param>
+		/// <param name="context">The context that this directive is being parsed in</param>
+		public void ParseIfndefDirective(List<Token> tokens, PreprocessorContext context)
 		{
-			PreprocessorBranch Branch = PreprocessorBranch.HasIfndefDirective;
-			if (State.IsCurrentBranchActive())
+			PreprocessorBranch branch = PreprocessorBranch.HasIfndefDirective;
+			if (_state.IsCurrentBranchActive())
 			{
 				// Make sure there's only one token
-				if (Tokens.Count != 1 || Tokens[0].Type != TokenType.Identifier)
+				if (tokens.Count != 1 || tokens[0].Type != TokenType.Identifier)
 				{
-					throw new PreprocessorException(Context, "Missing or invalid macro name for #ifndef directive");
+					throw new PreprocessorException(context, "Missing or invalid macro name for #ifndef directive");
 				}
 
 				// Check if the macro is defined
-				if(!State.IsMacroDefined(Tokens[0].Identifier!))
+				if (!_state.IsMacroDefined(tokens[0].Identifier!))
 				{
-					Branch |= PreprocessorBranch.Active | PreprocessorBranch.Taken;
+					branch |= PreprocessorBranch.Active | PreprocessorBranch.Taken;
 				}
 			}
-			State.PushBranch(Branch);
+			_state.PushBranch(branch);
 		}
 
 		/// <summary>
 		/// Parse an #elif directive
 		/// </summary>
-		/// <param name="Tokens">List of tokens for this directive</param>
-		/// <param name="Context">The context that this directive is being parsed in</param>
-		public void ParseElifDirective(List<Token> Tokens, PreprocessorContext Context)
+		/// <param name="tokens">List of tokens for this directive</param>
+		/// <param name="context">The context that this directive is being parsed in</param>
+		public void ParseElifDirective(List<Token> tokens, PreprocessorContext context)
 		{
 			// Check we're in a branch, and haven't already read an #else directive
-			PreprocessorBranch Branch;
-			if (!State.TryPopBranch(out Branch))
+			if (!_state.TryPopBranch(out PreprocessorBranch branch))
 			{
-				throw new PreprocessorException(Context, "#elif directive outside conditional block");
+				throw new PreprocessorException(context, "#elif directive outside conditional block");
 			}
-			if (Branch.HasFlag(PreprocessorBranch.Complete))
+			if (branch.HasFlag(PreprocessorBranch.Complete))
 			{
-				throw new PreprocessorException(Context, "#elif directive cannot appear after #else");
+				throw new PreprocessorException(context, "#elif directive cannot appear after #else");
 			}
 
 			// Pop the current branch state at this depth, so we can test against whether the parent state is enabled
-			Branch = (Branch | PreprocessorBranch.HasElifDirective) & ~PreprocessorBranch.Active;
-			if (State.IsCurrentBranchActive())
+			branch = (branch | PreprocessorBranch.HasElifDirective) & ~PreprocessorBranch.Active;
+			if (_state.IsCurrentBranchActive())
 			{
 				// Read a line into the buffer and expand the macros in it
-				List<Token> ExpandedTokens = new List<Token>();
-				ExpandMacros(Tokens, ExpandedTokens, true, Context);
+				List<Token> expandedTokens = new();
+				ExpandMacros(tokens, expandedTokens, true, context);
 
 				// Check we're at the end of a conditional block
-				if (!Branch.HasFlag(PreprocessorBranch.Taken))
+				if (!branch.HasFlag(PreprocessorBranch.Taken))
 				{
-					long Result = PreprocessorExpression.Evaluate(Context, ExpandedTokens);
-					if (Result != 0)
+					long result = PreprocessorExpression.Evaluate(context, expandedTokens);
+					if (result != 0)
 					{
-						Branch |= PreprocessorBranch.Active | PreprocessorBranch.Taken;
+						branch |= PreprocessorBranch.Active | PreprocessorBranch.Taken;
 					}
 				}
 			}
-			State.PushBranch(Branch);
+			_state.PushBranch(branch);
 		}
 
 		/// <summary>
 		/// Parse an #else directive
 		/// </summary>
-		/// <param name="Tokens">List of tokens in the directive</param>
-		/// <param name="Context">The context that this directive is being parsed in</param>
-		public void ParseElseDirective(List<Token> Tokens, PreprocessorContext Context)
+		/// <param name="tokens">List of tokens in the directive</param>
+		/// <param name="context">The context that this directive is being parsed in</param>
+		public void ParseElseDirective(List<Token> tokens, PreprocessorContext context)
 		{
 			// Make sure there's nothing else on the line
-			if (Tokens.Count > 0)
+			if (tokens.Count > 0)
 			{
-				throw new PreprocessorException(Context, "Garbage after #else directive");
+				throw new PreprocessorException(context, "Garbage after #else directive");
 			}
 
 			// Check we're in a branch, and haven't already read an #else directive
-			PreprocessorBranch Branch;
-			if (!State.TryPopBranch(out Branch))
+			if (!_state.TryPopBranch(out PreprocessorBranch branch))
 			{
-				throw new PreprocessorException(Context, "#else directive without matching #if directive");
+				throw new PreprocessorException(context, "#else directive without matching #if directive");
 			}
-			if ((Branch & PreprocessorBranch.Complete) != 0)
+			if ((branch & PreprocessorBranch.Complete) != 0)
 			{
-				throw new PreprocessorException(Context, "Only one #else directive can appear in a conditional block");
+				throw new PreprocessorException(context, "Only one #else directive can appear in a conditional block");
 			}
 
 			// Check whether to take this branch, but only allow activating if the parent state is active.
-			Branch &= ~PreprocessorBranch.Active;
-			if(State.IsCurrentBranchActive() && !Branch.HasFlag(PreprocessorBranch.Taken))
+			branch &= ~PreprocessorBranch.Active;
+			if (_state.IsCurrentBranchActive() && !branch.HasFlag(PreprocessorBranch.Taken))
 			{
-				Branch |= PreprocessorBranch.Active | PreprocessorBranch.Taken;
+				branch |= PreprocessorBranch.Active | PreprocessorBranch.Taken;
 			}
-			State.PushBranch(Branch | PreprocessorBranch.Complete);
+			_state.PushBranch(branch | PreprocessorBranch.Complete);
 		}
 
 		/// <summary>
 		/// Parse an #endif directive
 		/// </summary>
-		/// <param name="Tokens">List of tokens in the directive</param>
-		/// <param name="Context">The context that this directive is being parsed in</param>
-		public void ParseEndifDirective(List<Token> Tokens, PreprocessorContext Context)
+		/// <param name="tokens">List of tokens in the directive</param>
+		/// <param name="context">The context that this directive is being parsed in</param>
+		[SuppressMessage("Style", "IDE0060:Remove unused parameter", Justification = "<Pending>")]
+		public void ParseEndifDirective(List<Token> tokens, PreprocessorContext context)
 		{
 			// Pop the branch off the stack
-			PreprocessorBranch Branch;
-			if(!State.TryPopBranch(out Branch))
+			if (!_state.TryPopBranch(out _))
 			{
-				throw new PreprocessorException(Context, "#endif directive without matching #if/#ifdef/#ifndef directive");
+				throw new PreprocessorException(context, "#endif directive without matching #if/#ifdef/#ifndef directive");
 			}
 		}
 
 		/// <summary>
 		/// Parse a #pragma directive
 		/// </summary>
-		/// <param name="Tokens">List of tokens in the directive</param>
-		/// <param name="Context">The context that this directive is being parsed in</param>
-		public void ParsePragmaDirective(List<Token> Tokens, PreprocessorContext Context)
+		/// <param name="tokens">List of tokens in the directive</param>
+		/// <param name="context">The context that this directive is being parsed in</param>
+		public void ParsePragmaDirective(List<Token> tokens, PreprocessorContext context)
 		{
-			if(State.IsCurrentBranchActive())
+			if (_state.IsCurrentBranchActive())
 			{
-				if(Tokens.Count == 1 && Tokens[0].Identifier == Identifiers.Once)
+				if (tokens.Count == 1 && tokens[0].Identifier == Identifiers.Once)
 				{
-					SourceFile SourceFile = GetCurrentSourceFile(Context)!;
-					PragmaOnceFiles.Add(SourceFile.File);
-					State.MarkPragmaOnce();
+					SourceFile sourceFile = GetCurrentSourceFile(context)!;
+					_pragmaOnceFiles.Add(sourceFile.File);
+					_state.MarkPragmaOnce();
 				}
 			}
 		}
@@ -942,32 +930,32 @@ namespace UnrealBuildTool
 		/// <summary>
 		/// Expand macros in the given sequence.
 		/// </summary>
-		/// <param name="InputTokens">Sequence of input tokens</param>
-		/// <param name="OutputTokens">List to receive the expanded tokens</param>
-		/// <param name="bIsConditional">Whether a conditional expression is being evaluated (and 'defined' expressions are valid)</param>
-		/// <param name="Context">The context that this directive is being parsed in</param>
-		public void ExpandMacros(IEnumerable<Token> InputTokens, List<Token> OutputTokens, bool bIsConditional, PreprocessorContext Context)
+		/// <param name="inputTokens">Sequence of input tokens</param>
+		/// <param name="outputTokens">List to receive the expanded tokens</param>
+		/// <param name="isConditional">Whether a conditional expression is being evaluated (and 'defined' expressions are valid)</param>
+		/// <param name="context">The context that this directive is being parsed in</param>
+		public void ExpandMacros(IEnumerable<Token> inputTokens, List<Token> outputTokens, bool isConditional, PreprocessorContext context)
 		{
-			List<PreprocessorMacro> IgnoreMacros = new List<PreprocessorMacro>();
-			ExpandMacrosRecursively(InputTokens, OutputTokens, bIsConditional, IgnoreMacros, Context);
+			List<PreprocessorMacro> ignoreMacros = new();
+			ExpandMacrosRecursively(inputTokens, outputTokens, isConditional, ignoreMacros, context);
 		}
 
 		/// <summary>
 		/// Expand macros in the given sequence, ignoring previously expanded macro names from a list.
 		/// </summary>
-		/// <param name="InputTokens">Sequence of input tokens</param>
-		/// <param name="OutputTokens">List to receive the expanded tokens</param>
-		/// <param name="bIsConditional">Whether a conditional expression is being evaluated (and 'defined' expressions are valid)</param>
-		/// <param name="IgnoreMacros">List of macros to ignore</param>
-		/// <param name="Context">The context that this directive is being parsed in</param>
-		void ExpandMacrosRecursively(IEnumerable<Token> InputTokens, List<Token> OutputTokens, bool bIsConditional, List<PreprocessorMacro> IgnoreMacros, PreprocessorContext Context)
+		/// <param name="inputTokens">Sequence of input tokens</param>
+		/// <param name="outputTokens">List to receive the expanded tokens</param>
+		/// <param name="isConditional">Whether a conditional expression is being evaluated (and 'defined' expressions are valid)</param>
+		/// <param name="ignoreMacros">List of macros to ignore</param>
+		/// <param name="context">The context that this directive is being parsed in</param>
+		void ExpandMacrosRecursively(IEnumerable<Token> inputTokens, List<Token> outputTokens, bool isConditional, List<PreprocessorMacro> ignoreMacros, PreprocessorContext context)
 		{
-			IEnumerator<Token> InputEnumerator = InputTokens.GetEnumerator();
-			if(InputEnumerator.MoveNext())
+			IEnumerator<Token> inputEnumerator = inputTokens.GetEnumerator();
+			if (inputEnumerator.MoveNext())
 			{
-				for(;;)
+				for (; ; )
 				{
-					if(!ReadExpandedToken(InputEnumerator, OutputTokens, bIsConditional, IgnoreMacros, Context))
+					if (!ReadExpandedToken(inputEnumerator, outputTokens, isConditional, ignoreMacros, context))
 					{
 						break;
 					}
@@ -978,182 +966,179 @@ namespace UnrealBuildTool
 		/// <summary>
 		/// Merges an optional leading space flag into the given token (recycling the original token if possible).
 		/// </summary>
-		/// <param name="Token">The token to merge a leading space into</param>
-		/// <param name="bHasLeadingSpace">The leading space flag</param>
+		/// <param name="token">The token to merge a leading space into</param>
+		/// <param name="hasLeadingSpace">The leading space flag</param>
 		/// <returns>New token with the leading space flag set, or the existing token</returns>
-		Token MergeLeadingSpace(Token Token, bool bHasLeadingSpace)
+		static Token MergeLeadingSpace(Token token, bool hasLeadingSpace)
 		{
-			Token Result = Token;
-			if(bHasLeadingSpace && !Result.HasLeadingSpace)
+			Token result = token;
+			if (hasLeadingSpace && !result.HasLeadingSpace)
 			{
-				Result = Result.AddFlags(TokenFlags.HasLeadingSpace);
+				result = result.AddFlags(TokenFlags.HasLeadingSpace);
 			}
-			return Result;
+			return result;
 		}
 
 		/// <summary>
 		/// Read a token from an enumerator and substitute it if it's a macro or 'defined' expression (reading more tokens as necessary to complete the expression).
 		/// </summary>
-		/// <param name="InputEnumerator">The enumerator of input tokens</param>
-		/// <param name="OutputTokens">List to receive the expanded tokens</param>
-		/// <param name="bIsConditional">Whether a conditional expression is being evaluated (and 'defined' expressions are valid)</param>
-		/// <param name="IgnoreMacros">List of macros to ignore</param>
-		/// <param name="Context">The context that this directive is being parsed in</param>
+		/// <param name="inputEnumerator">The enumerator of input tokens</param>
+		/// <param name="outputTokens">List to receive the expanded tokens</param>
+		/// <param name="isConditional">Whether a conditional expression is being evaluated (and 'defined' expressions are valid)</param>
+		/// <param name="ignoreMacros">List of macros to ignore</param>
+		/// <param name="context">The context that this directive is being parsed in</param>
 		/// <returns>Result from calling the enumerator's MoveNext() method</returns>
-		bool ReadExpandedToken(IEnumerator<Token> InputEnumerator, List<Token> OutputTokens, bool bIsConditional, List<PreprocessorMacro> IgnoreMacros, PreprocessorContext Context)
+		bool ReadExpandedToken(IEnumerator<Token> inputEnumerator, List<Token> outputTokens, bool isConditional, List<PreprocessorMacro> ignoreMacros, PreprocessorContext context)
 		{
 			// Capture the first token, then move to the next
-			OutputTokens.Add(InputEnumerator.Current);
-			bool bMoveNext = InputEnumerator.MoveNext();
+			outputTokens.Add(inputEnumerator.Current);
+			bool moveNext = inputEnumerator.MoveNext();
 
 			// If it's an identifier, try to resolve it as a macro
-			if (OutputTokens[OutputTokens.Count - 1].Identifier == Identifiers.Defined && bIsConditional)
+			if (outputTokens[^1].Identifier == Identifiers.Defined && isConditional)
 			{
 				// Remove the 'defined' keyword
-				OutputTokens.RemoveAt(OutputTokens.Count - 1);
+				outputTokens.RemoveAt(outputTokens.Count - 1);
 
 				// Make sure there's another token
-				if (!bMoveNext)
+				if (!moveNext)
 				{
-					throw new PreprocessorException(Context, "Invalid syntax for 'defined' expression");
+					throw new PreprocessorException(context, "Invalid syntax for 'defined' expression");
 				}
 
 				// Check for the form 'defined identifier'
-				Token NameToken;
-				if(InputEnumerator.Current.Type == TokenType.Identifier)
+				Token nameToken;
+				if (inputEnumerator.Current.Type == TokenType.Identifier)
 				{
-					NameToken = InputEnumerator.Current;
+					nameToken = inputEnumerator.Current;
 				}
 				else
 				{
 					// Otherwise assume the form 'defined ( identifier )'
-					if(InputEnumerator.Current.Type != TokenType.LeftParen || !InputEnumerator.MoveNext() || InputEnumerator.Current.Type != TokenType.Identifier)
+					if (inputEnumerator.Current.Type != TokenType.LeftParen || !inputEnumerator.MoveNext() || inputEnumerator.Current.Type != TokenType.Identifier)
 					{
-						throw new PreprocessorException(Context, "Invalid syntax for 'defined' expression");
+						throw new PreprocessorException(context, "Invalid syntax for 'defined' expression");
 					}
-					NameToken = InputEnumerator.Current;
-					if(!InputEnumerator.MoveNext() || InputEnumerator.Current.Type != TokenType.RightParen)
+					nameToken = inputEnumerator.Current;
+					if (!inputEnumerator.MoveNext() || inputEnumerator.Current.Type != TokenType.RightParen)
 					{
-						throw new PreprocessorException(Context, "Invalid syntax for 'defined' expression");
+						throw new PreprocessorException(context, "Invalid syntax for 'defined' expression");
 					}
 				}
 
 				// Insert a token for whether it's defined or not
-				OutputTokens.Add(new Token(TokenType.Number, TokenFlags.None, State.IsMacroDefined(NameToken.Identifier!)? OneLiteral : ZeroLiteral));
-				bMoveNext = InputEnumerator.MoveNext();
+				outputTokens.Add(new Token(TokenType.Number, TokenFlags.None, _state.IsMacroDefined(nameToken.Identifier!) ? s_oneLiteral : s_zeroLiteral));
+				moveNext = inputEnumerator.MoveNext();
 			}
 			else
 			{
 				// Repeatedly try to expand the last token into the list
-				while(OutputTokens[OutputTokens.Count - 1].Type == TokenType.Identifier && !OutputTokens[OutputTokens.Count - 1].Flags.HasFlag(TokenFlags.DisableExpansion))
+				while (outputTokens[^1].Type == TokenType.Identifier && !outputTokens[^1].Flags.HasFlag(TokenFlags.DisableExpansion))
 				{
 					// Try to get a macro for the current token
-					PreprocessorMacro? Macro;
-					if (!State.TryGetMacro(OutputTokens[OutputTokens.Count - 1].Identifier!, out Macro) || IgnoreMacros.Contains(Macro))
+					if (!_state.TryGetMacro(outputTokens[^1].Identifier!, out PreprocessorMacro? macro) || ignoreMacros.Contains(macro))
 					{
 						break;
 					}
-					if(Macro.IsFunctionMacro && (!bMoveNext || InputEnumerator.Current.Type != TokenType.LeftParen))
+					if (macro.IsFunctionMacro && (!moveNext || inputEnumerator.Current.Type != TokenType.LeftParen))
 					{
 						break;
 					}
 
 					// Remove the macro name from the output list
-					bool bHasLeadingSpace = OutputTokens[OutputTokens.Count - 1].HasLeadingSpace;
-					OutputTokens.RemoveAt(OutputTokens.Count - 1);
+					bool hasLeadingSpace = outputTokens[^1].HasLeadingSpace;
+					outputTokens.RemoveAt(outputTokens.Count - 1);
 
 					// Save the initial number of tokens in the list, so we can tell if it expanded
-					int NumTokens = OutputTokens.Count;
+					int numTokens = outputTokens.Count;
 
 					// If it's an object macro, expand it immediately into the output buffer
-					if(Macro.IsObjectMacro)
+					if (macro.IsObjectMacro)
 					{
 						// Expand the macro tokens into the output buffer
-						ExpandObjectMacro(Macro, OutputTokens, bIsConditional, IgnoreMacros, Context);
+						ExpandObjectMacro(macro, outputTokens, isConditional, ignoreMacros, context);
 					}
 					else
 					{
 						// Read balanced token for argument list
-						List<Token> ArgumentTokens = new List<Token>();
-						bMoveNext = ReadBalancedToken(InputEnumerator, ArgumentTokens, Context);
+						List<Token> argumentTokens = new();
+						moveNext = ReadBalancedToken(inputEnumerator, argumentTokens, context);
 
 						// Expand the macro tokens into the output buffer
-						ExpandFunctionMacro(Macro, ArgumentTokens, OutputTokens, bIsConditional, IgnoreMacros, Context);
+						ExpandFunctionMacro(macro, argumentTokens, outputTokens, isConditional, ignoreMacros, context);
 					}
 
 					// If the macro expanded to nothing, quit
-					if(OutputTokens.Count <= NumTokens)
+					if (outputTokens.Count <= numTokens)
 					{
 						break;
 					}
 
 					// Make sure the space is propagated to the expanded macro
-					OutputTokens[NumTokens] = MergeLeadingSpace(OutputTokens[NumTokens], bHasLeadingSpace);
+					outputTokens[numTokens] = MergeLeadingSpace(outputTokens[numTokens], hasLeadingSpace);
 
 					// Mark any tokens matching the macro name as not to be expanded again. This can happen with recursive object macros, eg. #define DWORD ::DWORD
-					for(int Idx = NumTokens; Idx < OutputTokens.Count; Idx++)
+					for (int idx = numTokens; idx < outputTokens.Count; idx++)
 					{
-						if(OutputTokens[Idx].Type == TokenType.Identifier && OutputTokens[Idx].Identifier == Macro.Name)
+						if (outputTokens[idx].Type == TokenType.Identifier && outputTokens[idx].Identifier == macro.Name)
 						{
-							OutputTokens[Idx] = OutputTokens[Idx].AddFlags(TokenFlags.DisableExpansion);
+							outputTokens[idx] = outputTokens[idx].AddFlags(TokenFlags.DisableExpansion);
 						}
 					}
 				}
 			}
-            return bMoveNext;
+			return moveNext;
 		}
 
 		/// <summary>
 		/// Gets a string for the __FILE__ macro
 		/// </summary>
-		/// <param name="Context">Context to scan to find the current file</param>
+		/// <param name="context">Context to scan to find the current file</param>
 		/// <returns>String representing the current context</returns>
-		string GetCurrentFileMacroValue(PreprocessorContext Context)
+		static string GetCurrentFileMacroValue(PreprocessorContext context)
 		{
-			SourceFile? SourceFile = GetCurrentSourceFile(Context);
-			if(SourceFile == null)
+			SourceFile? sourceFile = GetCurrentSourceFile(context);
+			if (sourceFile == null)
 			{
 				return "<unknown>";
 			}
 			else
 			{
-				return SourceFile.Location.FullName;
+				return sourceFile.Location.FullName;
 			}
 		}
 
 		/// <summary>
 		/// Gets a string for the current file
 		/// </summary>
-		/// <param name="Context">Context to scan to find the current file</param>
+		/// <param name="context">Context to scan to find the current file</param>
 		/// <returns>Current source file being parsed</returns>
-		SourceFile? GetCurrentSourceFile(PreprocessorContext Context)
+		static SourceFile? GetCurrentSourceFile(PreprocessorContext context)
 		{
-			SourceFile? SourceFile = null;
-			for(PreprocessorContext? OuterContext = Context; OuterContext != null; OuterContext = OuterContext.Outer)
+			SourceFile? sourceFile = null;
+			for (PreprocessorContext? outerContext = context; outerContext != null; outerContext = outerContext.Outer)
 			{
-				PreprocessorFileContext? OuterFileContext = OuterContext as PreprocessorFileContext;
-				if(OuterFileContext != null)
+				if (outerContext is PreprocessorFileContext outerFileContext)
 				{
-					SourceFile = OuterFileContext.SourceFile;
+					sourceFile = outerFileContext.SourceFile;
 					break;
 				}
 			}
-			return SourceFile;
+			return sourceFile;
 		}
 
 		/// <summary>
 		/// Gets the current line number
 		/// </summary>
-		/// <param name="Context">Context to scan to find the current file</param>
+		/// <param name="context">Context to scan to find the current file</param>
 		/// <returns>Line number in the first file encountered</returns>
-		int GetCurrentLine(PreprocessorContext Context)
+		static int GetCurrentLine(PreprocessorContext context)
 		{
-			for(PreprocessorContext? OuterContext = Context; OuterContext != null; OuterContext = OuterContext.Outer)
+			for (PreprocessorContext? outerContext = context; outerContext != null; outerContext = outerContext.Outer)
 			{
-				PreprocessorFileContext? OuterFileContext = OuterContext as PreprocessorFileContext;
-				if(OuterFileContext != null)
+				if (outerContext is PreprocessorFileContext outerFileContext)
 				{
-					return OuterFileContext.SourceFile.Markup[OuterFileContext.MarkupIdx].LineNumber;
+					return outerFileContext.SourceFile.Markup[outerFileContext.MarkupIdx].LineNumber;
 				}
 			}
 			return 0;
@@ -1162,46 +1147,46 @@ namespace UnrealBuildTool
 		/// <summary>
 		/// Expand an object macro
 		/// </summary>
-		/// <param name="Macro">The functional macro</param>
-		/// <param name="OutputTokens">The list to receive the output tokens</param>
-		/// <param name="bIsConditional">Whether the macro is being expanded in a conditional context, allowing use of the 'defined' keyword</param>
-		/// <param name="IgnoreMacros">List of macros currently being expanded, which should be ignored for recursion</param>
-		/// <param name="Context">The context that this directive is being parsed in</param>
-		void ExpandObjectMacro(PreprocessorMacro Macro, List<Token> OutputTokens, bool bIsConditional, List<PreprocessorMacro> IgnoreMacros, PreprocessorContext Context)
+		/// <param name="macro">The functional macro</param>
+		/// <param name="outputTokens">The list to receive the output tokens</param>
+		/// <param name="isConditional">Whether the macro is being expanded in a conditional context, allowing use of the 'defined' keyword</param>
+		/// <param name="ignoreMacros">List of macros currently being expanded, which should be ignored for recursion</param>
+		/// <param name="context">The context that this directive is being parsed in</param>
+		void ExpandObjectMacro(PreprocessorMacro macro, List<Token> outputTokens, bool isConditional, List<PreprocessorMacro> ignoreMacros, PreprocessorContext context)
 		{
 			// Special handling for the __LINE__ directive, since we need an updated line number for the current token
-			if(Macro.Name == Identifiers.__FILE__)
+			if (macro.Name == Identifiers.__FILE__)
 			{
-				Token Token = new Token(TokenType.String, TokenFlags.None, String.Format("\"{0}\"", GetCurrentFileMacroValue(Context).Replace("\\", "\\\\")));
-				OutputTokens.Add(Token);
+				Token token = new(TokenType.String, TokenFlags.None, String.Format("\"{0}\"", GetCurrentFileMacroValue(context).Replace("\\", "\\\\")));
+				outputTokens.Add(token);
 			}
-			else if(Macro.Name == Identifiers.__LINE__)
+			else if (macro.Name == Identifiers.__LINE__)
 			{
-				Token Token = new Token(TokenType.Number, TokenFlags.None, GetCurrentLine(Context).ToString());
-				OutputTokens.Add(Token);
+				Token token = new(TokenType.Number, TokenFlags.None, GetCurrentLine(context).ToString());
+				outputTokens.Add(token);
 			}
-			else if(Macro.Name == Identifiers.__COUNTER__)
+			else if (macro.Name == Identifiers.__COUNTER__)
 			{
-				Token Token = new Token(TokenType.Number, TokenFlags.None, (Counter++).ToString());
-				OutputTokens.Add(Token);
+				Token token = new(TokenType.Number, TokenFlags.None, (_counter++).ToString());
+				outputTokens.Add(token);
 			}
 			else
 			{
-				int OutputTokenCount = OutputTokens.Count;
+				int outputTokenCount = outputTokens.Count;
 
 				// Expand all the macros
-				IgnoreMacros.Add(Macro);
-				ExpandMacrosRecursively(Macro.Tokens, OutputTokens, bIsConditional, IgnoreMacros, Context);
-				IgnoreMacros.RemoveAt(IgnoreMacros.Count - 1);
+				ignoreMacros.Add(macro);
+				ExpandMacrosRecursively(macro.Tokens, outputTokens, isConditional, ignoreMacros, context);
+				ignoreMacros.RemoveAt(ignoreMacros.Count - 1);
 
 				// Concatenate any adjacent tokens
-				for(int Idx = OutputTokenCount + 1; Idx < OutputTokens.Count - 1; Idx++)
+				for (int idx = outputTokenCount + 1; idx < outputTokens.Count - 1; idx++)
 				{
-					if(OutputTokens[Idx].Type == TokenType.HashHash)
+					if (outputTokens[idx].Type == TokenType.HashHash)
 					{
-						OutputTokens[Idx - 1] = Token.Concatenate(OutputTokens[Idx - 1], OutputTokens[Idx + 1], Context);
-						OutputTokens.RemoveRange(Idx, 2);
-						Idx--;
+						outputTokens[idx - 1] = Token.Concatenate(outputTokens[idx - 1], outputTokens[idx + 1], context);
+						outputTokens.RemoveRange(idx, 2);
+						idx--;
 					}
 				}
 			}
@@ -1210,240 +1195,240 @@ namespace UnrealBuildTool
 		/// <summary>
 		/// Expand a function macro
 		/// </summary>
-		/// <param name="Macro">The functional macro</param>
-		/// <param name="ArgumentListTokens">Identifiers for each argument token</param>
-		/// <param name="OutputTokens">The list to receive the output tokens</param>
-		/// <param name="bIsConditional">Whether the macro is being expanded in a conditional context, allowing use of the 'defined' keyword</param>
-		/// <param name="IgnoreMacros">List of macros currently being expanded, which should be ignored for recursion</param>
-		/// <param name="Context">The context that this macro is being expanded</param>
-		void ExpandFunctionMacro(PreprocessorMacro Macro, List<Token> ArgumentListTokens, List<Token> OutputTokens, bool bIsConditional, List<PreprocessorMacro> IgnoreMacros, PreprocessorContext Context)
+		/// <param name="macro">The functional macro</param>
+		/// <param name="argumentListTokens">Identifiers for each argument token</param>
+		/// <param name="outputTokens">The list to receive the output tokens</param>
+		/// <param name="isConditional">Whether the macro is being expanded in a conditional context, allowing use of the 'defined' keyword</param>
+		/// <param name="ignoreMacros">List of macros currently being expanded, which should be ignored for recursion</param>
+		/// <param name="context">The context that this macro is being expanded</param>
+		void ExpandFunctionMacro(PreprocessorMacro macro, List<Token> argumentListTokens, List<Token> outputTokens, bool isConditional, List<PreprocessorMacro> ignoreMacros, PreprocessorContext context)
 		{
 			// Replace any newlines with spaces, and merge them with the following token
-			for(int Idx = 0; Idx < ArgumentListTokens.Count; Idx++)
+			for (int idx = 0; idx < argumentListTokens.Count; idx++)
 			{
-				if(ArgumentListTokens[Idx].Type == TokenType.Newline)
+				if (argumentListTokens[idx].Type == TokenType.Newline)
 				{
-					if(Idx + 1 < ArgumentListTokens.Count)
+					if (idx + 1 < argumentListTokens.Count)
 					{
-						ArgumentListTokens[Idx + 1] = MergeLeadingSpace(ArgumentListTokens[Idx + 1], true);
+						argumentListTokens[idx + 1] = MergeLeadingSpace(argumentListTokens[idx + 1], true);
 					}
-					ArgumentListTokens.RemoveAt(Idx--);
+					argumentListTokens.RemoveAt(idx--);
 				}
 			}
 
 			// Split the arguments out into separate lists
-			List<List<Token>> Arguments = new List<List<Token>>();
-			if(ArgumentListTokens.Count > 2)
+			List<List<Token>> arguments = new();
+			if (argumentListTokens.Count > 2)
 			{
-				for(int Idx = 1;;Idx++)
+				for (int idx = 1; ; idx++)
 				{
-					if (!Macro.HasVariableArgumentList || Arguments.Count < Macro.Parameters!.Count)
+					if (!macro.HasVariableArgumentList || arguments.Count < macro.Parameters!.Count)
 					{
-						Arguments.Add(new List<Token>());
+						arguments.Add(new List<Token>());
 					}
 
-					List<Token> Argument = Arguments[Arguments.Count - 1];
+					List<Token> argument = arguments[^1];
 
-					int InitialIdx = Idx;
-					while (Idx < ArgumentListTokens.Count - 1 && ArgumentListTokens[Idx].Type != TokenType.Comma)
+					int initialIdx = idx;
+					while (idx < argumentListTokens.Count - 1 && argumentListTokens[idx].Type != TokenType.Comma)
 					{
-						if (!ReadBalancedToken(ArgumentListTokens, ref Idx, Argument))
+						if (!ReadBalancedToken(argumentListTokens, ref idx, argument))
 						{
-							throw new PreprocessorException(Context, "Invalid argument");
+							throw new PreprocessorException(context, "Invalid argument");
 						}
 					}
 
-					if(Argument.Count > 0 && Arguments[Arguments.Count - 1][0].HasLeadingSpace)
+					if (argument.Count > 0 && arguments[^1][0].HasLeadingSpace)
 					{
-						Argument[0] = Argument[0].RemoveFlags(TokenFlags.HasLeadingSpace);
+						argument[0] = argument[0].RemoveFlags(TokenFlags.HasLeadingSpace);
 					}
 
-					bool bHasLeadingSpace = false;
-					for(int TokenIdx = 0; TokenIdx < Argument.Count; TokenIdx++)
+					bool hasLeadingSpace = false;
+					for (int tokenIdx = 0; tokenIdx < argument.Count; tokenIdx++)
 					{
-						if(Argument[TokenIdx].Text.Length == 0)
+						if (argument[tokenIdx].Text.Length == 0)
 						{
-							bHasLeadingSpace |= Argument[TokenIdx].HasLeadingSpace;
-							Argument.RemoveAt(TokenIdx--);
+							hasLeadingSpace |= argument[tokenIdx].HasLeadingSpace;
+							argument.RemoveAt(tokenIdx--);
 						}
 						else
 						{
-							Argument[TokenIdx] = MergeLeadingSpace(Argument[TokenIdx], bHasLeadingSpace);
-							bHasLeadingSpace = false;
+							argument[tokenIdx] = MergeLeadingSpace(argument[tokenIdx], hasLeadingSpace);
+							hasLeadingSpace = false;
 						}
 					}
 
-					if (Argument.Count == 0)
+					if (argument.Count == 0)
 					{
-						Argument.Add(new Token(TokenType.Placemarker, TokenFlags.None));
-						Argument.Add(new Token(TokenType.Placemarker, bHasLeadingSpace? TokenFlags.HasLeadingSpace : TokenFlags.None));
+						argument.Add(new Token(TokenType.Placemarker, TokenFlags.None));
+						argument.Add(new Token(TokenType.Placemarker, hasLeadingSpace ? TokenFlags.HasLeadingSpace : TokenFlags.None));
 					}
 
-					if(Idx == ArgumentListTokens.Count - 1)
+					if (idx == argumentListTokens.Count - 1)
 					{
 						break;
 					}
 
-					if (ArgumentListTokens[Idx].Type != TokenType.Comma)
+					if (argumentListTokens[idx].Type != TokenType.Comma)
 					{
-						throw new PreprocessorException(Context, "Expected ',' between arguments");
+						throw new PreprocessorException(context, "Expected ',' between arguments");
 					}
 
-					if (Macro.HasVariableArgumentList && Arguments.Count == Macro.Parameters!.Count && Idx < ArgumentListTokens.Count - 1)
+					if (macro.HasVariableArgumentList && arguments.Count == macro.Parameters!.Count && idx < argumentListTokens.Count - 1)
 					{
-						Arguments[Arguments.Count - 1].Add(ArgumentListTokens[Idx]);
+						arguments[^1].Add(argumentListTokens[idx]);
 					}
 				}
 			}
 
 			// Add an empty variable argument if one was not specified
-			if (Macro.HasVariableArgumentList && Arguments.Count == Macro.Parameters!.Count - 1)
+			if (macro.HasVariableArgumentList && arguments.Count == macro.Parameters!.Count - 1)
 			{
-				Arguments.Add(new List<Token> { new Token(TokenType.Placemarker, TokenFlags.None) });
+				arguments.Add(new List<Token> { new Token(TokenType.Placemarker, TokenFlags.None) });
 			}
 
 			// Validate the argument list
-			if (Arguments.Count != Macro.Parameters!.Count)
+			if (arguments.Count != macro.Parameters!.Count)
 			{
-				throw new PreprocessorException(Context, "Incorrect number of arguments to macro");
+				throw new PreprocessorException(context, "Incorrect number of arguments to macro");
 			}
 
 			// Expand each one of the arguments
-			List<List<Token>> ExpandedArguments = new List<List<Token>>();
-			for (int Idx = 0; Idx < Arguments.Count; Idx++)
+			List<List<Token>> expandedArguments = new();
+			for (int idx = 0; idx < arguments.Count; idx++)
 			{
-				List<Token> NewArguments = new List<Token>();
-				ExpandMacrosRecursively(Arguments[Idx], NewArguments, bIsConditional, IgnoreMacros, Context);
-				ExpandedArguments.Add(NewArguments);
+				List<Token> newArguments = new();
+				ExpandMacrosRecursively(arguments[idx], newArguments, isConditional, ignoreMacros, context);
+				expandedArguments.Add(newArguments);
 			}
 
 			// Substitute all the argument tokens
-			List<Token> ExpandedTokens = new List<Token>();
-			for(int Idx = 0; Idx < Macro.Tokens.Count; Idx++)
+			List<Token> expandedTokens = new();
+			for (int idx = 0; idx < macro.Tokens.Count; idx++)
 			{
-				Token Token = Macro.Tokens[Idx];
-				if(Token.Type == TokenType.Hash && Idx + 1 < Macro.Tokens.Count)
+				Token token = macro.Tokens[idx];
+				if (token.Type == TokenType.Hash && idx + 1 < macro.Tokens.Count)
 				{
 					// Stringizing operator
-					int ParamIdx = Macro.FindParameterIndex(Macro.Tokens[++Idx].Identifier!);
-					if (ParamIdx == -1)
+					int paramIdx = macro.FindParameterIndex(macro.Tokens[++idx].Identifier!);
+					if (paramIdx == -1)
 					{
-						throw new PreprocessorException(Context, "{0} is not an argument name", Macro.Tokens[Idx].Text);
+						throw new PreprocessorException(context, "{0} is not an argument name", macro.Tokens[idx].Text);
 					}
-					ExpandedTokens.Add(new Token(TokenType.String, Token.Flags & TokenFlags.HasLeadingSpace, String.Format("\"{0}\"", Token.Format(Arguments[ParamIdx]).Replace("\\", "\\\\").Replace("\"", "\\\""))));
+					expandedTokens.Add(new Token(TokenType.String, token.Flags & TokenFlags.HasLeadingSpace, String.Format("\"{0}\"", Token.Format(arguments[paramIdx]).Replace("\\", "\\\\").Replace("\"", "\\\""))));
 				}
-				else if(Macro.HasVariableArgumentList && Idx + 2 < Macro.Tokens.Count && Token.Type == TokenType.Comma && Macro.Tokens[Idx + 1].Type == TokenType.HashHash && Macro.Tokens[Idx + 2].Identifier == Identifiers.__VA_ARGS__)
+				else if (macro.HasVariableArgumentList && idx + 2 < macro.Tokens.Count && token.Type == TokenType.Comma && macro.Tokens[idx + 1].Type == TokenType.HashHash && macro.Tokens[idx + 2].Identifier == Identifiers.__VA_ARGS__)
 				{
 					// Special MSVC/GCC extension: ', ## __VA_ARGS__' removes the comma if __VA_ARGS__ is empty. MSVC seems to format the result with a forced space.
-					List<Token> ExpandedArgument = ExpandedArguments[ExpandedArguments.Count - 1];
-					if(ExpandedArgument.Any(x => x.Text.Length > 0))
+					List<Token> expandedArgument = expandedArguments[^1];
+					if (expandedArgument.Any(x => x.Text.Length > 0))
 					{
-						ExpandedTokens.Add(Token);
-						AppendTokensWithWhitespace(ExpandedTokens, ExpandedArgument, false);
-						Idx += 2;
+						expandedTokens.Add(token);
+						AppendTokensWithWhitespace(expandedTokens, expandedArgument, false);
+						idx += 2;
 					}
 					else
 					{
-						ExpandedTokens.Add(new Token(TokenType.Placemarker, Token.Flags & TokenFlags.HasLeadingSpace));
-						ExpandedTokens.Add(new Token(TokenType.Placemarker, TokenFlags.HasLeadingSpace));
-						Idx += 2;
+						expandedTokens.Add(new Token(TokenType.Placemarker, token.Flags & TokenFlags.HasLeadingSpace));
+						expandedTokens.Add(new Token(TokenType.Placemarker, TokenFlags.HasLeadingSpace));
+						idx += 2;
 					}
 				}
-				else if (Token.Type == TokenType.Identifier)
+				else if (token.Type == TokenType.Identifier)
 				{
 					// Expand a parameter
-					int ParamIdx = Macro.FindParameterIndex(Token.Identifier!);
-					if(ParamIdx == -1)
+					int paramIdx = macro.FindParameterIndex(token.Identifier!);
+					if (paramIdx == -1)
 					{
-						ExpandedTokens.Add(Token);
+						expandedTokens.Add(token);
 					}
-					else if(Idx > 0 && Macro.Tokens[Idx - 1].Type == TokenType.HashHash)
+					else if (idx > 0 && macro.Tokens[idx - 1].Type == TokenType.HashHash)
 					{
-						AppendTokensWithWhitespace(ExpandedTokens, Arguments[ParamIdx], Token.HasLeadingSpace);
+						AppendTokensWithWhitespace(expandedTokens, arguments[paramIdx], token.HasLeadingSpace);
 					}
-					else if(Idx + 1 < Macro.Tokens.Count && Macro.Tokens[Idx + 1].Type == TokenType.HashHash)
+					else if (idx + 1 < macro.Tokens.Count && macro.Tokens[idx + 1].Type == TokenType.HashHash)
 					{
-						AppendTokensWithWhitespace(ExpandedTokens, Arguments[ParamIdx], Token.HasLeadingSpace);
+						AppendTokensWithWhitespace(expandedTokens, arguments[paramIdx], token.HasLeadingSpace);
 					}
 					else
 					{
-						AppendTokensWithWhitespace(ExpandedTokens, ExpandedArguments[ParamIdx], Token.HasLeadingSpace);
+						AppendTokensWithWhitespace(expandedTokens, expandedArguments[paramIdx], token.HasLeadingSpace);
 					}
 				}
 				else
 				{
-					ExpandedTokens.Add(Token);
+					expandedTokens.Add(token);
 				}
 			}
 
 			// Concatenate adjacent tokens
-			for (int Idx = 1; Idx < ExpandedTokens.Count - 1; Idx++)
+			for (int idx = 1; idx < expandedTokens.Count - 1; idx++)
 			{
-				if(ExpandedTokens[Idx].Type == TokenType.HashHash)
+				if (expandedTokens[idx].Type == TokenType.HashHash)
 				{
-					Token ConcatenatedToken = Token.Concatenate(ExpandedTokens[Idx - 1], ExpandedTokens[Idx + 1], Context);
-					ExpandedTokens.RemoveRange(Idx, 2);
-					ExpandedTokens[--Idx] = ConcatenatedToken;
+					Token concatenatedToken = Token.Concatenate(expandedTokens[idx - 1], expandedTokens[idx + 1], context);
+					expandedTokens.RemoveRange(idx, 2);
+					expandedTokens[--idx] = concatenatedToken;
 				}
 			}
 
 			// Finally, return the expansion of this
-			IgnoreMacros.Add(Macro);
-			ExpandMacrosRecursively(ExpandedTokens, OutputTokens, bIsConditional, IgnoreMacros, Context);
-			IgnoreMacros.RemoveAt(IgnoreMacros.Count - 1);
+			ignoreMacros.Add(macro);
+			ExpandMacrosRecursively(expandedTokens, outputTokens, isConditional, ignoreMacros, context);
+			ignoreMacros.RemoveAt(ignoreMacros.Count - 1);
 		}
 
 		/// <summary>
 		/// Appends a list of tokens to another list, setting the leading whitespace flag to the given value
 		/// </summary>
-		/// <param name="OutputTokens">List to receive the appended tokens</param>
-		/// <param name="InputTokens">List of tokens to append</param>
-		/// <param name="bHasLeadingSpace">Whether there is space before the first token</param>
-		void AppendTokensWithWhitespace(List<Token> OutputTokens, List<Token> InputTokens, bool bHasLeadingSpace)
+		/// <param name="outputTokens">List to receive the appended tokens</param>
+		/// <param name="inputTokens">List of tokens to append</param>
+		/// <param name="hasLeadingSpace">Whether there is space before the first token</param>
+		static void AppendTokensWithWhitespace(List<Token> outputTokens, List<Token> inputTokens, bool hasLeadingSpace)
 		{
-			if(InputTokens.Count > 0)
+			if (inputTokens.Count > 0)
 			{
-				OutputTokens.Add(MergeLeadingSpace(InputTokens[0], bHasLeadingSpace));
-				OutputTokens.AddRange(InputTokens.Skip(1));
+				outputTokens.Add(MergeLeadingSpace(inputTokens[0], hasLeadingSpace));
+				outputTokens.AddRange(inputTokens.Skip(1));
 			}
 		}
 
 		/// <summary>
 		/// Copies a single token from one list of tokens to another, or if it's an opening parenthesis, the entire subexpression until the closing parenthesis.
 		/// </summary>
-		/// <param name="InputTokens">The input token list</param>
-		/// <param name="InputIdx">First token index in the input token list. Set to the last uncopied token index on return.</param>
-		/// <param name="OutputTokens">List to recieve the output tokens</param>
+		/// <param name="inputTokens">The input token list</param>
+		/// <param name="inputIdx">First token index in the input token list. Set to the last uncopied token index on return.</param>
+		/// <param name="outputTokens">List to recieve the output tokens</param>
 		/// <returns>True if a balanced expression was read, or false if the end of the list was encountered before finding a matching token</returns>
-		bool ReadBalancedToken(List<Token> InputTokens, ref int InputIdx, List<Token> OutputTokens)
+		bool ReadBalancedToken(List<Token> inputTokens, ref int inputIdx, List<Token> outputTokens)
 		{
 			// Copy a single token to the output list
-			Token Token = InputTokens[InputIdx++];
-			OutputTokens.Add(Token);
+			Token token = inputTokens[inputIdx++];
+			outputTokens.Add(token);
 
 			// If it was the start of a subexpression, copy until the closing parenthesis
-			if (Token.Type == TokenType.LeftParen)
+			if (token.Type == TokenType.LeftParen)
 			{
 				// Copy the contents of the subexpression
-				for (;;)
+				for (; ; )
 				{
-					if (InputIdx == InputTokens.Count)
+					if (inputIdx == inputTokens.Count)
 					{
 						return false;
 					}
-					if (InputTokens[InputIdx].Type == TokenType.RightParen)
+					if (inputTokens[inputIdx].Type == TokenType.RightParen)
 					{
 						break;
 					}
-					if (!ReadBalancedToken(InputTokens, ref InputIdx, OutputTokens))
+					if (!ReadBalancedToken(inputTokens, ref inputIdx, outputTokens))
 					{
 						return false;
 					}
 				}
 
 				// Copy the closing parenthesis
-				Token = InputTokens[InputIdx++];
-				OutputTokens.Add(Token);
+				token = inputTokens[inputIdx++];
+				outputTokens.Add(token);
 			}
 			return true;
 		}
@@ -1451,37 +1436,37 @@ namespace UnrealBuildTool
 		/// <summary>
 		/// Copies a single token from one list of tokens to another, or if it's an opening parenthesis, the entire subexpression until the closing parenthesis.
 		/// </summary>
-		/// <param name="InputEnumerator">The input token list</param>
-		/// <param name="OutputTokens">List to recieve the output tokens</param>
-		/// <param name="Context">The context that the parser is in</param>
+		/// <param name="inputEnumerator">The input token list</param>
+		/// <param name="outputTokens">List to recieve the output tokens</param>
+		/// <param name="context">The context that the parser is in</param>
 		/// <returns>True if a balanced expression was read, or false if the end of the list was encountered before finding a matching token</returns>
-		bool ReadBalancedToken(IEnumerator<Token> InputEnumerator, List<Token> OutputTokens, PreprocessorContext Context)
+		bool ReadBalancedToken(IEnumerator<Token> inputEnumerator, List<Token> outputTokens, PreprocessorContext context)
 		{
 			// Copy a single token to the output list
-			Token Token = InputEnumerator.Current;
-			bool bMoveNext = InputEnumerator.MoveNext();
-			OutputTokens.Add(Token);
+			Token token = inputEnumerator.Current;
+			bool moveNext = inputEnumerator.MoveNext();
+			outputTokens.Add(token);
 
 			// If it was the start of a subexpression, copy until the closing parenthesis
-			if (Token.Type == TokenType.LeftParen)
+			if (token.Type == TokenType.LeftParen)
 			{
 				// Copy the contents of the subexpression
-				for (;;)
+				for (; ; )
 				{
-					if (!bMoveNext)
+					if (!moveNext)
 					{
-						throw new PreprocessorException(Context, "Unbalanced token sequence");
+						throw new PreprocessorException(context, "Unbalanced token sequence");
 					}
-					if (InputEnumerator.Current.Type == TokenType.RightParen)
+					if (inputEnumerator.Current.Type == TokenType.RightParen)
 					{
-						OutputTokens.Add(InputEnumerator.Current);
-						bMoveNext = InputEnumerator.MoveNext();
+						outputTokens.Add(inputEnumerator.Current);
+						moveNext = inputEnumerator.MoveNext();
 						break;
 					}
-					bMoveNext = ReadBalancedToken(InputEnumerator, OutputTokens, Context);
+					moveNext = ReadBalancedToken(inputEnumerator, outputTokens, context);
 				}
 			}
-			return bMoveNext;
+			return moveNext;
 		}
 	}
 }
diff --git a/Engine/Source/Programs/UnrealBuildTool/Preprocessor/PreprocessorContext.cs b/Engine/Source/Programs/UnrealBuildTool/Preprocessor/PreprocessorContext.cs
--- a/Engine/Source/Programs/UnrealBuildTool/Preprocessor/PreprocessorContext.cs
+++ b/Engine/Source/Programs/UnrealBuildTool/Preprocessor/PreprocessorContext.cs
@@ -1,10 +1,6 @@
 // Copyright Epic Games, Inc. All Rights Reserved.
 
 using System;
-using System.Collections.Generic;
-using System.Linq;
-using System.Text;
-using System.Threading.Tasks;
 using UnrealBuildBase;
 
 namespace UnrealBuildTool
@@ -22,10 +18,10 @@ namespace UnrealBuildTool
 		/// <summary>
 		/// Constructor
 		/// </summary>
-		/// <param name="Outer">The outer context</param>
-		public PreprocessorContext(PreprocessorContext? Outer)
+		/// <param name="outer">The outer context</param>
+		public PreprocessorContext(PreprocessorContext? outer)
 		{
-			this.Outer = Outer;
+			Outer = outer;
 		}
 	}
 
@@ -60,35 +56,35 @@ namespace UnrealBuildTool
 		/// <summary>
 		/// The source file being read
 		/// </summary>
-		public SourceFile SourceFile;
+		public readonly SourceFile SourceFile;
 
 		/// <summary>
 		/// The directory containing this file. When searching for included files, MSVC will check this directory.
 		/// </summary>
-		public DirectoryItem Directory;
+		public readonly DirectoryItem Directory;
 
 		/// <summary>
 		/// Index of the current markup object being processed
 		/// </summary>
-		public int MarkupIdx;
+		public int MarkupIdx { get; set; }
 
 		/// <summary>
 		/// Index of the next fragment to be read
 		/// </summary>
-		public int FragmentIdx;
+		public int FragmentIdx { get; set; }
 
 		/// <summary>
 		/// Constructor
 		/// </summary>
-		/// <param name="SourceFile">The source file being parsed</param>
-		/// <param name="Outer">The outer context</param>
-		public PreprocessorFileContext(SourceFile SourceFile, PreprocessorContext? Outer)
-			: base(Outer)
+		/// <param name="sourceFile">The source file being parsed</param>
+		/// <param name="outer">The outer context</param>
+		public PreprocessorFileContext(SourceFile sourceFile, PreprocessorContext? outer)
+			: base(outer)
 		{
-			this.SourceFile = SourceFile;
-			this.Directory = DirectoryItem.GetItemByDirectoryReference(SourceFile.Location.Directory);
-			this.MarkupIdx = 0;
-			this.FragmentIdx = 0;
+			SourceFile = sourceFile;
+			Directory = DirectoryItem.GetItemByDirectoryReference(sourceFile.Location.Directory);
+			MarkupIdx = 0;
+			FragmentIdx = 0;
 		}
 
 		/// <summary>
diff --git a/Engine/Source/Programs/UnrealBuildTool/Preprocessor/PreprocessorExpression.cs b/Engine/Source/Programs/UnrealBuildTool/Preprocessor/PreprocessorExpression.cs
--- a/Engine/Source/Programs/UnrealBuildTool/Preprocessor/PreprocessorExpression.cs
+++ b/Engine/Source/Programs/UnrealBuildTool/Preprocessor/PreprocessorExpression.cs
@@ -2,9 +2,7 @@
 
 using System;
 using System.Collections.Generic;
-using System.Linq;
 using System.Text;
-using System.Threading.Tasks;
 
 namespace UnrealBuildTool
 {
@@ -16,423 +14,423 @@ namespace UnrealBuildTool
 		/// <summary>
 		/// Precedence table for binary operators (indexed by TokenType). Lower values indicate higher precedence.
 		/// </summary>
-		static byte[] Precedence;
+		static readonly byte[] s_precedence;
 
 		/// <summary>
 		/// Static constructor. Initializes the Precedence table.
 		/// </summary>
 		static PreprocessorExpression()
 		{
-			Precedence = new byte[(int)TokenType.Max];
-			Precedence[(int)TokenType.Multiply] = 10;
-			Precedence[(int)TokenType.Divide] = 10;
-			Precedence[(int)TokenType.Modulo] = 10;
-			Precedence[(int)TokenType.Plus] = 9;
-			Precedence[(int)TokenType.Minus] = 9;
-			Precedence[(int)TokenType.LeftShift] = 8;
-			Precedence[(int)TokenType.RightShift] = 8;
-			Precedence[(int)TokenType.CompareLess] = 7;
-			Precedence[(int)TokenType.CompareLessOrEqual] = 7;
-			Precedence[(int)TokenType.CompareGreater] = 7;
-			Precedence[(int)TokenType.CompareGreaterOrEqual] = 7;
-			Precedence[(int)TokenType.CompareEqual] = 6;
-			Precedence[(int)TokenType.CompareNotEqual] = 6;
-			Precedence[(int)TokenType.BitwiseAnd] = 5;
-			Precedence[(int)TokenType.BitwiseXor] = 4;
-			Precedence[(int)TokenType.BitwiseOr] = 3;
-			Precedence[(int)TokenType.LogicalAnd] = 2;
-			Precedence[(int)TokenType.LogicalOr] = 1;
+			s_precedence = new byte[(int)TokenType.Max];
+			s_precedence[(int)TokenType.Multiply] = 10;
+			s_precedence[(int)TokenType.Divide] = 10;
+			s_precedence[(int)TokenType.Modulo] = 10;
+			s_precedence[(int)TokenType.Plus] = 9;
+			s_precedence[(int)TokenType.Minus] = 9;
+			s_precedence[(int)TokenType.LeftShift] = 8;
+			s_precedence[(int)TokenType.RightShift] = 8;
+			s_precedence[(int)TokenType.CompareLess] = 7;
+			s_precedence[(int)TokenType.CompareLessOrEqual] = 7;
+			s_precedence[(int)TokenType.CompareGreater] = 7;
+			s_precedence[(int)TokenType.CompareGreaterOrEqual] = 7;
+			s_precedence[(int)TokenType.CompareEqual] = 6;
+			s_precedence[(int)TokenType.CompareNotEqual] = 6;
+			s_precedence[(int)TokenType.BitwiseAnd] = 5;
+			s_precedence[(int)TokenType.BitwiseXor] = 4;
+			s_precedence[(int)TokenType.BitwiseOr] = 3;
+			s_precedence[(int)TokenType.LogicalAnd] = 2;
+			s_precedence[(int)TokenType.LogicalOr] = 1;
 		}
 
 		/// <summary>
 		/// Evaluate a preprocessor expression
 		/// </summary>
-		/// <param name="Context">The current preprocessor context</param>
-		/// <param name="Tokens">List of tokens in the expression</param>
+		/// <param name="context">The current preprocessor context</param>
+		/// <param name="tokens">List of tokens in the expression</param>
 		/// <returns>Value of the expression</returns>
-		public static long Evaluate(PreprocessorContext Context, List<Token> Tokens)
+		public static long Evaluate(PreprocessorContext context, List<Token> tokens)
 		{
-			int Idx = 0;
-			long Result = EvaluateTernary(Context, Tokens, ref Idx);
-			if (Idx != Tokens.Count)
+			int idx = 0;
+			long result = EvaluateTernary(context, tokens, ref idx);
+			if (idx != tokens.Count)
 			{
-				throw new PreprocessorException(Context, $"Garbage after end of expression: {Token.Format(Tokens)}");
+				throw new PreprocessorException(context, $"Garbage after end of expression: {Token.Format(tokens)}");
 			}
-			return Result;
+			return result;
 		}
 
 		/// <summary>
 		/// Evaluates a ternary expression (a? b :c)
 		/// </summary>
-		/// <param name="Context">The current preprocessor context</param>
-		/// <param name="Tokens">List of tokens in the expression</param>
-		/// <param name="Idx">Index into the token list of the next token to read</param>
+		/// <param name="context">The current preprocessor context</param>
+		/// <param name="tokens">List of tokens in the expression</param>
+		/// <param name="idx">Index into the token list of the next token to read</param>
 		/// <returns>Value of the expression</returns>
-		static long EvaluateTernary(PreprocessorContext Context, List<Token> Tokens, ref int Idx)
+		static long EvaluateTernary(PreprocessorContext context, List<Token> tokens, ref int idx)
 		{
-			long Result = EvaluateUnary(Context, Tokens, ref Idx);
-			if(Idx < Tokens.Count && Precedence[(int)Tokens[Idx].Type] != 0)
+			long result = EvaluateUnary(context, tokens, ref idx);
+			if (idx < tokens.Count && s_precedence[(int)tokens[idx].Type] != 0)
 			{
-				Result = EvaluateBinary(Context, Tokens, ref Idx, Result, 1);
+				result = EvaluateBinary(context, tokens, ref idx, result, 1);
 			}
-			if(Idx < Tokens.Count && Tokens[Idx].Type == TokenType.QuestionMark)
+			if (idx < tokens.Count && tokens[idx].Type == TokenType.QuestionMark)
 			{
 				// Read the left expression
-				Idx++;
-				long Lhs = EvaluateTernary(Context, Tokens, ref Idx);
+				idx++;
+				long lhs = EvaluateTernary(context, tokens, ref idx);
 
 				// Check for the colon in the middle
-				if(Tokens[Idx].Type != TokenType.Colon)
+				if (tokens[idx].Type != TokenType.Colon)
 				{
-					throw new PreprocessorException(Context, $"Expected colon for conditional operator, found {Tokens[Idx].Text}");
+					throw new PreprocessorException(context, $"Expected colon for conditional operator, found {tokens[idx].Text}");
 				}
 
 				// Read the right expression
-				Idx++;
-				long Rhs = EvaluateTernary(Context, Tokens, ref Idx);
+				idx++;
+				long rhs = EvaluateTernary(context, tokens, ref idx);
 
 				// Evaluate it
-				Result = (Result != 0) ? Lhs : Rhs;
+				result = (result != 0) ? lhs : rhs;
 			}
-			return Result;
+			return result;
 		}
 
 		/// <summary>
 		/// Recursive path for evaluating a sequence of binary operators, given a known LHS. Implementation of the shunting yard algorithm.
 		/// </summary>
-		/// <param name="Context">The current preprocessor context</param>
-		/// <param name="Tokens">List of tokens in the expression</param>
-		/// <param name="Idx">Index into the token list of the next token to read</param>
-		/// <param name="Lhs">The LHS value for the binary expression</param>
-		/// <param name="MinPrecedence">Minimum precedence value to consume</param>
+		/// <param name="context">The current preprocessor context</param>
+		/// <param name="tokens">List of tokens in the expression</param>
+		/// <param name="idx">Index into the token list of the next token to read</param>
+		/// <param name="lhs">The LHS value for the binary expression</param>
+		/// <param name="minPrecedence">Minimum precedence value to consume</param>
 		/// <returns>Value of the expression</returns>
-		static long EvaluateBinary(PreprocessorContext Context, List<Token> Tokens, ref int Idx, long Lhs, byte MinPrecedence)
+		static long EvaluateBinary(PreprocessorContext context, List<Token> tokens, ref int idx, long lhs, byte minPrecedence)
 		{
-			while(Idx < Tokens.Count)
+			while (idx < tokens.Count)
 			{
 				// Read the next operator token
-				Token Operator = Tokens[Idx];
+				Token tokenOp = tokens[idx];
 
 				// Check if this operator binds more tightly than the outer expression, and exit if it does.
-				int OperatorPrecedence = Precedence[(int)Operator.Type];
-				if(OperatorPrecedence < MinPrecedence)
+				int operatorPrecedence = s_precedence[(int)tokenOp.Type];
+				if (operatorPrecedence < minPrecedence)
 				{
 					break;
 				}
 
 				// Move to the next token
-				Idx++;
+				idx++;
 
 				// Evaluate the immediate RHS value, then recursively evaluate any subsequent binary operators that bind more tightly to it than this.
-				long Rhs = EvaluateUnary(Context, Tokens, ref Idx);
-				while(Idx < Tokens.Count)
+				long rhs = EvaluateUnary(context, tokens, ref idx);
+				while (idx < tokens.Count)
 				{
-					byte NextOperatorPrecedence = Precedence[(int)Tokens[Idx].Type];
-					if(NextOperatorPrecedence <= OperatorPrecedence)
+					byte nextOperatorPrecedence = s_precedence[(int)tokens[idx].Type];
+					if (nextOperatorPrecedence <= operatorPrecedence)
 					{
 						break;
 					}
-					Rhs = EvaluateBinary(Context, Tokens, ref Idx, Rhs, NextOperatorPrecedence);
+					rhs = EvaluateBinary(context, tokens, ref idx, rhs, nextOperatorPrecedence);
 				}
 
 				// Evalute the result of applying the operator to the LHS and RHS values
-				switch(Operator.Type)
+				switch (tokenOp.Type)
 				{
 					case TokenType.LogicalOr:
-						Lhs = ((Lhs != 0) || (Rhs != 0))? 1 : 0;
+						lhs = ((lhs != 0) || (rhs != 0)) ? 1 : 0;
 						break;
 					case TokenType.LogicalAnd:
-						Lhs = ((Lhs != 0) && (Rhs != 0))? 1 : 0;
+						lhs = ((lhs != 0) && (rhs != 0)) ? 1 : 0;
 						break;
 					case TokenType.BitwiseOr:
-						Lhs |= Rhs;
+						lhs |= rhs;
 						break;
 					case TokenType.BitwiseXor:
-						Lhs ^= Rhs;
+						lhs ^= rhs;
 						break;
 					case TokenType.BitwiseAnd:
-						Lhs &= Rhs;
+						lhs &= rhs;
 						break;
 					case TokenType.CompareEqual:
-						Lhs = (Lhs == Rhs)? 1 : 0;
+						lhs = (lhs == rhs) ? 1 : 0;
 						break;
 					case TokenType.CompareNotEqual:
-						Lhs = (Lhs != Rhs)? 1 : 0;
+						lhs = (lhs != rhs) ? 1 : 0;
 						break;
 					case TokenType.CompareLess:
-						Lhs = (Lhs < Rhs)? 1 : 0;
+						lhs = (lhs < rhs) ? 1 : 0;
 						break;
 					case TokenType.CompareGreater:
-						Lhs = (Lhs > Rhs)? 1 : 0;
+						lhs = (lhs > rhs) ? 1 : 0;
 						break;
 					case TokenType.CompareLessOrEqual:
-						Lhs = (Lhs <= Rhs)? 1 : 0;
+						lhs = (lhs <= rhs) ? 1 : 0;
 						break;
 					case TokenType.CompareGreaterOrEqual:
-						Lhs = (Lhs >= Rhs)? 1 : 0;
+						lhs = (lhs >= rhs) ? 1 : 0;
 						break;
 					case TokenType.LeftShift:
-						Lhs <<= (int)Rhs;
+						lhs <<= (int)rhs;
 						break;
 					case TokenType.RightShift:
-						Lhs >>= (int)Rhs;
+						lhs >>= (int)rhs;
 						break;
 					case TokenType.Plus:
-						Lhs += Rhs;
+						lhs += rhs;
 						break;
 					case TokenType.Minus:
-						Lhs -= Rhs;
+						lhs -= rhs;
 						break;
 					case TokenType.Multiply:
-						Lhs *= Rhs;
+						lhs *= rhs;
 						break;
 					case TokenType.Divide:
-						Lhs /= Rhs;
+						lhs /= rhs;
 						break;
 					case TokenType.Modulo:
-						Lhs %= Rhs;
+						lhs %= rhs;
 						break;
 					default:
-						throw new NotImplementedException($"Binary operator '{Operator.Type}' has not been implemented");
+						throw new NotImplementedException($"Binary operator '{tokenOp.Type}' has not been implemented");
 				}
 			}
 
-			return Lhs;
+			return lhs;
 		}
 
 		/// <summary>
 		/// Evaluates a unary expression (+, -, !, ~)
 		/// </summary>
-		/// <param name="Context">The current preprocessor context</param>
-		/// <param name="Tokens">List of tokens in the expression</param>
-		/// <param name="Idx">Index into the token list of the next token to read</param>
+		/// <param name="context">The current preprocessor context</param>
+		/// <param name="tokens">List of tokens in the expression</param>
+		/// <param name="idx">Index into the token list of the next token to read</param>
 		/// <returns>Value of the expression</returns>
-		static long EvaluateUnary(PreprocessorContext Context, List<Token> Tokens, ref int Idx)
+		static long EvaluateUnary(PreprocessorContext context, List<Token> tokens, ref int idx)
 		{
-			if(Idx == Tokens.Count)
+			if (idx == tokens.Count)
 			{
-				throw new PreprocessorException(Context, "Early end of expression");
+				throw new PreprocessorException(context, "Early end of expression");
 			}
 
-			switch(Tokens[Idx].Type)
+			switch (tokens[idx].Type)
 			{
 				case TokenType.Plus:
-					Idx++;
-					return +EvaluateUnary(Context, Tokens, ref Idx);
+					idx++;
+					return +EvaluateUnary(context, tokens, ref idx);
 				case TokenType.Minus:
-					Idx++;
-					return -EvaluateUnary(Context, Tokens, ref Idx);
+					idx++;
+					return -EvaluateUnary(context, tokens, ref idx);
 				case TokenType.LogicalNot:
-					Idx++;
-					return (EvaluateUnary(Context, Tokens, ref Idx) == 0) ? 1 : 0;
+					idx++;
+					return (EvaluateUnary(context, tokens, ref idx) == 0) ? 1 : 0;
 				case TokenType.BitwiseNot:
-					Idx++;
-					return ~EvaluateUnary(Context, Tokens, ref Idx);
+					idx++;
+					return ~EvaluateUnary(context, tokens, ref idx);
 				case TokenType.Identifier:
-					Identifier Text = Tokens[Idx].Identifier!;
-					if(Text == Identifiers.Sizeof)
+					Identifier text = tokens[idx].Identifier!;
+					if (text == Identifiers.Sizeof)
 					{
-						throw new NotImplementedException(Text.ToString());
+						throw new NotImplementedException(text.ToString());
 					}
-					else if(Text == Identifiers.Alignof)
+					else if (text == Identifiers.Alignof)
 					{
-						throw new NotImplementedException(Text.ToString());
+						throw new NotImplementedException(text.ToString());
 					}
-					else if (Text == Identifiers.__has_builtin || Text == Identifiers.__has_feature
-						|| Text == Identifiers.__has_warning || Text == Identifiers.__building_module
-						|| Text == Identifiers.__pragma || Text == Identifiers.__builtin_return_address
-						|| Text == Identifiers.__builtin_frame_address || Text == Identifiers.__has_keyword
-						|| Text == Identifiers.__has_extension || Text == Identifiers.__is_target_arch
-						|| Text == Identifiers.__is_identifier)
+					else if (text == Identifiers.__has_builtin || text == Identifiers.__has_feature
+						|| text == Identifiers.__has_warning || text == Identifiers.__building_module
+						|| text == Identifiers.__pragma || text == Identifiers.__builtin_return_address
+						|| text == Identifiers.__builtin_frame_address || text == Identifiers.__has_keyword
+						|| text == Identifiers.__has_extension || text == Identifiers.__is_target_arch
+						|| text == Identifiers.__is_identifier)
 					{
-						if (Tokens[Idx + 1].Type != TokenType.LeftParen || Tokens[Idx + 3].Type != TokenType.RightParen)
+						if (tokens[idx + 1].Type != TokenType.LeftParen || tokens[idx + 3].Type != TokenType.RightParen)
 						{
-							throw new NotImplementedException(Text.ToString());
+							throw new NotImplementedException(text.ToString());
 						}
-						Idx += 4;
+						idx += 4;
 						return 0;
 					}
-					else if (Text == Identifiers.__has_attribute || Text == Identifiers.__has_declspec_attribute
-						|| Text == Identifiers.__has_c_attribute || Text == Identifiers.__has_cpp_attribute
-						|| Text == Identifiers.__has_include || Text == Identifiers.__has_include_next)
+					else if (text == Identifiers.__has_attribute || text == Identifiers.__has_declspec_attribute
+						|| text == Identifiers.__has_c_attribute || text == Identifiers.__has_cpp_attribute
+						|| text == Identifiers.__has_include || text == Identifiers.__has_include_next)
 					{
-						Idx += Tokens.Count - Idx;
+						idx += tokens.Count - idx;
 						return 0;
 					}
 
 					else
 					{
-						return EvaluatePrimary(Context, Tokens, ref Idx);
+						return EvaluatePrimary(context, tokens, ref idx);
 					}
 				default:
-					return EvaluatePrimary(Context, Tokens, ref Idx);
+					return EvaluatePrimary(context, tokens, ref idx);
 			}
 		}
 
 		/// <summary>
 		/// Evaluates a primary expression
 		/// </summary>
-		/// <param name="Context">The current preprocessor context</param>
-		/// <param name="Tokens">List of tokens in the expression</param>
-		/// <param name="Idx">Index into the token list of the next token to read</param>
+		/// <param name="context">The current preprocessor context</param>
+		/// <param name="tokens">List of tokens in the expression</param>
+		/// <param name="idx">Index into the token list of the next token to read</param>
 		/// <returns>Value of the expression</returns>
-		static long EvaluatePrimary(PreprocessorContext Context, List<Token> Tokens, ref int Idx)
+		static long EvaluatePrimary(PreprocessorContext context, List<Token> tokens, ref int idx)
 		{
-			if(Tokens[Idx].Type == TokenType.Identifier)
+			if (tokens[idx].Type == TokenType.Identifier)
 			{
-				Idx++;
+				idx++;
 				return 0;
 			}
-			else if(Tokens[Idx].Type == TokenType.LeftParen)
+			else if (tokens[idx].Type == TokenType.LeftParen)
 			{
 				// Read the expression
-				Idx++;
-				long Result = EvaluateTernary(Context, Tokens, ref Idx);
+				idx++;
+				long result = EvaluateTernary(context, tokens, ref idx);
 
 				// Check for a closing parenthesis
-				if (Tokens[Idx].Type != TokenType.RightParen)
+				if (tokens[idx].Type != TokenType.RightParen)
 				{
-					throw new PreprocessorException(Context, "Missing closing parenthesis");
+					throw new PreprocessorException(context, "Missing closing parenthesis");
 				}
 
 				// Return the value
-				Idx++;
-				return Result;
+				idx++;
+				return result;
 			}
-			else if(Tokens[Idx].Type == TokenType.Number)
+			else if (tokens[idx].Type == TokenType.Number)
 			{
-				return ParseNumericLiteral(Context, Tokens[Idx++].Literal!);
+				return ParseNumericLiteral(context, tokens[idx++].Literal!);
 			}
 			else
 			{
-				throw new PreprocessorException(Context, "Unexpected token in expression: {0}", Tokens[Idx]);
+				throw new PreprocessorException(context, "Unexpected token in expression: {0}", tokens[idx]);
 			}
 		}
 
 		/// <summary>
 		/// Parse a numeric literal from the given token
 		/// </summary>
-		/// <param name="Context">The current preprocessor context</param>
-		/// <param name="Literal">The utf-8 literal characters</param>
+		/// <param name="context">The current preprocessor context</param>
+		/// <param name="literal">The utf-8 literal characters</param>
 		/// <returns>The literal value of the token</returns>
-		static long ParseNumericLiteral(PreprocessorContext Context, byte[] Literal)
+		static long ParseNumericLiteral(PreprocessorContext context, byte[] literal)
 		{
 			// Parse the token
-			long Value = 0;
-			if(Literal.Length >= 2 && Literal[0] == '0' && (Literal[1] == 'x' || Literal[1] == 'X'))
+			long value = 0;
+			if (literal.Length >= 2 && literal[0] == '0' && (literal[1] == 'x' || literal[1] == 'X'))
 			{
-				for (int Idx = 2; Idx < Literal.Length; Idx++)
+				for (int idx = 2; idx < literal.Length; idx++)
 				{
-					if (Literal[Idx] >= '0' && Literal[Idx] <= '9')
+					if (literal[idx] >= '0' && literal[idx] <= '9')
 					{
-						Value = (Value << 4) + (Literal[Idx] - '0');
+						value = (value << 4) + (literal[idx] - '0');
 					}
-					else if (Literal[Idx] >= 'a' && Literal[Idx] <= 'f')
+					else if (literal[idx] >= 'a' && literal[idx] <= 'f')
 					{
-						Value = (Value << 4) + (Literal[Idx] - 'a') + 10;
+						value = (value << 4) + (literal[idx] - 'a') + 10;
 					}
-					else if (Literal[Idx] >= 'A' && Literal[Idx] <= 'F')
+					else if (literal[idx] >= 'A' && literal[idx] <= 'F')
 					{
-						Value = (Value << 4) + (Literal[Idx] - 'A') + 10;
+						value = (value << 4) + (literal[idx] - 'A') + 10;
 					}
-					else if(IsValidIntegerSuffix(Literal, Idx))
+					else if (IsValidIntegerSuffix(literal, idx))
 					{
 						break;
 					}
 					else
 					{
-						throw new PreprocessorException(Context, "Invalid hexadecimal literal: '{0}'", Encoding.UTF8.GetString(Literal));
+						throw new PreprocessorException(context, "Invalid hexadecimal literal: '{0}'", Encoding.UTF8.GetString(literal));
 					}
 				}
 			}
-			else if(Literal[0] == '0')
+			else if (literal[0] == '0')
 			{
-				for (int Idx = 1; Idx < Literal.Length; Idx++)
+				for (int idx = 1; idx < literal.Length; idx++)
 				{
-					if (Literal[Idx] >= '0' && Literal[Idx] <= '7')
+					if (literal[idx] >= '0' && literal[idx] <= '7')
 					{
-						Value = (Value << 3) + (Literal[Idx] - '0');
+						value = (value << 3) + (literal[idx] - '0');
 					}
-					else if(IsValidIntegerSuffix(Literal, Idx))
+					else if (IsValidIntegerSuffix(literal, idx))
 					{
 						break;
 					}
 					else
 					{
-						throw new PreprocessorException(Context, "Invalid octal literal: '{0}'", Encoding.UTF8.GetString(Literal));
+						throw new PreprocessorException(context, "Invalid octal literal: '{0}'", Encoding.UTF8.GetString(literal));
 					}
 				}
 			}
 			else
 			{
-				for (int Idx = 0; Idx < Literal.Length; Idx++)
+				for (int idx = 0; idx < literal.Length; idx++)
 				{
-					if (Literal[Idx] >= '0' && Literal[Idx] <= '9')
+					if (literal[idx] >= '0' && literal[idx] <= '9')
 					{
-						Value = (Value * 10) + (Literal[Idx] - '0');
+						value = (value * 10) + (literal[idx] - '0');
 					}
-					else if(IsValidIntegerSuffix(Literal, Idx))
+					else if (IsValidIntegerSuffix(literal, idx))
 					{
 						break;
 					}
 					else
 					{
-						throw new PreprocessorException(Context, "Invalid decimal literal: '{0}'", Encoding.UTF8.GetString(Literal));
+						throw new PreprocessorException(context, "Invalid decimal literal: '{0}'", Encoding.UTF8.GetString(literal));
 					}
 				}
 			}
-			return Value;
+			return value;
 		}
 
 		/// <summary>
 		/// Checks whether the given literal has a valid integer suffix, starting at the given position
 		/// </summary>
-		/// <param name="Literal">The literal token</param>
-		/// <param name="Index">Index within the literal to start reading</param>
+		/// <param name="literal">The literal token</param>
+		/// <param name="index">Index within the literal to start reading</param>
 		/// <returns>The literal value of the token</returns>
-		static bool IsValidIntegerSuffix(byte[] Literal, int Index)
+		static bool IsValidIntegerSuffix(byte[] literal, int index)
 		{
-			bool bAllowTrailingUnsigned = true;
+			bool allowTrailingUnsigned = true;
 
-			if(Literal[Index] == 'u' || Literal[Index] == 'U')
+			if (literal[index] == 'u' || literal[index] == 'U')
 			{
 				// 'U'
-				Index++;
-				bAllowTrailingUnsigned = false;
+				index++;
+				allowTrailingUnsigned = false;
 			}
 
-			if(Index < Literal.Length && (Literal[Index] == 'l' || Literal[Index] == 'L'))
+			if (index < literal.Length && (literal[index] == 'l' || literal[index] == 'L'))
 			{
 				// 'LL' or 'L'
-				if(Index + 1 < Literal.Length && Literal[Index + 1] == Literal[Index])
+				if (index + 1 < literal.Length && literal[index + 1] == literal[index])
 				{
-					Index += 2;
+					index += 2;
 				}
 				else
 				{
-					Index++;
+					index++;
 				}
 			}
-			else if (Index + 2 < Literal.Length && (Literal[Index] == 'i' || Literal[Index] == 'I') && Literal[Index + 1] == '3' && Literal[Index + 2] == '2')
+			else if (index + 2 < literal.Length && (literal[index] == 'i' || literal[index] == 'I') && literal[index + 1] == '3' && literal[index + 2] == '2')
 			{
 				// 'I32' (Microsoft extension)
-				Index += 3;
-				bAllowTrailingUnsigned = false;
+				index += 3;
+				allowTrailingUnsigned = false;
 			}
-			else if(Index + 2 < Literal.Length && (Literal[Index] == 'i' || Literal[Index] == 'I') && Literal[Index + 1] == '6' && Literal[Index + 2] == '4')
+			else if (index + 2 < literal.Length && (literal[index] == 'i' || literal[index] == 'I') && literal[index + 1] == '6' && literal[index + 2] == '4')
 			{
 				// 'I64' (Microsoft extension)
-				Index += 3;
-				bAllowTrailingUnsigned = false;
+				index += 3;
+				allowTrailingUnsigned = false;
 			}
 
-			if(bAllowTrailingUnsigned && Index < Literal.Length && (Literal[Index] == 'u' || Literal[Index] == 'U'))
+			if (allowTrailingUnsigned && index < literal.Length && (literal[index] == 'u' || literal[index] == 'U'))
 			{
 				// 'U'
-				Index++;
+				index++;
 			}
 
-			return Index == Literal.Length;
+			return index == literal.Length;
 		}
 	}
 }
diff --git a/Engine/Source/Programs/UnrealBuildTool/Preprocessor/PreprocessorMacro.cs b/Engine/Source/Programs/UnrealBuildTool/Preprocessor/PreprocessorMacro.cs
--- a/Engine/Source/Programs/UnrealBuildTool/Preprocessor/PreprocessorMacro.cs
+++ b/Engine/Source/Programs/UnrealBuildTool/Preprocessor/PreprocessorMacro.cs
@@ -1,11 +1,10 @@
 // Copyright Epic Games, Inc. All Rights Reserved.
 
+using EpicGames.Core;
 using System;
 using System.Collections.Generic;
 using System.Linq;
 using System.Text;
-using System.Threading.Tasks;
-using EpicGames.Core;
 
 namespace UnrealBuildTool
 {
@@ -17,65 +16,68 @@ namespace UnrealBuildTool
 		/// <summary>
 		/// Name of the macro
 		/// </summary>
-		public Identifier Name;
+		public readonly Identifier Name;
 
 		/// <summary>
 		/// Parameter names for the macro. The '...' placeholder is represented by the __VA_ARGS__ string.
 		/// </summary>
-		public List<Identifier>? Parameters;
+		public readonly List<Identifier>? Parameters;
 
 		/// <summary>
 		/// Raw list of tokens for this macro
 		/// </summary>
-		public List<Token> Tokens;
+		public readonly List<Token> Tokens;
 
 		/// <summary>
 		/// Construct a preprocessor macro
 		/// </summary>
-		/// <param name="Name">Name of the macro</param>
-		/// <param name="Parameters">Parameter list for the macro. Should be null for object macros. Ownership of this list is transferred.</param>
-		/// <param name="Tokens">Tokens for the macro. Ownership of this list is transferred.</param>
-		public PreprocessorMacro(Identifier Name, List<Identifier>? Parameters, List<Token> Tokens)
+		/// <param name="name">Name of the macro</param>
+		/// <param name="parameters">Parameter list for the macro. Should be null for object macros. Ownership of this list is transferred.</param>
+		/// <param name="tokens">Tokens for the macro. Ownership of this list is transferred.</param>
+		public PreprocessorMacro(Identifier name, List<Identifier>? parameters, List<Token> tokens)
 		{
-			this.Name = Name;
-			this.Parameters = Parameters;
-			this.Tokens = Tokens;
+			Name = name;
+			Parameters = parameters;
+			Tokens = tokens;
 		}
 
 		/// <summary>
 		/// Read a macro from a binary archive
 		/// </summary>
-		/// <param name="Reader">Reader to serialize from</param>
-		public PreprocessorMacro(BinaryArchiveReader Reader)
+		/// <param name="reader">Reader to serialize from</param>
+		public PreprocessorMacro(BinaryArchiveReader reader)
 		{
-			Name = Reader.ReadIdentifier();
-			Parameters = Reader.ReadList(() => Reader.ReadIdentifier());
-			Tokens = Reader.ReadList(() => Reader.ReadToken())!;
+			Name = reader.ReadIdentifier();
+			Parameters = reader.ReadList(() => reader.ReadIdentifier());
+			Tokens = reader.ReadList(() => reader.ReadToken())!;
 		}
 
 		/// <summary>
 		/// Write a macro to a binary archive
 		/// </summary>
-		/// <param name="Writer">Writer to serialize to</param>
-		public void Write(BinaryArchiveWriter Writer)
+		/// <param name="writer">Writer to serialize to</param>
+		public void Write(BinaryArchiveWriter writer)
 		{
-			Writer.WriteIdentifier(Name);
-			Writer.WriteList(Parameters, x => Writer.WriteIdentifier(x));
-			Writer.WriteList(Tokens, x => Writer.WriteToken(x));
+			writer.WriteIdentifier(Name);
+			writer.WriteList(Parameters, x => writer.WriteIdentifier(x));
+			writer.WriteList(Tokens, x => writer.WriteToken(x));
 		}
 
 		/// <summary>
 		/// Finds the index of a parameter in the parameter list
 		/// </summary>
-		/// <param name="Parameter">Parameter name to look for</param>
+		/// <param name="parameter">Parameter name to look for</param>
 		/// <returns>Index of the parameter, or -1 if it's not found.</returns>
-		public int FindParameterIndex(Identifier Parameter)
+		public int FindParameterIndex(Identifier parameter)
 		{
-			for (int Idx = 0; Idx < Parameters!.Count; Idx++)
+			if (Parameters != null)
 			{
-				if (Parameters[Idx] == Parameter)
+				for (int idx = 0; idx < Parameters.Count; idx++)
 				{
-					return Idx;
+					if (Parameters[idx] == parameter)
+					{
+						return idx;
+					}
 				}
 			}
 			return -1;
@@ -84,31 +86,31 @@ namespace UnrealBuildTool
 		/// <summary>
 		/// Checks whether this macro definition is equivalent to another macro definition
 		/// </summary>
-		/// <param name="Other">The macro definition to compare against</param>
+		/// <param name="other">The macro definition to compare against</param>
 		/// <returns>True if the macro definitions are equivalent</returns>
-		public bool IsEquivalentTo(PreprocessorMacro Other)
+		public bool IsEquivalentTo(PreprocessorMacro other)
 		{
-			if(this != Other)
+			if (this != other)
 			{
-				if(Name != Other.Name || Tokens.Count != Other.Tokens.Count)
+				if (Name != other.Name || Tokens.Count != other.Tokens.Count)
 				{
 					return false;
 				}
-				if(Parameters != null)
+				if (Parameters != null)
 				{
-					if(Other.Parameters == null || Other.Parameters.Count != Parameters.Count || !Enumerable.SequenceEqual(Parameters, Other.Parameters))
+					if (other.Parameters == null || other.Parameters.Count != Parameters.Count || !Enumerable.SequenceEqual(Parameters, other.Parameters))
 					{
 						return false;
 					}
 				}
 				else
 				{
-					if(Other.Parameters != null)
+					if (other.Parameters != null)
 					{
 						return false;
 					}
 				}
-				if(!Enumerable.SequenceEqual(Tokens, Other.Tokens))
+				if (!Enumerable.SequenceEqual(Tokens, other.Tokens))
 				{
 					return false;
 				}
@@ -119,34 +121,22 @@ namespace UnrealBuildTool
 		/// <summary>
 		/// True if the macro is an object macro
 		/// </summary>
-		public bool IsObjectMacro
-		{
-			get { return Parameters == null; }
-		}
+		public bool IsObjectMacro => Parameters == null;
 
 		/// <summary>
 		/// True if the macro is a function macro
 		/// </summary>
-		public bool IsFunctionMacro
-		{
-			get { return Parameters != null; }
-		}
-		
+		public bool IsFunctionMacro => Parameters != null;
+
 		/// <summary>
 		/// The number of required parameters. For variadic macros, the last parameter is optional.
 		/// </summary>
-		public int MinRequiredParameters
-		{
-			get { return HasVariableArgumentList? (Parameters!.Count - 1) : Parameters!.Count; }
-		}
+		public int MinRequiredParameters => HasVariableArgumentList ? (Parameters!.Count - 1) : Parameters!.Count;
 
 		/// <summary>
 		/// True if the macro has a variable argument list
 		/// </summary>
-		public bool HasVariableArgumentList
-		{
-			get { return Parameters!.Count > 0 && Parameters[Parameters.Count - 1] == Identifiers.__VA_ARGS__; }
-		}
+		public bool HasVariableArgumentList => Parameters!.Count > 0 && Parameters[^1] == Identifiers.__VA_ARGS__;
 
 		/// <summary>
 		/// Converts this macro to a string for debugging
@@ -154,25 +144,25 @@ namespace UnrealBuildTool
 		/// <returns>The tokens in this macro</returns>
 		public override string ToString()
 		{
-			StringBuilder Result = new StringBuilder(Name.ToString());
+			StringBuilder result = new(Name.ToString());
 			if (Parameters != null)
 			{
-				Result.AppendFormat("({0})", String.Join(", ", Parameters));
+				result.AppendFormat("({0})", String.Join(", ", Parameters));
 			}
-			Result.Append("=");
+			result.Append('=');
 			if (Tokens.Count > 0)
 			{
-				Result.Append(Tokens[0].Text);
-				for (int Idx = 1; Idx < Tokens.Count; Idx++)
+				result.Append(Tokens[0].Text);
+				for (int idx = 1; idx < Tokens.Count; idx++)
 				{
-					if(Tokens[Idx].HasLeadingSpace)
+					if (Tokens[idx].HasLeadingSpace)
 					{
-						Result.Append(" ");
+						result.Append(' ');
 					}
-					Result.Append(Tokens[Idx].Text);
+					result.Append(Tokens[idx].Text);
 				}
 			}
-			return Result.ToString();
+			return result.ToString();
 		}
 	}
 }
diff --git a/Engine/Source/Programs/UnrealBuildTool/Preprocessor/PreprocessorState.cs b/Engine/Source/Programs/UnrealBuildTool/Preprocessor/PreprocessorState.cs
--- a/Engine/Source/Programs/UnrealBuildTool/Preprocessor/PreprocessorState.cs
+++ b/Engine/Source/Programs/UnrealBuildTool/Preprocessor/PreprocessorState.cs
@@ -3,9 +3,6 @@
 using System;
 using System.Collections.Generic;
 using System.Diagnostics.CodeAnalysis;
-using System.Linq;
-using System.Text;
-using System.Threading.Tasks;
 
 namespace UnrealBuildTool
 {
@@ -59,33 +56,27 @@ namespace UnrealBuildTool
 		/// <summary>
 		/// Stack of conditional branch states
 		/// </summary>
-		List<PreprocessorBranch> Branches = new List<PreprocessorBranch>();
+		readonly List<PreprocessorBranch> _branches = new();
 
 		/// <summary>
 		/// Mapping of name to macro definition
 		/// </summary>
-		Dictionary<Identifier, PreprocessorMacro> NameToMacro = new Dictionary<Identifier, PreprocessorMacro>();
+		readonly Dictionary<Identifier, PreprocessorMacro> _nameToMacro = new();
 
 		/// <summary>
 		/// The current transform. Any queries or state modifications will be recorded in this.
 		/// </summary>
-		PreprocessorTransform? Transform;
+		PreprocessorTransform? _transform;
 
 		/// <summary>
 		/// Enumerates the current branches. Do not use this for decision logic; any dependencies will not be tracked.
 		/// </summary>
-		public IEnumerable<PreprocessorBranch> CurrentBranches
-		{
-			get { return Branches; }
-		}
+		public IEnumerable<PreprocessorBranch> CurrentBranches => _branches;
 
 		/// <summary>
 		/// Enumerates the current macros. Do not use this for decision logic; any dependencies will not be tracked.
 		/// </summary>
-		public IEnumerable<PreprocessorMacro> CurrentMacros
-		{
-			get { return NameToMacro.Values; }
-		}
+		public IEnumerable<PreprocessorMacro> CurrentMacros => _nameToMacro.Values;
 
 		/// <summary>
 		/// Initialize an empty preprocessor state
@@ -97,19 +88,19 @@ namespace UnrealBuildTool
 		/// <summary>
 		/// Duplicates another preprocessor state. Throws an exception if a transform is currently being built in the other state.
 		/// </summary>
-		/// <param name="Other">The preprocessor state to copy</param>
-		public PreprocessorState(PreprocessorState Other)
+		/// <param name="other">The preprocessor state to copy</param>
+		public PreprocessorState(PreprocessorState other)
 		{
-			if(Other.Transform != null)
+			if (other._transform != null)
 			{
 				throw new Exception("Unable to copy another preprocessor state while a transform is being built.");
 			}
 
-			Branches.AddRange(Other.Branches);
+			_branches.AddRange(other._branches);
 
-			foreach(KeyValuePair<Identifier, PreprocessorMacro> Pair in Other.NameToMacro)
+			foreach (KeyValuePair<Identifier, PreprocessorMacro> pair in other._nameToMacro)
 			{
-				NameToMacro[Pair.Key] = Pair.Value;
+				_nameToMacro[pair.Key] = pair.Value;
 			}
 		}
 
@@ -119,8 +110,8 @@ namespace UnrealBuildTool
 		/// <returns>The new transform object</returns>
 		public PreprocessorTransform BeginCapture()
 		{
-			Transform = new PreprocessorTransform();
-			return Transform;
+			_transform = new PreprocessorTransform();
+			return _transform;
 		}
 
 		/// <summary>
@@ -129,9 +120,9 @@ namespace UnrealBuildTool
 		/// <returns>The transform object</returns>
 		public PreprocessorTransform? EndCapture()
 		{
-			PreprocessorTransform? PrevTransform = Transform;
-			Transform = null;
-			return PrevTransform;
+			PreprocessorTransform? prevTransform = _transform;
+			_transform = null;
+			return prevTransform;
 		}
 
 		/// <summary>
@@ -139,82 +130,78 @@ namespace UnrealBuildTool
 		/// </summary>
 		public void MarkPragmaOnce()
 		{
-			if(Transform != null)
+			if (_transform != null)
 			{
-				Transform.bHasPragmaOnce = true;
+				_transform.HasPragmaOnce = true;
 			}
 		}
 
 		/// <summary>
 		/// Set a macro to a specific definition
 		/// </summary>
-		/// <param name="Macro">Macro definition</param>
-		public void DefineMacro(PreprocessorMacro Macro)
+		/// <param name="macro">Macro definition</param>
+		public void DefineMacro(PreprocessorMacro macro)
 		{
-			NameToMacro[Macro.Name] = Macro;
+			_nameToMacro[macro.Name] = macro;
 
-			if(Transform != null)
+			if (_transform != null)
 			{
-				Transform.NewMacros[Macro.Name] = Macro;
+				_transform.NewMacros[macro.Name] = macro;
 			}
 		}
 
 		/// <summary>
 		/// Checks if a macro with the given name is defined
 		/// </summary>
-		/// <param name="Name">The macro name</param>
+		/// <param name="name">The macro name</param>
 		/// <returns>True if the macro is defined</returns>
-		public bool IsMacroDefined(Identifier Name)
+		public bool IsMacroDefined(Identifier name)
 		{
 			// Could account for the fact that we don't need the full definition later...
-			PreprocessorMacro? Macro;
-			return TryGetMacro(Name, out Macro);
+			return TryGetMacro(name, out _);
 		}
 
 		/// <summary>
 		/// Removes a macro definition
 		/// </summary>
-		/// <param name="Name">Name of the macro</param>
-		public void UndefMacro(Identifier Name)
+		/// <param name="name">Name of the macro</param>
+		public void UndefMacro(Identifier name)
 		{
-			NameToMacro.Remove(Name);
+			_nameToMacro.Remove(name);
 
-			if(Transform != null)
+			if (_transform != null)
 			{
-				Transform.NewMacros[Name] = null;
+				_transform.NewMacros[name] = null;
 			}
 		}
 
 		/// <summary>
 		/// Attemps to get the definition for a macro
 		/// </summary>
-		/// <param name="Name">Name of the macro</param>
-		/// <param name="Macro">Receives the macro definition, or null if it's not defined</param>
+		/// <param name="name">Name of the macro</param>
+		/// <param name="macro">Receives the macro definition, or null if it's not defined</param>
 		/// <returns>True if the macro is defined, otherwise false</returns>
-		public bool TryGetMacro(Identifier Name, [NotNullWhen(true)] out PreprocessorMacro? Macro)
+		public bool TryGetMacro(Identifier name, [NotNullWhen(true)] out PreprocessorMacro? macro)
 		{
-			bool bResult = NameToMacro.TryGetValue(Name, out Macro);
+			_nameToMacro.TryGetValue(name, out macro);
 
-			if(Transform != null && !Transform.NewMacros.ContainsKey(Name))
+			if (_transform != null && !_transform.NewMacros.ContainsKey(name))
 			{
-				Transform.RequiredMacros[Name] = Macro;
+				_transform.RequiredMacros[name] = macro;
 			}
 
-			return Macro != null;
+			return macro != null;
 		}
 
 		/// <summary>
 		/// Pushes the preprocessor branch onto the stack
 		/// </summary>
-		/// <param name="Branch">The branch state</param>
-		public void PushBranch(PreprocessorBranch Branch)
+		/// <param name="branch">The branch state</param>
+		public void PushBranch(PreprocessorBranch branch)
 		{
-			Branches.Add(Branch);
+			_branches.Add(branch);
 
-			if(Transform != null)
-			{
-				Transform.NewBranches.Add(Branch);
-			}
+			_transform?.NewBranches.Add(branch);
 		}
 
 		/// <summary>
@@ -223,41 +210,40 @@ namespace UnrealBuildTool
 		/// <returns>The popped branch state</returns>
 		public PreprocessorBranch PopBranch()
 		{
-			PreprocessorBranch Branch;
-			if(!TryPopBranch(out Branch))
+			if (!TryPopBranch(out PreprocessorBranch branch))
 			{
 				throw new InvalidOperationException("Branch stack is empty");
 			}
-			return Branch;
+			return branch;
 		}
 
 		/// <summary>
 		/// Attempts to pops a preprocessor branch from the stack
 		/// </summary>
-		/// <param name="Branch">On success, receives the preprocessor branch state</param>
+		/// <param name="branch">On success, receives the preprocessor branch state</param>
 		/// <returns>True if a branch was active, else false</returns>
-		public bool TryPopBranch(out PreprocessorBranch Branch)
+		public bool TryPopBranch(out PreprocessorBranch branch)
 		{
-			if(Branches.Count == 0)
+			if (_branches.Count == 0)
 			{
-				Branch = 0;
+				branch = 0;
 				return false;
 			}
 			else
 			{
-				Branch = Branches[Branches.Count - 1];
-				Branches.RemoveAt(Branches.Count - 1);
+				branch = _branches[^1];
+				_branches.RemoveAt(_branches.Count - 1);
 
-				if(Transform != null)
+				if (_transform != null)
 				{
-					if(Transform.NewBranches.Count > 0)
+					if (_transform.NewBranches.Count > 0)
 					{
-						Transform.NewBranches.RemoveAt(Transform.NewBranches.Count - 1);
+						_transform.NewBranches.RemoveAt(_transform.NewBranches.Count - 1);
 					}
 					else
 					{
-						Transform.RequiredBranches.Add(Branch);
-						Transform.bRequireTopmostActive = null;
+						_transform.RequiredBranches.Add(branch);
+						_transform.RequireTopmostActive = null;
 					}
 				}
 
@@ -271,54 +257,53 @@ namespace UnrealBuildTool
 		/// <returns>True if the branch is active, false otherwise</returns>
 		public bool IsCurrentBranchActive()
 		{
-			bool bActive = (Branches.Count == 0 || Branches[Branches.Count - 1].HasFlag(PreprocessorBranch.Active));
-			if(Transform != null && Transform.NewBranches.Count == 0)
+			bool active = (_branches.Count == 0 || _branches[^1].HasFlag(PreprocessorBranch.Active));
+			if (_transform != null && _transform.NewBranches.Count == 0)
 			{
-				Transform.bRequireTopmostActive = bActive;
+				_transform.RequireTopmostActive = active;
 			}
-			return bActive;
+			return active;
 		}
 
 		/// <summary>
 		/// Determines if the given transform can apply to the current preprocessor state
 		/// </summary>
-		/// <param name="Transform">The transform to test</param>
+		/// <param name="transform">The transform to test</param>
 		/// <returns>True if the transform can be applied to the current state</returns>
-		public bool CanApply(PreprocessorTransform Transform)
+		public bool CanApply(PreprocessorTransform transform)
 		{
 			// Check all the required branches match
-			for(int Idx = 0; Idx < Transform.RequiredBranches.Count; Idx++)
+			for (int idx = 0; idx < transform.RequiredBranches.Count; idx++)
 			{
-				if(Branches[Branches.Count - Idx - 1] != Transform.RequiredBranches[Idx])
+				if (_branches[_branches.Count - idx - 1] != transform.RequiredBranches[idx])
 				{
 					return false;
 				}
 			}
 
 			// Check the topmost branch is active
-			if(Transform.bRequireTopmostActive.HasValue)
+			if (transform.RequireTopmostActive.HasValue)
 			{
-				bool bTopmostActive = (Branches.Count == Transform.RequiredBranches.Count || Branches[Branches.Count - Transform.RequiredBranches.Count - 1].HasFlag(PreprocessorBranch.Active));
-				if(Transform.bRequireTopmostActive.Value != bTopmostActive)
+				bool topmostActive = _branches.Count == transform.RequiredBranches.Count || _branches[_branches.Count - transform.RequiredBranches.Count - 1].HasFlag(PreprocessorBranch.Active);
+				if (transform.RequireTopmostActive.Value != topmostActive)
 				{
 					return false;
 				}
 			}
 
 			// Check all the required macros match
-			foreach(KeyValuePair<Identifier, PreprocessorMacro?> RequiredPair in Transform.RequiredMacros)
+			foreach (KeyValuePair<Identifier, PreprocessorMacro?> requiredPair in transform.RequiredMacros)
 			{
-				PreprocessorMacro? Macro;
-				if(NameToMacro.TryGetValue(RequiredPair.Key, out Macro))
+				if (_nameToMacro.TryGetValue(requiredPair.Key, out PreprocessorMacro? macro))
 				{
-					if(RequiredPair.Value == null || !Macro.IsEquivalentTo(RequiredPair.Value))
+					if (requiredPair.Value == null || !macro.IsEquivalentTo(requiredPair.Value))
 					{
 						return false;
 					}
-				} 
+				}
 				else
 				{
-					if(RequiredPair.Value != null)
+					if (requiredPair.Value != null)
 					{
 						return false;
 					}
@@ -331,29 +316,29 @@ namespace UnrealBuildTool
 		/// <summary>
 		/// Apply the given transform to the current preprocessor state
 		/// </summary>
-		/// <param name="Transform">The transform to apply</param>
+		/// <param name="transform">The transform to apply</param>
 		/// <returns>True if the transform was applied to the current state</returns>
-		public bool TryToApply(PreprocessorTransform Transform)
+		public bool TryToApply(PreprocessorTransform transform)
 		{
-			if(!CanApply(Transform))
+			if (!CanApply(transform))
 			{
 				return false;
 			}
 
 			// Update the branch state
-			Branches.RemoveRange(Branches.Count - Transform.RequiredBranches.Count, Transform.RequiredBranches.Count);
-			Branches.AddRange(Transform.NewBranches);
+			_branches.RemoveRange(_branches.Count - transform.RequiredBranches.Count, transform.RequiredBranches.Count);
+			_branches.AddRange(transform.NewBranches);
 
 			// Update the macro definitions
-			foreach(KeyValuePair<Identifier, PreprocessorMacro?> NewMacro in Transform.NewMacros)
+			foreach (KeyValuePair<Identifier, PreprocessorMacro?> newMacro in transform.NewMacros)
 			{
-				if(NewMacro.Value == null)
+				if (newMacro.Value == null)
 				{
-					NameToMacro.Remove(NewMacro.Key);
+					_nameToMacro.Remove(newMacro.Key);
 				}
 				else
 				{
-					NameToMacro[NewMacro.Key] = NewMacro.Value;
+					_nameToMacro[newMacro.Key] = newMacro.Value;
 				}
 			}
 
diff --git a/Engine/Source/Programs/UnrealBuildTool/Preprocessor/PreprocessorTransform.cs b/Engine/Source/Programs/UnrealBuildTool/Preprocessor/PreprocessorTransform.cs
--- a/Engine/Source/Programs/UnrealBuildTool/Preprocessor/PreprocessorTransform.cs
+++ b/Engine/Source/Programs/UnrealBuildTool/Preprocessor/PreprocessorTransform.cs
@@ -1,11 +1,7 @@
 // Copyright Epic Games, Inc. All Rights Reserved.
 
-using System;
-using System.Collections.Generic;
-using System.Linq;
-using System.Text;
-using System.Threading.Tasks;
 using EpicGames.Core;
+using System.Collections.Generic;
 
 namespace UnrealBuildTool
 {
@@ -17,32 +13,32 @@ namespace UnrealBuildTool
 		/// <summary>
 		/// Whether the preprocessor must be in an active state after the branches below are popped. This cannot be treated as a 'pop', because it might be the global scope.
 		/// </summary>
-		public bool? bRequireTopmostActive;
+		public bool? RequireTopmostActive { get; set; } = null;
 
 		/// <summary>
 		/// List of branch states that will be popped from the stack
 		/// </summary>
-		public List<PreprocessorBranch> RequiredBranches;
+		public readonly List<PreprocessorBranch> RequiredBranches;
 
 		/// <summary>
 		/// Map of macro name to their required value
 		/// </summary>
-		public Dictionary<Identifier, PreprocessorMacro?> RequiredMacros;
+		public readonly Dictionary<Identifier, PreprocessorMacro?> RequiredMacros;
 
 		/// <summary>
 		/// List of new branches that will be pushed onto the stack
 		/// </summary>
-		public List<PreprocessorBranch> NewBranches;
+		public readonly List<PreprocessorBranch> NewBranches;
 
 		/// <summary>
 		/// List of new macros that will be defined
 		/// </summary>
-		public Dictionary<Identifier, PreprocessorMacro?> NewMacros;
+		public readonly Dictionary<Identifier, PreprocessorMacro?> NewMacros;
 
 		/// <summary>
 		/// This fragment contains a pragma once directive
 		/// </summary>
-		public bool bHasPragmaOnce;
+		public bool HasPragmaOnce { get; set; } = false;
 
 		/// <summary>
 		/// Constructor
@@ -58,38 +54,38 @@ namespace UnrealBuildTool
 		/// <summary>
 		/// Reads a transform from disk
 		/// </summary>
-		/// <param name="Reader">Archive to serialize from</param>
-		public PreprocessorTransform(BinaryArchiveReader Reader)
+		/// <param name="reader">Archive to serialize from</param>
+		public PreprocessorTransform(BinaryArchiveReader reader)
 		{
-			if(Reader.ReadBool())
+			if (reader.ReadBool())
 			{
-				bRequireTopmostActive = Reader.ReadBool();
+				RequireTopmostActive = reader.ReadBool();
 			}
 
-			RequiredBranches = Reader.ReadList(() => (PreprocessorBranch)Reader.ReadByte())!;
-			RequiredMacros = Reader.ReadDictionary(() => Reader.ReadIdentifier(), () => (PreprocessorMacro?)Reader.ReadObjectReference(() => new PreprocessorMacro(Reader)))!;
-			NewBranches = Reader.ReadList(() => (PreprocessorBranch)Reader.ReadByte())!;
-			NewMacros = Reader.ReadDictionary(() => Reader.ReadIdentifier(), () => (PreprocessorMacro?)Reader.ReadObjectReference(() => new PreprocessorMacro(Reader)))!;
-			bHasPragmaOnce = Reader.ReadBool();
+			RequiredBranches = reader.ReadList(() => (PreprocessorBranch)reader.ReadByte())!;
+			RequiredMacros = reader.ReadDictionary(() => reader.ReadIdentifier(), () => (PreprocessorMacro?)reader.ReadObjectReference(() => new PreprocessorMacro(reader)))!;
+			NewBranches = reader.ReadList(() => (PreprocessorBranch)reader.ReadByte())!;
+			NewMacros = reader.ReadDictionary(() => reader.ReadIdentifier(), () => (PreprocessorMacro?)reader.ReadObjectReference(() => new PreprocessorMacro(reader)))!;
+			HasPragmaOnce = reader.ReadBool();
 		}
 
 		/// <summary>
 		/// Writes a transform to disk
 		/// </summary>
-		/// <param name="Writer">Archive to serialize to</param>
-		public void Write(BinaryArchiveWriter Writer)
+		/// <param name="writer">Archive to serialize to</param>
+		public void Write(BinaryArchiveWriter writer)
 		{
-			Writer.WriteBool(bRequireTopmostActive.HasValue);
-			if(bRequireTopmostActive.HasValue)
+			writer.WriteBool(RequireTopmostActive.HasValue);
+			if (RequireTopmostActive.HasValue)
 			{
-				Writer.WriteBool(bRequireTopmostActive.Value);
+				writer.WriteBool(RequireTopmostActive.Value);
 			}
 
-			Writer.WriteList(RequiredBranches, x => Writer.WriteByte((byte)x));
-			Writer.WriteDictionary(RequiredMacros, k => Writer.WriteIdentifier(k), v => Writer.WriteObjectReference(v, () => v!.Write(Writer)));
-			Writer.WriteList(NewBranches, x => Writer.WriteByte((byte)x));
-			Writer.WriteDictionary(NewMacros, k => Writer.WriteIdentifier(k), v => Writer.WriteObjectReference(v, () => v!.Write(Writer)));
-			Writer.WriteBool(bHasPragmaOnce);
+			writer.WriteList(RequiredBranches, x => writer.WriteByte((byte)x));
+			writer.WriteDictionary(RequiredMacros, k => writer.WriteIdentifier(k), v => writer.WriteObjectReference(v, () => v!.Write(writer)));
+			writer.WriteList(NewBranches, x => writer.WriteByte((byte)x));
+			writer.WriteDictionary(NewMacros, k => writer.WriteIdentifier(k), v => writer.WriteObjectReference(v, () => v!.Write(writer)));
+			writer.WriteBool(HasPragmaOnce);
 		}
 	}
 }
diff --git a/Engine/Source/Programs/UnrealBuildTool/Preprocessor/SourceFile.cs b/Engine/Source/Programs/UnrealBuildTool/Preprocessor/SourceFile.cs
--- a/Engine/Source/Programs/UnrealBuildTool/Preprocessor/SourceFile.cs
+++ b/Engine/Source/Programs/UnrealBuildTool/Preprocessor/SourceFile.cs
@@ -1,14 +1,10 @@
 // Copyright Epic Games, Inc. All Rights Reserved.
 
+using EpicGames.Core;
 using System;
-using System.Collections.Concurrent;
 using System.Collections.Generic;
-using System.Diagnostics;
 using System.IO;
-using System.Linq;
 using System.Text;
-using System.Threading.Tasks;
-using EpicGames.Core;
 using UnrealBuildBase;
 
 namespace UnrealBuildTool
@@ -42,155 +38,150 @@ namespace UnrealBuildTool
 		/// <summary>
 		/// Macro which is defined as part of the header guard. If this macro is defined, we don't need to include the header more than once.
 		/// </summary>
-		public Identifier? HeaderGuardMacro;
+		public Identifier? HeaderGuardMacro { get; private set; }
 
 		/// <summary>
 		/// Parsed fragments in this file. Only set for files which are being optimized.
 		/// </summary>
-		public SourceFileFragment[] Fragments;
+		public readonly SourceFileFragment[] Fragments;
 
 		/// <summary>
 		/// Gets the location of this file
 		/// </summary>
-		public FileReference Location
-		{
-			get { return File.Location; }
-		}
+		public FileReference Location => File.Location;
 
 		/// <summary>
 		/// Construct a SourceFile from a file on disk
 		/// </summary>
-		/// <param name="File">Location of the file</param>
-		public SourceFile(FileItem File)
-			: this(File, ReadFileWithNullTerminator(File.Location))
+		/// <param name="file">Location of the file</param>
+		public SourceFile(FileItem file)
+			: this(file, ReadFileWithNullTerminator(file.Location))
 		{
-			this.LastWriteTimeUtc = File.LastWriteTimeUtc.Ticks;
+			LastWriteTimeUtc = file.LastWriteTimeUtc.Ticks;
 		}
 
 		/// <summary>
 		/// Construct a SourceFile from a raw byte array
 		/// </summary>
-		/// <param name="File">Location of the file</param>
-		/// <param name="Data">Contents of the file</param>
-		public SourceFile(FileItem File, byte[] Data)
+		/// <param name="file">Location of the file</param>
+		/// <param name="data">Contents of the file</param>
+		public SourceFile(FileItem file, byte[] data)
 		{
-			this.File = File;
+			File = file;
 
 			// Compute the hash
-			Hash = ContentHash.MD5(Data);
+			Hash = ContentHash.MD5(data);
 
 			// Read the preprocessor markup
-			TokenReader Reader = new TokenReader(Data);
-			Markup = SourceFileMarkup.Parse(Reader);
+			TokenReader reader = new(data);
+			Markup = SourceFileMarkup.Parse(reader);
 
 			// Try to parse a header guard from this file
 			ParseHeaderGuard();
 
 			// Create the fragments
-			List<SourceFileFragment> FragmentsList = new List<SourceFileFragment>();
-			for(int MaxIdx = 0; MaxIdx < Markup.Length; MaxIdx++)
+			List<SourceFileFragment> fragmentsList = new();
+			for (int maxIdx = 0; maxIdx < Markup.Length; maxIdx++)
 			{
-				if(Markup[MaxIdx].Type != SourceFileMarkupType.Include)
+				if (Markup[maxIdx].Type != SourceFileMarkupType.Include)
 				{
-					int MinIdx = MaxIdx++;
-					while(MaxIdx < Markup.Length && Markup[MaxIdx].Type != SourceFileMarkupType.Include)
+					int minIdx = maxIdx++;
+					while (maxIdx < Markup.Length && Markup[maxIdx].Type != SourceFileMarkupType.Include)
 					{
-						MaxIdx++;
+						maxIdx++;
 					}
-					FragmentsList.Add(new SourceFileFragment(this, MinIdx, MaxIdx));
+					fragmentsList.Add(new SourceFileFragment(this, minIdx, maxIdx));
 				}
 			}
-			Fragments = FragmentsList.ToArray();
+			Fragments = fragmentsList.ToArray();
 		}
 
 		/// <summary>
 		/// Read a source file from a binary reader
 		/// </summary>
-		/// <param name="Reader">The reader to serialize from</param>
-		public SourceFile(BinaryArchiveReader Reader)
+		/// <param name="reader">The reader to serialize from</param>
+		public SourceFile(BinaryArchiveReader reader)
 		{
-			File = Reader.ReadFileItem()!;
-			Hash = Reader.ReadContentHash()!;
-			LastWriteTimeUtc = Reader.ReadLong();
-			Markup = Reader.ReadArray(() => new SourceFileMarkup(Reader))!;
-			HeaderGuardMacro = Reader.ReadIdentifier();
-			Fragments = Reader.ReadArray(() => new SourceFileFragment(this, Reader))!;
+			File = reader.ReadFileItem()!;
+			Hash = reader.ReadContentHash()!;
+			LastWriteTimeUtc = reader.ReadLong();
+			Markup = reader.ReadArray(() => new SourceFileMarkup(reader))!;
+			HeaderGuardMacro = reader.ReadIdentifier();
+			Fragments = reader.ReadArray(() => new SourceFileFragment(this, reader))!;
 		}
 
 		/// <summary>
 		/// Read a source file from a binary reader
 		/// </summary>
-		/// <param name="Writer">The reader to serialize to</param>
-		public void Write(BinaryArchiveWriter Writer)
+		/// <param name="writer">The reader to serialize to</param>
+		public void Write(BinaryArchiveWriter writer)
 		{
-			Writer.WriteFileItem(File);
-			Writer.WriteContentHash(Hash);
-			Writer.WriteLong(LastWriteTimeUtc);
-			Writer.WriteArray(Markup, x => x.Write(Writer));
-			Writer.WriteIdentifier(HeaderGuardMacro);
-			Writer.WriteArray(Fragments, x => x.Write(Writer));
+			writer.WriteFileItem(File);
+			writer.WriteContentHash(Hash);
+			writer.WriteLong(LastWriteTimeUtc);
+			writer.WriteArray(Markup, x => x.Write(writer));
+			writer.WriteIdentifier(HeaderGuardMacro);
+			writer.WriteArray(Fragments, x => x.Write(writer));
 		}
 
 		/// <summary>
 		/// Reads the contents of a file into a buffer, with a null terminator
 		/// </summary>
-		/// <param name="Location">Location of the file</param>
+		/// <param name="location">Location of the file</param>
 		/// <returns>Array of bytes, with a null terminator</returns>
-		public static byte[] ReadFileWithNullTerminator(FileReference Location)
+		public static byte[] ReadFileWithNullTerminator(FileReference location)
 		{
-			using(FileStream Stream = FileReference.Open(Location, FileMode.Open, FileAccess.Read, FileShare.Read))
-			{
-				int Length = (int)Stream.Length;
+			using FileStream stream = FileReference.Open(location, FileMode.Open, FileAccess.Read, FileShare.Read);
+			int length = (int)stream.Length;
 
-				byte[] Data = new byte[Length + 1];
-				Stream.Read(Data, 0, Length);
+			byte[] data = new byte[length + 1];
+			stream.Read(data, 0, length);
 
-				return ConvertToNullTerminatedUtf8(Data, Length);
-			}
+			return ConvertToNullTerminatedUtf8(data, length);
 		}
 
 		/// <summary>
 		/// Converts a buffer from a given encoding to UTF8, appending a null terminator
 		/// </summary>
-		/// <param name="Data">The source data</param>
-		/// <param name="Length">Length of the data in the source buffer</param>
+		/// <param name="data">The source data</param>
+		/// <param name="length">Length of the data in the source buffer</param>
 		/// <returns>Array of utf-8 encoded characters, with a null terminator</returns>
-		public static byte[] ConvertToNullTerminatedUtf8(byte[] Data, int Length)
+		public static byte[] ConvertToNullTerminatedUtf8(byte[] data, int length)
 		{
-			if(Length >= 2)
+			if (length >= 2)
 			{
-				if(Data[0] == 0xfe && Data[1] == 0xff)
+				if (data[0] == 0xfe && data[1] == 0xff)
 				{
-					return ConvertToNullTerminatedUtf8(Encoding.BigEndianUnicode, Data, Length);
+					return ConvertToNullTerminatedUtf8(Encoding.BigEndianUnicode, data, length);
 				}
-				else if(Data[0] == 0xff && Data[1] == 0xfe)
+				else if (data[0] == 0xff && data[1] == 0xfe)
 				{
-					return ConvertToNullTerminatedUtf8(Encoding.Unicode, Data, Length);
+					return ConvertToNullTerminatedUtf8(Encoding.Unicode, data, length);
 				}
 			}
-			return Data;
+			return data;
 		}
 
 		/// <summary>
 		/// Converts a buffer from a given encoding to UTF8, appending a null terminator
 		/// </summary>
-		/// <param name="SourceEncoding">The encoding to convert from</param>
-		/// <param name="Data">The source data</param>
-		/// <param name="Length">Length of the data in the source buffer</param>
+		/// <param name="sourceEncoding">The encoding to convert from</param>
+		/// <param name="data">The source data</param>
+		/// <param name="length">Length of the data in the source buffer</param>
 		/// <returns>Array of utf-8 encoded characters, with a null terminator</returns>
-		public static byte[] ConvertToNullTerminatedUtf8(Encoding SourceEncoding, byte[] Data, int Length)
+		public static byte[] ConvertToNullTerminatedUtf8(Encoding sourceEncoding, byte[] data, int length)
 		{
-			byte[] NewBytes = Data;
-			if(!(SourceEncoding is UTF8Encoding))
+			byte[] newBytes = data;
+			if (sourceEncoding is not UTF8Encoding)
 			{
-				char[] Chars = SourceEncoding.GetChars(Data, 0, Length);
+				char[] chars = sourceEncoding.GetChars(data, 0, length);
 
-				int NewLength = Encoding.UTF8.GetByteCount(Chars);
-				NewBytes = new byte[NewLength + 1];
+				int newLength = Encoding.UTF8.GetByteCount(chars);
+				newBytes = new byte[newLength + 1];
 
-				Encoding.UTF8.GetBytes(Chars, 0, Chars.Length, NewBytes, 0);
+				Encoding.UTF8.GetBytes(chars, 0, chars.Length, newBytes, 0);
 			}
-			return NewBytes;
+			return newBytes;
 		}
 
 		/// <summary>
@@ -199,43 +190,43 @@ namespace UnrealBuildTool
 		void ParseHeaderGuard()
 		{
 			// Make sure there are enough markup entries in this list
-			if(Markup.Length >= 3)
+			if (Markup.Length >= 3)
 			{
-				int MinIdx = 0;
+				int minIdx = 0;
 
 				// Get the define used for the header guard
-				Identifier? PossibleHeaderGuardMacro;
-				if(Markup[MinIdx].Type == SourceFileMarkupType.Ifndef && Markup[MinIdx].Tokens!.Count == 1 && Markup[MinIdx].Tokens![0].Type == TokenType.Identifier)
+				Identifier? possibleHeaderGuardMacro;
+				if (Markup[minIdx].Type == SourceFileMarkupType.Ifndef && Markup[minIdx].Tokens!.Count == 1 && Markup[minIdx].Tokens![0].Type == TokenType.Identifier)
 				{
-					PossibleHeaderGuardMacro = Markup[MinIdx].Tokens![0].Identifier!;
+					possibleHeaderGuardMacro = Markup[minIdx].Tokens![0].Identifier!;
 				}
-				else if(Markup[MinIdx].Type == SourceFileMarkupType.If && Markup[MinIdx].Tokens!.Count == 3 && Markup[MinIdx].Tokens![0].Type == TokenType.LogicalNot && Markup[MinIdx].Tokens![1].Identifier == Identifiers.Defined && Markup[MinIdx].Tokens![2].Type == TokenType.Identifier)
+				else if (Markup[minIdx].Type == SourceFileMarkupType.If && Markup[minIdx].Tokens!.Count == 3 && Markup[minIdx].Tokens![0].Type == TokenType.LogicalNot && Markup[minIdx].Tokens![1].Identifier == Identifiers.Defined && Markup[minIdx].Tokens![2].Type == TokenType.Identifier)
 				{
-					PossibleHeaderGuardMacro = Markup[MinIdx].Tokens![2].Identifier!;
+					possibleHeaderGuardMacro = Markup[minIdx].Tokens![2].Identifier!;
 				}
 				else
 				{
-					PossibleHeaderGuardMacro = null;
+					possibleHeaderGuardMacro = null;
 				}
 
 				// Check it starts with an #if or #ifdef
-				if (Markup[MinIdx + 1].Type == SourceFileMarkupType.Define && Markup[MinIdx + 1].Tokens!.Count == 1 && Markup[MinIdx + 1].Tokens![0].Identifier == PossibleHeaderGuardMacro)
+				if (Markup[minIdx + 1].Type == SourceFileMarkupType.Define && Markup[minIdx + 1].Tokens!.Count == 1 && Markup[minIdx + 1].Tokens![0].Identifier == possibleHeaderGuardMacro)
 				{
 					// Find the point at which the include depth goes back to zero
-					int MaxIdx = MinIdx + 1;
-					for (int Depth = 1; MaxIdx < Markup.Length; MaxIdx++)
+					int maxIdx = minIdx + 1;
+					for (int depth = 1; maxIdx < Markup.Length; maxIdx++)
 					{
-						Depth += Markup[MaxIdx].GetConditionDepthDelta();
-						if (Depth == 0)
+						depth += Markup[maxIdx].GetConditionDepthDelta();
+						if (depth == 0)
 						{
 							break;
 						}
 					}
 
 					// Check it matched the end of the file
-					if (MaxIdx == Markup.Length - 1)
+					if (maxIdx == Markup.Length - 1)
 					{
-						HeaderGuardMacro = PossibleHeaderGuardMacro;
+						HeaderGuardMacro = possibleHeaderGuardMacro;
 					}
 				}
 			}
diff --git a/Engine/Source/Programs/UnrealBuildTool/Preprocessor/SourceFileFragment.cs b/Engine/Source/Programs/UnrealBuildTool/Preprocessor/SourceFileFragment.cs
--- a/Engine/Source/Programs/UnrealBuildTool/Preprocessor/SourceFileFragment.cs
+++ b/Engine/Source/Programs/UnrealBuildTool/Preprocessor/SourceFileFragment.cs
@@ -1,11 +1,6 @@
 // Copyright Epic Games, Inc. All Rights Reserved.
 
 using System;
-using System.Collections.Generic;
-using System.IO;
-using System.Linq;
-using System.Text;
-using System.Threading.Tasks;
 using EpicGames.Core;
 
 namespace UnrealBuildTool
@@ -34,43 +29,43 @@ namespace UnrealBuildTool
 		/// <summary>
 		/// List of known transforms for this fragment
 		/// </summary>
-		public PreprocessorTransform[] Transforms = new PreprocessorTransform[0];
+		public PreprocessorTransform[] Transforms { get; set; } = Array.Empty<PreprocessorTransform>();
 
 		/// <summary>
 		/// Constructor
 		/// </summary>
-		/// <param name="File">The source file that this fragment is part of</param>
-		/// <param name="MarkupMin">Index into the file's markup array of the start of this fragment (inclusive)</param>
-		/// <param name="MarkupMax">Index into the file's markup array of the end of this fragment (exclusive)</param>
-		public SourceFileFragment(SourceFile File, int MarkupMin, int MarkupMax)
+		/// <param name="file">The source file that this fragment is part of</param>
+		/// <param name="markupMin">Index into the file's markup array of the start of this fragment (inclusive)</param>
+		/// <param name="markupMax">Index into the file's markup array of the end of this fragment (exclusive)</param>
+		public SourceFileFragment(SourceFile file, int markupMin, int markupMax)
 		{
-			this.File = File;
-			this.MarkupMin = MarkupMin;
-			this.MarkupMax = MarkupMax;
+			File = file;
+			MarkupMin = markupMin;
+			MarkupMax = markupMax;
 		}
 
 		/// <summary>
 		/// Constructor
 		/// </summary>
-		/// <param name="File">The source file that this fragment is part of</param>
-		/// <param name="Reader">Reader to deserialize this object from</param>
-		public SourceFileFragment(SourceFile File, BinaryArchiveReader Reader)
+		/// <param name="file">The source file that this fragment is part of</param>
+		/// <param name="reader">Reader to deserialize this object from</param>
+		public SourceFileFragment(SourceFile file, BinaryArchiveReader reader)
 		{
-			this.File = File;
-			this.MarkupMin = Reader.ReadInt();
-			this.MarkupMax = Reader.ReadInt();
-			this.Transforms = Reader.ReadArray(() => new PreprocessorTransform(Reader))!;
+			File = file;
+			MarkupMin = reader.ReadInt();
+			MarkupMax = reader.ReadInt();
+			Transforms = reader.ReadArray(() => new PreprocessorTransform(reader))!;
 		}
 
 		/// <summary>
 		/// Write this fragment to an archive
 		/// </summary>
-		/// <param name="Writer">The writer to serialize to</param>
-		public void Write(BinaryArchiveWriter Writer)
+		/// <param name="writer">The writer to serialize to</param>
+		public void Write(BinaryArchiveWriter writer)
 		{
-			Writer.WriteInt(MarkupMin);
-			Writer.WriteInt(MarkupMax);
-			Writer.WriteArray(Transforms, x => x.Write(Writer));
+			writer.WriteInt(MarkupMin);
+			writer.WriteInt(MarkupMax);
+			writer.WriteArray(Transforms, x => x.Write(writer));
 		}
 
 		/// <summary>
@@ -79,7 +74,7 @@ namespace UnrealBuildTool
 		/// <returns>String representation of this fragment for the debugger</returns>
 		public override string ToString()
 		{
-			if(MarkupMax == 0)
+			if (MarkupMax == 0)
 			{
 				return File.ToString();
 			}
diff --git a/Engine/Source/Programs/UnrealBuildTool/Preprocessor/SourceFileMarkup.cs b/Engine/Source/Programs/UnrealBuildTool/Preprocessor/SourceFileMarkup.cs
--- a/Engine/Source/Programs/UnrealBuildTool/Preprocessor/SourceFileMarkup.cs
+++ b/Engine/Source/Programs/UnrealBuildTool/Preprocessor/SourceFileMarkup.cs
@@ -1,12 +1,9 @@
 // Copyright Epic Games, Inc. All Rights Reserved.
 
+using EpicGames.Core;
 using System;
 using System.Collections.Generic;
-using System.Linq;
 using System.Text;
-using System.Threading;
-using System.Threading.Tasks;
-using EpicGames.Core;
 
 namespace UnrealBuildTool
 {
@@ -100,51 +97,51 @@ namespace UnrealBuildTool
 		/// <summary>
 		/// The directive corresponding to this markup
 		/// </summary>
-		public SourceFileMarkupType Type;
+		public readonly SourceFileMarkupType Type;
 
 		/// <summary>
 		/// The one-based line number of this markup
 		/// </summary>
-		public int LineNumber;
+		public readonly int LineNumber;
 
 		/// <summary>
 		/// The tokens parsed for this markup. Set for directives.
 		/// </summary>
-		public List<Token>? Tokens;
+		public readonly List<Token>? Tokens;
 
 		/// <summary>
 		/// Construct the annotation with the given range
 		/// </summary>
-		/// <param name="Type">The type of this directive</param>
-		/// <param name="LineNumber">The line number of this markup</param>
-		/// <param name="Tokens">List of tokens</param>
-		public SourceFileMarkup(SourceFileMarkupType Type, int LineNumber, List<Token>? Tokens)
+		/// <param name="type">The type of this directive</param>
+		/// <param name="lineNumber">The line number of this markup</param>
+		/// <param name="tokens">List of tokens</param>
+		public SourceFileMarkup(SourceFileMarkupType type, int lineNumber, List<Token>? tokens)
 		{
-			this.Type = Type;
-			this.LineNumber = LineNumber;
-			this.Tokens = Tokens;
+			Type = type;
+			LineNumber = lineNumber;
+			Tokens = tokens;
 		}
 
 		/// <summary>
 		/// Constructs a markup object using data read from an archive
 		/// </summary>
-		/// <param name="Reader">The reader to deserialize from</param>
-		public SourceFileMarkup(BinaryArchiveReader Reader)
+		/// <param name="reader">The reader to deserialize from</param>
+		public SourceFileMarkup(BinaryArchiveReader reader)
 		{
-			Type = (SourceFileMarkupType)Reader.ReadByte();
-			LineNumber = Reader.ReadInt();
-			Tokens = Reader.ReadList(() => Reader.ReadToken());
+			Type = (SourceFileMarkupType)reader.ReadByte();
+			LineNumber = reader.ReadInt();
+			Tokens = reader.ReadList(() => reader.ReadToken());
 		}
 
 		/// <summary>
 		/// Serializes this object to a binary archive
 		/// </summary>
-		/// <param name="Writer">Writer to serialize to</param>
-		public void Write(BinaryArchiveWriter Writer)
+		/// <param name="writer">Writer to serialize to</param>
+		public void Write(BinaryArchiveWriter writer)
 		{
-			Writer.WriteByte((byte)Type);
-			Writer.WriteInt(LineNumber);
-			Writer.WriteList(Tokens, x => Writer.WriteToken(x));
+			writer.WriteByte((byte)Type);
+			writer.WriteInt(LineNumber);
+			writer.WriteList(Tokens, x => writer.WriteToken(x));
 		}
 
 		/// <summary>
@@ -153,17 +150,11 @@ namespace UnrealBuildTool
 		/// <returns>True if this object is a conditional preprocessor directive</returns>
 		public bool IsConditionalPreprocessorDirective()
 		{
-			switch(Type)
+			return Type switch
 			{
-				case SourceFileMarkupType.If:
-				case SourceFileMarkupType.Ifdef:
-				case SourceFileMarkupType.Ifndef:
-				case SourceFileMarkupType.Elif:
-				case SourceFileMarkupType.Else:
-				case SourceFileMarkupType.Endif:
-					return true;
-			}
-			return false;
+				SourceFileMarkupType.If or SourceFileMarkupType.Ifdef or SourceFileMarkupType.Ifndef or SourceFileMarkupType.Elif or SourceFileMarkupType.Else or SourceFileMarkupType.Endif => true,
+				_ => false,
+			};
 		}
 
 		/// <summary>
@@ -171,11 +162,11 @@ namespace UnrealBuildTool
 		/// </summary>
 		public int GetConditionDepthDelta()
 		{
-			if(Type == SourceFileMarkupType.If || Type == SourceFileMarkupType.Ifdef || Type == SourceFileMarkupType.Ifndef)
+			if (Type == SourceFileMarkupType.If || Type == SourceFileMarkupType.Ifdef || Type == SourceFileMarkupType.Ifndef)
 			{
 				return +1;
 			}
-			else if(Type == SourceFileMarkupType.Endif)
+			else if (Type == SourceFileMarkupType.Endif)
 			{
 				return -1;
 			}
@@ -191,161 +182,161 @@ namespace UnrealBuildTool
 		/// <returns>String representation for debugging</returns>
 		public override string ToString()
 		{
-			StringBuilder Result = new StringBuilder();
-			Result.AppendFormat("[{0}] ", LineNumber);
+			StringBuilder result = new();
+			result.AppendFormat("[{0}] ", LineNumber);
 
-			if(Type == SourceFileMarkupType.Text)
+			if (Type == SourceFileMarkupType.Text)
 			{
-				Result.Append("...");
+				result.Append("...");
 			}
 			else
 			{
-				Result.Append("#");
-				if(Type != SourceFileMarkupType.OtherDirective)
+				result.Append('#');
+				if (Type != SourceFileMarkupType.OtherDirective)
 				{
-					Result.Append(Type.ToString().ToLowerInvariant());
+					result.Append(Type.ToString().ToLowerInvariant());
 				}
-				if(Tokens != null && Tokens.Count > 0)
+				if (Tokens != null && Tokens.Count > 0)
 				{
-					Result.Append(' ');
-					Token.Format(Tokens, Result);
+					result.Append(' ');
+					Token.Format(Tokens, result);
 				}
 			}
-			return Result.ToString();
+			return result.ToString();
 		}
 
 		/// <summary>
 		/// Create markup for the given file
 		/// </summary>
-		/// <param name="Reader">Reader for tokens in the file</param>
+		/// <param name="reader">Reader for tokens in the file</param>
 		/// <returns>Array of markup objects which split up the given text buffer</returns>
-		public static SourceFileMarkup[] Parse(TokenReader Reader)
+		public static SourceFileMarkup[] Parse(TokenReader reader)
 		{
-			List<SourceFileMarkup> Markup = new List<SourceFileMarkup>();
-			if(Reader.MoveNext())
+			List<SourceFileMarkup> markup = new();
+			if (reader.MoveNext())
 			{
-				bool bMoveNext = true;
-				while(bMoveNext)
+				bool moveNext = true;
+				while (moveNext)
 				{
-					int StartLineNumber = Reader.LineNumber;
-					if (Reader.Current.Type == TokenType.Hash)
+					int startLineNumber = reader.LineNumber;
+					if (reader.Current.Type == TokenType.Hash)
 					{
 						// Create the appropriate markup object for the directive
-						SourceFileMarkupType Type = SourceFileMarkupType.OtherDirective;
-						if(Reader.MoveNext())
+						SourceFileMarkupType type = SourceFileMarkupType.OtherDirective;
+						if (reader.MoveNext())
 						{
-							if(Reader.Current.Type == TokenType.Identifier)
+							if (reader.Current.Type == TokenType.Identifier)
 							{
-								Identifier Directive = Reader.Current.Identifier!;
-								if(Directive == Identifiers.Include)
+								Identifier directive = reader.Current.Identifier!;
+								if (directive == Identifiers.Include)
 								{
-									Type = SourceFileMarkupType.Include;
+									type = SourceFileMarkupType.Include;
 								}
-								else if(Directive == Identifiers.Define)
+								else if (directive == Identifiers.Define)
 								{
-									Type = SourceFileMarkupType.Define;
+									type = SourceFileMarkupType.Define;
 								}
-								else if(Directive == Identifiers.Undef)
+								else if (directive == Identifiers.Undef)
 								{
-									Type = SourceFileMarkupType.Undef;
+									type = SourceFileMarkupType.Undef;
 								}
-								else if(Directive == Identifiers.If)
+								else if (directive == Identifiers.If)
 								{
-									Type = SourceFileMarkupType.If;
+									type = SourceFileMarkupType.If;
 								}
-								else if(Directive == Identifiers.Ifdef)
+								else if (directive == Identifiers.Ifdef)
 								{
-									Type = SourceFileMarkupType.Ifdef;
+									type = SourceFileMarkupType.Ifdef;
 								}
-								else if(Directive == Identifiers.Ifndef)
+								else if (directive == Identifiers.Ifndef)
 								{
-									Type = SourceFileMarkupType.Ifndef;
+									type = SourceFileMarkupType.Ifndef;
 								}
-								else if(Directive == Identifiers.Elif)
+								else if (directive == Identifiers.Elif)
 								{
-									Type = SourceFileMarkupType.Elif;
+									type = SourceFileMarkupType.Elif;
 								}
-								else if(Directive == Identifiers.Else)
+								else if (directive == Identifiers.Else)
 								{
-									Type = SourceFileMarkupType.Else;
+									type = SourceFileMarkupType.Else;
 								}
-								else if(Directive == Identifiers.Endif)
+								else if (directive == Identifiers.Endif)
 								{
-									Type = SourceFileMarkupType.Endif;
+									type = SourceFileMarkupType.Endif;
 								}
-								else if(Directive == Identifiers.Pragma)
+								else if (directive == Identifiers.Pragma)
 								{
-									Type = SourceFileMarkupType.Pragma;
+									type = SourceFileMarkupType.Pragma;
 								}
-								else if(Directive == Identifiers.Error)
+								else if (directive == Identifiers.Error)
 								{
-									Type = SourceFileMarkupType.Error;
+									type = SourceFileMarkupType.Error;
 								}
-								else if(Directive == Identifiers.Warning)
+								else if (directive == Identifiers.Warning)
 								{
-									Type = SourceFileMarkupType.Warning;
+									type = SourceFileMarkupType.Warning;
 								}
 							}
-							else if(Reader.Current.Type == TokenType.Newline)
+							else if (reader.Current.Type == TokenType.Newline)
 							{
-								Type = SourceFileMarkupType.Empty;
+								type = SourceFileMarkupType.Empty;
 							}
 						}
 
 						// Create the token list
-						List<Token> Tokens = new List<Token>();
-						if(Type == SourceFileMarkupType.OtherDirective)
+						List<Token> tokens = new();
+						if (type == SourceFileMarkupType.OtherDirective)
 						{
-							Tokens.Add(Reader.Current);
+							tokens.Add(reader.Current);
 						}
 
 						// Read the first token
-						if(Type == SourceFileMarkupType.Empty)
+						if (type == SourceFileMarkupType.Empty)
 						{
-							bMoveNext = true;
+							moveNext = true;
 						}
-						else if(Type == SourceFileMarkupType.Include)
+						else if (type == SourceFileMarkupType.Include)
 						{
-							bMoveNext = Reader.MoveNextIncludePath();
+							moveNext = reader.MoveNextIncludePath();
 						}
-						else if(Type == SourceFileMarkupType.Error || Type == SourceFileMarkupType.Warning)
+						else if (type == SourceFileMarkupType.Error || type == SourceFileMarkupType.Warning)
 						{
-							bMoveNext = Reader.MoveNextTokenString();
+							moveNext = reader.MoveNextTokenString();
 						}
 						else
 						{
-							bMoveNext = Reader.MoveNext();
+							moveNext = reader.MoveNext();
 						}
 
 						// Read the rest of the tokens
-						while(bMoveNext && Reader.Current.Type != TokenType.Newline)
+						while (moveNext && reader.Current.Type != TokenType.Newline)
 						{
-							Tokens.Add(Reader.Current);
-							bMoveNext = Reader.MoveNext();
+							tokens.Add(reader.Current);
+							moveNext = reader.MoveNext();
 						}
 
 						// Create the markup
-						Markup.Add(new SourceFileMarkup(Type, StartLineNumber, Tokens));
+						markup.Add(new SourceFileMarkup(type, startLineNumber, tokens));
 
 						// Move to the next token
-						bMoveNext = Reader.MoveNext();
+						moveNext = reader.MoveNext();
 					}
-					else if(Reader.Current.Type != TokenType.Newline)
+					else if (reader.Current.Type != TokenType.Newline)
 					{
 						// Create the new fragment
-						Markup.Add(new SourceFileMarkup(SourceFileMarkupType.Text, StartLineNumber, null));
+						markup.Add(new SourceFileMarkup(SourceFileMarkupType.Text, startLineNumber, null));
 
 						// Move to the next directive
-						bMoveNext = Reader.MoveToNextDirective();
+						moveNext = reader.MoveToNextDirective();
 					}
 					else
 					{
 						// Skip the empty line
-						bMoveNext = Reader.MoveNext();
+						moveNext = reader.MoveNext();
 					}
 				}
 			}
-			return Markup.ToArray();
+			return markup.ToArray();
 		}
 	}
 }
diff --git a/Engine/Source/Programs/UnrealBuildTool/Preprocessor/Token.cs b/Engine/Source/Programs/UnrealBuildTool/Preprocessor/Token.cs
--- a/Engine/Source/Programs/UnrealBuildTool/Preprocessor/Token.cs
+++ b/Engine/Source/Programs/UnrealBuildTool/Preprocessor/Token.cs
@@ -1,12 +1,11 @@
 // Copyright Epic Games, Inc. All Rights Reserved.
 
+using EpicGames.Core;
 using System;
 using System.Collections.Generic;
 using System.Diagnostics;
 using System.Linq;
 using System.Text;
-using System.Threading.Tasks;
-using EpicGames.Core;
 
 namespace UnrealBuildTool
 {
@@ -89,154 +88,142 @@ namespace UnrealBuildTool
 		/// <summary>
 		/// Static token names
 		/// </summary>
-		static string[] StaticTokens;
+		static readonly string[] s_staticTokens;
 
 		/// <summary>
 		/// Contains the type of the token
 		/// </summary>
-		public TokenType Type
-		{
-			get;
-		}
+		public readonly TokenType Type;
 
 		/// <summary>
 		/// Properties for this token
 		/// </summary>
-		public TokenFlags Flags
-		{
-			get;
-		}
+		public readonly TokenFlags Flags;
 
 		/// <summary>
 		/// If this token is identifier, contains the identifier name
 		/// </summary>
-		public Identifier? Identifier
-		{
-			get;
-		}
+		public readonly Identifier? Identifier;
 
 		/// <summary>
 		/// If this token is a literal, contains the raw utf-8 representation of it.
 		/// </summary>
-		public byte[]? Literal
-		{
-			get;
-		}
+		public readonly byte[]? Literal;
 
 		/// <summary>
 		/// Constructor
 		/// </summary>
-		/// <param name="Type">The token type</param>
-		/// <param name="Flags">Flags for this token</param>
-		public Token(TokenType Type, TokenFlags Flags)
+		/// <param name="type">The token type</param>
+		/// <param name="flags">Flags for this token</param>
+		public Token(TokenType type, TokenFlags flags)
 		{
-			this.Type = Type;
-			this.Flags = Flags;
+			Type = type;
+			Flags = flags;
 		}
 
 		/// <summary>
 		/// Constructor
 		/// </summary>
-		/// <param name="Identifier">The token identifier</param>
-		/// <param name="Flags">Flags for this token</param>
-		public Token(Identifier Identifier, TokenFlags Flags)
+		/// <param name="identifier">The token identifier</param>
+		/// <param name="flags">Flags for this token</param>
+		public Token(Identifier identifier, TokenFlags flags)
 		{
-			this.Type = TokenType.Identifier;
-			this.Flags = Flags;
-			this.Identifier = Identifier;
+			Type = TokenType.Identifier;
+			Flags = flags;
+			Identifier = identifier;
 		}
 
 		/// <summary>
 		/// Constructs a literal token
 		/// </summary>
-		public Token(TokenType Type, TokenFlags Flags, byte[] Literal) 
-			: this(Type, Flags)
+		public Token(TokenType type, TokenFlags flags, byte[] literal)
+			: this(type, flags)
 		{
-			Debug.Assert(IsLiteral(Type));
-			this.Literal = Literal;
+			Debug.Assert(IsLiteral(type));
+			Literal = literal;
 		}
 
 		/// <summary>
 		/// Constructs a literal token
 		/// </summary>
-		public Token(TokenType Type, TokenFlags Flags, string Literal) 
-			: this(Type, Flags)
+		public Token(TokenType type, TokenFlags flags, string literal)
+			: this(type, flags)
 		{
-			Debug.Assert(IsLiteral(Type));
-			this.Literal = Encoding.UTF8.GetBytes(Literal);
+			Debug.Assert(IsLiteral(type));
+			Literal = Encoding.UTF8.GetBytes(literal);
 		}
 
 		/// <summary>
 		/// Reads a token from a binary archive
 		/// </summary>
-		/// <param name="Reader">Archive to read from</param>
-		public Token(BinaryArchiveReader Reader)
+		/// <param name="reader">Archive to read from</param>
+		public Token(BinaryArchiveReader reader)
 		{
-			Type = (TokenType)Reader.ReadByte();
-			Flags = (TokenFlags)Reader.ReadByte();
-			if(Type == TokenType.Identifier)
+			Type = (TokenType)reader.ReadByte();
+			Flags = (TokenFlags)reader.ReadByte();
+			if (Type == TokenType.Identifier)
 			{
-				Identifier = Reader.ReadIdentifier();
+				Identifier = reader.ReadIdentifier();
 			}
-			else if(IsLiteral(Type))
+			else if (IsLiteral(Type))
 			{
-				Literal = Reader.ReadByteArray();
+				Literal = reader.ReadByteArray();
 			}
 		}
 
 		/// <summary>
 		/// Writes a token to a binary archive
 		/// </summary>
-		/// <param name="Writer">The writer to serialize to</param>
-		public void Write(BinaryArchiveWriter Writer)
+		/// <param name="writer">The writer to serialize to</param>
+		public void Write(BinaryArchiveWriter writer)
 		{
-			Writer.WriteByte((byte)Type);
-			Writer.WriteByte((byte)Flags);
-			if(Type == TokenType.Identifier)
+			writer.WriteByte((byte)Type);
+			writer.WriteByte((byte)Flags);
+			if (Type == TokenType.Identifier)
 			{
-				Writer.WriteIdentifier(Identifier!);
+				writer.WriteIdentifier(Identifier!);
 			}
-			else if(IsLiteral(Type))
+			else if (IsLiteral(Type))
 			{
-				Writer.WriteByteArray(Literal);
+				writer.WriteByteArray(Literal);
 			}
 		}
 
 		/// <summary>
 		/// Checks if a token is equal to another object
 		/// </summary>
-		/// <param name="Other">The object to compare against</param>
+		/// <param name="other">The object to compare against</param>
 		/// <returns>True if the objects are equal, false otherwise</returns>
-		public override bool Equals(object? Other)
+		public override bool Equals(object? other)
 		{
-			return Equals(Other as Token);
+			return Equals(other as Token);
 		}
 
 		/// <summary>
 		/// Checks if two tokens are equivalent
 		/// </summary>
-		/// <param name="Other">The object to compare against</param>
+		/// <param name="other">The object to compare against</param>
 		/// <returns>True if the tokens are equal, false otherwise</returns>
-		public bool Equals(Token? Other)
+		public bool Equals(Token? other)
 		{
-			if(ReferenceEquals(Other, null))
+			if (other is null)
 			{
 				return false;
 			}
-			if(Type != Other.Type || Flags != Other.Flags || Identifier != Other.Identifier)
+			if (Type != other.Type || Flags != other.Flags || Identifier != other.Identifier)
 			{
 				return false;
 			}
-			if(Literal != null)
+			if (Literal != null)
 			{
-				if(Other.Literal == null || Literal.Length != Other.Literal.Length || !Enumerable.SequenceEqual(Literal, Other.Literal))
+				if (other.Literal == null || Literal.Length != other.Literal.Length || !Enumerable.SequenceEqual(Literal, other.Literal))
 				{
 					return false;
 				}
 			}
 			else
 			{
-				if(Other.Literal != null)
+				if (other.Literal != null)
 				{
 					return false;
 				}
@@ -247,30 +234,30 @@ namespace UnrealBuildTool
 		/// <summary>
 		/// Compares two tokens for equality
 		/// </summary>
-		/// <param name="A">The first token to compare</param>
-		/// <param name="B">The second token to compare</param>
+		/// <param name="lhs">The first token to compare</param>
+		/// <param name="rhs">The second token to compare</param>
 		/// <returns>True if the objects are equal, false otherwise</returns>
-		public static bool operator==(Token A, Token B)
+		public static bool operator ==(Token lhs, Token rhs)
 		{
-			if((object)A == null)
+			if (lhs is null)
 			{
-				return ((object)B == null);
+				return (rhs is null);
 			}
 			else
 			{
-				return A.Equals(B);
+				return lhs.Equals(rhs);
 			}
 		}
 
 		/// <summary>
 		/// Compares two tokens for inequality
 		/// </summary>
-		/// <param name="A">The first token to compare</param>
-		/// <param name="B">The second token to compare</param>
+		/// <param name="lhs">The first token to compare</param>
+		/// <param name="rhs">The second token to compare</param>
 		/// <returns>True if the objects are not equal, false otherwise</returns>
-		public static bool operator!=(Token A, Token B)
+		public static bool operator !=(Token lhs, Token rhs)
 		{
-			return !(A == B);
+			return !(lhs == rhs);
 		}
 
 		/// <summary>
@@ -279,16 +266,16 @@ namespace UnrealBuildTool
 		/// <returns>Hash code for the object</returns>
 		public override int GetHashCode()
 		{
-			int Result = (int)Type + (int)Flags * 7;
-			if(Identifier != null)
+			int result = (int)Type + (int)Flags * 7;
+			if (Identifier != null)
 			{
-				Result = (Result * 11) + Identifier.GetHashCode();
+				result = (result * 11) + Identifier.GetHashCode();
 			}
-			if(Literal != null)
+			if (Literal != null)
 			{
-				for(int Idx = 0; Idx < Literal.Length; Idx++)
+				for (int idx = 0; idx < Literal.Length; idx++)
 				{
-					Result = (Result * 13) + Literal[Idx];
+					result = (result * 13) + Literal[idx];
 				}
 			}
 			return base.GetHashCode();
@@ -301,17 +288,17 @@ namespace UnrealBuildTool
 		{
 			get
 			{
-				if(Identifier != null)
+				if (Identifier != null)
 				{
 					return Identifier.ToString();
 				}
-				else if(Literal != null)
+				else if (Literal != null)
 				{
 					return Encoding.UTF8.GetString(Literal);
 				}
 				else
 				{
-					return StaticTokens[(int)Type];
+					return s_staticTokens[(int)Type];
 				}
 			}
 		}
@@ -319,119 +306,116 @@ namespace UnrealBuildTool
 		/// <summary>
 		/// Returns a new token with different flags
 		/// </summary>
-		/// <param name="FlagsToAdd">Flags to add from the token</param>
+		/// <param name="flagsToAdd">Flags to add from the token</param>
 		/// <returns>New token with updated flags</returns>
-		public Token AddFlags(TokenFlags FlagsToAdd)
+		public Token AddFlags(TokenFlags flagsToAdd)
 		{
-			if(Identifier != null)
+			if (Identifier != null)
 			{
-				return new Token(Identifier, Flags | FlagsToAdd);
+				return new Token(Identifier, Flags | flagsToAdd);
 			}
-			else if(Literal != null)
+			else if (Literal != null)
 			{
-				return new Token(Type, Flags | FlagsToAdd, Literal);
+				return new Token(Type, Flags | flagsToAdd, Literal);
 			}
 			else
 			{
-				return new Token(Type, Flags | FlagsToAdd);
+				return new Token(Type, Flags | flagsToAdd);
 			}
 		}
 
 		/// <summary>
 		/// Returns a new token with different flags
 		/// </summary>
-		/// <param name="FlagsToRemove">Flags to remove from the token</param>
+		/// <param name="flagsToRemove">Flags to remove from the token</param>
 		/// <returns>New token with updated flags</returns>
-		public Token RemoveFlags(TokenFlags FlagsToRemove)
+		public Token RemoveFlags(TokenFlags flagsToRemove)
 		{
-			if(Identifier != null)
+			if (Identifier != null)
 			{
-				return new Token(Identifier, Flags & ~FlagsToRemove);
+				return new Token(Identifier, Flags & ~flagsToRemove);
 			}
-			else if(Literal != null)
+			else if (Literal != null)
 			{
-				return new Token(Type, Flags & ~FlagsToRemove, Literal);
+				return new Token(Type, Flags & ~flagsToRemove, Literal);
 			}
 			else
 			{
-				return new Token(Type, Flags & ~FlagsToRemove);
+				return new Token(Type, Flags & ~flagsToRemove);
 			}
 		}
 
 		/// <summary>
 		/// Accessor for whether this token has leading whitespace
 		/// </summary>
-		public bool HasLeadingSpace
-		{
-			get { return (Flags & TokenFlags.HasLeadingSpace) != 0; }
-		}
+		public bool HasLeadingSpace => (Flags & TokenFlags.HasLeadingSpace) != 0;
 
 		/// <summary>
 		/// Checks whether two tokens match
 		/// </summary>
-		/// <param name="Other">The token to compare against</param>
+		/// <param name="other">The token to compare against</param>
 		/// <returns>True if the tokens match, false otherwise</returns>
-		public bool Matches(Token Other)
+		public bool Matches(Token other)
 		{
-			if(Type != Other.Type)
+			if (Type != other.Type)
 			{
 				return false;
 			}
 
-			if(Literal == null)
+			if (Literal == null)
 			{
-				return Identifier == Other.Identifier;
+				return Identifier == other.Identifier;
 			}
-			else 
+			else
 			{
-				return Other.Literal != null && Enumerable.SequenceEqual(Literal, Other.Literal);
+				return other.Literal != null && Enumerable.SequenceEqual(Literal, other.Literal);
 			}
 		}
 
 		/// <summary>
 		/// Determines whether the given token type is a literal
 		/// </summary>
-		/// <param name="Type">The type to test</param>
+		/// <param name="type">The type to test</param>
 		/// <returns>True if the given type is a literal</returns>
-		public static bool IsLiteral(TokenType Type)
+		public static bool IsLiteral(TokenType type)
 		{
-			return Type == TokenType.Unknown || Type == TokenType.Character || Type == TokenType.String || Type == TokenType.Number || Type == TokenType.StringOfTokens || Type == TokenType.SystemInclude;
+			return type == TokenType.Unknown || type == TokenType.Character || type == TokenType.String || type == TokenType.Number || type == TokenType.StringOfTokens || type == TokenType.SystemInclude;
 		}
 
 		/// <summary>
 		/// Concatenate a sequence of tokens into a string
 		/// </summary>
-		/// <param name="Tokens">The sequence of tokens to concatenate</param>
+		/// <param name="tokens">The sequence of tokens to concatenate</param>
 		/// <returns>String containing the concatenated tokens</returns>
-		public static string Format(IEnumerable<Token> Tokens)
+		public static string Format(IEnumerable<Token> tokens)
 		{
-			StringBuilder Result = new StringBuilder();
-			Format(Tokens, Result);
-			return Result.ToString();
+			StringBuilder result = new();
+			Format(tokens, result);
+			return result.ToString();
 		}
 
 		/// <summary>
 		/// Concatenate a sequence of tokens into a string
 		/// </summary>
-		/// <param name="Tokens">The sequence of tokens to concatenate</param>
-		/// <param name="Result">Receives the formatted string</param>
-		public static void Format(IEnumerable<Token> Tokens, StringBuilder Result)
+		/// <param name="tokens">The sequence of tokens to concatenate</param>
+		/// <param name="result">Receives the formatted string</param>
+		public static void Format(IEnumerable<Token> tokens, StringBuilder result)
 		{
-			IEnumerator<Token> Enumerator = Tokens.GetEnumerator();
-			if(Enumerator.MoveNext())
+			IEnumerator<Token> enumerator = tokens.GetEnumerator();
+			if (enumerator.MoveNext())
 			{
-				Result.Append(Enumerator.Current.Text);
+				result.Append(enumerator.Current.Text);
 
-				Token LastToken = Enumerator.Current;
-				while(Enumerator.MoveNext())
+				Token lastToken = enumerator.Current;
+				while (enumerator.MoveNext())
 				{
-					Token Token = Enumerator.Current;
-					if(Token.HasLeadingSpace && (Token.Type != TokenType.LeftParen || LastToken.Type != TokenType.Identifier || LastToken.Identifier != Identifiers.__pragma))
+					Token token = enumerator.Current;
+					if (token.HasLeadingSpace && (token.Type != TokenType.LeftParen || lastToken.Type != TokenType.Identifier || lastToken.Identifier != Identifiers.__pragma))
 					{
-						Result.Append(" ");
+						result.Append(' ');
 					}
-					Result.Append(Token.Text);
-					LastToken = Token;
+					result.Append(token.Text);
+					lastToken = token;
 				}
 			}
 		}
@@ -439,33 +423,34 @@ namespace UnrealBuildTool
 		/// <summary>
 		/// Concatenate two tokens together
 		/// </summary>
-		/// <param name="FirstToken">The first token</param>
-		/// <param name="SecondToken">The second token</param>
-		/// <param name="Context">Current preprocessor context</param>
+		/// <param name="firstToken">The first token</param>
+		/// <param name="secondToken">The second token</param>
+		/// <param name="context">Current preprocessor context</param>
 		/// <returns>The combined token</returns>
-		public static Token Concatenate(Token FirstToken, Token SecondToken, PreprocessorContext Context)
+		[System.Diagnostics.CodeAnalysis.SuppressMessage("Style", "IDE0060:Remove unused parameter", Justification = "<Pending>")]
+		public static Token Concatenate(Token firstToken, Token secondToken, PreprocessorContext context)
 		{
-			string Text = FirstToken.Text.ToString() + SecondToken.Text.ToString();
+			string text = firstToken.Text.ToString() + secondToken.Text.ToString();
 
-			List<Token> Tokens = new List<Token>();
+			List<Token> tokens = new();
 
-			TokenReader Reader = new TokenReader(Text);
-			while(Reader.MoveNext())
+			TokenReader reader = new(text);
+			while (reader.MoveNext())
 			{
-				Tokens.Add(Reader.Current);
+				tokens.Add(reader.Current);
 			}
 
-			if(Tokens.Count == 0)
+			if (tokens.Count == 0)
 			{
 				return new Token(TokenType.Placemarker, TokenFlags.None);
 			}
-			else if(Tokens.Count == 1)
+			else if (tokens.Count == 1)
 			{
-				return Tokens[0];
+				return tokens[0];
 			}
 			else
 			{
-				return new Token(TokenType.Unknown, TokenFlags.None, Text);
+				return new Token(TokenType.Unknown, TokenFlags.None, text);
 			}
 		}
 
@@ -474,39 +459,39 @@ namespace UnrealBuildTool
 		/// </summary>
 		static Token()
 		{
-			StaticTokens = new string[(int)TokenType.Max];
-			StaticTokens[(int)TokenType.LeftParen] = "(";
-			StaticTokens[(int)TokenType.RightParen] = ")";
-			StaticTokens[(int)TokenType.Comma] = ",";
-			StaticTokens[(int)TokenType.Newline] = "\n";
-			StaticTokens[(int)TokenType.Ellipsis] = "...";
-			StaticTokens[(int)TokenType.Placemarker] = "";
-			StaticTokens[(int)TokenType.Dot] = ".";
-			StaticTokens[(int)TokenType.QuestionMark] = "?";
-			StaticTokens[(int)TokenType.Colon] = ":";
-			StaticTokens[(int)TokenType.LogicalNot] = "!";
-			StaticTokens[(int)TokenType.LogicalAnd] = "&&";
-			StaticTokens[(int)TokenType.LogicalOr] = "||";
-			StaticTokens[(int)TokenType.BitwiseXor] = "^";
-			StaticTokens[(int)TokenType.BitwiseAnd] = "&";
-			StaticTokens[(int)TokenType.BitwiseNot] = "~";
-			StaticTokens[(int)TokenType.BitwiseOr] = "|";
-			StaticTokens[(int)TokenType.Equals] = "=";
-			StaticTokens[(int)TokenType.LeftShift] = "<<";
-			StaticTokens[(int)TokenType.RightShift] = ">>";
-			StaticTokens[(int)TokenType.CompareEqual] = "==";
-			StaticTokens[(int)TokenType.CompareNotEqual] = "!=";
-			StaticTokens[(int)TokenType.CompareLessOrEqual] = "<=";
-			StaticTokens[(int)TokenType.CompareLess] = "<";
-			StaticTokens[(int)TokenType.CompareGreaterOrEqual] = ">=";
-			StaticTokens[(int)TokenType.CompareGreater] = ">";
-			StaticTokens[(int)TokenType.Plus] = "+";
-			StaticTokens[(int)TokenType.Minus] = "-";
-			StaticTokens[(int)TokenType.Multiply] = "*";
-			StaticTokens[(int)TokenType.Divide] = "/";
-			StaticTokens[(int)TokenType.Modulo] = "%";
-			StaticTokens[(int)TokenType.Hash] = "#";
-			StaticTokens[(int)TokenType.HashHash] = "##";
+			s_staticTokens = new string[(int)TokenType.Max];
+			s_staticTokens[(int)TokenType.LeftParen] = "(";
+			s_staticTokens[(int)TokenType.RightParen] = ")";
+			s_staticTokens[(int)TokenType.Comma] = ",";
+			s_staticTokens[(int)TokenType.Newline] = "\n";
+			s_staticTokens[(int)TokenType.Ellipsis] = "...";
+			s_staticTokens[(int)TokenType.Placemarker] = "";
+			s_staticTokens[(int)TokenType.Dot] = ".";
+			s_staticTokens[(int)TokenType.QuestionMark] = "?";
+			s_staticTokens[(int)TokenType.Colon] = ":";
+			s_staticTokens[(int)TokenType.LogicalNot] = "!";
+			s_staticTokens[(int)TokenType.LogicalAnd] = "&&";
+			s_staticTokens[(int)TokenType.LogicalOr] = "||";
+			s_staticTokens[(int)TokenType.BitwiseXor] = "^";
+			s_staticTokens[(int)TokenType.BitwiseAnd] = "&";
+			s_staticTokens[(int)TokenType.BitwiseNot] = "~";
+			s_staticTokens[(int)TokenType.BitwiseOr] = "|";
+			s_staticTokens[(int)TokenType.Equals] = "=";
+			s_staticTokens[(int)TokenType.LeftShift] = "<<";
+			s_staticTokens[(int)TokenType.RightShift] = ">>";
+			s_staticTokens[(int)TokenType.CompareEqual] = "==";
+			s_staticTokens[(int)TokenType.CompareNotEqual] = "!=";
+			s_staticTokens[(int)TokenType.CompareLessOrEqual] = "<=";
+			s_staticTokens[(int)TokenType.CompareLess] = "<";
+			s_staticTokens[(int)TokenType.CompareGreaterOrEqual] = ">=";
+			s_staticTokens[(int)TokenType.CompareGreater] = ">";
+			s_staticTokens[(int)TokenType.Plus] = "+";
+			s_staticTokens[(int)TokenType.Minus] = "-";
+			s_staticTokens[(int)TokenType.Multiply] = "*";
+			s_staticTokens[(int)TokenType.Divide] = "/";
+			s_staticTokens[(int)TokenType.Modulo] = "%";
+			s_staticTokens[(int)TokenType.Hash] = "#";
+			s_staticTokens[(int)TokenType.HashHash] = "##";
 		}
 	}
 
@@ -518,21 +503,21 @@ namespace UnrealBuildTool
 		/// <summary>
 		/// Read a token from a binary archive
 		/// </summary>
-		/// <param name="Reader">Reader to serialize data from</param>
+		/// <param name="reader">Reader to serialize data from</param>
 		/// <returns>Token read from the archive</returns>
-		public static Token ReadToken(this BinaryArchiveReader Reader)
+		public static Token ReadToken(this BinaryArchiveReader reader)
 		{
-			return new Token(Reader);
+			return new Token(reader);
 		}
 
 		/// <summary>
 		/// Write a token to a binary archive
 		/// </summary>
-		/// <param name="Writer">Writer to serialize data to</param>
-		/// <param name="Token">Token to write</param>
-		public static void WriteToken(this BinaryArchiveWriter Writer, Token Token)
+		/// <param name="writer">Writer to serialize data to</param>
+		/// <param name="token">Token to write</param>
+		public static void WriteToken(this BinaryArchiveWriter writer, Token token)
 		{
-			Token.Write(Writer);
+			token.Write(writer);
 		}
 	}
 }
diff --git a/Engine/Source/Programs/UnrealBuildTool/Preprocessor/TokenReader.cs b/Engine/Source/Programs/UnrealBuildTool/Preprocessor/TokenReader.cs
--- a/Engine/Source/Programs/UnrealBuildTool/Preprocessor/TokenReader.cs
+++ b/Engine/Source/Programs/UnrealBuildTool/Preprocessor/TokenReader.cs
@@ -3,9 +3,7 @@
 using System;
 using System.Collections;
 using System.Collections.Generic;
-using System.Linq;
 using System.Text;
-using System.Threading.Tasks;
 
 namespace UnrealBuildTool
 {
@@ -64,100 +62,100 @@ namespace UnrealBuildTool
 		/// <summary>
 		/// Map of utf-8 leading bytes to their class
 		/// </summary>
-		static FirstCharacterClass[] FirstCharacters;
+		static readonly FirstCharacterClass[] s_firstCharacters;
 
 		/// <summary>
 		/// Array of flags for different leading utf-8 sequences
 		/// </summary>
-		static CharacterFlags[] Characters;
+		static readonly CharacterFlags[] s_characters;
 
 		/// <summary>
 		/// The current buffer being read from. Encoded as UTF-8 with a null terminator.
 		/// </summary>
-		byte[] Data;
+		readonly byte[] _data;
 
 		/// <summary>
 		/// Current offset within the buffer
 		/// </summary>
-		int Offset;
+		int _offset;
 
 		/// <summary>
 		/// The current token
 		/// </summary>
-		Token? CurrentToken;
+		Token? _currentToken;
 
 		/// <summary>
 		/// Constructor
 		/// </summary>
-		/// <param name="Data">The raw byte stream to read from, encoded as UTF-8</param>
-		public TokenReader(byte[] Data)
+		/// <param name="data">The raw byte stream to read from, encoded as UTF-8</param>
+		public TokenReader(byte[] data)
 		{
-			this.Data = Data;
-			this.LineNumber = 0;
-			this.LineNumberAfterToken = 1;
+			_data = data;
+			LineNumber = 0;
+			LineNumberAfterToken = 1;
 
 			// Make sure the input data has a null terminator
-			if(Data.Length == 0 || Data[Data.Length - 1] != 0)
+			if (data.Length == 0 || data[^1] != 0)
 			{
 				throw new ArgumentException("Data parameter must be null terminated.");
 			}
 
 			// If the data contains a UTF-8 BOM, skip over it
-			if(Data[0] == 0xef && Data[1] == 0xbb && Data[2] == 0xbf)
+			if (data[0] == 0xef && data[1] == 0xbb && data[2] == 0xbf)
 			{
-				Offset = 3;
+				_offset = 3;
 			}
 		}
 
 		/// <summary>
 		/// Constructor
 		/// </summary>
-		/// <param name="Text">The text to read from</param>
-		public TokenReader(string Text)
-			: this(GetNullTerminatedByteArray(Text))
+		/// <param name="text">The text to read from</param>
+		public TokenReader(string text)
+			: this(GetNullTerminatedByteArray(text))
 		{
 		}
 
 		/// <summary>
 		/// Copy constructor
 		/// </summary>
-		/// <param name="Other">Token reader to copy from</param>
-		public TokenReader(TokenReader Other)
+		/// <param name="other">Token reader to copy from</param>
+		public TokenReader(TokenReader other)
 		{
-			this.Data = Other.Data;
-			this.Offset = Other.Offset;
-			this.LineNumber = Other.LineNumber;
-			this.LineNumberAfterToken = Other.LineNumberAfterToken;
+			_data = other._data;
+			_offset = other._offset;
+			LineNumber = other.LineNumber;
+			LineNumberAfterToken = other.LineNumberAfterToken;
 		}
 
 		/// <summary>
 		/// Gets a null terminated byte array from the given string
 		/// </summary>
-		/// <param name="Text">String to convert into bytes</param>
+		/// <param name="text">String to convert into bytes</param>
 		/// <returns>Array of bytes</returns>
-		public static byte[] GetNullTerminatedByteArray(string Text)
+		public static byte[] GetNullTerminatedByteArray(string text)
 		{
-			byte[] Bytes = new byte[Encoding.UTF8.GetByteCount(Text) + 1];
-			Encoding.UTF8.GetBytes(Text, 0, Text.Length, Bytes, 0);
-			return Bytes;
+			byte[] bytes = new byte[Encoding.UTF8.GetByteCount(text) + 1];
+			Encoding.UTF8.GetBytes(text, 0, text.Length, bytes, 0);
+			return bytes;
 		}
 
 		/// <summary>
 		/// Tokenize an input string
 		/// </summary>
-		/// <param name="Text">Text to tokenize</param>
+		/// <param name="text">Text to tokenize</param>
 		/// <returns>List of tokens parsed from the text</returns>
-		public static List<Token> Tokenize(string Text)
+		public static List<Token> Tokenize(string text)
 		{
-			List<Token> Tokens = new List<Token>();
+			List<Token> tokens = new();
 
-			TokenReader Reader = new TokenReader(Text);
-			while(Reader.MoveNext())
+			TokenReader reader = new(text);
+			while (reader.MoveNext())
 			{
-				Tokens.Add(Reader.Current);
+				tokens.Add(reader.Current);
 			}
 
-			return Tokens;
+			return tokens;
 		}
 
 		/// <summary>
@@ -181,18 +179,12 @@ namespace UnrealBuildTool
 		/// <summary>
 		/// Returns the current token
 		/// </summary>
-		public Token Current
-		{
-			get { return CurrentToken!; }
-		}
+		public Token Current => _currentToken!;
 
 		/// <summary>
 		/// Untyped implementation of Current for IEnumerator.
 		/// </summary>
-		object IEnumerator.Current
-		{
-			get { return Current; }
-		}
+		object IEnumerator.Current => Current;
 
 		/// <summary>
 		/// Override of IEnumerator.Dispose. Not required.
@@ -208,213 +200,213 @@ namespace UnrealBuildTool
 		public bool MoveNext()
 		{
 			// Skip past the leading whitespace
-			TokenFlags Flags = SkipWhitespace();
+			TokenFlags flags = SkipWhitespace();
 
 			// Update the line number to the one after the current token
 			LineNumber = LineNumberAfterToken;
 
 			// Capture the start of the token
-			int StartOffset = Offset;
+			int startOffset = _offset;
 
 			// Initialize the flags for the next token
-			byte FirstByte = Data[Offset++];
-			switch(FirstCharacters[FirstByte])
+			byte firstByte = _data[_offset++];
+			switch (s_firstCharacters[firstByte])
 			{
 				case FirstCharacterClass.Unknown:
-					if((FirstByte & 0x80) != 0)
+					if ((firstByte & 0x80) != 0)
 					{
-						while((Data[Offset] & 0x80) != 0)
+						while ((_data[_offset] & 0x80) != 0)
 						{
-							Offset++;
+							_offset++;
 						}
 					}
-					CurrentToken = CreateLiteral(TokenType.Unknown, Flags, StartOffset, Offset - StartOffset);
+					_currentToken = CreateLiteral(TokenType.Unknown, flags, startOffset, _offset - startOffset);
 					return true;
 				case FirstCharacterClass.Terminator:
-					Offset = StartOffset;
-					CurrentToken = new Token(TokenType.End, Flags);
+					_offset = startOffset;
+					_currentToken = new Token(TokenType.End, flags);
 					return false;
 				case FirstCharacterClass.Newline:
 					LineNumberAfterToken++;
-					CurrentToken = new Token(TokenType.Newline, Flags);
+					_currentToken = new Token(TokenType.Newline, flags);
 					return true;
 				case FirstCharacterClass.Identifier:
 					// Identifier (or text literal with prefix)
-					while((Characters[Data[Offset]] & CharacterFlags.Identifier) != 0)
+					while ((s_characters[_data[_offset]] & CharacterFlags.Identifier) != 0)
 					{
-						Offset++;
+						_offset++;
 					}
 
 					// Check if it's a prefixed text literal
-					if(Data[Offset] == '\'')
+					if (_data[_offset] == '\'')
 					{
-						Offset++;
+						_offset++;
 						SkipTextLiteral('\'');
-						CurrentToken = CreateLiteral(TokenType.Character, Flags, StartOffset, Offset - StartOffset);
+						_currentToken = CreateLiteral(TokenType.Character, flags, startOffset, _offset - startOffset);
 						return true;
 					}
-					else if(Data[Offset] == '\"')
+					else if (_data[_offset] == '\"')
 					{
-						Offset++;
+						_offset++;
 						SkipTextLiteral('\"');
-						CurrentToken = CreateLiteral(TokenType.String, Flags, StartOffset, Offset - StartOffset);
+						_currentToken = CreateLiteral(TokenType.String, flags, startOffset, _offset - startOffset);
 						return true;
 					}
 					else
 					{
-						string Name = Encoding.UTF8.GetString(Data, StartOffset, Offset - StartOffset);
-						Identifier Identifier = Identifier.FindOrAdd(Name);
-						CurrentToken = new Token(Identifier, Flags);
+						string name = Encoding.UTF8.GetString(_data, startOffset, _offset - startOffset);
+						Identifier identifier = Identifier.FindOrAdd(name);
+						_currentToken = new Token(identifier, flags);
 						return true;
 					}
 				case FirstCharacterClass.Number:
 					// Numeric literal
 					SkipNumericLiteral();
-					CurrentToken = CreateLiteral(TokenType.Number, Flags, StartOffset, Offset - StartOffset);
+					_currentToken = CreateLiteral(TokenType.Number, flags, startOffset, _offset - startOffset);
 					return true;
 				case FirstCharacterClass.Character:
 					// Character literal
 					SkipTextLiteral('\'');
-					CurrentToken = CreateLiteral(TokenType.Character, Flags, StartOffset, Offset - StartOffset);
+					_currentToken = CreateLiteral(TokenType.Character, flags, startOffset, _offset - startOffset);
 					return true;
 				case FirstCharacterClass.String:
 					// String literal
 					SkipTextLiteral('\"');
-					CurrentToken = CreateLiteral(TokenType.String, Flags, StartOffset, Offset - StartOffset);
+					_currentToken = CreateLiteral(TokenType.String, flags, startOffset, _offset - startOffset);
 					return true;
 				case FirstCharacterClass.Dot:
 					// Numeric literal, ellipsis, or dot
-					if((Characters[Data[Offset]] & CharacterFlags.Digit) != 0)
+					if ((s_characters[_data[_offset]] & CharacterFlags.Digit) != 0)
 					{
-						Offset++;
+						_offset++;
 						SkipNumericLiteral();
-						CurrentToken = CreateLiteral(TokenType.Number, Flags, StartOffset, Offset - StartOffset);
+						_currentToken = CreateLiteral(TokenType.Number, flags, startOffset, _offset - startOffset);
 						return true;
 					}
-					else if(Data[Offset] == '.' && Data[Offset + 1] == '.')
+					else if (_data[_offset] == '.' && _data[_offset + 1] == '.')
 					{
-						Offset += 2;
-						CurrentToken = new Token(TokenType.Ellipsis, Flags);
+						_offset += 2;
+						_currentToken = new Token(TokenType.Ellipsis, flags);
 						return true;
 					}
 					else
 					{
-						CurrentToken = new Token(TokenType.Dot, Flags);
+						_currentToken = new Token(TokenType.Dot, flags);
 						return true;
 					}
 				case FirstCharacterClass.QuestionMark:
-					CurrentToken = new Token(TokenType.QuestionMark, Flags);
+					_currentToken = new Token(TokenType.QuestionMark, flags);
 					return true;
 				case FirstCharacterClass.Colon:
-					CurrentToken = new Token(TokenType.Colon, Flags);
+					_currentToken = new Token(TokenType.Colon, flags);
 					return true;
 				case FirstCharacterClass.ExclamationMark:
-					if(ReadCharacter('='))
+					if (ReadCharacter('='))
 					{
-						CurrentToken = new Token(TokenType.CompareNotEqual, Flags);
+						_currentToken = new Token(TokenType.CompareNotEqual, flags);
 					}
 					else
 					{
-						CurrentToken = new Token(TokenType.LogicalNot, Flags);
+						_currentToken = new Token(TokenType.LogicalNot, flags);
 					}
 					return true;
 				case FirstCharacterClass.Pipe:
-					if(ReadCharacter('|'))
+					if (ReadCharacter('|'))
 					{
-						CurrentToken = new Token(TokenType.LogicalOr, Flags);
+						_currentToken = new Token(TokenType.LogicalOr, flags);
 					}
 					else
 					{
-						CurrentToken = new Token(TokenType.BitwiseOr, Flags);
+						_currentToken = new Token(TokenType.BitwiseOr, flags);
 					}
 					return true;
 				case FirstCharacterClass.Ampersand:
-					if(ReadCharacter('&'))
+					if (ReadCharacter('&'))
 					{
-						CurrentToken = new Token(TokenType.LogicalAnd, Flags);
+						_currentToken = new Token(TokenType.LogicalAnd, flags);
 					}
 					else
 					{
-						CurrentToken = new Token(TokenType.BitwiseAnd, Flags);
+						_currentToken = new Token(TokenType.BitwiseAnd, flags);
 					}
 					return true;
 				case FirstCharacterClass.Caret:
-					CurrentToken = new Token(TokenType.BitwiseXor, Flags);
+					_currentToken = new Token(TokenType.BitwiseXor, flags);
 					return true;
 				case FirstCharacterClass.Equals:
-					if(ReadCharacter('='))
+					if (ReadCharacter('='))
 					{
-						CurrentToken = new Token(TokenType.CompareEqual, Flags);
+						_currentToken = new Token(TokenType.CompareEqual, flags);
 					}
 					else
 					{
-						CurrentToken = new Token(TokenType.Equals, Flags);
+						_currentToken = new Token(TokenType.Equals, flags);
 					}
 					return true;
 				case FirstCharacterClass.LeftTriangleBracket:
-					if(ReadCharacter('<'))
+					if (ReadCharacter('<'))
 					{
-						CurrentToken = new Token(TokenType.LeftShift, Flags);
+						_currentToken = new Token(TokenType.LeftShift, flags);
 					}
-					else if(ReadCharacter('='))
+					else if (ReadCharacter('='))
 					{
-						CurrentToken = new Token(TokenType.CompareLessOrEqual, Flags);
+						_currentToken = new Token(TokenType.CompareLessOrEqual, flags);
 					}
 					else
 					{
-						CurrentToken = new Token(TokenType.CompareLess, Flags);
+						_currentToken = new Token(TokenType.CompareLess, flags);
 					}
 					return true;
 				case FirstCharacterClass.RightTriangleBracket:
-					if(ReadCharacter('>'))
+					if (ReadCharacter('>'))
 					{
-						CurrentToken = new Token(TokenType.RightShift, Flags);
+						_currentToken = new Token(TokenType.RightShift, flags);
 					}
-					else if(ReadCharacter('='))
+					else if (ReadCharacter('='))
 					{
-						CurrentToken = new Token(TokenType.CompareGreaterOrEqual, Flags);
+						_currentToken = new Token(TokenType.CompareGreaterOrEqual, flags);
 					}
 					else
 					{
-						CurrentToken = new Token(TokenType.CompareGreater, Flags);
+						_currentToken = new Token(TokenType.CompareGreater, flags);
 					}
 					return true;
 				case FirstCharacterClass.Plus:
-					CurrentToken = new Token(TokenType.Plus, Flags);
+					_currentToken = new Token(TokenType.Plus, flags);
 					return true;
 				case FirstCharacterClass.Minus:
-					CurrentToken = new Token(TokenType.Minus, Flags);
+					_currentToken = new Token(TokenType.Minus, flags);
 					return true;
 				case FirstCharacterClass.Star:
-					CurrentToken = new Token(TokenType.Multiply, Flags);
+					_currentToken = new Token(TokenType.Multiply, flags);
 					return true;
 				case FirstCharacterClass.Slash:
-					CurrentToken = new Token(TokenType.Divide, Flags);
+					_currentToken = new Token(TokenType.Divide, flags);
 					return true;
 				case FirstCharacterClass.PercentSign:
-					CurrentToken = new Token(TokenType.Modulo, Flags);
+					_currentToken = new Token(TokenType.Modulo, flags);
 					return true;
 				case FirstCharacterClass.Tilde:
-					CurrentToken = new Token(TokenType.BitwiseNot, Flags);
+					_currentToken = new Token(TokenType.BitwiseNot, flags);
 					return true;
 				case FirstCharacterClass.LeftParen:
-					CurrentToken = new Token(TokenType.LeftParen, Flags);
+					_currentToken = new Token(TokenType.LeftParen, flags);
 					return true;
 				case FirstCharacterClass.RightParen:
-					CurrentToken = new Token(TokenType.RightParen, Flags);
+					_currentToken = new Token(TokenType.RightParen, flags);
 					return true;
 				case FirstCharacterClass.Hash:
-					if(ReadCharacter('#'))
+					if (ReadCharacter('#'))
 					{
-						CurrentToken = new Token(TokenType.HashHash, Flags);
+						_currentToken = new Token(TokenType.HashHash, flags);
 					}
 					else
 					{
-						CurrentToken = new Token(TokenType.Hash, Flags);
+						_currentToken = new Token(TokenType.Hash, flags);
 					}
 					return true;
 				case FirstCharacterClass.Comma:
-					CurrentToken = new Token(TokenType.Comma, Flags);
+					_currentToken = new Token(TokenType.Comma, flags);
 					return true;
 				default:
 					throw new NotImplementedException();
@@ -427,15 +419,15 @@ namespace UnrealBuildTool
 		/// <returns>True if a token was parsed</returns>
 		public bool MoveNextTokenString()
 		{
-			TokenFlags Flags = SkipWhitespace();
+			SkipWhitespace();
 
-			int StartOffset = Offset;
-			while(Data[Offset] != 0 && Data[Offset] != '\n')
+			int startOffset = _offset;
+			while (_data[_offset] != 0 && _data[_offset] != '\n')
 			{
-				Offset++;
+				_offset++;
 			}
 
-			CurrentToken = CreateLiteral(TokenType.StringOfTokens, TokenFlags.None, StartOffset, Offset - StartOffset);
+			_currentToken = CreateLiteral(TokenType.StringOfTokens, TokenFlags.None, startOffset, _offset - startOffset);
 			return true;
 		}
 
@@ -445,39 +437,39 @@ namespace UnrealBuildTool
 		/// <returns>True if a token was read</returns>
 		public bool MoveNextIncludePath()
 		{
-			bool bResult = MoveNext();
-			if(bResult && Current.Type == TokenType.CompareLess)
+			bool result = MoveNext();
+			if (result && Current.Type == TokenType.CompareLess)
 			{
-				int StartOffset = Offset - 1;
-				for(int EndOffset = Offset; Data[EndOffset] != 0 && Data[EndOffset] != '\n'; EndOffset++)
+				int startOffset = _offset - 1;
+				for (int endOffset = _offset; _data[endOffset] != 0 && _data[endOffset] != '\n'; endOffset++)
 				{
-					if(Data[EndOffset] == '>')
+					if (_data[endOffset] == '>')
 					{
-						CurrentToken = CreateLiteral(TokenType.SystemInclude, TokenFlags.None, StartOffset, EndOffset + 1 - StartOffset);
-						Offset = EndOffset + 1;
+						_currentToken = CreateLiteral(TokenType.SystemInclude, TokenFlags.None, startOffset, endOffset + 1 - startOffset);
+						_offset = endOffset + 1;
 						break;
 					}
 				}
 			}
-			return bResult;
+			return result;
 		}
 
 		/// <summary>
 		/// Creates a literal token
 		/// </summary>
-		/// <param name="Type">Type of token to create</param>
-		/// <param name="Flags">Flags for the token</param>
-		/// <param name="Offset">Offset of the literal within the source data stream</param>
-		/// <param name="Length">Length of the literal</param>
+		/// <param name="type">Type of token to create</param>
+		/// <param name="flags">Flags for the token</param>
+		/// <param name="offset">Offset of the literal within the source data stream</param>
+		/// <param name="length">Length of the literal</param>
 		/// <returns>New token for the literal</returns>
-		private Token CreateLiteral(TokenType Type, TokenFlags Flags, int Offset, int Length)
+		private Token CreateLiteral(TokenType type, TokenFlags flags, int offset, int length)
 		{
-			byte[] Literal = new byte[Length];
-			for(int Idx = 0; Idx < Length; Idx++)
+			byte[] literal = new byte[length];
+			for (int idx = 0; idx < length; idx++)
 			{
-				Literal[Idx] = Data[Offset + Idx];
+				literal[idx] = _data[offset + idx];
 			}
-			return new Token(Type, Flags, Literal);
+			return new Token(type, flags, literal);
 		}
 
 		/// <summary>
@@ -487,61 +479,61 @@ namespace UnrealBuildTool
 		public bool MoveToNextDirective()
 		{
 			// Scan lines until we reach a directive
-			for(;;)
+			for (; ; )
 			{
 				// Move to the next newline
-				while(Data[Offset] != '\n')
+				while (_data[_offset] != '\n')
 				{
-					if(Data[Offset] == 0)
+					if (_data[_offset] == 0)
 					{
 						return false;
 					}
-					else if(Data[Offset] == '\\' && Data[Offset + 1] == '\r' && Data[Offset + 2] == '\n')
+					else if (_data[_offset] == '\\' && _data[_offset + 1] == '\r' && _data[_offset + 2] == '\n')
 					{
 						LineNumberAfterToken++;
-						Offset += 3;
+						_offset += 3;
 					}
-					else if(Data[Offset] == '\\' && Data[Offset + 1] == '\n')
+					else if (_data[_offset] == '\\' && _data[_offset + 1] == '\n')
 					{
 						LineNumberAfterToken++;
-						Offset += 2;
+						_offset += 2;
 					}
-					else if (Data[Offset] == '/' && Data[Offset + 1] == '*')
+					else if (_data[_offset] == '/' && _data[_offset + 1] == '*')
 					{
-						Offset += 2;
-						for(; Data[Offset] != 0; Offset++)
+						_offset += 2;
+						for (; _data[_offset] != 0; _offset++)
 						{
-							if(Data[Offset] == '\n')
+							if (_data[_offset] == '\n')
 							{
 								LineNumberAfterToken++;
 								continue;
 							}
-							if(Data[Offset] == '*' && Data[Offset + 1] == '/')
+							if (_data[_offset] == '*' && _data[_offset + 1] == '/')
 							{
-								Offset += 2;
+								_offset += 2;
 								break;
 							}
 						}
 					}
 					else
 					{
-						Offset++;
+						_offset++;
 					}
 				}
 
 				// Move past the newline
-				Offset++;
+				_offset++;
 				LineNumberAfterToken++;
 
 				// Skip any horizontal whitespace
-				TokenFlags Flags = SkipWhitespace();
+				TokenFlags flags = SkipWhitespace();
 
 				// Check if this is a line marker
-				if(Data[Offset] == '#' && Data[Offset + 1] != '#')
+				if (_data[_offset] == '#' && _data[_offset + 1] != '#')
 				{
-					Offset++;
+					_offset++;
 					LineNumber = LineNumberAfterToken;
-					CurrentToken = new Token(TokenType.Hash, Flags);
+					_currentToken = new Token(TokenType.Hash, flags);
 					return true;
 				}
 			}
@@ -558,13 +550,13 @@ namespace UnrealBuildTool
 		/// <summary>
 		/// Attempts to read a given character from the stream
 		/// </summary>
-		/// <param name="Character">Character to read</param>
+		/// <param name="character">Character to read</param>
 		/// <returns>True if the character was read (and the current position was updated)</returns>
-		bool ReadCharacter(char Character)
+		bool ReadCharacter(char character)
 		{
-			if(Data[Offset] == Character)
+			if (_data[_offset] == character)
 			{
-				Offset++;
+				_offset++;
 				return true;
 			}
 			return false;
@@ -576,68 +568,68 @@ namespace UnrealBuildTool
 		/// <returns>Flags for the following token</returns>
 		TokenFlags SkipWhitespace()
 		{
-			TokenFlags Flags = TokenFlags.None;
-			for(;;)
+			TokenFlags flags = TokenFlags.None;
+			for (; ; )
 			{
 				// Quickly skip over trivial whitespace
-				while((Characters[Data[Offset]] & CharacterFlags.Whitespace) != 0)
+				while ((s_characters[_data[_offset]] & CharacterFlags.Whitespace) != 0)
 				{
-					Offset++;
-					Flags |= TokenFlags.HasLeadingSpace;
+					_offset++;
+					flags |= TokenFlags.HasLeadingSpace;
 				}
 
 				// Look at what's next
-				char Character = (char)Data[Offset];
-				if (Character == '\\')
+				char character = (char)_data[_offset];
+				if (character == '\\')
 				{
-					byte NextCharacter = Data[Offset + 1];
-					if(NextCharacter == '\r' && Data[Offset + 2] == '\n')
+					byte nextCharacter = _data[_offset + 1];
+					if (nextCharacter == '\r' && _data[_offset + 2] == '\n')
 					{
 						LineNumberAfterToken++;
-						Offset += 3;
+						_offset += 3;
 					}
-					else if (NextCharacter == '\n')
+					else if (nextCharacter == '\n')
 					{
 						LineNumberAfterToken++;
-						Offset += 2;
+						_offset += 2;
 					}
 					else
 					{
 						break;
 					}
 				}
-				else if (Character == '/')
+				else if (character == '/')
 				{
-					byte NextCharacter = Data[Offset + 1];
-					if (NextCharacter == '/')
+					byte nextCharacter = _data[_offset + 1];
+					if (nextCharacter == '/')
 					{
-						Offset += 2;
-						for(; Data[Offset] != 0; Offset++)
+						_offset += 2;
+						for (; _data[_offset] != 0; _offset++)
 						{
-							if(Data[Offset] == '\n')
+							if (_data[_offset] == '\n')
 							{
 								break;
 							}
 						}
-						Flags |= TokenFlags.HasLeadingSpace;
+						flags |= TokenFlags.HasLeadingSpace;
 					}
-					else if (NextCharacter == '*')
+					else if (nextCharacter == '*')
 					{
-						Offset += 2;
-						for(; Data[Offset] != 0; Offset++)
+						_offset += 2;
+						for (; _data[_offset] != 0; _offset++)
 						{
-							if(Data[Offset] == '\n')
+							if (_data[_offset] == '\n')
 							{
 								LineNumberAfterToken++;
 								continue;
 							}
-							if(Data[Offset] == '*' && Data[Offset + 1] == '/')
+							if (_data[_offset] == '*' && _data[_offset + 1] == '/')
 							{
-								Offset += 2;
+								_offset += 2;
 								break;
 							}
 						}
-						Flags |= TokenFlags.HasLeadingSpace;
+						flags |= TokenFlags.HasLeadingSpace;
 					}
 					else
 					{
@@ -646,9 +638,9 @@ namespace UnrealBuildTool
 				}
 				else
 				{
-					if(Character == '\r')
+					if (character == '\r')
 					{
-						Offset++;
+						_offset++;
 					}
 					else
 					{
@@ -656,30 +648,30 @@ namespace UnrealBuildTool
 					}
 				}
 			}
-			return Flags;
+			return flags;
 		}
 
 		/// <summary>
 		/// Skip past a text literal (a quoted character literal or string literal)
 		/// </summary>
-		/// <param name="LastCharacter">The terminating character to look for, ignoring escape sequences</param>
-		void SkipTextLiteral(char LastCharacter)
+		/// <param name="lastCharacter">The terminating character to look for, ignoring escape sequences</param>
+		void SkipTextLiteral(char lastCharacter)
 		{
-			for(;;)
+			for (; ; )
 			{
-				char Character = (char)Data[Offset];
-				if(Character == '\0')
+				char character = (char)_data[_offset];
+				if (character == '\0')
 				{
 					throw new Exception("Unexpected end of file in text literal");
 				}
 
-				Offset++;
+				_offset++;
 
-				if(Character == '\\' && Data[Offset] != 0)
+				if (character == '\\' && _data[_offset] != 0)
 				{
-					Offset++;
+					_offset++;
 				}
-				else if(Character == LastCharacter)
+				else if (character == lastCharacter)
 				{
 					break;
 				}
@@ -691,15 +683,15 @@ namespace UnrealBuildTool
 		/// </summary>
 		void SkipNumericLiteral()
 		{
-			for (;;)
+			for (; ; )
 			{
-				while((Characters[Data[Offset]] & CharacterFlags.NumberTail) != 0)
+				while ((s_characters[_data[_offset]] & CharacterFlags.NumberTail) != 0)
 				{
-					Offset++;
+					_offset++;
 				}
-				if((Data[Offset] == '+' || Data[Offset] == '-') && (Data[Offset - 1] == 'e' || Data[Offset - 1] == 'E'))
+				if ((_data[_offset] == '+' || _data[_offset] == '-') && (_data[_offset - 1] == 'e' || _data[_offset - 1] == 'E'))
 				{
-					Offset++;
+					_offset++;
 				}
 				else
 				{
@@ -713,109 +705,109 @@ namespace UnrealBuildTool
 		/// </summary>
 		static TokenReader()
 		{
-			FirstCharacters = new FirstCharacterClass[256];
+			s_firstCharacters = new FirstCharacterClass[256];
 
 			// End of file
-			FirstCharacters[0] = FirstCharacterClass.Terminator;
+			s_firstCharacters[0] = FirstCharacterClass.Terminator;
 
 			// Horizontal whitespace
-			FirstCharacters[' '] = FirstCharacterClass.Whitespace;
-			FirstCharacters['\t'] = FirstCharacterClass.Whitespace;
-			FirstCharacters['\v'] = FirstCharacterClass.Whitespace;
-			FirstCharacters['\r'] = FirstCharacterClass.Whitespace;
+			s_firstCharacters[' '] = FirstCharacterClass.Whitespace;
+			s_firstCharacters['\t'] = FirstCharacterClass.Whitespace;
+			s_firstCharacters['\v'] = FirstCharacterClass.Whitespace;
+			s_firstCharacters['\r'] = FirstCharacterClass.Whitespace;
 
 			// Newline
-			FirstCharacters['\n'] = FirstCharacterClass.Newline;
+			s_firstCharacters['\n'] = FirstCharacterClass.Newline;
 
 			// Identifiers
-			FirstCharacters['_'] = FirstCharacterClass.Identifier;
-			for(int Idx = 'a'; Idx <= 'z'; Idx++)
+			s_firstCharacters['_'] = FirstCharacterClass.Identifier;
+			for (int idx = 'a'; idx <= 'z'; idx++)
 			{
-				FirstCharacters[Idx] = FirstCharacterClass.Identifier;
+				s_firstCharacters[idx] = FirstCharacterClass.Identifier;
 			}
-			for(int Idx = 'A'; Idx <= 'Z'; Idx++)
+			for (int idx = 'A'; idx <= 'Z'; idx++)
 			{
-				FirstCharacters[Idx] = FirstCharacterClass.Identifier;
+				s_firstCharacters[idx] = FirstCharacterClass.Identifier;
 			}
 
 			// Numeric literals
-			for(int Idx = '0'; Idx <= '9'; Idx++)
+			for (int idx = '0'; idx <= '9'; idx++)
 			{
-				FirstCharacters[Idx] = FirstCharacterClass.Number;
+				s_firstCharacters[idx] = FirstCharacterClass.Number;
 			}
 
 			// Character literals
-			FirstCharacters['\''] = FirstCharacterClass.Character;
+			s_firstCharacters['\''] = FirstCharacterClass.Character;
 
 			// String literals
-			FirstCharacters['\"'] = FirstCharacterClass.String;
+			s_firstCharacters['\"'] = FirstCharacterClass.String;
 
 			// Other symbols
-			FirstCharacters['.'] = FirstCharacterClass.Dot;
-			FirstCharacters['?'] = FirstCharacterClass.QuestionMark;
-			FirstCharacters[':'] = FirstCharacterClass.Colon;
-			FirstCharacters['!'] = FirstCharacterClass.ExclamationMark;
-			FirstCharacters['|'] = FirstCharacterClass.Pipe;
-			FirstCharacters['&'] = FirstCharacterClass.Ampersand;
-			FirstCharacters['^'] = FirstCharacterClass.Caret;
-			FirstCharacters['='] = FirstCharacterClass.Equals;
-			FirstCharacters['<'] = FirstCharacterClass.LeftTriangleBracket;
-			FirstCharacters['>'] = FirstCharacterClass.RightTriangleBracket;
-			FirstCharacters['+'] = FirstCharacterClass.Plus;
-			FirstCharacters['-'] = FirstCharacterClass.Minus;
-			FirstCharacters['*'] = FirstCharacterClass.Star;
-			FirstCharacters['/'] = FirstCharacterClass.Slash;
-			FirstCharacters['%'] = FirstCharacterClass.PercentSign;
-			FirstCharacters['~'] = FirstCharacterClass.Tilde;
-			FirstCharacters['('] = FirstCharacterClass.LeftParen;
-			FirstCharacters[')'] = FirstCharacterClass.RightParen;
-			FirstCharacters['#'] = FirstCharacterClass.Hash;
-			FirstCharacters[','] = FirstCharacterClass.Comma;
+			s_firstCharacters['.'] = FirstCharacterClass.Dot;
+			s_firstCharacters['?'] = FirstCharacterClass.QuestionMark;
+			s_firstCharacters[':'] = FirstCharacterClass.Colon;
+			s_firstCharacters['!'] = FirstCharacterClass.ExclamationMark;
+			s_firstCharacters['|'] = FirstCharacterClass.Pipe;
+			s_firstCharacters['&'] = FirstCharacterClass.Ampersand;
+			s_firstCharacters['^'] = FirstCharacterClass.Caret;
+			s_firstCharacters['='] = FirstCharacterClass.Equals;
+			s_firstCharacters['<'] = FirstCharacterClass.LeftTriangleBracket;
+			s_firstCharacters['>'] = FirstCharacterClass.RightTriangleBracket;
+			s_firstCharacters['+'] = FirstCharacterClass.Plus;
+			s_firstCharacters['-'] = FirstCharacterClass.Minus;
+			s_firstCharacters['*'] = FirstCharacterClass.Star;
+			s_firstCharacters['/'] = FirstCharacterClass.Slash;
+			s_firstCharacters['%'] = FirstCharacterClass.PercentSign;
+			s_firstCharacters['~'] = FirstCharacterClass.Tilde;
+			s_firstCharacters['('] = FirstCharacterClass.LeftParen;
+			s_firstCharacters[')'] = FirstCharacterClass.RightParen;
+			s_firstCharacters['#'] = FirstCharacterClass.Hash;
+			s_firstCharacters[','] = FirstCharacterClass.Comma;
 
 			// Flags for secondary characters
-			Characters = new CharacterFlags[256];
+			s_characters = new CharacterFlags[256];
 
 			// Identifiers
-			Characters['_'] |= CharacterFlags.Identifier;
-			for(int Idx = 'a'; Idx <= 'z'; Idx++)
+			s_characters['_'] |= CharacterFlags.Identifier;
+			for (int idx = 'a'; idx <= 'z'; idx++)
 			{
-				Characters[Idx] |= CharacterFlags.Identifier;
+				s_characters[idx] |= CharacterFlags.Identifier;
 			}
-			for(int Idx = 'A'; Idx <= 'Z'; Idx++)
+			for (int idx = 'A'; idx <= 'Z'; idx++)
 			{
-				Characters[Idx] |= CharacterFlags.Identifier;
+				s_characters[idx] |= CharacterFlags.Identifier;
 			}
-			for(int Idx = '0'; Idx <= '9'; Idx++)
+			for (int idx = '0'; idx <= '9'; idx++)
 			{
-				Characters[Idx] |= CharacterFlags.Identifier;
+				s_characters[idx] |= CharacterFlags.Identifier;
 			}
 
 			// Numbers
-			for(int Idx = '0'; Idx <= '9'; Idx++)
+			for (int idx = '0'; idx <= '9'; idx++)
 			{
-				Characters[Idx] |= CharacterFlags.Digit;
+				s_characters[idx] |= CharacterFlags.Digit;
 			}
 
 			// Preprocessing number tail
-			Characters['.'] |= CharacterFlags.NumberTail;
-			Characters['_'] |= CharacterFlags.NumberTail;
-			for(int Idx = '0'; Idx <= '9'; Idx++)
+			s_characters['.'] |= CharacterFlags.NumberTail;
+			s_characters['_'] |= CharacterFlags.NumberTail;
+			for (int idx = '0'; idx <= '9'; idx++)
 			{
-				Characters[Idx] |= CharacterFlags.NumberTail;
+				s_characters[idx] |= CharacterFlags.NumberTail;
 			}
-			for(int Idx = 'a'; Idx <= 'z'; Idx++)
+			for (int idx = 'a'; idx <= 'z'; idx++)
 			{
-				Characters[Idx] |= CharacterFlags.NumberTail;
+				s_characters[idx] |= CharacterFlags.NumberTail;
 			}
-			for(int Idx = 'A'; Idx <= 'Z'; Idx++)
+			for (int idx = 'A'; idx <= 'Z'; idx++)
 			{
-				Characters[Idx] |= CharacterFlags.NumberTail;
+				s_characters[idx] |= CharacterFlags.NumberTail;
 			}
 
 			// Whitespace
-			Characters[' '] |= CharacterFlags.Whitespace;
-			Characters['\t'] |= CharacterFlags.Whitespace;
-			Characters['\v'] |= CharacterFlags.Whitespace;
+			s_characters[' '] |= CharacterFlags.Whitespace;
+			s_characters['\t'] |= CharacterFlags.Whitespace;
+			s_characters['\v'] |= CharacterFlags.Whitespace;
 		}
 	}
 }
