- Big cleanup of Modern mode to simplify the .xcconfig files, and the generating code, as well as moveing less used stuff to other files

diff --git a/Engine/Source/Programs/UnrealBuildTool/ProjectFiles/Xcode/XcodeProject.File.cs b/Engine/Source/Programs/UnrealBuildTool/ProjectFiles/Xcode/XcodeProject.File.cs
new file mode 100644
--- /dev/null
+++ b/Engine/Source/Programs/UnrealBuildTool/ProjectFiles/Xcode/XcodeProject.File.cs
@@ -0,0 +1,581 @@
+﻿// Copyright Epic Games, Inc. All Rights Reserved.
+
+using System;
+using System.Collections.Generic;
+using System.IO;
+using System.Linq;
+using System.Text;
+using EpicGames.Core;
+using UnrealBuildBase;
+
+namespace UnrealBuildTool.XcodeProjectXcconfig
+{
+	/// <summary>
+	/// Info needed to make a file a member of specific group
+	/// </summary>
+	class XcodeSourceFile : ProjectFile.SourceFile
+	{
+		/// <summary>
+		/// Constructor
+		/// </summary>
+		public XcodeSourceFile(FileReference InitFilePath, DirectoryReference? InitRelativeBaseFolder, string? FileRefGuid=null)
+			: base(InitFilePath, InitRelativeBaseFolder)
+		{
+			FileGuid = XcodeProjectFileGenerator.MakeXcodeGuid();
+			this.FileRefGuid = FileRefGuid ?? XcodeProjectFileGenerator.MakeXcodeGuid();
+		}
+
+		/// <summary>
+		/// File Guid for use in Xcode project
+		/// </summary>
+		public string FileGuid
+		{
+			get;
+			private set;
+		}
+
+		public void ReplaceGuids(string NewFileGuid, string NewFileRefGuid)
+		{
+			FileGuid = NewFileGuid;
+			FileRefGuid = NewFileRefGuid;
+		}
+
+		/// <summary>
+		/// File reference Guid for use in Xcode project
+		/// </summary>
+		public string FileRefGuid
+		{
+			get;
+			private set;
+		}
+	}
+
+	/// <summary>
+	/// Represents a group of files shown in Xcode's project navigator as a folder
+	/// </summary>
+	internal class XcodeFileGroup
+	{
+		public XcodeFileGroup(string InName, string InPath, bool InIsReference)
+		{
+			GroupName = InName;
+			GroupPath = InPath;
+			GroupGuid = XcodeProjectFileGenerator.MakeXcodeGuid();
+			bIsReference = InIsReference;
+		}
+
+		public XcodeFileGroup(XcodeFileGroup Other)
+		{
+			GroupName = Other.GroupName;
+			GroupPath = Other.GroupPath;
+			GroupGuid = Other.GroupGuid;
+			bIsReference = Other.bIsReference;
+			Files.AddRange(Other.Files);
+			foreach (string Key in Other.Children.Keys)
+			{
+				Children.Add(Key, new XcodeFileGroup(Other.Children[Key]));
+			}
+		}
+
+		public string GroupGuid;
+		public string GroupName;
+		public string GroupPath;
+		public Dictionary<string, XcodeFileGroup> Children = new Dictionary<string, XcodeFileGroup>();
+		public List<XcodeSourceFile> Files = new List<XcodeSourceFile>();
+		public bool bIsReference;
+	}
+
+	class XcconfigFile
+	{
+		public string Name;
+		public string Guid;
+		public FileReference FileRef;
+		internal StringBuilder Text;
+
+		public XcconfigFile(DirectoryReference XcodeProjectDirectory, UnrealTargetPlatform? Platform, string ConfigName)
+		{
+			Name = ConfigName;
+			FileRef = FileReference.Combine(XcodeProjectDirectory, "Xcconfigs" + (Platform == null ? "" : Platform.ToString()), $"{ConfigName}.xcconfig");
+			Guid = XcodeProjectFileGenerator.MakeXcodeGuid();
+			Text = new StringBuilder();
+		}
+
+		public void AppendLine(string Line)
+		{
+			Text.Append(Line);
+			Text.Append(ProjectFileGenerator.NewLine);
+		}
+
+		public void Write()
+		{
+			// write the file to disk
+			DirectoryReference.CreateDirectory(FileRef.Directory);
+			FileReference.WriteAllTextIfDifferent(FileRef, Text.ToString());
+		}
+	}
+
+
+	class XcodeFileCollection
+	{
+		static readonly string[] ResourceExtensions =
+		{
+			".xcassets",
+			".storyboardc",
+			".storyboard",
+			".xcframework",
+			".framework",
+			".lproj",
+			// no extensions will be folders, like "cookeddata" that we want to copy as a resource
+			"",
+		};
+
+		static readonly string[] SourceCodeExtensions =
+		{
+			".c",
+			".cc",
+			".cpp",
+			".m",
+			".mm",
+		};
+
+		static readonly Dictionary<string, string> FileTypeMap = new()
+		{
+			{ ".c", "sourcecode.c.objc" },
+			{ ".m", "sourcecode.c.objc" },
+			{ ".cc", "sourcecode.cpp.objcpp" },
+			{ ".cpp", "sourcecode.cpp.objcpp" },
+			{ ".mm", "sourcecode.cpp.objcpp" },
+			{ ".h", "sourcecode.c.h" },
+			{ ".hpp", "sourcecode.cpp.h" },
+			{ ".inl", "sourcecode.c.h" },
+			{ ".pch", "sourcecode.c.h" },
+			{ ".framework", "wrapper.framework" },
+			{ ".plist", "text.plist.xml" },
+			{ ".png", "image.png" },
+			{ ".icns", "image.icns" },
+			{ ".xcassets", "folder.assetcatalog" },
+			{ ".storyboardc", "wrapper.storyboardc" },
+			{ ".storyboard", "file.storyboard" },
+			{ ".xcframework", "wrapper.xcframework" },
+			{ ".bundle", "\"wrapper.plug-in\"" }
+		};
+
+		const string FileTypeDefault = "file.txt";
+
+
+		class ManualFileReference
+		{
+			public string Guid = XcodeProjectFileGenerator.MakeXcodeGuid();
+			public string RelativePath = "";
+			public string FileKey = "";
+			public string FileType = "";
+			public string SourceTree = "";
+			public string GroupName = "";
+		}
+
+		// we need relative paths, so this is the directory that paths are relative to
+		private DirectoryReference ProjectDirectory;
+
+		// the location of the uproject file, or the engine directory for non-uproject projects
+		public DirectoryReference RootDirectory;
+
+		public string MainGroupGuid = "";
+		//public string ProductRefGroupGuid = XcodeProjectFileGenerator.MakeXcodeGuid();
+
+		internal Dictionary<string, XcodeFileGroup> Groups = new();
+		internal Dictionary<XcodeSourceFile, FileReference?> BuildableFilesToResponseFile = new();
+		internal List<Tuple<XcodeSourceFile, string>> BuildableResourceFiles = new();
+		internal List<XcodeSourceFile> AllFiles = new();
+
+		private List<ManualFileReference> ManualFiles = new();
+		private Dictionary<string, string> GuidsForGroups = new();
+
+
+		public XcodeFileCollection(XcodeProjectFile ProjectFile)
+		{
+			ProjectDirectory = ProjectFile.ProjectFilePath.Directory;
+			RootDirectory = Unreal.EngineDirectory;
+			SetUProjectLocation(ProjectFile.UnrealData.UProjectFileLocation);
+		}
+
+		public XcodeFileCollection(XcodeFileCollection Other)
+		{
+			ProjectDirectory = Other.ProjectDirectory;
+			RootDirectory = Other.RootDirectory;
+			MainGroupGuid = Other.MainGroupGuid;
+			AllFiles.AddRange(Other.AllFiles);
+			BuildableResourceFiles.AddRange(Other.BuildableResourceFiles);
+			ManualFiles.AddRange(Other.ManualFiles);
+
+			BuildableFilesToResponseFile = new(Other.BuildableFilesToResponseFile);
+			GuidsForGroups = new(Other.GuidsForGroups);
+
+			foreach (string Key in Other.Groups.Keys)
+			{
+				Groups.Add(Key, new XcodeFileGroup(Other.Groups[Key]));
+			}
+
+		}
+
+		public void SetUProjectLocation(FileReference? UProjectLocation)
+		{
+			RootDirectory = UProjectLocation == null || UProjectLocation.ContainsName("Programs", 0) ? (Unreal.EngineDirectory) : UProjectLocation!.Directory;
+		}
+
+		public void ProcessFile(XcodeSourceFile File, bool bIsForBuild, bool bIsFolder, string? GroupName=null, string ExtraResourceSettings="")
+		{
+			// remember all buildable files
+			if (bIsForBuild)
+			{
+				string FileExtension = File.Reference.GetExtension();
+
+				if (IsSourceCode(FileExtension))
+				{
+					// look if it contains any of the exluded names, and is so, don't include it
+					if (!ExcludedFolders.Any(x => File.Reference.ContainsName(x, 0)))
+					{
+						// will fill in the response file later
+						BuildableFilesToResponseFile[File] = null;
+					}
+				}
+				else if (IsResourceFile(FileExtension))
+				{
+					BuildableResourceFiles.Add(Tuple.Create(File, ExtraResourceSettings));
+				}
+			}
+			
+			if (GroupName != "")
+			{
+				// group the files by path
+				XcodeFileGroup? Group = GroupName == null ? FindGroupByAbsolutePath(File.Reference.Directory) : null;
+				if (Group != null)
+				{
+					AllFiles.Add(File);
+					Group.Files.Add(File);
+				}
+				else
+				{
+					GroupName = GroupName ?? (File.Reference.IsUnderDirectory(Unreal.EngineDirectory) ? "EngineReferences" : "ExternalReferences");
+					if (bIsFolder)
+					{
+						AddFolderReference(File.FileRefGuid, File.Reference.MakeRelativeTo(ProjectDirectory), GroupName);
+					}
+					else
+					{
+						AddFileReference(File.FileRefGuid, File.Reference.FullName, GetFileKey(File.Reference.GetExtension()), GetFileType(File.Reference.GetExtension()), "<absolute>", GroupName);
+					}
+				}
+			}
+		}
+
+		/// <summary>
+		/// The project needs to point to the product group, so we expose the Guid here
+		/// </summary>
+		/// <returns></returns>
+		public string GetProductGroupGuid()
+		{
+			return GuidsForGroups["Products"];
+		}
+
+		public void AddFileReference(string Guid, string RelativePath, string FileKey, string FileType, string SourceTree/*="SOURCE_ROOT"*/, string GroupName)
+		{
+			ManualFiles.Add(new ManualFileReference()
+			{
+				Guid = Guid,
+				RelativePath = RelativePath,
+				FileKey = FileKey,
+				FileType = FileType,
+				SourceTree = SourceTree.StartsWith("<") ? $"\"{SourceTree}\"" : SourceTree,
+				GroupName = GroupName,
+			});
+
+			// make sure each unique group has a guid
+			if (!GuidsForGroups.ContainsKey(GroupName))
+			{
+				GuidsForGroups[GroupName] = XcodeProjectFileGenerator.MakeXcodeGuid();
+			}
+		}
+
+		public void AddFolderReference(string Guid, string RelativePath, string GroupName)
+		{
+			ManualFiles.Add(new ManualFileReference()
+			{
+				Guid = Guid,
+				RelativePath = RelativePath,
+				FileKey = "lastKnownFileType",
+				FileType = "folder",
+				SourceTree = "\"<group>\"",
+				GroupName = GroupName,
+			});
+
+			// make sure each unique group has a guid
+			if (!GuidsForGroups.ContainsKey(GroupName))
+			{
+				GuidsForGroups[GroupName] = XcodeProjectFileGenerator.MakeXcodeGuid();
+			}
+		}
+
+
+		public void Write(StringBuilder Content)
+		{
+			Content.WriteLine("/* Begin PBXBuildFile section */");
+			foreach (KeyValuePair<XcodeSourceFile, FileReference?> Pair in BuildableFilesToResponseFile)
+			{
+				string CompileFlags = "";
+				if (Pair.Value != null)
+				{
+					CompileFlags = $" settings = {{COMPILER_FLAGS = \"@{Pair.Value.FullName}\"; }};";
+				}
+				Content.WriteLine($"\t\t{Pair.Key.FileGuid} = {{isa = PBXBuildFile; fileRef = {Pair.Key.FileRefGuid};{CompileFlags} }}; /* {Pair.Key.Reference.GetFileName()} */");
+			}
+			foreach (Tuple<XcodeSourceFile, string> Resource in BuildableResourceFiles)
+			{
+				XcodeSourceFile File = Resource.Item1;
+				string ExtraSettings = Resource.Item2;
+				Content.WriteLine($"\t\t{File.FileGuid} = {{isa = PBXBuildFile; fileRef = {File.FileRefGuid}; {ExtraSettings}}}; /* {File.Reference.GetFileName()} */");
+			}
+			Content.WriteLine("/* End PBXBuildFile section */");
+			Content.WriteLine();
+
+
+			Content.WriteLine("/* Begin PBXFileReference section */");
+			foreach (XcodeSourceFile File in AllFiles)
+			{
+				string FileName = File.Reference.GetFileName();
+				string FileExtension = Path.GetExtension(FileName);
+				string FileType = GetFileType(FileExtension);
+				string FileKey = GetFileKey(FileExtension);
+				string RelativePath = Utils.CleanDirectorySeparators(File.Reference.MakeRelativeTo(ProjectDirectory), '/');
+				string SourceTree = "SOURCE_ROOT";
+
+				Content.WriteLine($"\t\t{File.FileRefGuid} = {{isa = PBXFileReference; {FileKey} = {FileType}; name = \"{FileName}\"; path = \"{RelativePath}\"; sourceTree = {SourceTree}; }};");
+			}
+			foreach (ManualFileReference File in ManualFiles)
+			{
+				string FileName = Path.GetFileName(File.RelativePath);
+				Content.WriteLine($"\t\t{File.Guid} = {{isa = PBXFileReference; {File.FileKey} = {File.FileType}; name = \"{FileName}\"; path = \"{File.RelativePath}\"; sourceTree = {File.SourceTree}; }};");
+			}
+			Content.WriteLine("/* End PBXFileReference section */");
+			Content.WriteLine();
+
+			WriteGroups(Content);
+		}
+
+		private void WriteGroups(StringBuilder Content)
+		{
+			XcodeFileGroup? RootGroup = FindRootFileGroup(Groups);
+			if (RootGroup == null)
+			{
+				return;
+			}
+
+			string XcconfigsRefGroupGuid = XcodeProjectFileGenerator.MakeXcodeGuid();
+
+			Content.WriteLine("/* Begin PBXGroup section */");
+
+
+			// write main/root group and it's children (adding the manual groups to the main group, which is the true param here0
+			WriteGroup(Content, RootGroup, bIsRootGroup: true);
+
+			// write some manual groups
+			foreach (var Pair in GuidsForGroups)
+			{
+				Content.WriteLine($"\t\t{Pair.Value} /* {Pair.Key} */ = {{");
+				Content.WriteLine("\t\t\tisa = PBXGroup;");
+				Content.WriteLine("\t\t\tchildren = (");
+				foreach (ManualFileReference Ref in ManualFiles.Where(x => x.GroupName == Pair.Key))
+				{
+					Content.WriteLine($"\t\t\t\t{Ref.Guid} /* {Path.GetFileName(Ref.RelativePath)} */,");
+				}
+				Content.WriteLine("\t\t\t);");
+				Content.WriteLine($"\t\t\tname = {Pair.Key};");
+				Content.WriteLine("\t\t\tsourceTree = \"<group>\";");
+				Content.WriteLine("\t\t};");
+
+			}
+
+			Content.WriteLine("/* End PBXGroup section */");
+		}
+
+
+		private void WriteGroup(StringBuilder Content, XcodeFileGroup Group, bool bIsRootGroup)
+		{
+			if (!Group.bIsReference)
+			{
+				Content.WriteLine($"\t\t{Group.GroupGuid} = {{");
+				Content.WriteLine("\t\t\tisa = PBXGroup;");
+				Content.WriteLine("\t\t\tchildren = (");
+				foreach (XcodeFileGroup ChildGroup in Group.Children.Values)
+				{
+					Content.WriteLine($"\t\t\t\t{ChildGroup.GroupGuid} /* {ChildGroup.GroupName} */,");
+				}
+				foreach (XcodeSourceFile File in Group.Files)
+				{
+					Content.WriteLine($"\t\t\t\t{File.FileRefGuid} /* {File.Reference.GetFileName()} */,");
+				}
+
+				if (bIsRootGroup)
+				{
+					foreach (var Pair in GuidsForGroups)
+					{
+						Content.WriteLine($"\t\t\t\t{Pair.Value} /* {Pair.Key} */,");
+					}
+				}
+				Content.WriteLine("\t\t\t);");
+				if (!bIsRootGroup)
+				{
+					Content.WriteLine("\t\t\tname = \"" + Group.GroupName + "\";");
+					Content.WriteLine("\t\t\tpath = \"" + Group.GroupPath + "\";");
+					Content.WriteLine("\t\t\tsourceTree = \"<absolute>\";");
+				}
+				else
+				{
+					Content.WriteLine("\t\t\tsourceTree = \"<group>\";");
+				}
+				Content.WriteLine("\t\t};");
+
+				foreach (XcodeFileGroup ChildGroup in Group.Children.Values)
+				{
+					WriteGroup(Content, ChildGroup, bIsRootGroup: false);
+				}
+			}
+		}
+
+		public static string GetRootGroupGuid(Dictionary<string, XcodeFileGroup> GroupsDict)
+		{
+			return FindRootFileGroup(GroupsDict)!.GroupGuid;
+		}
+
+		private static XcodeFileGroup? FindRootFileGroup(Dictionary<string, XcodeFileGroup> GroupsDict)
+		{
+			foreach (XcodeFileGroup Group in GroupsDict.Values)
+			{
+				if (Group.Children.Count > 1 || Group.Files.Count > 0)
+				{
+					return Group;
+				}
+				else
+				{
+					XcodeFileGroup? Found = FindRootFileGroup(Group.Children);
+					if (Found != null)
+					{
+						return Found;
+					}
+				}
+			}
+			return null;
+		}
+
+
+		/// <summary>
+		/// Gets Xcode file type based on its extension
+		/// </summary>
+		internal static string GetFileType(string Extension)
+		{
+			string? FileType;
+			if (FileTypeMap.TryGetValue(Extension, out FileType))
+			{
+				return FileType;
+			}
+			return FileTypeDefault;
+		}
+
+		// most types Xcode wants as explicitFileType, but some want lastKnownFileType
+		internal static string GetFileKey(string Extension)
+		{
+			switch (Extension)
+			{
+				case ".storyboard":
+				case ".bundle":
+					return "lastKnownFileType";
+				default:
+					return "explicitFileType";
+			}
+		}
+		/// <summary>
+		/// Returns true if Extension is a known extension for files containing source code
+		/// </summary>
+		internal static bool IsSourceCode(string Extension)
+		{
+			return SourceCodeExtensions.Contains(Extension);
+		}
+
+		/// <summary>
+		/// Returns true if Extension is a known extension for files containing resource data
+		/// </summary>
+		internal static bool IsResourceFile(string Extension)
+		{
+			return ResourceExtensions.Contains(Extension);
+		}
+
+		/// <summary>
+		/// Cache the folders excluded by Mac
+		/// </summary>
+		static string[] ExcludedFolders = PlatformExports.GetExcludedFolderNames(UnrealTargetPlatform.Mac);
+
+		/// <summary>
+		/// Returns a project navigator group to which the file should belong based on its path.
+		/// Creates a group tree if it doesn't exist yet.
+		/// </summary>
+		internal XcodeFileGroup? FindGroupByAbsolutePath(DirectoryReference DirectoryPath)
+		{
+			string AbsolutePath = DirectoryPath.FullName;
+			string[] Parts = AbsolutePath.Split(Path.DirectorySeparatorChar);
+			string CurrentPath = "/";
+			Dictionary<string, XcodeFileGroup> CurrentSubGroups = Groups;
+
+			for (int Index = 1; Index < Parts.Count(); ++Index)
+			{
+				string Part = Parts[Index];
+
+				if (CurrentPath.Length > 1)
+				{
+					CurrentPath += Path.DirectorySeparatorChar;
+				}
+
+				CurrentPath += Part;
+
+				XcodeFileGroup CurrentGroup;
+				if (!CurrentSubGroups.ContainsKey(CurrentPath))
+				{
+					// we only want files under the project/engine dir, we have to make external references to other files, otherwise the grouping gets all broken
+					if (!DirectoryPath.IsUnderDirectory(RootDirectory))
+					{
+						return null;
+					}
+
+
+					CurrentGroup = new XcodeFileGroup(Path.GetFileName(CurrentPath), CurrentPath, GetFileType(Path.GetExtension(CurrentPath)).StartsWith("folder"));
+					if (Groups.Count == 0)
+					{
+						MainGroupGuid = CurrentGroup.GroupGuid;
+					}
+					CurrentSubGroups.Add(CurrentPath, CurrentGroup);
+				}
+				else
+				{
+					CurrentGroup = CurrentSubGroups[CurrentPath];
+				}
+
+				if (CurrentPath == AbsolutePath)
+				{
+					return CurrentGroup;
+				}
+
+				CurrentSubGroups = CurrentGroup.Children;
+			}
+
+			return null;
+		}
+
+		/// <summary>
+		/// Convert all paths to Apple/Unix format (with forward slashes)
+		/// </summary>
+		/// <param name="InPath">The path to convert</param>
+		/// <returns>The normalized path</returns>
+		internal static string ConvertPath(string InPath)
+		{
+			return InPath.Replace("\\", "/");
+		}
+	}
+
+}
diff --git a/Engine/Source/Programs/UnrealBuildTool/ProjectFiles/Xcode/XcodeProject.Scheme.cs b/Engine/Source/Programs/UnrealBuildTool/ProjectFiles/Xcode/XcodeProject.Scheme.cs
new file mode 100644
--- /dev/null
+++ b/Engine/Source/Programs/UnrealBuildTool/ProjectFiles/Xcode/XcodeProject.Scheme.cs
@@ -0,0 +1,294 @@
+﻿// Copyright Epic Games, Inc. All Rights Reserved.
+
+using System;
+using System.Collections.Generic;
+using System.IO;
+using System.Linq;
+using System.Text;
+using EpicGames.Core;
+
+
+namespace UnrealBuildTool.XcodeProjectXcconfig
+{
+	static class XcodeSchemeFile
+	{ 
+		private static FileReference GetUserSchemeManagementFilePath(DirectoryReference ProjectFile, UnrealTargetPlatform? Platform)
+		{
+			return FileReference.Combine(XcodeUtils.ProjectDirPathForPlatform(ProjectFile, Platform), $"xcuserdata/{Environment.UserName}.xcuserdatad/xcschemes/xcschememanagement.plist");
+		}
+
+		private static DirectoryReference GetProjectSchemeDirectory(DirectoryReference ProjectFile, UnrealTargetPlatform? Platform)
+		{
+			return DirectoryReference.Combine(XcodeUtils.ProjectDirPathForPlatform(ProjectFile, Platform), "xcshareddata/xcschemes");
+		}
+
+		private static FileReference GetProjectSchemeFilePathForTarget(DirectoryReference ProjectFile, UnrealTargetPlatform? Platform, string TargetName)
+		{
+			return FileReference.Combine(GetProjectSchemeDirectory(ProjectFile, Platform), TargetName + ".xcscheme");
+		}
+
+		public static void WriteSchemeFile(UnrealData UnrealData, UnrealTargetPlatform? Platform, string ProjectName, List<XcodeRunTarget> RunTargets, string? BuildTargetGuid, string? IndexTargetGuid)
+		{
+			StringBuilder Content = new StringBuilder();
+
+			string ProductName = UnrealData.ProductName;
+			string GameProjectPath = UnrealData.UProjectFileLocation != null ? UnrealData.UProjectFileLocation.FullName : "";
+		
+			foreach (XcodeRunTarget RunTarget in RunTargets)
+			{
+				string TargetName = RunTarget.Name;
+				string TargetGuid = RunTarget.Guid;
+				
+				bool bHasEditorConfiguration = RunTarget.BuildConfigList!.BuildConfigs.Any(x => x.Info.ProjectTarget!.TargetRules!.Type == TargetType.Editor);
+				bool bUseEditorConfiguration = bHasEditorConfiguration && !UnrealData.bMakeProjectPerTarget && !XcodeProjectFileGenerator.bGenerateRunOnlyProject;
+				string DefaultConfiguration = bUseEditorConfiguration ? "Development Editor" : "Development";
+				
+				FileReference SchemeFilePath = GetProjectSchemeFilePathForTarget(UnrealData.XcodeProjectFileLocation, Platform, TargetName);
+
+				DirectoryReference.CreateDirectory(SchemeFilePath.Directory);
+
+				string? OldCommandLineArguments = null;
+				if (FileReference.Exists(SchemeFilePath))
+				{
+					string OldContents = File.ReadAllText(SchemeFilePath.FullName);
+					int OldCommandLineArgumentsStart = OldContents.IndexOf("<CommandLineArguments>") + "<CommandLineArguments>".Length;
+					int OldCommandLineArgumentsEnd = OldContents.IndexOf("</CommandLineArguments>");
+					if (OldCommandLineArgumentsStart != -1 && OldCommandLineArgumentsEnd != -1)
+					{
+						OldCommandLineArguments = OldContents.Substring(OldCommandLineArgumentsStart, OldCommandLineArgumentsEnd - OldCommandLineArgumentsStart);
+					}
+				}
+
+				string[] ScriptLines = new string[]
+				{
+					"# if a UBTGenerated plist doesn&apos;t exist, create it now for Xcode&apos;s dependency needs",
+					"if [ ! -f &quot;${PROJECT_DIR}/${INFOPLIST_FILE}&quot; ]",
+					"then",
+					"  echo Creating ${PROJECT_DIR}/${INFOPLIST_FILE}...",
+					"  touch &quot;${PROJECT_DIR}/${INFOPLIST_FILE}&quot;",
+					"fi",
+					"",
+					"# if we had previously archived the build, then Xcode would have made a soft-link, but it will fail when building normally, so delete a soft-link",
+					"if [ -L &quot;${CONFIGURATION_BUILD_DIR}/${PRODUCT_NAME}.app&quot; ]",
+					"then",
+					"  rm &quot;${CONFIGURATION_BUILD_DIR}/${PRODUCT_NAME}.app&quot;",
+					"fi",
+					"",
+					"# make sure we have a cookeddata directory if needed",
+					"if [[ &quot;${PLATFORM_NAME}&quot; == &quot;iphoneos&quot; || &quot;${PLATFORM_NAME}&quot; == &quot;appletvos&quot; ]]",
+					"then",
+					"  mkdir -p &quot;${CONFIGURATION_BUILD_DIR}/../../../Saved/StagedBuilds/${UE_TARGET_PLATFORM_NAME}/cookeddata&quot;",
+					"fi",
+					"",
+				};
+				// convert special characters to be xml happy
+				IEnumerable<string> Lines = ScriptLines.Concat(UnrealData.ExtraPreBuildScriptLines);
+				Lines = Lines.Select(x => x.Replace("\"", "&quot;").Replace("&&", "&amp;&amp;"));
+
+				Content.WriteLine("<?xml version=\"1.0\" encoding=\"UTF-8\"?>");
+				Content.WriteLine("<Scheme");
+				Content.WriteLine("   LastUpgradeVersion = \"2000\"");
+				Content.WriteLine("   version = \"1.3\">");
+				Content.WriteLine("   <BuildAction");
+				Content.WriteLine("      parallelizeBuildables = \"YES\"");
+				Content.WriteLine("      buildImplicitDependencies = \"YES\">");
+				Content.WriteLine("      <PreActions>");
+				Content.WriteLine("	        <ExecutionAction");
+				Content.WriteLine("            ActionType = \"Xcode.IDEStandardExecutionActionsCore.ExecutionActionType.ShellScriptAction\">");
+				Content.WriteLine("            <ActionContent");
+				Content.WriteLine("               title = \"Touch UBT generated tiles\"");
+				Content.WriteLine("               scriptText = \"" + string.Join("&#10;", Lines) + "\">");
+				Content.WriteLine("               <EnvironmentBuildable>");
+				Content.WriteLine("                  <BuildableReference");
+				Content.WriteLine("                     BuildableIdentifier = \"primary\"");
+				Content.WriteLine("                     BlueprintIdentifier = \"" + TargetGuid + "\"");
+				Content.WriteLine("                     BuildableName = \"" + ProductName + ".app\"");
+				Content.WriteLine("                     BlueprintName = \"" + TargetName + "\"");
+				Content.WriteLine("                     ReferencedContainer = \"container:" + ProjectName + ".xcodeproj\">");
+				Content.WriteLine("                  </BuildableReference>");
+				Content.WriteLine("               </EnvironmentBuildable>");
+				Content.WriteLine("            </ActionContent>");
+				Content.WriteLine("         </ExecutionAction>");
+				Content.WriteLine("      </PreActions>");
+				Content.WriteLine("      <BuildActionEntries>");
+				Content.WriteLine("         <BuildActionEntry");
+				Content.WriteLine("            buildForTesting = \"YES\"");
+				Content.WriteLine("            buildForRunning = \"YES\"");
+				Content.WriteLine("            buildForProfiling = \"YES\"");
+				Content.WriteLine("            buildForArchiving = \"YES\"");
+				Content.WriteLine("            buildForAnalyzing = \"YES\">");
+				Content.WriteLine("            <BuildableReference");
+				Content.WriteLine("               BuildableIdentifier = \"primary\"");
+				Content.WriteLine("               BlueprintIdentifier = \"" + TargetGuid + "\"");
+				Content.WriteLine("               BuildableName = \"" + ProductName + ".app\"");
+				Content.WriteLine("               BlueprintName = \"" + TargetName + "\"");
+				Content.WriteLine("               ReferencedContainer = \"container:" + ProjectName + ".xcodeproj\">");
+				Content.WriteLine("            </BuildableReference>");
+				Content.WriteLine("         </BuildActionEntry>");
+				Content.WriteLine("      </BuildActionEntries>");
+				Content.WriteLine("   </BuildAction>");
+				Content.WriteLine("   <TestAction");
+				Content.WriteLine("      buildConfiguration = \"" + DefaultConfiguration + "\"");
+				Content.WriteLine("      selectedDebuggerIdentifier = \"Xcode.DebuggerFoundation.Debugger.LLDB\"");
+				Content.WriteLine("      selectedLauncherIdentifier = \"Xcode.DebuggerFoundation.Launcher.LLDB\"");
+				Content.WriteLine("      shouldUseLaunchSchemeArgsEnv = \"YES\">");
+				Content.WriteLine("      <MacroExpansion>");
+				Content.WriteLine("         <BuildableReference");
+				Content.WriteLine("            BuildableIdentifier = \"primary\"");
+				Content.WriteLine("            BlueprintIdentifier = \"" + TargetGuid + "\"");
+				Content.WriteLine("            BuildableName = \"" + ProductName + ".app\"");
+				Content.WriteLine("            BlueprintName = \"" + TargetName + "\"");
+				Content.WriteLine("            ReferencedContainer = \"container:" + ProjectName + ".xcodeproj\">");
+				Content.WriteLine("         </BuildableReference>");
+				Content.WriteLine("      </MacroExpansion>");
+				Content.WriteLine("      <Testables>");
+				Content.WriteLine("      </Testables>");
+				Content.WriteLine("   </TestAction>");
+				Content.WriteLine("   <LaunchAction");
+				Content.WriteLine("      buildConfiguration = \"" + DefaultConfiguration + "\"");
+				Content.WriteLine("      selectedDebuggerIdentifier = \"Xcode.DebuggerFoundation.Debugger.LLDB\"");
+				Content.WriteLine("      selectedLauncherIdentifier = \"Xcode.DebuggerFoundation.Launcher.LLDB\"");
+				Content.WriteLine("      launchStyle = \"0\"");
+				Content.WriteLine("      useCustomWorkingDirectory = \"NO\"");
+				Content.WriteLine("      ignoresPersistentStateOnLaunch = \"NO\"");
+				Content.WriteLine("      debugDocumentVersioning = \"NO\"");
+				Content.WriteLine("      debugServiceExtension = \"internal\"");
+				Content.WriteLine("      allowLocationSimulation = \"YES\">");
+				Content.WriteLine("      <BuildableProductRunnable");
+				Content.WriteLine("         runnableDebuggingMode = \"0\">");
+				Content.WriteLine("         <BuildableReference");
+				Content.WriteLine("            BuildableIdentifier = \"primary\"");
+				Content.WriteLine("            BlueprintIdentifier = \"" + TargetGuid + "\"");
+				Content.WriteLine("            BuildableName = \"" + ProductName + ".app\"");
+				Content.WriteLine("            BlueprintName = \"" + TargetName + "\"");
+				Content.WriteLine("            ReferencedContainer = \"container:" + ProjectName + ".xcodeproj\">");
+				Content.WriteLine("         </BuildableReference>");
+				Content.WriteLine("      </BuildableProductRunnable>");
+				if (string.IsNullOrEmpty(OldCommandLineArguments))
+				{
+					if (bHasEditorConfiguration && GameProjectPath != "")
+					{
+						Content.WriteLine("      <CommandLineArguments>");
+						if (UnrealData.bIsForeignProject)
+						{
+							Content.WriteLine("         <CommandLineArgument");
+							Content.WriteLine("            argument = \"&quot;" + GameProjectPath + "&quot;\"");
+							Content.WriteLine("            isEnabled = \"YES\">");
+							Content.WriteLine("         </CommandLineArgument>");
+						}
+						else
+						{
+							Content.WriteLine("         <CommandLineArgument");
+							Content.WriteLine("            argument = \"" + Path.GetFileNameWithoutExtension(GameProjectPath) + "\"");
+							Content.WriteLine("            isEnabled = \"YES\">");
+							Content.WriteLine("         </CommandLineArgument>");
+						}
+						// Always add a configuration argument
+						Content.WriteLine("         <CommandLineArgument");
+						Content.WriteLine("            argument = \"-RunConfig=$(Configuration)\"");
+						Content.WriteLine("            isEnabled = \"YES\">");
+						Content.WriteLine("         </CommandLineArgument>");
+						Content.WriteLine("      </CommandLineArguments>");
+					}
+				}
+				else
+				{
+					Content.WriteLine("      <CommandLineArguments>" + OldCommandLineArguments + "</CommandLineArguments>");
+				}
+				Content.WriteLine("   </LaunchAction>");
+				Content.WriteLine("   <ProfileAction");
+				Content.WriteLine("      buildConfiguration = \"" + DefaultConfiguration + "\"");
+				Content.WriteLine("      shouldUseLaunchSchemeArgsEnv = \"YES\"");
+				Content.WriteLine("      savedToolIdentifier = \"\"");
+				Content.WriteLine("      useCustomWorkingDirectory = \"NO\"");
+				Content.WriteLine("      debugDocumentVersioning = \"NO\">");
+				Content.WriteLine("      <BuildableProductRunnable");
+				Content.WriteLine("         runnableDebuggingMode = \"0\">");
+				Content.WriteLine("         <BuildableReference");
+				Content.WriteLine("            BuildableIdentifier = \"primary\"");
+				Content.WriteLine("            BlueprintIdentifier = \"" + TargetGuid + "\"");
+				Content.WriteLine("            BuildableName = \"" + ProductName + ".app\"");
+				Content.WriteLine("            BlueprintName = \"" + TargetName + "\"");
+				Content.WriteLine("            ReferencedContainer = \"container:" + ProjectName + ".xcodeproj\">");
+				Content.WriteLine("         </BuildableReference>");
+				Content.WriteLine("      </BuildableProductRunnable>");
+				Content.WriteLine("   </ProfileAction>");
+				Content.WriteLine("   <AnalyzeAction");
+				Content.WriteLine("      buildConfiguration = \"" + DefaultConfiguration + "\">");
+				Content.WriteLine("   </AnalyzeAction>");
+				Content.WriteLine("   <ArchiveAction");
+				Content.WriteLine("      buildConfiguration = \"" + DefaultConfiguration + "\"");
+				Content.WriteLine("      revealArchiveInOrganizer = \"YES\">");
+				Content.WriteLine("   </ArchiveAction>");
+				Content.WriteLine("</Scheme>");
+
+				File.WriteAllText(SchemeFilePath.FullName, Content.ToString(), new UTF8Encoding());
+				Content.Clear();
+			}
+
+			Content.WriteLine("<?xml version=\"1.0\" encoding=\"UTF-8\"?>");
+			Content.WriteLine("<!DOCTYPE plist PUBLIC \"-//Apple//DTD PLIST 1.0//EN\" \"http://www.apple.com/DTDs/PropertyList-1.0.dtd\">");
+			Content.WriteLine("<plist version=\"1.0\">");
+			Content.WriteLine("<dict>");
+			Content.WriteLine("\t<key>SchemeUserState</key>");
+			Content.WriteLine("\t<dict>");
+			foreach (XcodeRunTarget RunTarget in RunTargets)
+			{
+				Content.WriteLine("\t\t<key>" + RunTarget.Name + ".xcscheme_^#shared#^_</key>");
+				Content.WriteLine("\t\t<dict>");
+				Content.WriteLine("\t\t\t<key>orderHint</key>");
+				Content.WriteLine("\t\t\t<integer>1</integer>");
+				Content.WriteLine("\t\t</dict>");
+			}
+			Content.WriteLine("\t</dict>");
+			Content.WriteLine("\t<key>SuppressBuildableAutocreation</key>");
+			Content.WriteLine("\t<dict>");
+			foreach (XcodeRunTarget RunTarget in RunTargets)
+			{
+				Content.WriteLine("\t\t<key>" + RunTarget.Guid + "</key>");
+				Content.WriteLine("\t\t<dict>");
+				Content.WriteLine("\t\t\t<key>primary</key>");
+				Content.WriteLine("\t\t\t<true/>");
+				Content.WriteLine("\t\t</dict>");
+			}
+			if (BuildTargetGuid != null)
+			{
+				Content.WriteLine("\t\t<key>" + BuildTargetGuid + "</key>");
+				Content.WriteLine("\t\t<dict>");
+				Content.WriteLine("\t\t\t<key>primary</key>");
+				Content.WriteLine("\t\t\t<true/>");
+				Content.WriteLine("\t\t</dict>");
+			}
+			if (IndexTargetGuid != null)
+			{
+				Content.WriteLine("\t\t<key>" + IndexTargetGuid + "</key>");
+				Content.WriteLine("\t\t<dict>");
+				Content.WriteLine("\t\t\t<key>primary</key>");
+				Content.WriteLine("\t\t\t<true/>");
+				Content.WriteLine("\t\t</dict>");
+			}
+			Content.WriteLine("\t</dict>");
+			Content.WriteLine("</dict>");
+			Content.WriteLine("</plist>");
+
+			FileReference ManagementFile = GetUserSchemeManagementFilePath(UnrealData.XcodeProjectFileLocation, Platform);
+			if (!DirectoryReference.Exists(ManagementFile.Directory))
+			{
+				DirectoryReference.CreateDirectory(ManagementFile.Directory);
+			}
+
+			File.WriteAllText(ManagementFile.FullName, Content.ToString(), new UTF8Encoding());
+		}
+
+		public static void CleanSchemeFile(UnrealData UnrealData, UnrealTargetPlatform? Platform)
+		{
+			// clean this up because we don't want it persisting if we narrow our project list
+			DirectoryReference SchemeDir = GetProjectSchemeDirectory(UnrealData.XcodeProjectFileLocation, Platform);
+
+			if (DirectoryReference.Exists(SchemeDir))
+			{
+				DirectoryReference.Delete(SchemeDir, true);
+			}
+		}
+	}
+}
diff --git a/Engine/Source/Programs/UnrealBuildTool/ProjectFiles/Xcode/XcodeProject.Utils.cs b/Engine/Source/Programs/UnrealBuildTool/ProjectFiles/Xcode/XcodeProject.Utils.cs
new file mode 100644
--- /dev/null
+++ b/Engine/Source/Programs/UnrealBuildTool/ProjectFiles/Xcode/XcodeProject.Utils.cs
@@ -0,0 +1,288 @@
+﻿// Copyright Epic Games, Inc. All Rights Reserved.
+
+using System;
+using System.Collections.Generic;
+using System.IO;
+using System.Linq;
+using System.Text;
+using EpicGames.Core;
+using Microsoft.Extensions.Logging;
+using System.Text.RegularExpressions;
+
+
+namespace UnrealBuildTool.XcodeProjectXcconfig
+{
+	static class StringBuilderExtensions
+	{
+		public static void WriteLine(this StringBuilder SB, string Line = "")
+		{
+			SB.Append(Line);
+			SB.Append(ProjectFileGenerator.NewLine);
+		}
+		public static void WriteLine(this StringBuilder SB, int Indent, string Line = "")
+		{
+			SB.Append(new String('\t', Indent));
+			SB.Append(Line);
+			SB.Append(ProjectFileGenerator.NewLine);
+		}
+	}
+
+	static class XcodeUtils
+	{
+		// null platform means use the old way, with multiple platforms per project
+		private static string Suffix(UnrealTargetPlatform? Platform)
+		{
+			return Platform == null ? "" : $" ({Platform})";
+		}
+
+		public static DirectoryReference ProjectDirPathForPlatform(DirectoryReference ProjectFilePath, UnrealTargetPlatform? Platform)
+		{
+			FileReference ProjectAsFile = new FileReference(ProjectFilePath.FullName);
+			return DirectoryReference.Combine(ProjectAsFile.Directory, $"{ProjectAsFile.GetFileNameWithoutExtension()}{Suffix(Platform)}{ProjectAsFile.GetExtension()}");
+		}
+
+
+		private static IEnumerable<UnrealTargetConfiguration> GetSupportedConfigurations()
+		{
+			return new UnrealTargetConfiguration[] {
+				UnrealTargetConfiguration.Debug,
+				UnrealTargetConfiguration.DebugGame,
+				UnrealTargetConfiguration.Development,
+				UnrealTargetConfiguration.Test,
+				UnrealTargetConfiguration.Shipping
+			};
+		}
+
+		public static bool ShouldIncludeProjectInWorkspace(ProjectFile Proj, ILogger Logger)
+		{
+			foreach (Project ProjectTarget in Proj.ProjectTargets)
+			{
+				foreach (UnrealTargetPlatform Platform in XcodeProjectFileGenerator.XcodePlatforms)
+				{
+					foreach (UnrealTargetConfiguration Config in GetSupportedConfigurations())
+					{
+						if (MSBuildProjectFile.IsValidProjectPlatformAndConfiguration(ProjectTarget, Platform, Config, Logger))
+						{
+							return true;
+						}
+					}
+				}
+			}
+
+			return false;
+		}
+
+
+		public static string MakeExecutableFileName(string BinaryName, UnrealTargetPlatform Platform, UnrealTargetConfiguration Configuration, UnrealArchitectures Architectures, UnrealTargetConfiguration UndecoratedConfiguration)
+		{
+			StringBuilder Result = new StringBuilder();
+
+			Result.Append(BinaryName);
+
+			if (Configuration != UndecoratedConfiguration)
+			{
+				Result.AppendFormat("-{0}-{1}", Platform.ToString(), Configuration.ToString());
+			}
+
+			if (UnrealArchitectureConfig.ForPlatform(Platform).RequiresArchitectureFilenames(Architectures))
+			{
+				Result.Append(Architectures.ToString());
+			}
+
+			return Result.ToString();
+		}
+
+
+		// cache for the below function
+		static Dictionary<string, UnrealArchitectures> CachedMacProjectArchitectures = new();
+
+		/// <summary>
+		/// Returns the Mac architectures that should be configured for the provided target. If the target has a project we'll adhere
+		/// to whether it's set as Intel/Universal/Apple unless the type is denied (pretty much just Editor)
+		/// 
+		/// If the target has no project we'll support allow-listed targets for installed builds and all non-editor architectures 
+		/// for source builds. Not all programs are going to compile for Apple Silicon, but being able to build and fail is useful...
+		/// </summary>
+		/// <param name="TargetName">The target we're generatin forg</param>
+		/// <param name="InProjectFile">Path to the project file, or null if the target has no project</param>
+		/// <returns></returns>
+		public static UnrealArchitectures GetSupportedMacArchitectures(string TargetName, FileReference? InProjectFile)
+		{
+			// All architectures supported
+			UnrealArchitectures AllArchitectures = new(new[] { UnrealArch.X64, UnrealArch.Arm64 });
+
+			// Add a way on the command line of forcing a project file with all architectures (there isn't a good way to let this be
+			// set and checked where we can access it).
+			bool ForceAllArchitectures = Environment.GetCommandLineArgs().Contains("AllArchitectures", StringComparer.OrdinalIgnoreCase);
+
+			if (ForceAllArchitectures)
+			{
+				return AllArchitectures;
+			}
+
+			// First time seeing this target?
+			if (!CachedMacProjectArchitectures.ContainsKey(TargetName))
+			{
+				CachedMacProjectArchitectures[TargetName] = UnrealArchitectureConfig.ForPlatform(UnrealTargetPlatform.Mac).ProjectSupportedArchitectures(InProjectFile, TargetName);
+			}
+
+			return CachedMacProjectArchitectures[TargetName];
+		}
+
+
+		public static void FindPlistId(MetadataItem PlistItem, string Key, ref string BundleId)
+		{
+			if (PlistItem.File == null || !FileReference.Exists(PlistItem.File))
+			{
+				return;
+			}
+
+			string Identifier = Plist($"Print :{Key}", PlistItem.File.FullName);
+			
+			// handle error
+			if (Identifier == "" || Identifier.StartsWith("Print:"))
+			{
+				if (PlistItem.Mode == MetadataMode.UsePremade)
+				{
+					Log.TraceErrorOnce($"Premade .plist file '{PlistItem.File}' was found, but it did not contain {Key} (Key is missing or value is empty)");
+				}
+			}
+			else
+			{
+				BundleId = Identifier;
+			}
+		}
+
+		private static string? ActivePlistFile;
+
+		public static void SetActivePlistFile(string PlistFile)
+		{
+			ActivePlistFile = PlistFile;
+		}
+
+		public static string Plist(string Command, string PlistFile)
+		{
+			Command = Command.Replace("\"", "\\\"");
+			return Utils.RunLocalProcessAndReturnStdOut("/usr/libexec/PlistBuddy", $"-c \"{Command}\" \"{PlistFile}\"");
+		}
+
+		public static string Plist(string Command)
+		{
+			return Plist(Command, ActivePlistFile!);
+		}
+
+		public static void PlistSetAdd(string Entry, string Value, string Type = "string")
+		{
+			string AddOutput = Plist($"Add {Entry} {Type} {Value}");
+			// error will be non-empty string
+			if (AddOutput != "")
+			{
+				Plist($"Set {Entry} {Value}");
+			}
+		}
+
+		public static bool PlistSetUpdate(string Entry, string Value)
+		{
+			// see if the setting is already there
+			string ExistingSetting = Plist($"Print {Entry}");
+
+			// Print errors start with Print
+			if (!ExistingSetting.StartsWith("Print:") && ExistingSetting != Value)
+			{
+				Plist($"Set {Entry} {Value}");
+
+				return true;
+			}
+
+			return false;
+		}
+
+		public static IEnumerable<string> PlistArray(string Entry)
+		{
+			return Plist($"Print {Entry}")
+				.Replace("Array {", "")
+				.Replace("}", "")
+				.Trim()
+				.ReplaceLineEndings()
+				.Split(Environment.NewLine)
+				.Select(x => x.Trim());
+		}
+
+		public static List<string> PlistObjects()
+		{
+			List<string> Result = new();
+
+			IEnumerable<string> Lines = Plist("print :objects")
+				.ReplaceLineEndings()
+				.Split(Environment.NewLine);
+
+			Regex Regex = new Regex("^\\s*(\\S*) = Dict {$");
+			foreach (string Line in Lines)
+			{
+				Match Match = Regex.Match(Line);
+				if (Match.Success)
+				{
+					Result.Add(Match.Groups[1].Value);
+				}
+			}
+			return Result;
+		}
+
+		public static string? PlistFixPath(string Entry, string RelativeToProject)
+		{
+			string ExistingPath = Plist($"Print {Entry}");
+			// skip of errors, or it's an absolute path
+			if (!ExistingPath.StartsWith("Print:") && !ExistingPath.StartsWith("/"))
+			{
+				// fixup the path to be relative to new project instead of old
+				string FixedPath = Utils.CollapseRelativeDirectories(Path.Combine(RelativeToProject, ExistingPath));
+				// and set it back
+				Plist($"Set {Entry} {FixedPath}");
+
+				return FixedPath;
+			}
+
+			return null;
+		}
+
+
+
+
+		public static List<string> GetSupportedOrientations(ConfigHierarchy Ini)
+		{
+			List<string> Orientations = new();
+
+			bool bSupported = true;
+			if (Ini.TryGetValue("/Script/IOSRuntimeSettings.IOSRuntimeSettings", "bSupportsPortraitOrientation", out bSupported) && bSupported)
+			{
+				Orientations.Add("UIInterfaceOrientationPortrait");
+			}
+			if (Ini.TryGetValue("/Script/IOSRuntimeSettings.IOSRuntimeSettings", "bSupportsUpsideDownOrientation", out bSupported) && bSupported)
+			{
+				Orientations.Add("UIInterfaceOrientationPortraitUpsideDown");
+			}
+
+			string? PreferredLandscapeOrientation;
+			Ini.TryGetValue("/Script/IOSRuntimeSettings.IOSRuntimeSettings", "PreferredLandscapeOrientation", out PreferredLandscapeOrientation);
+			bool bSupportsLandscapeLeft = false;
+			Ini.TryGetValue("/Script/IOSRuntimeSettings.IOSRuntimeSettings", "bSupportsLandscapeLeftOrientation", out bSupportsLandscapeLeft);
+			bool bSupportsLandscapeRight = false;
+			Ini.TryGetValue("/Script/IOSRuntimeSettings.IOSRuntimeSettings", "bSupportsLandscapeRightOrientation", out bSupportsLandscapeRight);
+
+			if (bSupportsLandscapeLeft && PreferredLandscapeOrientation == "LandscapeLeft")
+			{
+				Orientations.Add("UIInterfaceOrientationLandscapeLeft");
+			}
+			if (bSupportsLandscapeRight)
+			{
+				Orientations.Add("UIInterfaceOrientationLandscapeRight");
+			}
+			if (bSupportsLandscapeLeft && PreferredLandscapeOrientation != "LandscapeLeft")
+			{
+				Orientations.Add("UIInterfaceOrientationLandscapeLeft");
+			}
+
+			return Orientations;
+		}
+	}
+}
diff --git a/Engine/Source/Programs/UnrealBuildTool/ProjectFiles/Xcode/XcodeProject.cs b/Engine/Source/Programs/UnrealBuildTool/ProjectFiles/Xcode/XcodeProject.cs
--- a/Engine/Source/Programs/UnrealBuildTool/ProjectFiles/Xcode/XcodeProject.cs
+++ b/Engine/Source/Programs/UnrealBuildTool/ProjectFiles/Xcode/XcodeProject.cs
@@ -3,19 +3,12 @@
 using System;
 using System.Collections.Generic;
 using System.IO;
-using System.Xml;
-using System.Xml.XPath;
-using System.Xml.Linq;
 using System.Linq;
 using System.Text;
 using EpicGames.Core;
 using UnrealBuildBase;
 using Microsoft.Extensions.Logging;
 using System.Text.RegularExpressions;
-using System.Diagnostics.CodeAnalysis;
-using System.Security.Cryptography;
-using static UnrealBuildTool.Unity;
-
 
 /*****
 
@@ -53,22 +46,6 @@ Future ideas:
 
 namespace UnrealBuildTool.XcodeProjectXcconfig
 {
-	static class StringBuilderExtensions
-	{
-		public static void WriteLine(this StringBuilder SB, string Line = "")
-		{
-			SB.Append(Line);
-			SB.Append(ProjectFileGenerator.NewLine);
-		}
-		public static void WriteLine(this StringBuilder SB, int Indent, string Line = "")
-		{
-			SB.Append(new String('\t', Indent));
-			SB.Append(Line);
-			SB.Append(ProjectFileGenerator.NewLine);
-		}
-	}
-
-
 	class UnrealBuildConfig
 	{
 		public UnrealBuildConfig(string InDisplayName, string InBuildTarget, FileReference? InMacExecutablePath, FileReference? InIOSExecutablePath, FileReference? InTVOSExecutablePath,
@@ -259,10 +236,12 @@ namespace UnrealBuildTool.XcodeProjectXcconfig
 	class UnrealData
 	{
 		public bool bIsStubProject;
+		public bool bIsForeignProject;
 		public bool bMakeProjectPerTarget;
 
+		public TargetRules TargetRules => _TargetRules!;
 		public bool bIsAppBundle;
-		public bool bHasEditorConfiguration;
+
 		public bool bUseAutomaticSigning = false;
 		public bool bIsMergingProjects = false;
 		public bool bWriteCodeSigningSettings = true;
@@ -295,7 +274,6 @@ namespace UnrealBuildTool.XcodeProjectXcconfig
 		// Display name, can be overridden from commandline
 		public string DisplayName;
 
-
 		/// <summary>
 		///  Used to mark the project for distribution (some platforms require this)
 		/// </summary>
@@ -322,6 +300,7 @@ namespace UnrealBuildTool.XcodeProjectXcconfig
 		public ILogger? Logger;
 
 		private XcodeProjectFile? ProjectFile;
+		private TargetRules? _TargetRules;
 
 		public static bool bSupportsMac => Supports(UnrealTargetPlatform.Mac);
 		public static bool bSupportsIOS => Supports(UnrealTargetPlatform.IOS);
@@ -372,6 +351,7 @@ namespace UnrealBuildTool.XcodeProjectXcconfig
 		public bool Initialize(XcodeProjectFile ProjectFile, List<UnrealTargetConfiguration> Configurations, ILogger Logger)
 		{
 			this.ProjectFile = ProjectFile;
+			this.bIsForeignProject = ProjectFile.IsForeignProject;
 			this.Logger = Logger;
 
 			FindUProjectFileLocation(ProjectFile);
@@ -410,9 +390,16 @@ namespace UnrealBuildTool.XcodeProjectXcconfig
 				return false;
 			}
 
+			// verify all configs share the same TargetRules
+			if (!AllConfigs.All(x => x.ProjectTarget!.TargetRules == AllConfigs[0].ProjectTarget!.TargetRules))
+			{
+				throw new BuildException("All Configs must share a TargetRules. This indicates bMakeProjectPerTarget is returning false");
+			}
+
+			_TargetRules = AllConfigs[0].ProjectTarget!.TargetRules;
+			
 			// this project makes an app bundle (.app directory instead of a raw executable or dylib) if none of the fings make a non-appbundle
 			bIsAppBundle = !AllConfigs.Any(x => x.ProjectTarget!.TargetRules!.bIsBuildingConsoleApplication || x.ProjectTarget.TargetRules.bShouldCompileAsDLL);
-			bHasEditorConfiguration = AllConfigs.Any(x => x.ProjectTarget!.TargetRules!.Type == TargetType.Editor);
 
 			// read config settings
 			if (AllConfigs.Any(x => x.bSupportsIOS))
@@ -681,14 +668,14 @@ namespace UnrealBuildTool.XcodeProjectXcconfig
 
 											if (BuildPlatform.Platform == UnrealTargetPlatform.Mac)
 											{
-												string MacExecutableName = MakeExecutableFileName(ExeName, UnrealTargetPlatform.Mac, Configuration, ProjectTarget.TargetRules.Architectures, ProjectTarget.TargetRules.UndecoratedConfiguration);
+												string MacExecutableName = XcodeUtils.MakeExecutableFileName(ExeName, UnrealTargetPlatform.Mac, Configuration, ProjectTarget.TargetRules.Architectures, ProjectTarget.TargetRules.UndecoratedConfiguration);
 												string IOSExecutableName = MacExecutableName.Replace("-Mac-", "-IOS-");
 												string TVOSExecutableName = MacExecutableName.Replace("-Mac-", "-TVOS-");
 												BuildConfigs.Add(new UnrealBuildConfig(ConfigName, TargetName, FileReference.Combine(MacBinaryDir, MacExecutableName), FileReference.Combine(IOSBinaryDir, IOSExecutableName), FileReference.Combine(TVOSBinaryDir, TVOSExecutableName), ProjectTarget, Configuration));
 											}
 											else if (BuildPlatform.Platform == UnrealTargetPlatform.IOS || BuildPlatform.Platform == UnrealTargetPlatform.TVOS)
 											{
-												string IOSExecutableName = MakeExecutableFileName(ExeName, UnrealTargetPlatform.IOS, Configuration, ProjectTarget.TargetRules.Architectures, ProjectTarget.TargetRules.UndecoratedConfiguration);
+												string IOSExecutableName = XcodeUtils.MakeExecutableFileName(ExeName, UnrealTargetPlatform.IOS, Configuration, ProjectTarget.TargetRules.Architectures, ProjectTarget.TargetRules.UndecoratedConfiguration);
 												string TVOSExecutableName = IOSExecutableName.Replace("-IOS-", "-TVOS-");
 												//string MacExecutableName = IOSExecutableName.Replace("-IOS-", "-Mac-");
 												BuildConfigs.Add(new UnrealBuildConfig(ConfigName, TargetName, FileReference.Combine(MacBinaryDir, IOSExecutableName), FileReference.Combine(IOSBinaryDir, IOSExecutableName), FileReference.Combine(TVOSBinaryDir, TVOSExecutableName), ProjectTarget, Configuration));
@@ -704,209 +691,8 @@ namespace UnrealBuildTool.XcodeProjectXcconfig
 
 			return BuildConfigs;
 		}
-
-		public static IEnumerable<UnrealTargetConfiguration> GetSupportedConfigurations()
-		{
-			return new UnrealTargetConfiguration[] {
-				UnrealTargetConfiguration.Debug,
-				UnrealTargetConfiguration.DebugGame,
-				UnrealTargetConfiguration.Development,
-				UnrealTargetConfiguration.Test,
-				UnrealTargetConfiguration.Shipping
-			};
-		}
-
-		public static bool ShouldIncludeProjectInWorkspace(ProjectFile Proj, ILogger Logger)
-		{
-			return CanBuildProjectLocally(Proj, Logger);
-		}
-
-		public static bool CanBuildProjectLocally(ProjectFile Proj, ILogger Logger)
-		{
-			foreach (Project ProjectTarget in Proj.ProjectTargets)
-			{
-				foreach (UnrealTargetPlatform Platform in XcodeProjectFileGenerator.XcodePlatforms)
-				{
-					foreach (UnrealTargetConfiguration Config in GetSupportedConfigurations())
-					{
-						if (MSBuildProjectFile.IsValidProjectPlatformAndConfiguration(ProjectTarget, Platform, Config, Logger))
-						{
-							return true;
-						}
-					}
-				}
-			}
-
-			return false;
-		}
-
-		private static string MakeExecutableFileName(string BinaryName, UnrealTargetPlatform Platform, UnrealTargetConfiguration Configuration, UnrealArchitectures Architectures, UnrealTargetConfiguration UndecoratedConfiguration)
-		{
-			StringBuilder Result = new StringBuilder();
-
-			Result.Append(BinaryName);
-
-			if (Configuration != UndecoratedConfiguration)
-			{
-				Result.AppendFormat("-{0}-{1}", Platform.ToString(), Configuration.ToString());
-			}
-
-			UEBuildPlatform BuildPlatform = UEBuildPlatform.GetBuildPlatform(Platform);
-			if (UnrealArchitectureConfig.ForPlatform(Platform).RequiresArchitectureFilenames(Architectures))
-			{
-				Result.Append(Architectures.ToString());
-			}
-
-			return Result.ToString();
-		}
-
-
-		// cache for the below function
-		Dictionary<string, UnrealArchitectures> CachedMacProjectArchitectures = new();
-
-		/// <summary>
-		/// Returns the Mac architectures that should be configured for the provided target. If the target has a project we'll adhere
-		/// to whether it's set as Intel/Universal/Apple unless the type is denied (pretty much just Editor)
-		/// 
-		/// If the target has no project we'll support allow-listed targets for installed builds and all non-editor architectures 
-		/// for source builds. Not all programs are going to compile for Apple Silicon, but being able to build and fail is useful...
-		/// </summary>
-		/// <param name="Config">Build config for the target we're generating</param>
-		/// <param name="InProjectFile">Path to the project file, or null if the target has no project</param>
-		/// <returns></returns>
-		public UnrealArchitectures GetSupportedMacArchitectures(UnrealBuildConfig Config, FileReference? InProjectFile)
-		{
-			// All architectures supported
-			UnrealArchitectures AllArchitectures = new(new[] { UnrealArch.X64, UnrealArch.Arm64 });
-
-			// Add a way on the command line of forcing a project file with all architectures (there isn't a good way to let this be
-			// set and checked where we can access it).
-			bool ForceAllArchitectures = Environment.GetCommandLineArgs().Contains("AllArchitectures", StringComparer.OrdinalIgnoreCase);
-
-			if (ForceAllArchitectures)
-			{
-				return AllArchitectures;
-			}
-
-			string TargetName = Config.BuildTarget;
-
-			// First time seeing this target?
-			if (!CachedMacProjectArchitectures.ContainsKey(TargetName))
-			{
-				CachedMacProjectArchitectures[TargetName] = UnrealArchitectureConfig.ForPlatform(UnrealTargetPlatform.Mac).ProjectSupportedArchitectures(InProjectFile, TargetName);
-			}
-
-			return CachedMacProjectArchitectures[TargetName];
-		}
-
-	}
-
-	/// <summary>
-	/// Info needed to make a file a member of specific group
-	/// </summary>
-	class XcodeSourceFile : ProjectFile.SourceFile
-	{
-		/// <summary>
-		/// Constructor
-		/// </summary>
-		public XcodeSourceFile(FileReference InitFilePath, DirectoryReference? InitRelativeBaseFolder, string? FileRefGuid=null)
-			: base(InitFilePath, InitRelativeBaseFolder)
-		{
-			FileGuid = XcodeProjectFileGenerator.MakeXcodeGuid();
-			this.FileRefGuid = FileRefGuid ?? XcodeProjectFileGenerator.MakeXcodeGuid();
-		}
-
-		/// <summary>
-		/// File Guid for use in Xcode project
-		/// </summary>
-		public string FileGuid
-		{
-			get;
-			private set;
-		}
-
-		public void ReplaceGuids(string NewFileGuid, string NewFileRefGuid)
-		{
-			FileGuid = NewFileGuid;
-			FileRefGuid = NewFileRefGuid;
-		}
-
-		/// <summary>
-		/// File reference Guid for use in Xcode project
-		/// </summary>
-		public string FileRefGuid
-		{
-			get;
-			private set;
-		}
-	}
-
-	/// <summary>
-	/// Represents a group of files shown in Xcode's project navigator as a folder
-	/// </summary>
-	internal class XcodeFileGroup
-	{
-		public XcodeFileGroup(string InName, string InPath, bool InIsReference)
-		{
-			GroupName = InName;
-			GroupPath = InPath;
-			GroupGuid = XcodeProjectFileGenerator.MakeXcodeGuid();
-			bIsReference = InIsReference;
-		}
-
-		public XcodeFileGroup(XcodeFileGroup Other)
-		{
-			GroupName = Other.GroupName;
-			GroupPath = Other.GroupPath;
-			GroupGuid = Other.GroupGuid;
-			bIsReference = Other.bIsReference;
-			Files.AddRange(Other.Files);
-			foreach (string Key in Other.Children.Keys)
-			{
-				Children.Add(Key, new XcodeFileGroup(Other.Children[Key]));
-			}
-		}
-
-		public string GroupGuid;
-		public string GroupName;
-		public string GroupPath;
-		public Dictionary<string, XcodeFileGroup> Children = new Dictionary<string, XcodeFileGroup>();
-		public List<XcodeSourceFile> Files = new List<XcodeSourceFile>();
-		public bool bIsReference;
-	}
-
-	class XcconfigFile
-	{
-		public string Name;
-		public string Guid;
-		public FileReference FileRef;
-		internal StringBuilder Text;
-
-		public XcconfigFile(DirectoryReference XcodeProjectDirectory, UnrealTargetPlatform? Platform, string ConfigName)
-		{
-			Name = ConfigName;
-			FileRef = FileReference.Combine(XcodeProjectDirectory, "Xcconfigs" + (Platform == null ? "" : Platform.ToString()), $"{ConfigName}.xcconfig");
-			Guid = XcodeProjectFileGenerator.MakeXcodeGuid();
-			Text = new StringBuilder();
-		}
-
-		public void AppendLine(string Line)
-		{
-			Text.Append(Line);
-			Text.Append(ProjectFileGenerator.NewLine);
-		}
-
-		public void Write()
-		{
-			// write the file to disk
-			DirectoryReference.CreateDirectory(FileRef.Directory);
-			FileReference.WriteAllTextIfDifferent(FileRef, Text.ToString());
-		}
 	}
 
-
-
-
 	abstract class XcodeProjectNode
 	{
 		// keeps a list of other node this node references, which is used when writing out the whole xcode project file
@@ -951,10 +737,10 @@ namespace UnrealBuildTool.XcodeProjectXcconfig
 		}
 
 
-		public void CreateXcconfigFile(XcodeProject Project, string Name)
+		public void CreateXcconfigFile(XcodeProject Project, UnrealTargetPlatform? Platform, string Name)
 		{
 			DirectoryReference XcodeProjectDirectory = Project.UnrealData.XcodeProjectFileLocation.ParentDirectory!;
-			Xcconfig = new XcconfigFile(XcodeProjectDirectory, Project.Platform, Name);
+			Xcconfig = new XcconfigFile(XcodeProjectDirectory, Platform, Name);
 			Project.FileCollection.AddFileReference(Xcconfig.Guid, Xcconfig.FileRef.MakeRelativeTo(XcodeProjectDirectory), "explicitFileType", "test.xcconfig", "\"<group>\"", "Xcconfigs");
 		}
 
@@ -1186,10 +972,14 @@ namespace UnrealBuildTool.XcodeProjectXcconfig
 	{
 		public string Guid = XcodeProjectFileGenerator.MakeXcodeGuid();
 		public UnrealBuildConfig Info;
+		// Because we don't make project-wide .xcconfig files, we need to specify at the projet level that all of the platforms are supported,
+		// so that the _Build targets will have any platform as a supported platform, otherwise it can only compile for Mac (default Xcode platform)
+		private bool bIncludeAllPlatforms;
 
-		public XcodeBuildConfig(UnrealBuildConfig Info)
+		public XcodeBuildConfig(UnrealBuildConfig Info, bool bIncludeAllPlatforms)
 		{
 			this.Info = Info;
+			this.bIncludeAllPlatforms = bIncludeAllPlatforms;
 		}
 
 		public override void Write(StringBuilder Content)
@@ -1201,6 +991,11 @@ namespace UnrealBuildTool.XcodeProjectXcconfig
 				Content.WriteLine(3,	$"baseConfigurationReference = {Xcconfig.Guid} /* {Xcconfig.Name}.xcconfig */;");
 			}
 			Content.WriteLine(3,	"buildSettings = {");
+			if (bIncludeAllPlatforms)
+			{
+				Content.WriteLine(4,		$"SUPPORTED_PLATFORMS = \"macosx iphonesimulator iphoneos appletvsimulator appletvos\";");
+				Content.WriteLine(4,		$"ONLY_ACTIVE_ARCH = YES;");
+			}
 			Content.WriteLine(3,	"};");
 			Content.WriteLine(3,	$"name = \"{Info.DisplayName}\";");
 			Content.WriteLine(2, "};");
@@ -1224,7 +1019,7 @@ namespace UnrealBuildTool.XcodeProjectXcconfig
 		}
 
 		public UnrealTargetPlatform? Platform;
-		public XcodeBuildConfigList(UnrealTargetPlatform? Platform, string TargetName, UnrealData UnrealData)
+		public XcodeBuildConfigList(UnrealTargetPlatform? Platform, string TargetName, UnrealData UnrealData, bool bIncludeAllPlatformsInConfig)
 		{
 			this.Platform = Platform;
 			this.UnrealData = UnrealData;
@@ -1237,7 +1032,7 @@ namespace UnrealBuildTool.XcodeProjectXcconfig
 			this.TargetName = TargetName;
 
 			// create build config objects for each info passed in, and them as references
-			IEnumerable<XcodeBuildConfig> Configs = UnrealData.AllConfigs.Select(x => new XcodeBuildConfig(x));
+			IEnumerable<XcodeBuildConfig> Configs = UnrealData.AllConfigs.Select(x => new XcodeBuildConfig(x, bIncludeAllPlatformsInConfig));
 			// filter out configs that dont match a platform if we are single-platform mode
 			Configs = Configs.Where(x => Platform == null || x.Info.Supports((UnrealTargetPlatform)Platform));
 			BuildConfigs = Configs.ToList();
@@ -1288,9 +1083,12 @@ namespace UnrealBuildTool.XcodeProjectXcconfig
 		string TargetTypeName;
 		Type TargetType;
 
-		// UE5_Build, EngineTest_SubIndex1, etc
+		// UnrealEngine_Build, etc
 		public string Name;
 
+		// QAGame, QAGameEditor, etc
+		private string UnrealTargetName;
+
 		// list of build configs this target supports (for instance, the Index target only indexes a Development config) 
 		public XcodeBuildConfigList? BuildConfigList;
 
@@ -1335,6 +1133,7 @@ namespace UnrealBuildTool.XcodeProjectXcconfig
 			}
 
 			// set up names
+			UnrealTargetName = UnrealData.XcodeProjectName;
 			Name = OverrideName ?? (UnrealData.XcodeProjectName + ConfigName);
 		}
 
@@ -1362,7 +1161,7 @@ namespace UnrealBuildTool.XcodeProjectXcconfig
 				string BuildToolPath = UEDir + "/Engine/Build/BatchFiles/Mac/XcodeBuild.sh";
 
 				// insert elements to call UBT when building
-				Content.WriteLine(3,	$"buildArgumentsString = \"$(ACTION) $(UE_TARGET_NAME) $(PLATFORM_NAME) $(UE_TARGET_CONFIG) {UProjectParam}\";");
+				Content.WriteLine(3,	$"buildArgumentsString = \"$(ACTION) {UnrealTargetName} $(PLATFORM_NAME) $(CONFIGURATION) {UProjectParam}\";");
 				Content.WriteLine(3,	$"buildToolPath = \"{BuildToolPath}\";");
 				Content.WriteLine(3,	$"buildWorkingDirectory = \"{UEDir}\";");
 			}
@@ -1409,11 +1208,12 @@ namespace UnrealBuildTool.XcodeProjectXcconfig
 		{
 			this.TargetType = TargetType;
 			this.Platform = Platform;
-			BuildConfigList = new XcodeBuildConfigList(Platform, Name, Project.UnrealData);
+			this.UnrealData = Project.UnrealData;
+
+			BuildConfigList = new XcodeBuildConfigList(Platform, Name, Project.UnrealData, bIncludeAllPlatformsInConfig:false);
 			References.Add(BuildConfigList);
 
 			// add the Product item to the project to be visible in left pane
-			UnrealData = Project.UnrealData;
 			Project.FileCollection.AddFileReference(ProductGuid, UnrealData.ProductName, "explicitFileType", Project.UnrealData.bIsAppBundle ? "wrapper.application" : "\"compiled.mach-o.executable\"", "BUILT_PRODUCTS_DIR", "Products");
 
 			if (Project.UnrealData.bIsAppBundle)
@@ -1485,10 +1285,12 @@ namespace UnrealBuildTool.XcodeProjectXcconfig
 				AddDependency(BuildTarget, Project);
 			}
 
-			CreateXcconfigFile(Project, Name);
-
-			// hook up each buildconfig to this Xcconfig 
-			BuildConfigList!.BuildConfigs.ForEach(x => x.Xcconfig = Xcconfig);
+			CreateXcconfigFile(Project, Platform, $"{Name}");
+			// create per-config Xcconfig files
+			foreach (XcodeBuildConfig Config in BuildConfigList.BuildConfigs)
+			{
+				Config.CreateXcconfigFile(Project, Platform, $"{Name}_{Config.Info.DisplayName.Replace(" ", "")}");
+			}
 		}
 
 		protected void ProcessFrameworks(XcodeResourcesBuildPhase ResourcesBuildPhase, XcodeProject Project, XcodeProjectFile ProjectFile, ILogger Logger)
@@ -1597,36 +1399,243 @@ namespace UnrealBuildTool.XcodeProjectXcconfig
 
 		public override void WriteXcconfigFile()
 		{
-			// pull in the stuff that UBT updates
+			// gather general, all-platform, data we are doing to put into the configs
+			UnrealBuildConfig BuildConfig = BuildConfigList!.BuildConfigs[0].Info;
+			TargetRules TargetRules = BuildConfig.ProjectTarget!.TargetRules!;
+
+			MetadataPlatform MetadataPlatform;
+			string TargetPlatformName = Platform.ToString();
+			if (TargetRules.Type != TargetType.Game)
+			{
+				TargetPlatformName += TargetRules.Type.ToString();
+			}
+
+			// get ini file for the platform
 			DirectoryReference ProjectOrEngineDir = UnrealData.UProjectFileLocation?.Directory ?? Unreal.EngineDirectory;
+			ConfigHierarchy PlatformIni = ConfigCache.ReadHierarchy(ConfigHierarchyType.Engine, UnrealData.UProjectFileLocation?.Directory, Platform);
+
+			// settings for all platforms
+			bool bAutomaticSigning = true;
+			string? SigningTeam;
+			string? SigningPrefix;
+			string SupportedPlatforms;
+			string SDKRoot;
+			DirectoryReference ConfigBuildDir;
+			string DeploymentTarget;
+			string DeploymentTargetKey;
+			string? SigningCert = null;
+			string? ProvisioningProfile = null;
+			string? SupportedDevices = null;
+			string? MarketingVersion = null;
+			string BundleIdentifier;
+			List<string> ExtraConfigLines = new();
+
+			// get codesigning info (all platforms)
+			PlatformIni.TryGetValue("XcodeConfiguration", "bUseModernCodeSigning", out bAutomaticSigning);
+			PlatformIni.TryGetValue("XcodeConfiguration", "ModernSigningTeam", out SigningTeam);
+			PlatformIni.TryGetValue("XcodeConfiguration", "ModernSigningPrefix", out SigningPrefix);
+
+
+			if (Platform == UnrealTargetPlatform.Mac)
+			{
+				// editor vs game metadata
+				bool bIsEditor = UnrealData.TargetRules.Type == TargetType.Editor;
+				MetadataPlatform = bIsEditor ? MetadataPlatform.MacEditor : MetadataPlatform.Mac;
+
+				SDKRoot = "macosx";
+				SupportedPlatforms = "macosx";
+				DeploymentTargetKey = "MACOSX_DEPLOYMENT_TARGET";
+				DeploymentTarget = MacToolChain.Settings.MacOSVersion;
+				ConfigBuildDir = BuildConfig.MacExecutablePath!.Directory;
+				BundleIdentifier = bIsEditor ? "com.epicgames.UnrealEditor" : UnrealData.BundleIdentifier;
+
+				// @todo: get a version for  games, like IOS has
+				MarketingVersion = MacToolChain.LoadEngineDisplayVersion();
+
+				string SupportedMacArchitectures = string.Join(" ", XcodeUtils.GetSupportedMacArchitectures(BuildConfig.BuildTarget, UnrealData.UProjectFileLocation).Architectures.Select(x => x.AppleName));
+				ExtraConfigLines.Add($"VALID_ARCHS = {SupportedMacArchitectures}");
+			}
+			else
+			{
+				MetadataPlatform = MetadataPlatform.IOS;
+
+				// get IOS (same as TVOS) BundleID, and if there's a specified plist with a bundleID, use it, as Xcode would warn if they don't match
+				BundleIdentifier = UnrealData.BundleIdentifier;
+
+				// short version string
+				PlatformIni.GetString($"/Script/IOSRuntimeSettings.IOSRuntimeSettings", "VersionInfo", out MarketingVersion);
+
+				if (Platform == UnrealTargetPlatform.IOS)
+				{
+					SDKRoot = "iphoneos";
+					SupportedPlatforms = "iphoneos"; // iphonesimulator
+					DeploymentTargetKey = "IPHONEOS_DEPLOYMENT_TARGET";
+					ConfigBuildDir = DirectoryReference.Combine(BuildConfig.IOSExecutablePath!.Directory, "Payload");
+					SupportedDevices = UnrealData.IOSProjectSettings!.RuntimeDevices;
+					DeploymentTarget = UnrealData.IOSProjectSettings.RuntimeVersion;
+
+					if (UnrealData.IOSProvisioningData != null && UnrealData.bWriteCodeSigningSettings)
+					{
+						ProvisioningProfile = UnrealData.IOSProvisioningData.MobileProvisionUUID;
+						SigningCert = UnrealData.IOSProvisioningData.SigningCertificate ?? (UnrealData.bForDistribution ? "Apple Distribution" : "Apple Developer");
+					}
+
+					// only iphone deals with orientation
+					List<string> SupportedOrientations = XcodeUtils.GetSupportedOrientations(PlatformIni);
+					ExtraConfigLines.Add($"INFOPLIST_KEY_UISupportedInterfaceOrientations = \"{string.Join(" ", SupportedOrientations)}\"");
+				}
+				else // tvos
+				{
+					SDKRoot = "appletvos";
+					SupportedPlatforms = "appletvos"; // appletvsimulator
+					DeploymentTargetKey = "TVOS_DEPLOYMENT_TARGET";
+					ConfigBuildDir = DirectoryReference.Combine(BuildConfig.TVOSExecutablePath!.Directory, "Payload");
+					SupportedDevices = UnrealData.TVOSProjectSettings!.RuntimeDevices;
+					DeploymentTarget = UnrealData.TVOSProjectSettings.RuntimeVersion;
+
+					if (UnrealData.TVOSProvisioningData != null && UnrealData.bWriteCodeSigningSettings)
+					{
+						ProvisioningProfile = UnrealData.TVOSProvisioningData.MobileProvisionUUID;
+						SigningCert = UnrealData.TVOSProvisioningData.SigningCertificate ?? (UnrealData.bForDistribution ? "Apple Distribution" : "Apple Developer");
+					}
+				}
+			}
+
+
+			// get metadata for the platform set above
+			MetadataItem PlistMetadata = UnrealData.Metadata!.PlistFiles[MetadataPlatform];
+			MetadataItem EntitlementsMetadata = UnrealData.Metadata!.PlistFiles[MetadataPlatform];
+			// now pull the bundle id's out, as xcode will warn if they don't match (this has to happen after each platform set bundle id above)
+			XcodeUtils.FindPlistId(PlistMetadata, "CFBundleIdentifier", ref BundleIdentifier);
+
+
+			// include another xcconfig for versions that UBT writes out
 			Xcconfig!.AppendLine($"#include? \"{ProjectOrEngineDir}/Intermediate/Build/Versions.xcconfig\"");
+
+			// write out some UE variables that can be used in premade .plist files, etc
 			Xcconfig.AppendLine("");
+			Xcconfig.AppendLine("// Unreal project-wide variables");
+			Xcconfig.AppendLine($"UE_PRODUCT_NAME = {UnrealData.ProductName}");
+			Xcconfig.AppendLine($"UE_PRODUCT_NAME_STRIPPED = {UnrealData.ProductName.Replace("_", "").Replace(" ", "")}");
+			Xcconfig.AppendLine($"UE_DISPLAY_NAME = {UnrealData.DisplayName}");
+			Xcconfig.AppendLine($"UE_SIGNING_PREFIX = {SigningPrefix}");
+			Xcconfig.AppendLine($"UE_ENGINE_DIR = {Unreal.EngineDirectory}");
+			Xcconfig.AppendLine($"UE_PLATFORM_NAME = {Platform}");
+			Xcconfig.AppendLine($"UE_TARGET_NAME = {UnrealData.TargetRules.Name}");
+			Xcconfig.AppendLine($"UE_TARGET_PLATFORM_NAME = {TargetPlatformName}");
 
+			Xcconfig.AppendLine("");
+			Xcconfig.AppendLine("// Constant settings (same for all platforms and targets)");
 			// #jira UE-143619: Pre Monterey macOS requires this option for a packaged app to run on iOS15 due to new code signature format. Could be removed once Monterey is miniuS.
 			Xcconfig.AppendLine("OTHER_CODE_SIGN_FLAGS = --generate-entitlement-der");
-
-			Xcconfig.AppendLine($"MACOSX_DEPLOYMENT_TARGET = {MacToolChain.Settings.MacOSVersion};");
 			Xcconfig.AppendLine("INFOPLIST_OUTPUT_FORMAT = xml");
 			Xcconfig.AppendLine("COMBINE_HIDPI_IMAGES = YES");
+			Xcconfig.AppendLine("USE_HEADERMAP = NO");
+			Xcconfig.AppendLine("ONLY_ACTIVE_ARCH = YES");
 
-			if (Platform == UnrealTargetPlatform.Mac)
+
+			Xcconfig.AppendLine("");
+			Xcconfig.AppendLine("// Platform settings");
+			Xcconfig.AppendLine($"SUPPORTED_PLATFORMS = {SupportedPlatforms}");
+			Xcconfig.AppendLine($"SDKROOT = {SDKRoot}");
+			Xcconfig.AppendLine($"CONFIGURATION_BUILD_DIR = {ConfigBuildDir}");
+			foreach (string Line in ExtraConfigLines)
 			{
-				Xcconfig.AppendLine($"SUPPORTED_PLATFORMS = macosx");
-				Xcconfig.AppendLine($"SDKROOT = macosx");
+				Xcconfig.AppendLine(Line);
 			}
-			else if (Platform == UnrealTargetPlatform.IOS)
+
+			Xcconfig.AppendLine("");
+			Xcconfig.AppendLine("// Project settings");
+			Xcconfig.AppendLine($"TARGETED_DEVICE_FAMILY = {SupportedDevices}");
+			Xcconfig.AppendLine($"PRODUCT_BUNDLE_IDENTIFIER = {BundleIdentifier}");
+			Xcconfig.AppendLine($"{DeploymentTargetKey} = {DeploymentTarget}");
+
+
+			Xcconfig.AppendLine("");
+			Xcconfig.AppendLine("// Plist settings");
+			Xcconfig.AppendLine($"INFOPLIST_FILE = {PlistMetadata.XcodeProjectRelative}");
+			if (PlistMetadata.Mode == MetadataMode.UpdateTemplate)
 			{
-				Xcconfig.AppendLine($"SUPPORTED_PLATFORMS = iphoneos");// iphonesimulator");
-				Xcconfig.AppendLine($"SDKROOT = iphoneos");
+				// allow Xcode to generate the final plist file from our input, some INFOPLIST settings and other settings 
+				Xcconfig.AppendLine("GENERATE_INFOPLIST_FILE = YES");
+				Xcconfig.AppendLine($"ASSETCATALOG_COMPILER_APPICON_NAME = AppIcon");
+				Xcconfig.AppendLine($"CURRENT_PROJECT_VERSION = $(UE_{Platform.ToString().ToUpper()}_BUILD_VERSION)");
+				Xcconfig.AppendLine($"MARKETING_VERSION = {MarketingVersion}");
 			}
-			else
+
+			if (UnrealData.bWriteCodeSigningSettings)
 			{
-				Xcconfig.AppendLine($"SUPPORTED_PLATFORMS = appletvos");
-				Xcconfig.AppendLine($"SDKROOT = appletvos");
+				Xcconfig.AppendLine("");
+				Xcconfig.AppendLine("// Code-signing settings");
+				Xcconfig.AppendLine("CODE_SIGN_STYLE = " + (bAutomaticSigning ? "Automatic" : "Manual"));
+				if (EntitlementsMetadata.Mode == MetadataMode.UsePremade)
+				{
+					Xcconfig.AppendLine($"CODE_SIGN_ENTITLEMENTS = {EntitlementsMetadata.XcodeProjectRelative}");
+				}
+				if (!string.IsNullOrEmpty(SigningTeam))
+				{
+					Xcconfig.AppendLine($"DEVELOPMENT_TEAM = {SigningTeam}");
+				}
+				if (!string.IsNullOrEmpty(ProvisioningProfile))
+				{
+					Xcconfig.AppendLine($"PROVISIONING_PROFILE_SPECIFIER = {ProvisioningProfile}");
+				}
+				if (SigningCert != null)
+				{
+					Xcconfig.AppendLine($"CODE_SIGN_IDENTITY = {SigningCert}");
+				}
 			}
 
-
 			Xcconfig.Write();
+
+
+
+			// Now for each config write out the specific settings
+			DirectoryReference? GameDir = UnrealData.UProjectFileLocation?.Directory;
+			string? GamePath = GameDir != null ? XcodeFileCollection.ConvertPath(GameDir.FullName) : null;
+			foreach (UnrealBuildConfig Config in UnrealData.AllConfigs)
+			{
+				XcodeBuildConfig? MatchedConfig = BuildConfigList!.BuildConfigs.FirstOrDefault(x => x.Info == Config);
+				if (MatchedConfig == null || !Config.Supports(Platform))
+				{
+					continue;
+				}
+
+				string ProductName;
+				string ExecutableName;
+				string ExecutableKey = $"UE_{Platform.ToString().ToUpper()}_EXECUTABLE_NAME";
+				if (Platform == UnrealTargetPlatform.Mac)
+				{
+					ExecutableName = Config.MacExecutablePath!.GetFileName();
+					ProductName = ExecutableName;
+				}
+				else // IOS,TVOS
+				{
+					ExecutableName = ((Platform == UnrealTargetPlatform.IOS) ? Config.IOSExecutablePath : Config.TVOSExecutablePath)!.GetFileName();
+					// @todo: change to Executable when we stop using payload!
+					ProductName = Config.BuildTarget;
+				}
+
+
+				// hook up the Buildconfig that matches this info to this xcconfig file
+				XcconfigFile ConfigXcconfig = MatchedConfig.Xcconfig!;
+
+				ConfigXcconfig.AppendLine("// pull in the shared settings for all configs for this target");
+				ConfigXcconfig.AppendLine($"#include \"{Xcconfig.Name}.xcconfig\"");
+				ConfigXcconfig.AppendLine("");
+				ConfigXcconfig.AppendLine("// Unreal per-config variables");
+				ConfigXcconfig.AppendLine($"UE_TARGET_CONFIG = {Config.BuildConfig}");
+				ConfigXcconfig.AppendLine($"PRODUCT_NAME = {ProductName}");
+				ConfigXcconfig.AppendLine($"{ExecutableKey} = {ExecutableName}");
+
+				// debug settings
+				if (Config.BuildConfig == UnrealTargetConfiguration.Debug)
+				{
+					ConfigXcconfig.AppendLine("ENABLE_TESTABILITY = YES");
+				}
+
+				ConfigXcconfig.Write();
+			}
 		}
 	}
 
@@ -1635,7 +1644,7 @@ namespace UnrealBuildTool.XcodeProjectXcconfig
 		public XcodeBuildTarget(XcodeProject Project)
 			: base(XcodeTarget.Type.Build, Project.UnrealData)
 		{
-			BuildConfigList = new XcodeBuildConfigList(Project.Platform, Name, Project.UnrealData);
+			BuildConfigList = new XcodeBuildConfigList(Project.Platform, Name, Project.UnrealData, bIncludeAllPlatformsInConfig: false);
 			References.Add(BuildConfigList);
 		}
 	}
@@ -1650,11 +1659,11 @@ namespace UnrealBuildTool.XcodeProjectXcconfig
 		{
 			UnrealData = Project.UnrealData;
 
-			BuildConfigList = new XcodeBuildConfigList(Project.Platform, Name, UnrealData);
+			BuildConfigList = new XcodeBuildConfigList(Project.Platform, Name, UnrealData, bIncludeAllPlatformsInConfig: false);
 			References.Add(BuildConfigList);
 
-			CreateXcconfigFile(Project, Name);
-			// hook up each buildconfig to this Xcconfig 
+			CreateXcconfigFile(Project, Project.Platform, Name);
+			// hook up each buildconfig to this Xcconfig
 			BuildConfigList.BuildConfigs.ForEach(x => x.Xcconfig = Xcconfig);
 
 
@@ -1675,8 +1684,7 @@ namespace UnrealBuildTool.XcodeProjectXcconfig
 
 		public override void WriteXcconfigFile()
 		{
-			// write out settings that apply whether or not we have subtargets, which #include this one, or no subtargets, and we write out more
-			// @todo move tis to the subtarget and remember it from the Module
+			// write out settings that for compiling natively
 			Xcconfig!.AppendLine("CLANG_CXX_LANGUAGE_STANDARD = c++17");
 			Xcconfig.AppendLine("GCC_WARN_CHECK_SWITCH_STATEMENTS = NO");
 			Xcconfig.AppendLine("GCC_PRECOMPILE_PREFIX_HEADER = YES");
@@ -1730,19 +1738,9 @@ namespace UnrealBuildTool.XcodeProjectXcconfig
 				References.Add(RunTarget);
 			}
 
-			ProjectBuildConfigs = new XcodeBuildConfigList(Platform, UnrealData.ProductName, UnrealData);
+			ProjectBuildConfigs = new XcodeBuildConfigList(Platform, UnrealData.ProductName, UnrealData, bIncludeAllPlatformsInConfig: true);
 			References.Add(ProjectBuildConfigs);
-
-			// create the Projet xcconfig
-			CreateXcconfigFile(this, $"{UnrealData.ProductName}_Project");
-
-			// create per-config Xcconfig files
-			foreach (XcodeBuildConfig Config in ProjectBuildConfigs.BuildConfigs)
-			{
-				Config.CreateXcconfigFile(this, $"{UnrealData.ProductName}_{Config.Info.DisplayName.Replace(" ", "")}");
-			}
-
-
+			
 			// make an indexing target if we aren't just a run-only project, and it has buildable source files
 			if (!XcodeProjectFileGenerator.bGenerateRunOnlyProject && UnrealData.BatchedFiles.Count != 0)
 			{				
@@ -1791,996 +1789,31 @@ namespace UnrealBuildTool.XcodeProjectXcconfig
 
 			Content.WriteLine("/* End PBXProject section */");
 		}
+	}
 
-		private List<string> GetSupportedOrientations(ConfigHierarchy Ini)
+	class XcodeProjectFile : ProjectFile
+	{
+
+		/// <summary>
+		/// Constructs a new project file object
+		/// </summary>
+		/// <param name="InitFilePath">The path to the project file on disk</param>
+		/// <param name="BaseDir">The base directory for files within this project</param>
+		/// <param name="bIsForDistribution">True for distribution builds</param>
+		/// <param name="BundleID">Override option for bundle identifier</param>
+		/// <param name="AppName"></param>
+		/// <param name="bMakeProjectPerTarget"></param>
+		public XcodeProjectFile(FileReference InitFilePath, DirectoryReference BaseDir, bool bIsForDistribution, string BundleID, string AppName, bool bMakeProjectPerTarget)
+			: base(InitFilePath, BaseDir)
 		{
-			List<string> Orientations = new();
+			UnrealData = new UnrealData(InitFilePath, bIsForDistribution, BundleID, AppName, bMakeProjectPerTarget, IsStubProject);
 
-			bool bSupported = true;
-			if (Ini.TryGetValue("/Script/IOSRuntimeSettings.IOSRuntimeSettings", "bSupportsPortraitOrientation", out bSupported) && bSupported)
-			{
-				Orientations.Add("UIInterfaceOrientationPortrait");
-			}
-			if (Ini.TryGetValue("/Script/IOSRuntimeSettings.IOSRuntimeSettings", "bSupportsUpsideDownOrientation", out bSupported) && bSupported)
-			{
-				Orientations.Add("UIInterfaceOrientationPortraitUpsideDown");
-			}
+			// create the container for all the files that will 
+			SharedFileCollection = new XcodeFileCollection(this);
+			FileCollection = SharedFileCollection;
+		}
 
-			string? PreferredLandscapeOrientation;
-			Ini.TryGetValue("/Script/IOSRuntimeSettings.IOSRuntimeSettings", "PreferredLandscapeOrientation", out PreferredLandscapeOrientation);
-			bool bSupportsLandscapeLeft = false;
-			Ini.TryGetValue("/Script/IOSRuntimeSettings.IOSRuntimeSettings", "bSupportsLandscapeLeftOrientation", out bSupportsLandscapeLeft);
-			bool bSupportsLandscapeRight = false;
-			Ini.TryGetValue("/Script/IOSRuntimeSettings.IOSRuntimeSettings", "bSupportsLandscapeRightOrientation", out bSupportsLandscapeRight);
-
-			if (bSupportsLandscapeLeft && PreferredLandscapeOrientation == "LandscapeLeft")
-			{
-				Orientations.Add("UIInterfaceOrientationLandscapeLeft");
-			}
-			if (bSupportsLandscapeRight)
-			{
-				Orientations.Add("UIInterfaceOrientationLandscapeRight");
-			}
-			if (bSupportsLandscapeLeft && PreferredLandscapeOrientation != "LandscapeLeft")
-			{
-				Orientations.Add("UIInterfaceOrientationLandscapeLeft");
-			}
-
-			return Orientations;
-		}
-
-		void FindPlistId(MetadataItem PlistItem, string Key, ref string BundleId)
-		{
-			if (PlistItem.File == null || !FileReference.Exists(PlistItem.File))
-			{
-				return;
-			}
-
-			string Identifier = XcodeProjectFile.Plist($"Print :{Key}", PlistItem.File.FullName);
-			// handle error
-			if (Identifier == "" || Identifier.StartsWith("Print:"))
-			{
-				if (PlistItem.Mode == MetadataMode.UsePremade)
-				{
-					Log.TraceErrorOnce($"Premade .plist file '{PlistItem.File}' was found, but it did not contain {Key} (Key is missing or value is empty)");
-				}
-			}
-			else
-			{
-				BundleId = Identifier;
-			}
-		}
-
-		public override void WriteXcconfigFile()
-		{
-			if (Xcconfig == null)
-			{
-				return;
-			}
-
-			string UEDir = XcodeFileCollection.ConvertPath(Path.GetFullPath(Directory.GetCurrentDirectory() + "../../.."));
-
-			ConfigHierarchy SharedPlatformIni = ConfigCache.ReadHierarchy(ConfigHierarchyType.Engine, UnrealData.UProjectFileLocation?.Directory, UnrealTargetPlatform.Mac);
-
-			// read some settings from new config settings
-			bool bAutomaticSigning = true;
-			string? SigningTeam;
-			string? SigningPrefix;
-			SharedPlatformIni.TryGetValue("XcodeConfiguration", "bUseModernCodeSigning", out bAutomaticSigning);
-			SharedPlatformIni.TryGetValue("XcodeConfiguration", "ModernSigningTeam", out SigningTeam);
-			SharedPlatformIni.TryGetValue("XcodeConfiguration", "ModernSigningPrefix", out SigningPrefix);
-
-			List<string> SupportedDevices = new();
-
-			StringBuilder Content = new();
-
-			// figure out the directory that contains the Binaries and Intermediate directories
-			string ProjectRootDir = (UnrealData.UProjectFileLocation == null) ? $"{UEDir}/Engine" : UnrealData.UProjectFileLocation.Directory.FullName;
-
-			Xcconfig.AppendLine("// Unreal projcet-wide variables");
-			Xcconfig.AppendLine($"UE_PRODUCT_NAME = {UnrealData.ProductName}");
-			Xcconfig.AppendLine($"UE_PRODUCT_NAME_STRIPPED = {UnrealData.ProductName.Replace("_", "").Replace(" ", "")}");
-			Xcconfig.AppendLine($"UE_DISPLAY_NAME = {UnrealData.DisplayName}");
-			Xcconfig.AppendLine($"UE_SIGNING_PREFIX = {SigningPrefix}");
-			Xcconfig.AppendLine($"UE_ENGINE_DIR = {Unreal.EngineDirectory}");
-			Xcconfig.AppendLine("");
-
-
-			// Info.plist generation settings
-
-			////////////////////////
-			// IOS
-			////////////////////////
-
-			if (ProjectBuildConfigs.bSupportsIOS)
-			{
-				string SettingFilter = Platform != null ? "" : "[sdk=iphoneos*]";
-				ConfigHierarchy Ini = ConfigCache.ReadHierarchy(ConfigHierarchyType.Engine, UnrealData.UProjectFileLocation?.Directory, UnrealTargetPlatform.IOS);
-
-				Xcconfig.AppendLine("// IOS Settings");
-
-				SupportedDevices.Add(UnrealData.IOSProjectSettings!.RuntimeDevices);
-
-				Xcconfig.AppendLine($"CONFIGURATION_BUILD_DIR{SettingFilter} = {UnrealData.AllConfigs[0].IOSExecutablePath!.Directory}/Payload");
-				Xcconfig.AppendLine($"IPHONEOS_DEPLOYMENT_TARGET = {UnrealData.IOSProjectSettings.RuntimeVersion}");
-
-				if (UnrealData.Metadata!.PlistFiles[MetadataPlatform.IOS].Mode == MetadataMode.UpdateTemplate)
-				{
-					// allow Xcode to generate the final plist file from our input, some INFOPLIST settings and other settings 
-					Xcconfig.AppendLine("GENERATE_INFOPLIST_FILE = YES");
-					Xcconfig.AppendLine($"ASSETCATALOG_COMPILER_APPICON_NAME = AppIcon");
-					Xcconfig.AppendLine($"CURRENT_PROJECT_VERSION{SettingFilter} = $(UE_IOS_BUILD_VERSION)");
-
-					// short version string
-					string BundleShortVersion;
-					Ini.GetString("/Script/IOSRuntimeSettings.IOSRuntimeSettings", "VersionInfo", out BundleShortVersion);
-					Xcconfig.AppendLine($"MARKETING_VERSION{SettingFilter} = {BundleShortVersion}");
-				}
-
-
-				if (UnrealData.IOSProvisioningData != null && UnrealData.bWriteCodeSigningSettings)
-				{
-					if (!string.IsNullOrEmpty(UnrealData.IOSProvisioningData.MobileProvisionUUID))
-					{
-						Xcconfig.AppendLine($"PROVISIONING_PROFILE_SPECIFIER{SettingFilter} = {UnrealData.IOSProvisioningData.MobileProvisionUUID}");
-					}
-					string IOSCert = UnrealData.IOSProvisioningData.SigningCertificate ?? (UnrealData.bForDistribution ? "Apple Distribution" : "Apple Developer");
-					Xcconfig.AppendLine($"CODE_SIGN_IDENTITY{SettingFilter} = {IOSCert}");
-				}
-
-
-				// Info plist generation settings
-				//Xcconfig.AppendLine("INFOPLIST_KEY_UIStatusBarHidden = YES");
-				//Xcconfig.AppendLine("INFOPLIST_KEY_UIRequiresFullScreen = YES");
-
-				List<string> SupportedOrientations = GetSupportedOrientations(Ini);
-				Content.WriteLine(3, $"INFOPLIST_KEY_UISupportedInterfaceOrientations = \"{string.Join(" ", SupportedOrientations)}\"");
-
-				//if (UnrealData.IOSProjectSettings.BundleName != "" && UnrealData.IOSProjectSettings.BundleName != "[PROJECT_NAME]")
-				//{
-				//	string DisplayName = UnrealData.IOSProjectSettings.BundleDisplayName.Replace("[PROJECT_NAME]", ProjectName).Replace("_", "");
-				//	Content.WriteLine(3, $"INFOPLIST_KEY_CFBundleName = {DisplayName}");
-				//}
-
-				//if (UnrealData.IOSProjectSettings.BundleDisplayName != "" && UnrealData.IOSProjectSettings.BundleDisplayName != "[PROJECT_NAME]")
-				//{
-				//	string DisplayName = UnrealData.IOSProjectSettings.BundleDisplayName.Replace("[PROJECT_NAME]", ProjectName).Replace("_", "");
-				//	Content.WriteLine(3, $"INFOPLIST_KEY_CFBundleDisplayName = {DisplayName}");
-				//}
-
-			}
-
-			////////////////////////
-			// TVOS
-			////////////////////////
-
-			if (ProjectBuildConfigs.bSupportsTVOS)
-			{
-				string SettingFilter = Platform != null ? "" : "[sdk=appletvos*]";
-				ConfigHierarchy Ini = ConfigCache.ReadHierarchy(ConfigHierarchyType.Engine, UnrealData.UProjectFileLocation?.Directory, UnrealTargetPlatform.TVOS);
-
-				Xcconfig.AppendLine("");
-				Xcconfig.AppendLine("// TVOS Settings");
-
-				SupportedDevices.Add(UnrealData.TVOSProjectSettings!.RuntimeDevices);
-
-				Xcconfig.AppendLine($"CONFIGURATION_BUILD_DIR{SettingFilter} = {UnrealData.AllConfigs[0].TVOSExecutablePath!.Directory}/Payload");
-				Xcconfig.AppendLine($"TVOS_DEPLOYMENT_TARGET = {UnrealData.TVOSProjectSettings.RuntimeVersion}");
-
-				if (UnrealData.Metadata!.PlistFiles[MetadataPlatform.IOS].Mode == MetadataMode.UpdateTemplate)
-				{
-					// allow Xcode to generate the final plist file from our input, some INFOPLIST settings and other settings 
-					Xcconfig.AppendLine("GENERATE_INFOPLIST_FILE = YES");
-					Xcconfig.AppendLine($"ASSETCATALOG_COMPILER_APPICON_NAME = AppIcon");
-					Xcconfig.AppendLine($"CURRENT_PROJECT_VERSION{SettingFilter} = $(UE_TVOS_BUILD_VERSION)");
-
-					// short version string
-					string BundleShortVersion;
-					Ini.GetString("/Script/IOSRuntimeSettings.IOSRuntimeSettings", "VersionInfo", out BundleShortVersion);
-					Xcconfig.AppendLine($"MARKETING_VERSION{SettingFilter} = {BundleShortVersion}");
-				}
-
-				if (UnrealData.TVOSProvisioningData != null && UnrealData.bWriteCodeSigningSettings)
-				{
-					if (!string.IsNullOrEmpty(UnrealData.TVOSProvisioningData.MobileProvisionUUID))
-					{
-						Xcconfig.AppendLine($"PROVISIONING_PROFILE_SPECIFIER{SettingFilter} = {UnrealData.TVOSProvisioningData.MobileProvisionUUID}");
-					}
-					string TVOSCert = UnrealData.TVOSProvisioningData.SigningCertificate ?? (UnrealData.bForDistribution ? "Apple Distribution" : "Apple Developer");
-					Xcconfig.AppendLine($"CODE_SIGN_IDENTITY{SettingFilter} = {TVOSCert}");
-				}
-
-
-				// Info plist generation settings
-				//Xcconfig.AppendLine("INFOPLIST_KEY_UIStatusBarHidden = YES");
-				//Xcconfig.AppendLine("INFOPLIST_KEY_UIRequiresFullScreen = YES");
-			}
-
-			////////////////////////
-			// MAC
-			////////////////////////
-
-			if (ProjectBuildConfigs.bSupportsMac)
-			{
-				string SettingFilter = Platform != null ? "" : "[sdk=macosx*]";
-
-				Xcconfig.AppendLine("");
-				Xcconfig.AppendLine("// Mac Settings");
-			}
-
-
-			////////////////////////
-			// IOS / TVOS
-			////////////////////////
-
-
-			////////////////////////
-			// MAC / IOS / TVOS
-			////////////////////////
-
-			Xcconfig.AppendLine("");
-			Xcconfig.AppendLine("// Cross-platform settings, and defaults for anything not specified above (Xcode uses first-seen value, not last-seen)");
-			Xcconfig.AppendLine("USE_HEADERMAP = NO");
-			Xcconfig.AppendLine("ONLY_ACTIVE_ARCH = YES");
-			Xcconfig.AppendLine($"TARGETED_DEVICE_FAMILY = {string.Join(",", SupportedDevices)}");
-
-
-			if (UnrealData.bWriteCodeSigningSettings)
-			{
-				Xcconfig.AppendLine("CODE_SIGN_STYLE = " + (bAutomaticSigning ? "Automatic" : "Manual"));
-				if (!string.IsNullOrEmpty(SigningTeam))
-				{
-					Xcconfig.AppendLine($"DEVELOPMENT_TEAM = {SigningTeam}");
-				}
-			}
-
-			//if (UnrealData.EntitlementsLocation != null)
-			//{
-			//	string RelativeLoc = UnrealData.EntitlementsLocation.MakeRelativeTo(DirectoryReference.Combine(UnrealData.XcodeProjectFileLocation.ParentDirectory!));
-			//	Xcconfig.AppendLine($"CODE_SIGN_ENTITLEMENTS = {RelativeLoc}");
-			//}
-
-			Xcconfig.Write();
-
-			// Now for each config write out the specific settings
-
-			DirectoryReference? GameDir = UnrealData.UProjectFileLocation?.Directory;
-			string? GamePath = GameDir != null ? XcodeFileCollection.ConvertPath(GameDir.FullName) : null;
-
-			// now pull the bundle id's out, as xcode will warn if they don't match
-			string MacEditorBundleID = "com.epicgames.UnrealEditor";
-			string MacGameBundleID = UnrealData.BundleIdentifier;
-			string IOSBundleID = UnrealData.BundleIdentifier;
-			if (UnrealData.bSupportsMac)
-			{
-				FindPlistId(UnrealData.Metadata!.PlistFiles[MetadataPlatform.MacEditor], "CFBundleIdentifier", ref MacEditorBundleID);
-				FindPlistId(UnrealData.Metadata.PlistFiles[MetadataPlatform.Mac], "CFBundleIdentifier", ref MacGameBundleID);
-			}
-			if (UnrealData.bSupportsIOS || UnrealData.bSupportsTVOS)
-			{
-				FindPlistId(UnrealData.Metadata!.PlistFiles[MetadataPlatform.IOS], "CFBundleIdentifier", ref IOSBundleID);
-			}
-
-			// Get Mac architectures supported by this project
-			foreach (UnrealBuildConfig Config in UnrealData.AllConfigs)
-			{
-				XcodeBuildConfig? MatchedConfig = ProjectBuildConfigs.BuildConfigs.FirstOrDefault(x => x.Info == Config);
-				if (MatchedConfig == null || !Config.Supports(Platform))
-				{
-					continue;
-				}
-
-				// hook up the Buildconfig that matches this info to this xcconfig file
-				XcconfigFile ConfigXcconfig = MatchedConfig.Xcconfig!;
-
-				ConfigXcconfig.AppendLine("// pull in the shared settings for all configs");
-				ConfigXcconfig.AppendLine($"#include \"{UnrealData.ProductName}_Project.xcconfig\"");
-				ConfigXcconfig.AppendLine("");
-
-				bool bIsUnrealGame = Config.BuildTarget.Equals("UnrealGame", StringComparison.InvariantCultureIgnoreCase);
-				bool bIsUnrealClient = Config.BuildTarget.Equals("UnrealClient", StringComparison.InvariantCultureIgnoreCase);
-
-				// Setup Mac stuff
-				FileReference MacExecutablePath = Config.MacExecutablePath!;
-				string MacExecutableDir = XcodeFileCollection.ConvertPath(MacExecutablePath.Directory.FullName);
-				string MacExecutableFileName = MacExecutablePath.GetFileName();
-				string IOSExecutableFilename = Config.IOSExecutablePath!.GetFileName();
-				string TVOSExecutableFilename = Config.TVOSExecutablePath!.GetFileName();
-				string SupportedMacArchitectures = string.Join(" ", UnrealData.GetSupportedMacArchitectures(Config, UnrealData.UProjectFileLocation).Architectures.Select(x => x.AppleName));
-
-				// get defeault PRODUCT_NAME values (if Premade plist, then pull it, like we do with BundleId)
-				string MacProductName = MacExecutableFileName;
-				string IOSProductName = Config.BuildTarget; // @todo: change to Path.GetFileName(Config.IOSExecutablePath) when we stop using payload
-
-				//if (UnrealData.bSupportsMac)
-				//{
-				//	FindBundleID(UnrealData.Metadata!.PlistFiles[MetadataPlatform.Mac], ref MacProductName);
-				//}
-				//if (UnrealData.bSupportsIOS || UnrealData.bSupportsTVOS)
-				//{
-				//	FindBundleID(UnrealData.Metadata!.PlistFiles[MetadataPlatform.IOS], ref IOSProductName);
-				//}
-
-
-				string SupportedPlatforms = "";
-				if (Config.ProjectTarget!.TargetRules != null)
-				{
-					if (ProjectBuildConfigs.bSupportsIOS)
-					{
-						SupportedPlatforms += "iphoneos iphonesimulator ";
-					}
-					if (ProjectBuildConfigs.bSupportsTVOS)
-					{
-						SupportedPlatforms += "appletvos ";
-					}
-					if (ProjectBuildConfigs.bSupportsMac)
-					{
-						SupportedPlatforms += "macosx ";
-					}
-				}
-				else
-				{
-					// @todo when does this case happen?
-					throw new BuildException("Project without a target??");
-//					SupportedPlatforms = "macosx";
-				}
-
-
-				ConfigXcconfig.AppendLine("// Unreal per-config variables");
-				if (ProjectBuildConfigs.bSupportsMac)
-				{
-					ConfigXcconfig.AppendLine($"UE_MAC_EXECUTABLE_NAME = {MacExecutableFileName}");
-				}
-				if (ProjectBuildConfigs.bSupportsIOS)
-				{
-					ConfigXcconfig.AppendLine($"UE_IOS_EXECUTABLE_NAME = {IOSExecutableFilename}");
-				}
-				if (ProjectBuildConfigs.bSupportsTVOS)
-				{
-					ConfigXcconfig.AppendLine($"UE_TVOS_EXECUTABLE_NAME = {TVOSExecutableFilename}");
-				}
-				ConfigXcconfig.AppendLine($"UE_TARGET_NAME = {Config.BuildTarget}");
-				ConfigXcconfig.AppendLine($"UE_TARGET_CONFIG = {Config.BuildConfig}");
-				
-				
-				ConfigXcconfig.AppendLine($"UE_PLATFORM_NAME = {ProjectBuildConfigs.Platform}");
-
-				string TargetPlatformName = ProjectBuildConfigs.Platform.ToString()!;
-				if (ProjectBuildConfigs.BuildConfigs[0].Info.ProjectTarget!.TargetRules!.Type != TargetType.Game)
-				{
-					TargetPlatformName += ProjectBuildConfigs.BuildConfigs[0].Info.ProjectTarget!.TargetRules!.Type.ToString();
-				}
-				ConfigXcconfig.AppendLine($"UE_TARGET_PLATFORM_NAME = {TargetPlatformName}");
-
-
-				ConfigXcconfig.AppendLine("");
-
-				// debug settings
-				if (Config.BuildConfig == UnrealTargetConfiguration.Debug)
-				{
-					ConfigXcconfig.AppendLine("ENABLE_TESTABILITY = YES");
-				}
-
-
-				// list the supported platforms, which will narrow down what other settings are active
-				ConfigXcconfig.AppendLine($"SUPPORTED_PLATFORMS = {SupportedPlatforms}");
-
-				// @otodo move VALID_ARCHS up to Project config once we are always universal
-				if (ProjectBuildConfigs.bSupportsMac)
-				{
-					bool bIsEditor = (Config.ProjectTarget!.TargetRules!.Type == TargetType.Editor);
-					MetadataItem PlistData = UnrealData.Metadata!.PlistFiles[bIsEditor ? MetadataPlatform.MacEditor : MetadataPlatform.Mac];
-					MetadataItem EntitlementsData = UnrealData.Metadata.EntitlementsFiles[bIsEditor ? MetadataPlatform.MacEditor : MetadataPlatform.Mac];
-
-					string SettingFilter = Platform != null ? "" : "[sdk=macosx*]";
-
-					ConfigXcconfig.AppendLine("");
-					ConfigXcconfig.AppendLine("// Mac setup");
-					ConfigXcconfig.AppendLine($"VALID_ARCHS{SettingFilter} = {SupportedMacArchitectures}");
-					ConfigXcconfig.AppendLine($"CONFIGURATION_BUILD_DIR{SettingFilter} = {MacExecutableDir}");
-					ConfigXcconfig.AppendLine($"INFOPLIST_FILE{SettingFilter} = {PlistData.XcodeProjectRelative}");
-					ConfigXcconfig.AppendLine($"PRODUCT_BUNDLE_IDENTIFIER{SettingFilter} = " + (bIsEditor ? MacEditorBundleID : MacGameBundleID));
-					ConfigXcconfig.AppendLine($"PRODUCT_NAME{SettingFilter} = {MacProductName}");
-
-					if (PlistData.Mode == MetadataMode.UpdateTemplate)
-					{
-						ConfigXcconfig.AppendLine($"GENERATE_INFOPLIST_FILE{SettingFilter} = YES");
-						ConfigXcconfig.AppendLine($"ASSETCATALOG_COMPILER_APPICON_NAME{SettingFilter} = AppIcon");
-
-						// @todo: get a version for  games, like IOS has
-						ConfigXcconfig.AppendLine($"MARKETING_VERSION{SettingFilter} = {MacToolChain.LoadEngineDisplayVersion()}");
-						ConfigXcconfig.AppendLine($"CURRENT_PROJECT_VERSION{SettingFilter} = $(UE_MAC_BUILD_VERSION)");
-					}
-
-					if (EntitlementsData.Mode == MetadataMode.UsePremade)
-					{
-						ConfigXcconfig.AppendLine($"CODE_SIGN_ENTITLEMENTS{SettingFilter} = {EntitlementsData.XcodeProjectRelative}");
-					}
-				}
-
-				if (ProjectBuildConfigs.bSupportsIOS)
-				{
-					string SettingFilter = Platform != null ? "" : "[sdk=iphoneos*]";
-
-					ConfigXcconfig.AppendLine("");
-					ConfigXcconfig.AppendLine("// IOS setup");
-					ConfigXcconfig.AppendLine($"INFOPLIST_FILE{SettingFilter} = {UnrealData.Metadata!.PlistFiles[MetadataPlatform.IOS].XcodeProjectRelative}");
-					ConfigXcconfig.AppendLine($"PRODUCT_BUNDLE_IDENTIFIER{SettingFilter} = {IOSBundleID}");
-					ConfigXcconfig.AppendLine($"PRODUCT_NAME{SettingFilter} = {IOSProductName}");
-
-					if (UnrealData.Metadata!.PlistFiles[MetadataPlatform.IOS].Mode == MetadataMode.UpdateTemplate)
-					{
-					}
-					if (UnrealData.Metadata.EntitlementsFiles[MetadataPlatform.IOS].Mode == MetadataMode.UsePremade)
-					{
-						ConfigXcconfig.AppendLine($"CODE_SIGN_ENTITLEMENTS{SettingFilter} = {UnrealData.Metadata.EntitlementsFiles[MetadataPlatform.IOS].XcodeProjectRelative}");
-					}
-				}
-				if (ProjectBuildConfigs.bSupportsTVOS)
-				{
-					string SettingFilter = Platform != null ? "" : "[sdk=appletvos*]";
-
-					ConfigXcconfig.AppendLine("");
-					ConfigXcconfig.AppendLine("// TVOS setup");
-					ConfigXcconfig.AppendLine($"INFOPLIST_FILE{SettingFilter} = {UnrealData.Metadata!.PlistFiles[MetadataPlatform.IOS].XcodeProjectRelative}");
-					ConfigXcconfig.AppendLine($"PRODUCT_BUNDLE_IDENTIFIER{SettingFilter} = {IOSBundleID}");
-					if (UnrealData.Metadata.EntitlementsFiles[MetadataPlatform.IOS].Mode == MetadataMode.UsePremade)
-					{
-						ConfigXcconfig.AppendLine($"CODE_SIGN_ENTITLEMENTS{SettingFilter} = {UnrealData.Metadata.EntitlementsFiles[MetadataPlatform.IOS].XcodeProjectRelative}");
-					}
-				}
-
-				// @todo this is a reminder of plist settings we are not handling yet
-#if false
-				//// Prepare a temp Info.plist file so Xcode has some basic info about the target immediately after opening the project.
-				//// This is needed for the target to pass the settings validation before code signing. UBT will overwrite this plist file later, with proper contents.
-				//if (BuildHostPlatform.Current.Platform == UnrealTargetPlatform.Mac)
-				//{
-				//	bool bCreateMacInfoPlist = !File.Exists(MacInfoPlistPath);
-				//	bool bCreateIOSInfoPlist = !File.Exists(IOSInfoPlistPath) && IOSRunTimeVersion != null;
-				//	bool bCreateTVOSInfoPlist = !File.Exists(TVOSInfoPlistPath) && TVOSRunTimeVersion != null;
-				//	if (bCreateMacInfoPlist || bCreateIOSInfoPlist || bCreateTVOSInfoPlist)
-				//	{
-				//		DirectoryReference? ProjectPath = GameDir;
-				//		DirectoryReference EngineDir = DirectoryReference.Combine(new DirectoryReference(UEDir), "Engine");
-				//		string GameName = Config.BuildTarget;
-				//		bool bIsClient = false;
-				//		if (ProjectPath == null)
-				//		{
-				//			ProjectPath = EngineDir;
-				//		}
-				//		if (bIsUnrealGame)
-				//		{
-				//			ProjectPath = EngineDir;
-				//			GameName = "UnrealGame";
-				//			bIsClient = (UnrealData.AppName == "UnrealClient");
-				//		}
-
-				//		if (bCreateMacInfoPlist)
-				//		{
-				//			Directory.CreateDirectory(Path.GetDirectoryName(MacInfoPlistPath)!);
-				//			UEDeployMac.GeneratePList(ProjectPath.FullName, bIsUnrealGame, GameName, Config.BuildTarget, EngineDir.FullName, MacExecutableFileName);
-				//		}
-				//		if (bCreateIOSInfoPlist)
-				//		{
-				//			// get the receipt
-				//			FileReference ReceiptFilename;
-				//			if (bIsUnrealGame)
-				//			{
-				//				ReceiptFilename = TargetReceipt.GetDefaultPath(Unreal.EngineDirectory, "UnrealGame", UnrealTargetPlatform.IOS, Config.BuildConfig, "");
-				//			}
-				//			else
-				//			{
-				//				ReceiptFilename = TargetReceipt.GetDefaultPath(ProjectPath, GameName, UnrealTargetPlatform.IOS, Config.BuildConfig, "");
-				//			}
-				//			Directory.CreateDirectory(Path.GetDirectoryName(IOSInfoPlistPath)!);
-				//			bool bSupportPortrait, bSupportLandscape;
-				//			TargetReceipt? Receipt;
-				//			TargetReceipt.TryRead(ReceiptFilename, out Receipt);
-				//			bool bBuildAsFramework = UEDeployIOS.GetCompileAsDll(Receipt);
-				//			UEDeployIOS.GenerateIOSPList(UnrealData.UProjectFileLocation, Config.BuildConfig, ProjectPath.FullName, bIsUnrealGame, GameName, bIsClient, Config.BuildTarget, EngineDir.FullName, ProjectPath + "/Binaries/IOS/Payload", null, UnrealData.BundleIdentifier, bBuildAsFramework, UnrealData.Logger!, out bSupportPortrait, out bSupportLandscape);
-				//		}
-				//		if (bCreateTVOSInfoPlist)
-				//		{
-				//			Directory.CreateDirectory(Path.GetDirectoryName(TVOSInfoPlistPath)!);
-				//			UEDeployTVOS.GenerateTVOSPList(ProjectPath.FullName, bIsUnrealGame, GameName, bIsClient, Config.BuildTarget, EngineDir.FullName, ProjectPath + "/Binaries/TVOS/Payload", null, UnrealData.BundleIdentifier, UnrealData.Logger!);
-				//		}
-				//	}
-				//}
-#endif
-
-				ConfigXcconfig.Write();
-			}
-		}
-	}
-
-	class XcodeFileCollection
-	{
-		static readonly string[] ResourceExtensions =
-		{
-			".xcassets",
-			".storyboardc",
-			".storyboard",
-			".xcframework",
-			".framework",
-			".lproj",
-			// no extensions will be folders, like "cookeddata" that we want to copy as a resource
-			"",
-		};
-
-		static readonly string[] SourceCodeExtensions =
-		{
-			".c",
-			".cc",
-			".cpp",
-			".m",
-			".mm",
-		};
-
-		static readonly Dictionary<string, string> FileTypeMap = new()
-		{
-			{ ".c", "sourcecode.c.objc" },
-			{ ".m", "sourcecode.c.objc" },
-			{ ".cc", "sourcecode.cpp.objcpp" },
-			{ ".cpp", "sourcecode.cpp.objcpp" },
-			{ ".mm", "sourcecode.cpp.objcpp" },
-			{ ".h", "sourcecode.c.h" },
-			{ ".inl", "sourcecode.c.h" },
-			{ ".pch", "sourcecode.c.h" },
-			{ ".framework", "wrapper.framework" },
-			{ ".plist", "text.plist.xml" },
-			{ ".png", "image.png" },
-			{ ".icns", "image.icns" },
-			{ ".xcassets", "folder.assetcatalog" },
-			{ ".storyboardc", "wrapper.storyboardc" },
-			{ ".storyboard", "file.storyboard" },
-			{ ".xcframework", "wrapper.xcframework" },
-			{ ".bundle", "\"wrapper.plug-in\"" }
-		};
-
-		const string FileTypeDefault = "file.txt";
-
-
-		class ManualFileReference
-		{
-			public string Guid = XcodeProjectFileGenerator.MakeXcodeGuid();
-			public string RelativePath = "";
-			public string FileKey = "";
-			public string FileType = "";
-			public string SourceTree = "";
-			public string GroupName = "";
-		}
-
-		// we need relativ paths, so this is the directory that paths are relative to
-		private DirectoryReference ProjectDirectory;
-
-		// the location of the uproject file, or the engine directory for non-uproject projects
-		public DirectoryReference RootDirectory;
-
-		public string MainGroupGuid = "";
-		//public string ProductRefGroupGuid = XcodeProjectFileGenerator.MakeXcodeGuid();
-
-		internal Dictionary<string, XcodeFileGroup> Groups = new();
-		internal Dictionary<XcodeSourceFile, FileReference?> BuildableFilesToResponseFile = new();
-		internal List<Tuple<XcodeSourceFile, string>> BuildableResourceFiles = new();
-		internal List<XcodeSourceFile> AllFiles = new();
-
-		private List<ManualFileReference> ManualFiles = new();
-		private Dictionary<string, string> GuidsForGroups = new();
-
-
-		public XcodeFileCollection(XcodeProjectFile ProjectFile)
-		{
-			ProjectDirectory = ProjectFile.ProjectFilePath.Directory;
-			RootDirectory = Unreal.EngineDirectory;
-			SetUProjectLocation(ProjectFile.UnrealData.UProjectFileLocation);
-		}
-
-		public XcodeFileCollection(XcodeFileCollection Other)
-		{
-			ProjectDirectory = Other.ProjectDirectory;
-			RootDirectory = Other.RootDirectory;
-			MainGroupGuid = Other.MainGroupGuid;
-			AllFiles.AddRange(Other.AllFiles);
-			BuildableResourceFiles.AddRange(Other.BuildableResourceFiles);
-			ManualFiles.AddRange(Other.ManualFiles);
-
-			BuildableFilesToResponseFile = new(Other.BuildableFilesToResponseFile);
-			GuidsForGroups = new(Other.GuidsForGroups);
-
-			foreach (string Key in Other.Groups.Keys)
-			{
-				Groups.Add(Key, new XcodeFileGroup(Other.Groups[Key]));
-			}
-
-		}
-
-		public void SetUProjectLocation(FileReference? UProjectLocation)
-		{
-			RootDirectory = UProjectLocation == null || UProjectLocation.ContainsName("Programs", 0) ? (Unreal.EngineDirectory) : UProjectLocation!.Directory;
-		}
-
-		public void ProcessFile(XcodeSourceFile File, bool bIsForBuild, bool bIsFolder, string? GroupName=null, string ExtraResourceSettings="")
-		{
-			// remember all buildable files
-			if (bIsForBuild)
-			{
-				string FileExtension = File.Reference.GetExtension();
-
-				if (IsSourceCode(FileExtension))
-				{
-					// look if it contains any of the exluded names, and is so, don't include it
-					if (!ExcludedFolders.Any(x => File.Reference.ContainsName(x, 0)))
-					{
-						// will fill in the response file later
-						BuildableFilesToResponseFile[File] = null;
-					}
-				}
-				else if (IsResourceFile(FileExtension))
-				{
-					BuildableResourceFiles.Add(Tuple.Create(File, ExtraResourceSettings));
-				}
-			}
-			
-			if (GroupName != "")
-			{
-				// group the files by path
-				XcodeFileGroup? Group = GroupName == null ? FindGroupByAbsolutePath(File.Reference.Directory) : null;
-				if (Group != null)
-				{
-					AllFiles.Add(File);
-					Group.Files.Add(File);
-				}
-				else
-				{
-					GroupName = GroupName ?? (File.Reference.IsUnderDirectory(Unreal.EngineDirectory) ? "EngineReferences" : "ExternalReferences");
-					if (bIsFolder)
-					{
-						AddFolderReference(File.FileRefGuid, File.Reference.MakeRelativeTo(ProjectDirectory), GroupName);
-					}
-					else
-					{
-						AddFileReference(File.FileRefGuid, File.Reference.FullName, GetFileKey(File.Reference.GetExtension()), GetFileType(File.Reference.GetExtension()), "<absolute>", GroupName);
-					}
-				}
-			}
-		}
-
-		/// <summary>
-		/// The project needs to point to the product group, so we expose the Guid here
-		/// </summary>
-		/// <returns></returns>
-		public string GetProductGroupGuid()
-		{
-			return GuidsForGroups["Products"];
-		}
-
-		public void AddFileReference(string Guid, string RelativePath, string FileKey, string FileType, string SourceTree/*="SOURCE_ROOT"*/, string GroupName)
-		{
-			ManualFiles.Add(new ManualFileReference()
-			{
-				Guid = Guid,
-				RelativePath = RelativePath,
-				FileKey = FileKey,
-				FileType = FileType,
-				SourceTree = SourceTree.StartsWith("<") ? $"\"{SourceTree}\"" : SourceTree,
-				GroupName = GroupName,
-			});
-
-			// make sure each unique group has a guid
-			if (!GuidsForGroups.ContainsKey(GroupName))
-			{
-				GuidsForGroups[GroupName] = XcodeProjectFileGenerator.MakeXcodeGuid();
-			}
-		}
-
-		public void AddFolderReference(string Guid, string RelativePath, string GroupName)
-		{
-			ManualFiles.Add(new ManualFileReference()
-			{
-				Guid = Guid,
-				RelativePath = RelativePath,
-				FileKey = "lastKnownFileType",
-				FileType = "folder",
-				SourceTree = "\"<group>\"",
-				GroupName = GroupName,
-			});
-
-			// make sure each unique group has a guid
-			if (!GuidsForGroups.ContainsKey(GroupName))
-			{
-				GuidsForGroups[GroupName] = XcodeProjectFileGenerator.MakeXcodeGuid();
-			}
-		}
-
-
-		public void Write(StringBuilder Content)
-		{
-			Content.WriteLine("/* Begin PBXBuildFile section */");
-			foreach (KeyValuePair<XcodeSourceFile, FileReference?> Pair in BuildableFilesToResponseFile)
-			{
-				string CompileFlags = "";
-				if (Pair.Value != null)
-				{
-					CompileFlags = $" settings = {{COMPILER_FLAGS = \"@{Pair.Value.FullName}\"; }};";
-				}
-				Content.WriteLine($"\t\t{Pair.Key.FileGuid} = {{isa = PBXBuildFile; fileRef = {Pair.Key.FileRefGuid};{CompileFlags} }}; /* {Pair.Key.Reference.GetFileName()} */");
-			}
-			foreach (Tuple<XcodeSourceFile, string> Resource in BuildableResourceFiles)
-			{
-				XcodeSourceFile File = Resource.Item1;
-				string ExtraSettings = Resource.Item2;
-				Content.WriteLine($"\t\t{File.FileGuid} = {{isa = PBXBuildFile; fileRef = {File.FileRefGuid}; {ExtraSettings}}}; /* {File.Reference.GetFileName()} */");
-			}
-			Content.WriteLine("/* End PBXBuildFile section */");
-			Content.WriteLine();
-
-
-			Content.WriteLine("/* Begin PBXFileReference section */");
-			foreach (XcodeSourceFile File in AllFiles)
-			{
-				string FileName = File.Reference.GetFileName();
-				string FileExtension = Path.GetExtension(FileName);
-				string FileType = GetFileType(FileExtension);
-				string FileKey = GetFileKey(FileExtension);
-				string RelativePath = Utils.CleanDirectorySeparators(File.Reference.MakeRelativeTo(ProjectDirectory), '/');
-				string SourceTree = "SOURCE_ROOT";
-
-				Content.WriteLine($"\t\t{File.FileRefGuid} = {{isa = PBXFileReference; {FileKey} = {FileType}; name = \"{FileName}\"; path = \"{RelativePath}\"; sourceTree = {SourceTree}; }};");
-			}
-			foreach (ManualFileReference File in ManualFiles)
-			{
-				string FileName = Path.GetFileName(File.RelativePath);
-				Content.WriteLine($"\t\t{File.Guid} = {{isa = PBXFileReference; {File.FileKey} = {File.FileType}; name = \"{FileName}\"; path = \"{File.RelativePath}\"; sourceTree = {File.SourceTree}; }};");
-			}
-			Content.WriteLine("/* End PBXFileReference section */");
-			Content.WriteLine();
-
-			WriteGroups(Content);
-		}
-
-		private void WriteGroups(StringBuilder Content)
-		{
-			XcodeFileGroup? RootGroup = FindRootFileGroup(Groups);
-			if (RootGroup == null)
-			{
-				return;
-			}
-
-			string XcconfigsRefGroupGuid = XcodeProjectFileGenerator.MakeXcodeGuid();
-
-			Content.WriteLine("/* Begin PBXGroup section */");
-
-
-			// write main/root group and it's children (adding the manual groups to the main group, which is the true param here0
-			WriteGroup(Content, RootGroup, bIsRootGroup: true);
-
-			// write some manual groups
-			foreach (var Pair in GuidsForGroups)
-			{
-				Content.WriteLine($"\t\t{Pair.Value} /* {Pair.Key} */ = {{");
-				Content.WriteLine("\t\t\tisa = PBXGroup;");
-				Content.WriteLine("\t\t\tchildren = (");
-				foreach (ManualFileReference Ref in ManualFiles.Where(x => x.GroupName == Pair.Key))
-				{
-					Content.WriteLine($"\t\t\t\t{Ref.Guid} /* {Path.GetFileName(Ref.RelativePath)} */,");
-				}
-				Content.WriteLine("\t\t\t);");
-				Content.WriteLine($"\t\t\tname = {Pair.Key};");
-				Content.WriteLine("\t\t\tsourceTree = \"<group>\";");
-				Content.WriteLine("\t\t};");
-
-			}
-
-			Content.WriteLine("/* End PBXGroup section */");
-		}
-
-
-		private void WriteGroup(StringBuilder Content, XcodeFileGroup Group, bool bIsRootGroup)
-		{
-			if (!Group.bIsReference)
-			{
-				Content.WriteLine($"\t\t{Group.GroupGuid} = {{");
-				Content.WriteLine("\t\t\tisa = PBXGroup;");
-				Content.WriteLine("\t\t\tchildren = (");
-				foreach (XcodeFileGroup ChildGroup in Group.Children.Values)
-				{
-					Content.WriteLine($"\t\t\t\t{ChildGroup.GroupGuid} /* {ChildGroup.GroupName} */,");
-				}
-				foreach (XcodeSourceFile File in Group.Files)
-				{
-					Content.WriteLine($"\t\t\t\t{File.FileRefGuid} /* {File.Reference.GetFileName()} */,");
-				}
-
-				if (bIsRootGroup)
-				{
-					foreach (var Pair in GuidsForGroups)
-					{
-						Content.WriteLine($"\t\t\t\t{Pair.Value} /* {Pair.Key} */,");
-					}
-				}
-				Content.WriteLine("\t\t\t);");
-				if (!bIsRootGroup)
-				{
-					Content.WriteLine("\t\t\tname = \"" + Group.GroupName + "\";");
-					Content.WriteLine("\t\t\tpath = \"" + Group.GroupPath + "\";");
-					Content.WriteLine("\t\t\tsourceTree = \"<absolute>\";");
-				}
-				else
-				{
-					Content.WriteLine("\t\t\tsourceTree = \"<group>\";");
-				}
-				Content.WriteLine("\t\t};");
-
-				foreach (XcodeFileGroup ChildGroup in Group.Children.Values)
-				{
-					WriteGroup(Content, ChildGroup, bIsRootGroup: false);
-				}
-			}
-		}
-
-		public static string GetRootGroupGuid(Dictionary<string, XcodeFileGroup> GroupsDict)
-		{
-			return FindRootFileGroup(GroupsDict)!.GroupGuid;
-		}
-
-		private static XcodeFileGroup? FindRootFileGroup(Dictionary<string, XcodeFileGroup> GroupsDict)
-		{
-			foreach (XcodeFileGroup Group in GroupsDict.Values)
-			{
-				if (Group.Children.Count > 1 || Group.Files.Count > 0)
-				{
-					return Group;
-				}
-				else
-				{
-					XcodeFileGroup? Found = FindRootFileGroup(Group.Children);
-					if (Found != null)
-					{
-						return Found;
-					}
-				}
-			}
-			return null;
-		}
-
-
-		/// <summary>
-		/// Gets Xcode file type based on its extension
-		/// </summary>
-		internal static string GetFileType(string Extension)
-		{
-			string? FileType;
-			if (FileTypeMap.TryGetValue(Extension, out FileType))
-			{
-				return FileType;
-			}
-			return FileTypeDefault;
-		}
-
-		// most types Xcode wants as explicitFileType, but some want lastKnownFileType
-		internal static string GetFileKey(string Extension)
-		{
-			switch (Extension)
-			{
-				case ".storyboard":
-				case ".bundle":
-					return "lastKnownFileType";
-				default:
-					return "explicitFileType";
-			}
-		}
-		/// <summary>
-		/// Returns true if Extension is a known extension for files containing source code
-		/// </summary>
-		internal static bool IsSourceCode(string Extension)
-		{
-			return SourceCodeExtensions.Contains(Extension);
-		}
-
-		/// <summary>
-		/// Returns true if Extension is a known extension for files containing resource data
-		/// </summary>
-		internal static bool IsResourceFile(string Extension)
-		{
-			return ResourceExtensions.Contains(Extension);
-		}
-
-		/// <summary>
-		/// Cache the folders excluded by Mac
-		/// </summary>
-		static string[] ExcludedFolders = PlatformExports.GetExcludedFolderNames(UnrealTargetPlatform.Mac);
-
-		/// <summary>
-		/// Returns a project navigator group to which the file should belong based on its path.
-		/// Creates a group tree if it doesn't exist yet.
-		/// </summary>
-		internal XcodeFileGroup? FindGroupByAbsolutePath(DirectoryReference DirectoryPath)
-		{
-			string AbsolutePath = DirectoryPath.FullName;
-			string[] Parts = AbsolutePath.Split(Path.DirectorySeparatorChar);
-			string CurrentPath = "/";
-			Dictionary<string, XcodeFileGroup> CurrentSubGroups = Groups;
-
-			for (int Index = 1; Index < Parts.Count(); ++Index)
-			{
-				string Part = Parts[Index];
-
-				if (CurrentPath.Length > 1)
-				{
-					CurrentPath += Path.DirectorySeparatorChar;
-				}
-
-				CurrentPath += Part;
-
-				XcodeFileGroup CurrentGroup;
-				if (!CurrentSubGroups.ContainsKey(CurrentPath))
-				{
-					// we only want files under the project/engine dir, we have to make external references to other files, otherwise the grouping gets all broken
-					if (!DirectoryPath.IsUnderDirectory(RootDirectory))
-					{
-						return null;
-					}
-
-
-					CurrentGroup = new XcodeFileGroup(Path.GetFileName(CurrentPath), CurrentPath, GetFileType(Path.GetExtension(CurrentPath)).StartsWith("folder"));
-					if (Groups.Count == 0)
-					{
-						MainGroupGuid = CurrentGroup.GroupGuid;
-					}
-					CurrentSubGroups.Add(CurrentPath, CurrentGroup);
-				}
-				else
-				{
-					CurrentGroup = CurrentSubGroups[CurrentPath];
-				}
-
-				if (CurrentPath == AbsolutePath)
-				{
-					return CurrentGroup;
-				}
-
-				CurrentSubGroups = CurrentGroup.Children;
-			}
-
-			return null;
-		}
-
-		/// <summary>
-		/// Convert all paths to Apple/Unix format (with forward slashes)
-		/// </summary>
-		/// <param name="InPath">The path to convert</param>
-		/// <returns>The normalized path</returns>
-		internal static string ConvertPath(string InPath)
-		{
-			return InPath.Replace("\\", "/");
-		}
-	}
-
-
-	class XcodeProjectFile : ProjectFile
-	{
-
-		/// <summary>
-		/// Constructs a new project file object
-		/// </summary>
-		/// <param name="InitFilePath">The path to the project file on disk</param>
-		/// <param name="BaseDir">The base directory for files within this project</param>
-		/// <param name="bIsForDistribution">True for distribution builds</param>
-		/// <param name="BundleID">Override option for bundle identifier</param>
-		/// <param name="AppName"></param>
-		/// <param name="bMakeProjectPerTarget"></param>
-		public XcodeProjectFile(FileReference InitFilePath, DirectoryReference BaseDir, bool bIsForDistribution, string BundleID, string AppName, bool bMakeProjectPerTarget)
-			: base(InitFilePath, BaseDir)
-		{
-			UnrealData = new UnrealData(InitFilePath, bIsForDistribution, BundleID, AppName, bMakeProjectPerTarget, IsStubProject);
-
-			// create the container for all the files that will 
-			SharedFileCollection = new XcodeFileCollection(this);
-			FileCollection = SharedFileCollection;
-		}
-
-		public UnrealData UnrealData;
+		public UnrealData UnrealData;
 
 		/// <summary>
 		///  The PBXPRoject node, root of everything
@@ -2897,25 +1930,14 @@ namespace UnrealBuildTool.XcodeProjectXcconfig
 		}
 
 
-		// null platform means use the old way, with multiple platforms per project
-		private string Suffix(UnrealTargetPlatform? Platform)
-		{
-			return Platform == null ? "" : $" ({Platform})";
-		}
-
-		public DirectoryReference ProjectDirPathForPlatform(UnrealTargetPlatform? Platform)
-		{
-			return DirectoryReference.Combine(ProjectFilePath.Directory, $"{ProjectFilePath.GetFileNameWithoutExtension()}{Suffix(Platform)}{ProjectFilePath.GetExtension()}");
-		}
-
 		public FileReference ProjectFilePathForPlatform(UnrealTargetPlatform? Platform)
 		{
-			return new FileReference(ProjectDirPathForPlatform(Platform).FullName);
+			return new FileReference(XcodeUtils.ProjectDirPathForPlatform(UnrealData.XcodeProjectFileLocation, Platform).FullName);
 		}
 
 		public FileReference PBXFilePathForPlatform(UnrealTargetPlatform? Platform)
 		{
-			return FileReference.Combine(ProjectDirPathForPlatform(Platform), "project.pbxproj");
+			return FileReference.Combine(XcodeUtils.ProjectDirPathForPlatform(UnrealData.XcodeProjectFileLocation, Platform), "project.pbxproj");
 		}
 
 		/// Implements Project interface
@@ -2998,8 +2020,7 @@ namespace UnrealBuildTool.XcodeProjectXcconfig
 					bSuccess = ProjectFileGenerator.WriteFileIfChanged(PBXFilePath.FullName, Content.ToString(), Logger, new UTF8Encoding()) && bSuccess;
 				}
 
-
-				bool bNeedScheme = UnrealData.CanBuildProjectLocally(this, Logger);
+				bool bNeedScheme = XcodeUtils.ShouldIncludeProjectInWorkspace(this, Logger);
 				if (bNeedScheme)
 				{
 					if (bSuccess)
@@ -3007,19 +2028,12 @@ namespace UnrealBuildTool.XcodeProjectXcconfig
 						string ProjectName = ProjectFilePathForPlatform(Platform).GetFileNameWithoutAnyExtensions();
 						string? BuildTargetGuid = XcodeProjectNode.GetNodesOfType<XcodeBuildTarget>(RootProject).FirstOrDefault()?.Guid;
 						string? IndexTargetGuid = XcodeProjectNode.GetNodesOfType<XcodeIndexTarget>(RootProject).FirstOrDefault()?.Guid;
-						WriteSchemeFile(Platform, ProjectName, UnrealData.ProductName, RootProject.RunTargets, BuildTargetGuid, IndexTargetGuid, 
-							UnrealData.UProjectFileLocation != null ? UnrealData.UProjectFileLocation.FullName : "", UnrealData.ExtraPreBuildScriptLines);
+						XcodeSchemeFile.WriteSchemeFile(UnrealData, Platform, ProjectName, RootProject.RunTargets, BuildTargetGuid, IndexTargetGuid);
 					}
 				}
 				else
 				{
-					// clean this up because we don't want it persisting if we narrow our project list
-					DirectoryReference SchemeDir = GetProjectSchemeDirectory(Platform);
-
-					if (DirectoryReference.Exists(SchemeDir))
-					{
-						DirectoryReference.Delete(SchemeDir, true);
-					}
+					XcodeSchemeFile.CleanSchemeFile(UnrealData, Platform);
 				}
 			}
 			return bSuccess;
@@ -3045,445 +2059,18 @@ namespace UnrealBuildTool.XcodeProjectXcconfig
 					FileCollection.AddFileReference(XcodeProjectFileGenerator.MakeXcodeGuid(), Data.XcodeProjectRelative, "explicitFileType", "text.plist", "\"<group>\"", "Metadata");
 				}
 			}
-#if false
-			ConfigHierarchy Ini = ConfigCache.ReadHierarchy(ConfigHierarchyType.Engine, UnrealData.UProjectFileLocation?.Directory, Platform);
-
-			/////////////////////
-			// Entitlements
-			/////////////////////
-
-			FileReference? ExistingEntitlements = UnrealData.UProjectFileLocation == null ? null :
-				FileReference.Combine(UnrealData.UProjectFileLocation.Directory, $"Build/Xcode/{UnrealData.ProductName}.entitlements");
-
-			if (ExistingEntitlements != null && FileReference.Exists(ExistingEntitlements))
-			{
-				UnrealData.EntitlementsLocation = ExistingEntitlements;
-				UnrealData.bEntitlementsWasPremade = true;
-				string RelativeLoc = UnrealData.EntitlementsLocation.MakeRelativeTo(UnrealData.XcodeProjectFileLocation.ParentDirectory!);
-				FileCollection.AddFileReference(XcodeProjectFileGenerator.MakeXcodeGuid(), RelativeLoc, "text.plist", "\"<group>\"", "Metadata");
-			}
-			else
-			{
-				UnrealData.EntitlementsLocation = FileReference.Combine(UnrealData.XcodeProjectFileLocation.ParentDirectory!, "Metadata", $"{UnrealData.ProductName}.entitlements");
-				FileCollection.AddFileReference(XcodeProjectFileGenerator.MakeXcodeGuid(), $"Metadata/{UnrealData.EntitlementsLocation.GetFileName()}", "text.plist.entitlements", "\"<group>\"", "Metadata");
-
-				string Filename = UnrealData.EntitlementsLocation.FullName;
-				DirectoryReference.CreateDirectory(UnrealData.EntitlementsLocation.Directory);
-
-				// reset file to start (or create it)
-				Plist("Clear", Filename);
-
-				bool bValue = false;
-				// remote push notifications
-				if (Ini.GetBool("/Script/IOSRuntimeSettings.IOSRuntimeSettings", "bEnableRemoteNotificationsSupport", out bValue) && bValue)
-				{
-					// Xcode will replace this with production when using a distro cert
-					Plist("Add aps-environment string development", Filename);
-				}
-
-				// Sign in with Apple
-				if (Ini.GetBool("/Script/IOSRuntimeSettings.IOSRuntimeSettings", "bEnableSignInWithAppleSupport", out bValue) && bValue)
-				{
-					Plist("Add com.apple.developer.applesignin string Default", Filename);
-				}
-
-				// Add Multi-user support for tvOS
-				if (Ini.GetBool("/Script/IOSRuntimeSettings.IOSRuntimeSettings", "bRunAsCurrentUser", out bValue) && bValue)
-				{
-					Plist("Add com.apple.developer.user-management array", Filename);
-					Plist("Add com.apple.developer.user-management:0 string runs-as-current-user", Filename);
-				}
-			}
-#endif
-
-
-
-
-
-#if false
-			/////////////////////
-			// Info.plist
-			/////////////////////
-
-			FileReference? ExistingInfoPlist = UnrealData.UProjectFileLocation == null ? null :
-				FileReference.Combine(UnrealData.UProjectFileLocation.Directory, $"Build/Xcode/{UnrealData.ProductName}-Info.plist");
-			if (ExistingInfoPlist != null && FileReference.Exists(ExistingInfoPlist))
-			{
-				UnrealData.InfoPlistLocation = ExistingInfoPlist;
-				UnrealData.bInfoPlistWasPremade = true;
-				string RelativeLoc = UnrealData.InfoPlistLocation.MakeRelativeTo(UnrealData.XcodeProjectFileLocation.ParentDirectory!);
-				FileCollection.AddFileReference(XcodeProjectFileGenerator.MakeXcodeGuid(), RelativeLoc, "text.plist", "\"<group>\"", "Metadata");
-			}
-			else
-			{
-				UnrealData.InfoPlistLocation = FileReference.Combine(UnrealData.XcodeProjectFileLocation.ParentDirectory!, "Metadata", $"{UnrealData.ProductName}-Info.plist");
-				FileCollection.AddFileReference(XcodeProjectFileGenerator.MakeXcodeGuid(), $"Metadata/{UnrealData.InfoPlistLocation.GetFileName()}", "text.plist", "\"<group>\"", "Metadata");
-
-
-				string Filename = UnrealData.InfoPlistLocation.FullName;
-
-				// reset file to start (or create it)
-				Plist("Clear", Filename);
-
-				////////////////////////
-				// PLIST MAC / IOS
-				////////////////////////
-
-
-				////////////////////////
-				// PLIST MAC
-				////////////////////////
-
-				// disable non-exempt encryption
-				Plist("Add :ITSAppUsesNonExemptEncryption bool false", Filename);
-
-				////////////////////////
-				// PLIST IOS
-				////////////////////////
-
-				if (UnrealData.IOSProjectSettings != null)
-				{
-					bool HACK_bSupportsFilesApp = true;
-					if (UnrealData.IOSProjectSettings.bFileSharingEnabled)
-					{
-						Plist("Add :UIFileSharingEnabled bool true", Filename);
-					}
-					if (HACK_bSupportsFilesApp)
-					{
-						Plist("Add :LSSupportsOpeningDocumentsInPlace bool true", Filename);
-					}
-					Plist("Add :UIViewControllerBasedStatusBarAppearance bool false", Filename);
-				}
-			}
-
-#endif
-
-
-
-#if false
-
-			// in CFBundleURLSchemes array
-			if (bEnableFacebookSupport)
-			{
-				// This is needed for facebook login to redirect back to the app after completion.
-				Text.AppendLine(string.Format("\t\t\t\t<string>fb{0}</string>", FacebookAppID));
-			}
-			if (bEnableGoogleSupport)
-			{
-				Text.AppendLine(string.Format("\t\t\t\t<string>{0}</string>", GoogleReversedClientId));
-			}
-
-			Text.AppendLine("\t<key>CFBundleVersion</key>");
-			Text.AppendLine(string.Format("\t<string>{0}</string>", VersionUtilities.UpdateBundleVersion(OldPListData, InEngineDir)));
-			Text.AppendLine("\t<key>CFBundleShortVersionString</key>");
-			Text.AppendLine(string.Format("\t<string>{0}</string>", BundleShortVersion));
-
-
-#endif
-
-
-			// deleted
-#if false
-
-			Text.AppendLine("\t<key>CFBundleDevelopmentRegion</key>");
-			Text.AppendLine("\t<string>English</string>");
-
-			Text.AppendLine("\t<key>CFBundleExecutable</key>");
-			string BundleExecutable = bIsUnrealGame ?
-				(bIsClient ? "UnrealClient" : "UnrealGame") :
-				(bIsClient ? GameName + "Client" : GameName);
-			Text.AppendLine(string.Format("\t<string>{0}</string>", BundleExecutable));
-
- 
-#endif
-
-			// moved into .plist
-#if false
-			Text.AppendLine("\t\t\t<key>CFBundleURLName</key>");
-			Text.AppendLine("\t\t\t<string>com.Epic.Unreal</string>");
-			Text.AppendLine("\t\t\t<key>CFBundleURLSchemes</key>");
-			Text.AppendLine("\t\t\t<array>");
-			Text.AppendLine(string.Format("\t\t\t\t<string>{0}</string>", bIsUnrealGame ? "UnrealGame" : GameName));
-			Text.AppendLine("\t\t\t</array>");
-			Text.AppendLine("\t\t</dict>");
-			Text.AppendLine("\t</array>");
-
-
-			Text.AppendLine("\t<key>UIRequiredDeviceCapabilities</key>");
-			Text.AppendLine("\t<array>");
-			foreach(string Line in RequiredCaps.Split("\r\n".ToCharArray()))
-			{
-				if (!string.IsNullOrWhiteSpace(Line))
-				{
-					Text.AppendLine(Line);
-				}
-			}
-			Text.AppendLine("\t</array>");
-
-			Text.AppendLine("\t<key>CFBundleIcons</key>");
-			Text.AppendLine("\t<dict>");
-			Text.AppendLine("\t\t<key>CFBundlePrimaryIcon</key>");
-			Text.AppendLine("\t\t<dict>");
-			Text.AppendLine("\t\t\t<key>CFBundleIconFiles</key>");
-			Text.AppendLine("\t\t\t<array>");
-			Text.AppendLine("\t\t\t\t<string>AppIcon60x60</string>");
-			Text.AppendLine("\t\t\t</array>");
-			Text.AppendLine("\t\t\t<key>CFBundleIconName</key>");
-			Text.AppendLine("\t\t\t<string>AppIcon</string>");
-			Text.AppendLine("\t\t\t<key>UIPrerenderedIcon</key>");
-			Text.AppendLine("\t\t\t<true/>");
-			Text.AppendLine("\t\t</dict>");
-			Text.AppendLine("\t</dict>");
-			Text.AppendLine("\t<key>CFBundleIcons~ipad</key>");
-			Text.AppendLine("\t<dict>");
-			Text.AppendLine("\t\t<key>CFBundlePrimaryIcon</key>");
-			Text.AppendLine("\t\t<dict>");
-			Text.AppendLine("\t\t\t<key>CFBundleIconFiles</key>");
-			Text.AppendLine("\t\t\t<array>");
-			Text.AppendLine("\t\t\t\t<string>AppIcon60x60</string>");
-			Text.AppendLine("\t\t\t\t<string>AppIcon76x76</string>");
-			Text.AppendLine("\t\t\t</array>");
-			Text.AppendLine("\t\t\t<key>CFBundleIconName</key>");
-			Text.AppendLine("\t\t\t<string>AppIcon</string>");
-			Text.AppendLine("\t\t\t<key>UIPrerenderedIcon</key>");
-			Text.AppendLine("\t\t\t<true/>");
-			Text.AppendLine("\t\t</dict>");
-			Text.AppendLine("\t</dict>");
-			Text.AppendLine("\t<key>UILaunchStoryboardName</key>");
-			Text.AppendLine("\t<string>LaunchScreen</string>");
-
-			if (File.Exists(DirectoryReference.FromFile(ProjectFile) + "/Build/IOS/Resources/Interface/LaunchScreen.storyboard") && VersionUtilities.bCustomLaunchscreenStoryboard)
-			{
-				string LaunchStoryboard = DirectoryReference.FromFile(ProjectFile) + "/Build/IOS/Resources/Interface/LaunchScreen.storyboard";
-
-				if (BuildHostPlatform.Current.Platform == UnrealTargetPlatform.Mac)
-				{
-					string outputStoryboard = LaunchStoryboard + "c";
-					string argsStoryboard = "--compile " + outputStoryboard + " " + LaunchStoryboard;
-					string stdOutLaunchScreen = Utils.RunLocalProcessAndReturnStdOut("ibtool", argsStoryboard, Logger);
-
-					Logger.LogInformation("LaunchScreen Storyboard compilation results : {Results}", stdOutLaunchScreen);
-				}
-				else
-				{
-					Logger.LogWarning("Custom Launchscreen compilation storyboard only compatible on Mac for now");
-				}
-			}
-
-
-			// add location services descriptions if used
-			if (!string.IsNullOrWhiteSpace(LocationAlwaysUsageDescription))
-			{
-				Text.AppendLine("\t<key>NSLocationAlwaysAndWhenInUseUsageDescription</key>");
-				Text.AppendLine(string.Format("\t<string>{0}</string>", LocationAlwaysUsageDescription));
-			}
-			if (!string.IsNullOrWhiteSpace(LocationWhenInUseDescription))
-			{
-				Text.AppendLine("\t<key>NSLocationWhenInUseUsageDescription</key>");
-				Text.AppendLine(string.Format("\t<string>{0}</string>", LocationWhenInUseDescription));
-			}
-			// disable HTTPS requirement
-			if (bDisableHTTPS)
-			{
-				Text.AppendLine("\t<key>NSAppTransportSecurity</key>");
-				Text.AppendLine("\t\t<dict>");
-				Text.AppendLine("\t\t\t<key>NSAllowsArbitraryLoads</key><true/>");
-				Text.AppendLine("\t\t</dict>");
-			}
-
-			if (bEnableFacebookSupport)
-			{
-				Text.AppendLine("\t<key>FacebookAppID</key>");
-				Text.AppendLine(string.Format("\t<string>{0}</string>", FacebookAppID));
-				Text.AppendLine("\t<key>FacebookDisplayName</key>");
-				Text.AppendLine(string.Format("\t<string>{0}</string>", FacebookDisplayName));
-				
-				if (!bEnableAutomaticLogging)
-				{
-					Text.AppendLine("<key>FacebookAutoLogAppEventsEnabled</key><false/>");
-				}
-				
-				if (!bEnableAdvertisingId)
-				{
-					Text.AppendLine("<key>FacebookAdvertiserIDCollectionEnabled</key><false/>");
-				}
-
-				Text.AppendLine("\t<key>LSApplicationQueriesSchemes</key>");
-				Text.AppendLine("\t<array>");
-				Text.AppendLine("\t\t<string>fbapi</string>");
-				Text.AppendLine("\t\t<string>fb-messenger-api</string>");
-				Text.AppendLine("\t\t<string>fb-messenger-share-api</string>");
-				Text.AppendLine("\t\t<string>fbauth2</string>");
-				Text.AppendLine("\t\t<string>fbshareextension</string>");
-				Text.AppendLine("\t</array>");
-			}
-
-			if (!string.IsNullOrEmpty(ExtraData))
-			{
-				ExtraData = ExtraData.Replace("\\n", "\n");
-				foreach(string Line in ExtraData.Split("\r\n".ToCharArray()))
-				{
-					if (!string.IsNullOrWhiteSpace(Line))
-					{
-						Text.AppendLine("\t" + Line);
-					}
-				}
-			}
-
-			// Add remote-notifications as background mode
-			if (bRemoteNotificationsSupported || bBackgroundFetch)
-			{
-				Text.AppendLine("\t<key>UIBackgroundModes</key>");
-				Text.AppendLine("\t<array>");
-                if (bRemoteNotificationsSupported)
-                {
-				    Text.AppendLine("\t\t<string>remote-notification</string>");
-                }
-                if (bBackgroundFetch)
-                {
-                    Text.AppendLine("\t\t<string>fetch</string>");
-                }
-				Text.AppendLine("\t</array>");
-			}
-
-			// write the iCloud container identifier, if present in the old file
-			if (!string.IsNullOrEmpty(OldPListData))
-			{
-				int index = OldPListData.IndexOf("ICloudContainerIdentifier");
-				if (index > 0)
-				{
-					index = OldPListData.IndexOf("<string>", index) + 8;
-					int length = OldPListData.IndexOf("</string>", index) - index;
-					string ICloudContainerIdentifier = OldPListData.Substring(index, length);
-					Text.AppendLine("\t<key>ICloudContainerIdentifier</key>");
-					Text.AppendLine(string.Format("\t<string>{0}</string>", ICloudContainerIdentifier));
-				}
-			}
-
-			Text.AppendLine("</dict>");
-			Text.AppendLine("</plist>");
-
-			// Create the intermediate directory if needed
-			if (!Directory.Exists(IntermediateDirectory))
-			{
-				Directory.CreateDirectory(IntermediateDirectory);
-			}
-
-			if (UPL != null)
-			{
-				// Allow UPL to modify the plist here
-				XDocument XDoc;
-				try
-				{
-					XDoc = XDocument.Parse(Text.ToString());
-				}
-				catch (Exception e)
-				{
-					throw new BuildException("plist is invalid {0}\n{1}", e, Text.ToString());
-				}
-
-				XDoc.DocumentType!.InternalSubset = "";
-				UPL.ProcessPluginNode("None", "iosPListUpdates", "", ref XDoc);
-				string result = XDoc.Declaration?.ToString() + "\n" + XDoc.ToString().Replace("<!DOCTYPE plist PUBLIC \"-//Apple//DTD PLIST 1.0//EN\" \"http://www.apple.com/DTDs/PropertyList-1.0.dtd\"[]>", "<!DOCTYPE plist PUBLIC \"-//Apple//DTD PLIST 1.0//EN\" \"http://www.apple.com/DTDs/PropertyList-1.0.dtd\">");
-				File.WriteAllText(PListFile, result);
-
-				Text = new StringBuilder(result);
-			}
-#endif
-		}
-
-		public static string Plist(string Command, string PlistFile)
-		{
-			Command = Command.Replace("\"", "\\\"");
-			string Output = Utils.RunLocalProcessAndReturnStdOut("/usr/libexec/PlistBuddy", $"-c \"{Command}\" \"{PlistFile}\"");
-
-//			Console.WriteLine($"{Command} ==> {Output}");
-
-			return Output;
 		}
 
 		private string Plist(string Command)
 		{
-			return Plist(Command, PBXFilePathForPlatform(CurrentPlistPlatform).FullName);
-		}
-
-		private void PlistSetAdd(string Entry, string Value, string Type="string")
-		{
-			string AddOutput = Plist($"Add {Entry} {Type} {Value}");
-			// error will be non-empty string
-			if (AddOutput != "")
-			{
-				Plist($"Set {Entry} {Value}");
-			}
-		}
-
-		private bool PlistSetUpdate(string Entry, string Value)
-		{
-			// see if the setting is already there
-			string ExistingSetting = Plist($"Print {Entry}");
-
-			// Print errors start with Print
-			if (!ExistingSetting.StartsWith("Print:") && ExistingSetting != Value)
-			{
-				Plist($"Set {Entry} {Value}");
-
-				return true;
-			}
-
-			return false;
-		}
-
-		private IEnumerable<string> PlistArray(string Entry)
-		{
-			return Plist($"Print {Entry}")
-				.Replace("Array {", "")
-				.Replace("}", "")
-				.Trim()
-				.ReplaceLineEndings()
-				.Split(Environment.NewLine)
-				.Select(x => x.Trim());
-		}
-
-		private List<string> PlistObjects()
-		{
-			List<string> Result = new();
-
-			IEnumerable<string> Lines = Plist("print :objects")
-				.ReplaceLineEndings()
-				.Split(Environment.NewLine);
-
-			Regex Regex = new Regex("^    (\\S*) = Dict {$");
-			foreach (string Line in Lines)
-			{
-				Match Match = Regex.Match(Line);
-				if (Match.Success)
-				{
-					Result.Add(Match.Groups[1].Value);
-				}
-			}
-			return Result;
-		}
-
-		private string? PlistFixPath(string Entry, string RelativeToProject)
-		{
-			string ExistingPath = Plist($"Print {Entry}");
-			// skip of errors, or it's an absolute path
-			if (!ExistingPath.StartsWith("Print:") && !ExistingPath.StartsWith("/"))
-			{
-				// fixup the path to be relative to new project instead of old
-				string FixedPath = Utils.CollapseRelativeDirectories(Path.Combine(RelativeToProject, ExistingPath));
-				// and set it back
-				Plist($"Set {Entry} {FixedPath}");
-
-				return FixedPath;
-			}
-
-			return null;
+			return XcodeUtils.Plist(Command);
 		}
 
-
 		bool MergeIntoTemplateProject(FileReference PBXProjFilePath, XcodeProject RootProject, FileReference TemplateProject)
 		{
+			// activate a file for plist reading/writing here
+			XcodeUtils.SetActivePlistFile(PBXFilePathForPlatform(CurrentPlistPlatform).FullName);
+
 			// copy existing template project to final location
 			if (FileReference.Exists(PBXProjFilePath))
 			{
@@ -3522,7 +2109,7 @@ namespace UnrealBuildTool.XcodeProjectXcconfig
 
 			// fixup paths that were relative to original project to be relative to merged project
 //			List<string> ObjectGuids = PlistObjects();
-			IEnumerable<string> MainGroupChildrenGuids = PlistArray($":objects:{TemplateMainGroupGuid}:children");
+			IEnumerable<string> MainGroupChildrenGuids = XcodeUtils.PlistArray($":objects:{TemplateMainGroupGuid}:children");
 
 			string RelativeFromMergedToTemplate = TemplateProject.Directory.ParentDirectory!.MakeRelativeTo(PBXProjFilePath.Directory.ParentDirectory!);
 
@@ -3534,10 +2121,10 @@ namespace UnrealBuildTool.XcodeProjectXcconfig
 				// if a Group has a path
 				if (IsA == "PBXGroup")
 				{
-					if ((FixedPath = PlistFixPath($":objects:{ChildGuid}:path", RelativeFromMergedToTemplate)) != null)
+					if ((FixedPath = XcodeUtils.PlistFixPath($":objects:{ChildGuid}:path", RelativeFromMergedToTemplate)) != null)
 					{
 						// if there wasn't a name before, it will now have a nasty path as the name, so add it now
-						PlistSetAdd($":objects:{ChildGuid}:name", Path.GetFileName(FixedPath));
+						XcodeUtils.PlistSetAdd($":objects:{ChildGuid}:name", Path.GetFileName(FixedPath));
 					}
 				}
 			}
@@ -3547,7 +2134,7 @@ namespace UnrealBuildTool.XcodeProjectXcconfig
 
 
 			// get all the targets in the template that are application types
-			IEnumerable<string> AppTargetGuids = PlistArray($":objects:{ProjectGuid}:targets")
+			IEnumerable<string> AppTargetGuids = XcodeUtils.PlistArray($":objects:{ProjectGuid}:targets")
 				.Where(TargetGuid => (Plist($"Print :objects:{TargetGuid}:productType") == "com.apple.product-type.application"));
 
 			// add a dependency on the build target from the app target(s)
@@ -3571,7 +2158,7 @@ namespace UnrealBuildTool.XcodeProjectXcconfig
 			//PlistSetAdd($":objects:{ProjectGuid}:baseConfigurationReference", RootProject.Xcconfig!.Guid, "string");
 
 			// re-get the list of targets now that we merged in the other file
-			IEnumerable<string> AllTargetGuids = PlistArray($":objects:{ProjectGuid}:targets");
+			IEnumerable<string> AllTargetGuids = XcodeUtils.PlistArray($":objects:{ProjectGuid}:targets");
 
 			List<string> NodesToFix = new() { ProjectGuid };
 			NodesToFix.AddRange(AllTargetGuids);
@@ -3584,7 +2171,7 @@ namespace UnrealBuildTool.XcodeProjectXcconfig
 				// get the config list, and from there we can get the configs
 				string ProjectBuildConfigListGuid = Plist($"Print :objects:{NodeGuid}:buildConfigurationList");
 
-				IEnumerable<string> ConfigGuids = PlistArray($":objects:{ProjectBuildConfigListGuid}:buildConfigurations");
+				IEnumerable<string> ConfigGuids = XcodeUtils.PlistArray($":objects:{ProjectBuildConfigListGuid}:buildConfigurations");
 				foreach (string ConfigGuid in ConfigGuids)
 				{
 					// find the matching unreal generated project build config to hook up to
@@ -3594,7 +2181,7 @@ namespace UnrealBuildTool.XcodeProjectXcconfig
 					if (ConfigName == "Release")
 					{
 						ConfigName = "Development";
-						if (UnrealData.bMakeProjectPerTarget && UnrealData.bHasEditorConfiguration)
+						if (UnrealData.bMakeProjectPerTarget && UnrealData.TargetRules.Type == TargetType.Editor)
 						{
 							ConfigName = "Development Editor";
 						}
@@ -3602,25 +2189,25 @@ namespace UnrealBuildTool.XcodeProjectXcconfig
 					}
 
 					// if there's a plist path, then it will need to be fixed up
-					PlistFixPath($":objects:{ConfigGuid}:buildSettings:INFOPLIST_FILE", RelativeFromMergedToTemplate);
+					XcodeUtils.PlistFixPath($":objects:{ConfigGuid}:buildSettings:INFOPLIST_FILE", RelativeFromMergedToTemplate);
 
 					if (bIsProject)
 					{
 						//Console.WriteLine("Looking for " + ConfigName);
 						XcodeBuildConfig Config = RootProject!.ProjectBuildConfigs.BuildConfigs.First(x => x.Info.DisplayName == ConfigName);
-						PlistSetAdd($":objects:{ConfigGuid}:baseConfigurationReference", Config.Xcconfig!.Guid, "string");
+						XcodeUtils.PlistSetAdd($":objects:{ConfigGuid}:baseConfigurationReference", Config.Xcconfig!.Guid, "string");
 					}
 
 					// the Build target used some ini settings to compile, and Run target must match, so we override a few settings, at
 					// whatever level they were already specified at (Projet and/or Target)
-					PlistSetUpdate($":objects:{ConfigGuid}:buildSettings:MACOSX_DEPLOYMENT_TARGET", MacToolChain.Settings.MacOSVersion);
+					XcodeUtils.PlistSetUpdate($":objects:{ConfigGuid}:buildSettings:MACOSX_DEPLOYMENT_TARGET", MacToolChain.Settings.MacOSVersion);
 					if (UnrealData.IOSProjectSettings != null)
 					{
-						PlistSetUpdate($":objects:{ConfigGuid}:buildSettings:IPHONEOS_DEPLOYMENT_TARGET", UnrealData.IOSProjectSettings.RuntimeVersion);
+						XcodeUtils.PlistSetUpdate($":objects:{ConfigGuid}:buildSettings:IPHONEOS_DEPLOYMENT_TARGET", UnrealData.IOSProjectSettings.RuntimeVersion);
 					}
 					if (UnrealData.TVOSProjectSettings != null)
 					{
-						PlistSetUpdate($":objects:{ConfigGuid}:buildSettings:TVOS_DEPLOYMENT_TARGET", UnrealData.TVOSProjectSettings.RuntimeVersion);
+						XcodeUtils.PlistSetUpdate($":objects:{ConfigGuid}:buildSettings:TVOS_DEPLOYMENT_TARGET", UnrealData.TVOSProjectSettings.RuntimeVersion);
 					}
 				}
 
@@ -3647,286 +2234,5 @@ namespace UnrealBuildTool.XcodeProjectXcconfig
 
 			return true;
 		}
-
-#region Schemes
-
-		private FileReference GetUserSchemeManagementFilePath(UnrealTargetPlatform? Platform)
-		{
-			return FileReference.Combine(ProjectDirPathForPlatform(Platform), $"xcuserdata/{Environment.UserName}.xcuserdatad/xcschemes/xcschememanagement.plist");
-		}
-
-		private DirectoryReference GetProjectSchemeDirectory(UnrealTargetPlatform? Platform)
-		{
-			return DirectoryReference.Combine(ProjectDirPathForPlatform(Platform), "xcshareddata/xcschemes");
-		}
-
-		private FileReference GetProjectSchemeFilePathForTarget(UnrealTargetPlatform? Platform, string TargetName)
-		{
-			return FileReference.Combine(GetProjectSchemeDirectory(Platform), TargetName + ".xcscheme");
-		}
-
-		private void WriteSchemeFile(UnrealTargetPlatform? Platform, string ProjectName, string ProductName, List<XcodeRunTarget> RunTargets, string? BuildTargetGuid, string? IndexTargetGuid, string GameProjectPath, IEnumerable<string> ExtraPreBuildScriptLines)
-		{
-			StringBuilder Content = new StringBuilder();
-
-			foreach (XcodeRunTarget RunTarget in RunTargets)
-			{
-				string TargetName = RunTarget.Name;
-				string TargetGuid = RunTarget.Guid;
-				
-				bool bHasEditorConfiguration = RunTarget.BuildConfigList!.BuildConfigs.Any(x => x.Info.ProjectTarget!.TargetRules!.Type == TargetType.Editor);
-				bool bUseEditorConfiguration = bHasEditorConfiguration && !UnrealData.bMakeProjectPerTarget && !XcodeProjectFileGenerator.bGenerateRunOnlyProject;
-				string DefaultConfiguration = bUseEditorConfiguration ? "Development Editor" : "Development";
-				
-				FileReference SchemeFilePath = GetProjectSchemeFilePathForTarget(Platform, TargetName);
-
-				DirectoryReference.CreateDirectory(SchemeFilePath.Directory);
-
-				string? OldCommandLineArguments = null;
-				if (FileReference.Exists(SchemeFilePath))
-				{
-					string OldContents = File.ReadAllText(SchemeFilePath.FullName);
-					int OldCommandLineArgumentsStart = OldContents.IndexOf("<CommandLineArguments>") + "<CommandLineArguments>".Length;
-					int OldCommandLineArgumentsEnd = OldContents.IndexOf("</CommandLineArguments>");
-					if (OldCommandLineArgumentsStart != -1 && OldCommandLineArgumentsEnd != -1)
-					{
-						OldCommandLineArguments = OldContents.Substring(OldCommandLineArgumentsStart, OldCommandLineArgumentsEnd - OldCommandLineArgumentsStart);
-					}
-				}
-
-				string[] ScriptLines = new string[]
-				{
-					"# if a UBTGenerated plist doesn&apos;t exist, create it now for Xcode&apos;s dependency needs",
-					"if [ ! -f &quot;${PROJECT_DIR}/${INFOPLIST_FILE}&quot; ]",
-					"then",
-					"  echo Creating ${PROJECT_DIR}/${INFOPLIST_FILE}...",
-					"  touch &quot;${PROJECT_DIR}/${INFOPLIST_FILE}&quot;",
-					"fi",
-					"",
-					"# if we had previously archived the build, then Xcode would have made a soft-link, but it will fail when building normally, so delete a soft-link",
-					"if [ -L &quot;${CONFIGURATION_BUILD_DIR}/${PRODUCT_NAME}.app&quot; ]",
-					"then",
-					"  rm &quot;${CONFIGURATION_BUILD_DIR}/${PRODUCT_NAME}.app&quot;",
-					"fi",
-					"",
-					"# make sure we have a cookeddata directory if needed",
-					"if [[ &quot;${PLATFORM_NAME}&quot; == &quot;iphoneos&quot; || &quot;${PLATFORM_NAME}&quot; == &quot;appletvos&quot; ]]",
-					"then",
-					"  mkdir -p &quot;${CONFIGURATION_BUILD_DIR}/../../../Saved/StagedBuilds/${UE_TARGET_PLATFORM_NAME}/cookeddata&quot;",
-					"fi",
-					"",
-				};
-				// convert special characters to be xml happy
-				IEnumerable<string> Lines = ScriptLines.Concat(ExtraPreBuildScriptLines);
-				Lines = Lines.Select(x => x.Replace("\"", "&quot;").Replace("&&", "&amp;&amp;"));
-
-				Content.WriteLine("<?xml version=\"1.0\" encoding=\"UTF-8\"?>");
-				Content.WriteLine("<Scheme");
-				Content.WriteLine("   LastUpgradeVersion = \"2000\"");
-				Content.WriteLine("   version = \"1.3\">");
-				Content.WriteLine("   <BuildAction");
-				Content.WriteLine("      parallelizeBuildables = \"YES\"");
-				Content.WriteLine("      buildImplicitDependencies = \"YES\">");
-				Content.WriteLine("      <PreActions>");
-				Content.WriteLine("	        <ExecutionAction");
-				Content.WriteLine("            ActionType = \"Xcode.IDEStandardExecutionActionsCore.ExecutionActionType.ShellScriptAction\">");
-				Content.WriteLine("            <ActionContent");
-				Content.WriteLine("               title = \"Touch UBT generated tiles\"");
-				Content.WriteLine("               scriptText = \"" + string.Join("&#10;", Lines) + "\">");
-				Content.WriteLine("               <EnvironmentBuildable>");
-				Content.WriteLine("                  <BuildableReference");
-				Content.WriteLine("                     BuildableIdentifier = \"primary\"");
-				Content.WriteLine("                     BlueprintIdentifier = \"" + TargetGuid + "\"");
-				Content.WriteLine("                     BuildableName = \"" + ProductName + ".app\"");
-				Content.WriteLine("                     BlueprintName = \"" + TargetName + "\"");
-				Content.WriteLine("                     ReferencedContainer = \"container:" + ProjectName + ".xcodeproj\">");
-				Content.WriteLine("                  </BuildableReference>");
-				Content.WriteLine("               </EnvironmentBuildable>");
-				Content.WriteLine("            </ActionContent>");
-				Content.WriteLine("         </ExecutionAction>");
-				Content.WriteLine("      </PreActions>");
-				Content.WriteLine("      <BuildActionEntries>");
-				Content.WriteLine("         <BuildActionEntry");
-				Content.WriteLine("            buildForTesting = \"YES\"");
-				Content.WriteLine("            buildForRunning = \"YES\"");
-				Content.WriteLine("            buildForProfiling = \"YES\"");
-				Content.WriteLine("            buildForArchiving = \"YES\"");
-				Content.WriteLine("            buildForAnalyzing = \"YES\">");
-				Content.WriteLine("            <BuildableReference");
-				Content.WriteLine("               BuildableIdentifier = \"primary\"");
-				Content.WriteLine("               BlueprintIdentifier = \"" + TargetGuid + "\"");
-				Content.WriteLine("               BuildableName = \"" + ProductName + ".app\"");
-				Content.WriteLine("               BlueprintName = \"" + TargetName + "\"");
-				Content.WriteLine("               ReferencedContainer = \"container:" + ProjectName + ".xcodeproj\">");
-				Content.WriteLine("            </BuildableReference>");
-				Content.WriteLine("         </BuildActionEntry>");
-				Content.WriteLine("      </BuildActionEntries>");
-				Content.WriteLine("   </BuildAction>");
-				Content.WriteLine("   <TestAction");
-				Content.WriteLine("      buildConfiguration = \"" + DefaultConfiguration + "\"");
-				Content.WriteLine("      selectedDebuggerIdentifier = \"Xcode.DebuggerFoundation.Debugger.LLDB\"");
-				Content.WriteLine("      selectedLauncherIdentifier = \"Xcode.DebuggerFoundation.Launcher.LLDB\"");
-				Content.WriteLine("      shouldUseLaunchSchemeArgsEnv = \"YES\">");
-				Content.WriteLine("      <MacroExpansion>");
-				Content.WriteLine("         <BuildableReference");
-				Content.WriteLine("            BuildableIdentifier = \"primary\"");
-				Content.WriteLine("            BlueprintIdentifier = \"" + TargetGuid + "\"");
-				Content.WriteLine("            BuildableName = \"" + ProductName + ".app\"");
-				Content.WriteLine("            BlueprintName = \"" + TargetName + "\"");
-				Content.WriteLine("            ReferencedContainer = \"container:" + ProjectName + ".xcodeproj\">");
-				Content.WriteLine("         </BuildableReference>");
-				Content.WriteLine("      </MacroExpansion>");
-				Content.WriteLine("      <Testables>");
-				Content.WriteLine("      </Testables>");
-				Content.WriteLine("   </TestAction>");
-				Content.WriteLine("   <LaunchAction");
-				Content.WriteLine("      buildConfiguration = \"" + DefaultConfiguration + "\"");
-				Content.WriteLine("      selectedDebuggerIdentifier = \"Xcode.DebuggerFoundation.Debugger.LLDB\"");
-				Content.WriteLine("      selectedLauncherIdentifier = \"Xcode.DebuggerFoundation.Launcher.LLDB\"");
-				Content.WriteLine("      launchStyle = \"0\"");
-				Content.WriteLine("      useCustomWorkingDirectory = \"NO\"");
-				Content.WriteLine("      ignoresPersistentStateOnLaunch = \"NO\"");
-				Content.WriteLine("      debugDocumentVersioning = \"NO\"");
-				Content.WriteLine("      debugServiceExtension = \"internal\"");
-				Content.WriteLine("      allowLocationSimulation = \"YES\">");
-				Content.WriteLine("      <BuildableProductRunnable");
-				Content.WriteLine("         runnableDebuggingMode = \"0\">");
-				Content.WriteLine("         <BuildableReference");
-				Content.WriteLine("            BuildableIdentifier = \"primary\"");
-				Content.WriteLine("            BlueprintIdentifier = \"" + TargetGuid + "\"");
-				Content.WriteLine("            BuildableName = \"" + ProductName + ".app\"");
-				Content.WriteLine("            BlueprintName = \"" + TargetName + "\"");
-				Content.WriteLine("            ReferencedContainer = \"container:" + ProjectName + ".xcodeproj\">");
-				Content.WriteLine("         </BuildableReference>");
-				Content.WriteLine("      </BuildableProductRunnable>");
-				if (string.IsNullOrEmpty(OldCommandLineArguments))
-				{
-					if (bHasEditorConfiguration && GameProjectPath != "")
-					{
-						Content.WriteLine("      <CommandLineArguments>");
-						if (IsForeignProject)
-						{
-							Content.WriteLine("         <CommandLineArgument");
-							Content.WriteLine("            argument = \"&quot;" + GameProjectPath + "&quot;\"");
-							Content.WriteLine("            isEnabled = \"YES\">");
-							Content.WriteLine("         </CommandLineArgument>");
-						}
-						else
-						{
-							Content.WriteLine("         <CommandLineArgument");
-							Content.WriteLine("            argument = \"" + Path.GetFileNameWithoutExtension(GameProjectPath) + "\"");
-							Content.WriteLine("            isEnabled = \"YES\">");
-							Content.WriteLine("         </CommandLineArgument>");
-						}
-						// Always add a configuration argument
-						Content.WriteLine("         <CommandLineArgument");
-						Content.WriteLine("            argument = \"-RunConfig=$(Configuration)\"");
-						Content.WriteLine("            isEnabled = \"YES\">");
-						Content.WriteLine("         </CommandLineArgument>");
-						Content.WriteLine("      </CommandLineArguments>");
-					}
-				}
-				else
-				{
-					Content.WriteLine("      <CommandLineArguments>" + OldCommandLineArguments + "</CommandLineArguments>");
-				}
-				Content.WriteLine("   </LaunchAction>");
-				Content.WriteLine("   <ProfileAction");
-				Content.WriteLine("      buildConfiguration = \"" + DefaultConfiguration + "\"");
-				Content.WriteLine("      shouldUseLaunchSchemeArgsEnv = \"YES\"");
-				Content.WriteLine("      savedToolIdentifier = \"\"");
-				Content.WriteLine("      useCustomWorkingDirectory = \"NO\"");
-				Content.WriteLine("      debugDocumentVersioning = \"NO\">");
-				Content.WriteLine("      <BuildableProductRunnable");
-				Content.WriteLine("         runnableDebuggingMode = \"0\">");
-				Content.WriteLine("         <BuildableReference");
-				Content.WriteLine("            BuildableIdentifier = \"primary\"");
-				Content.WriteLine("            BlueprintIdentifier = \"" + TargetGuid + "\"");
-				Content.WriteLine("            BuildableName = \"" + ProductName + ".app\"");
-				Content.WriteLine("            BlueprintName = \"" + TargetName + "\"");
-				Content.WriteLine("            ReferencedContainer = \"container:" + ProjectName + ".xcodeproj\">");
-				Content.WriteLine("         </BuildableReference>");
-				Content.WriteLine("      </BuildableProductRunnable>");
-				Content.WriteLine("   </ProfileAction>");
-				Content.WriteLine("   <AnalyzeAction");
-				Content.WriteLine("      buildConfiguration = \"" + DefaultConfiguration + "\">");
-				Content.WriteLine("   </AnalyzeAction>");
-				Content.WriteLine("   <ArchiveAction");
-				Content.WriteLine("      buildConfiguration = \"" + DefaultConfiguration + "\"");
-				Content.WriteLine("      revealArchiveInOrganizer = \"YES\">");
-				Content.WriteLine("   </ArchiveAction>");
-				Content.WriteLine("</Scheme>");
-
-				File.WriteAllText(SchemeFilePath.FullName, Content.ToString(), new UTF8Encoding());
-				Content.Clear();
-			}
-
-			Content.WriteLine("<?xml version=\"1.0\" encoding=\"UTF-8\"?>");
-			Content.WriteLine("<!DOCTYPE plist PUBLIC \"-//Apple//DTD PLIST 1.0//EN\" \"http://www.apple.com/DTDs/PropertyList-1.0.dtd\">");
-			Content.WriteLine("<plist version=\"1.0\">");
-			Content.WriteLine("<dict>");
-			Content.WriteLine("\t<key>SchemeUserState</key>");
-			Content.WriteLine("\t<dict>");
-			foreach (XcodeRunTarget RunTarget in RunTargets)
-			{
-				Content.WriteLine("\t\t<key>" + RunTarget.Name + ".xcscheme_^#shared#^_</key>");
-				Content.WriteLine("\t\t<dict>");
-				Content.WriteLine("\t\t\t<key>orderHint</key>");
-				Content.WriteLine("\t\t\t<integer>1</integer>");
-				Content.WriteLine("\t\t</dict>");
-			}
-			Content.WriteLine("\t</dict>");
-			Content.WriteLine("\t<key>SuppressBuildableAutocreation</key>");
-			Content.WriteLine("\t<dict>");
-			foreach (XcodeRunTarget RunTarget in RunTargets)
-			{
-				Content.WriteLine("\t\t<key>" + RunTarget.Guid + "</key>");
-				Content.WriteLine("\t\t<dict>");
-				Content.WriteLine("\t\t\t<key>primary</key>");
-				Content.WriteLine("\t\t\t<true/>");
-				Content.WriteLine("\t\t</dict>");
-			}
-			if (BuildTargetGuid != null)
-			{
-				Content.WriteLine("\t\t<key>" + BuildTargetGuid + "</key>");
-				Content.WriteLine("\t\t<dict>");
-				Content.WriteLine("\t\t\t<key>primary</key>");
-				Content.WriteLine("\t\t\t<true/>");
-				Content.WriteLine("\t\t</dict>");
-			}
-			if (IndexTargetGuid != null)
-			{
-				Content.WriteLine("\t\t<key>" + IndexTargetGuid + "</key>");
-				Content.WriteLine("\t\t<dict>");
-				Content.WriteLine("\t\t\t<key>primary</key>");
-				Content.WriteLine("\t\t\t<true/>");
-				Content.WriteLine("\t\t</dict>");
-			}
-			Content.WriteLine("\t</dict>");
-			Content.WriteLine("</dict>");
-			Content.WriteLine("</plist>");
-
-			FileReference ManagementFile = GetUserSchemeManagementFilePath(Platform);
-			if (!DirectoryReference.Exists(ManagementFile.Directory))
-			{
-				DirectoryReference.CreateDirectory(ManagementFile.Directory);
-			}
-
-			File.WriteAllText(ManagementFile.FullName, Content.ToString(), new UTF8Encoding());
-		}
-
-#endregion
-
-
-#region Utilities
-
-		public override string ToString()
-		{
-			throw new NotImplementedException();
-//			return ProjectFilePath.GetFileNameWithoutExtension();
-		}
-
-#endregion
-
 	}
 }
diff --git a/Engine/Source/Programs/UnrealBuildTool/ProjectFiles/Xcode/XcodeProjectFileGenerator.cs b/Engine/Source/Programs/UnrealBuildTool/ProjectFiles/Xcode/XcodeProjectFileGenerator.cs
--- a/Engine/Source/Programs/UnrealBuildTool/ProjectFiles/Xcode/XcodeProjectFileGenerator.cs
+++ b/Engine/Source/Programs/UnrealBuildTool/ProjectFiles/Xcode/XcodeProjectFileGenerator.cs
@@ -238,7 +238,7 @@ namespace UnrealBuildTool
 							IEnumerable<XcodeProjectXcconfig.XcodeProjectFile> SupportedProjects =
 									CurFolder.ChildProjects.Where(P => P.GetType() == typeof(XcodeProjectXcconfig.XcodeProjectFile))
 										.Select(P => (XcodeProjectXcconfig.XcodeProjectFile)P)
-										.Where(P => XcodeProjectXcconfig.UnrealData.ShouldIncludeProjectInWorkspace(P, Logger))
+										.Where(P => XcodeProjectXcconfig.XcodeUtils.ShouldIncludeProjectInWorkspace(P, Logger))
 										// @todo - still need to handle legacy project getting split up?
 										.Where(P => P.RootProjects.Count == 0 || P.RootProjects.ContainsValue(Platform))
 										.OrderBy(P => P.ProjectFilePath.GetFileName());
