Fixed an issue where -allcores would not be used if ProcessorCountMultiplier was set.
Fixed an issue where it was possible to do more parallel actions than cores available.

diff --git a/Engine/Source/Programs/UnrealBuildTool/Executors/ParallelExecutor.cs b/Engine/Source/Programs/UnrealBuildTool/Executors/ParallelExecutor.cs
--- a/Engine/Source/Programs/UnrealBuildTool/Executors/ParallelExecutor.cs
+++ b/Engine/Source/Programs/UnrealBuildTool/Executors/ParallelExecutor.cs
@@ -112,7 +112,7 @@ namespace UnrealBuildTool
 				Logger.LogInformation("Overriding MemoryPerAction with target-defined value of {Memory} bytes", MemoryPerActionBytesComputed / 1024 / 1024 / 1024);
 			}
 
-			return Utils.GetMaxActionsToExecuteInParallel(MaxLocalActions, ProcessorCountMultiplier, bAllCores, Convert.ToInt64(MemoryPerActionBytesComputed));
+			return Utils.GetMaxActionsToExecuteInParallel(MaxLocalActions, bAllCores ? 1.0f : ProcessorCountMultiplier, bAllCores, Convert.ToInt64(MemoryPerActionBytesComputed));
 		}
 
 		/// <summary>
diff --git a/Engine/Source/Programs/UnrealBuildTool/System/Utils.cs b/Engine/Source/Programs/UnrealBuildTool/System/Utils.cs
--- a/Engine/Source/Programs/UnrealBuildTool/System/Utils.cs
+++ b/Engine/Source/Programs/UnrealBuildTool/System/Utils.cs
@@ -1350,6 +1350,10 @@ namespace UnrealBuildTool
 				// The CPU has more logical cores than physical ones, aka uses hyper-threading. 
 				// Use multiplier if provided
 				MaxActionsToExecuteInParallel = (int)(NumPhysicalCores * ProcessorCountMultiplier);
+
+				// make sure we don't try to run more actions than cores we have
+				MaxActionsToExecuteInParallel = Math.Min(MaxActionsToExecuteInParallel, Math.Max(NumLogicalCores, NumPhysicalCores));
+
 				Log.TraceInformationOnce($"  Requested {ProcessorCountMultiplier} process count multiplier: limiting max parallel actions to {MaxActionsToExecuteInParallel}");
 			}
 			else if (ConsiderLogicalCores && NumLogicalCores > NumPhysicalCores)
