Use PrimaryProjectName to name VS Code workspace file.

diff --git a/Engine/Source/Programs/UnrealBuildTool/ProjectFiles/VisualStudioCode/VSCodeProjectFileGenerator.cs b/Engine/Source/Programs/UnrealBuildTool/ProjectFiles/VisualStudioCode/VSCodeProjectFileGenerator.cs
--- a/Engine/Source/Programs/UnrealBuildTool/ProjectFiles/VisualStudioCode/VSCodeProjectFileGenerator.cs
+++ b/Engine/Source/Programs/UnrealBuildTool/ProjectFiles/VisualStudioCode/VSCodeProjectFileGenerator.cs
@@ -1706,7 +1706,7 @@ namespace UnrealBuildTool
 
 			WorkspaceFile.EndRootObject();
 
-			string? WorkspaceName = bForeignProject ? GameProjectName : "UE5";
+			string? WorkspaceName = bForeignProject ? GameProjectName : PrimaryProjectName;
 			WorkspaceFile.Write(FileReference.Combine(PrimaryProjectPath, WorkspaceName + ".code-workspace"));
 		}
 
