Fix response file paths with special characters causing packaging errors

diff --git a/Engine/Source/Programs/UnrealBuildTool/System/Utils.cs b/Engine/Source/Programs/UnrealBuildTool/System/Utils.cs
--- a/Engine/Source/Programs/UnrealBuildTool/System/Utils.cs
+++ b/Engine/Source/Programs/UnrealBuildTool/System/Utils.cs
@@ -133,7 +133,8 @@ namespace UnrealBuildTool
 		/// <returns></returns>
 		public static string MakePathSafeToUseWithCommandLine(string InPath)
 		{
-			if (InPath.Contains(' ') && InPath[0] != '\"')
+			// just always quote paths if they aren't already
+			if (InPath[0] != '\"')
 			{
 				InPath = "\"" + InPath + "\"";
 			}
