UnrealBuildTool: Don't -SingleFile compile a file that doesn't exist such as a file open for delete

#rnx

diff --git a/Engine/Source/Programs/UnrealBuildTool/Modes/BuildMode.cs b/Engine/Source/Programs/UnrealBuildTool/Modes/BuildMode.cs
--- a/Engine/Source/Programs/UnrealBuildTool/Modes/BuildMode.cs
+++ b/Engine/Source/Programs/UnrealBuildTool/Modes/BuildMode.cs
@@ -799,6 +799,12 @@ namespace UnrealBuildTool
 			{
 				FileItem SourceFile = FileItem.GetItemByFileReference(FileRef);
 
+				// Ensure the file requested exists on disk
+				if (!SourceFile.Exists)
+				{
+					continue;
+				}
+
 				ISpecificFileAction? SpecificFileAction = null;
 
 				// Traverse upwards from file directory to try to find a SpecificFileAction.
