UnrealBuildTool: Add -UsePDB TargetRules argument for testing

#rnx

diff --git a/Engine/Source/Programs/UnrealBuildTool/Configuration/TargetRules.cs b/Engine/Source/Programs/UnrealBuildTool/Configuration/TargetRules.cs
--- a/Engine/Source/Programs/UnrealBuildTool/Configuration/TargetRules.cs
+++ b/Engine/Source/Programs/UnrealBuildTool/Configuration/TargetRules.cs
@@ -1647,6 +1647,7 @@ namespace UnrealBuildTool
 		/// Whether PDB files should be used for Visual C++ builds.
 		/// </summary>
 		[CommandLine("-NoPDB", Value = "false")]
+		[CommandLine("-UsePDB", Value = "true")]
 		[XmlConfigFile(Category = "BuildConfiguration")]
 		public bool bUsePDBFiles = false;
 
