PR #10107: Limit Fastbuild process count (Contributed by [REDACTED_PERSON])

diff --git a/Engine/Source/Programs/UnrealBuildTool/Executors/Experimental/FASTBuild.cs b/Engine/Source/Programs/UnrealBuildTool/Executors/Experimental/FASTBuild.cs
--- a/Engine/Source/Programs/UnrealBuildTool/Executors/Experimental/FASTBuild.cs
+++ b/Engine/Source/Programs/UnrealBuildTool/Executors/Experimental/FASTBuild.cs
@@ -1283,13 +1283,14 @@ namespace UnrealBuildTool
 			string ForceRemoteArgument		= bForceRemote ? "-forceremote" : "";
 			string NoStopOnErrorArgument	= bStopOnError ? "" : "-nostoponerror";
 			string IDEArgument				= IsApple() ? "" : "-ide";
+			string MaxProcesses				= "-j" + ((ParallelExecutor)LocalExecutor).NumParallelProcesses;
 
 			// Interesting flags for FASTBuild:
 			// -nostoponerror, -verbose, -monitor (if FASTBuild Monitor Visual Studio Extension is installed!)
 			// Yassine: The -clean is to bypass the FASTBuild internal
 			// dependencies checks (cached in the fdb) as it could create some conflicts with UBT.
 			// Basically we want FB to stupidly compile what UBT tells it to.
-			string FBCommandLine	= $"-monitor -summary {DistArgument} {CacheArgument} {IDEArgument} -clean -config \"{BffFilePath}\" {NoStopOnErrorArgument} {ForceRemoteArgument}";
+			string FBCommandLine	= $"-monitor -summary {DistArgument} {CacheArgument} {IDEArgument} {MaxProcesses} -clean -config \"{BffFilePath}\" {NoStopOnErrorArgument} {ForceRemoteArgument}";
 
 			Logger.LogInformation("FBuild Command Line Arguments: '{FBCommandLine}", FBCommandLine);
 
