UBT: added initial support for the internal AutoRTFM Clang variant

diff --git a/Engine/Source/Programs/UnrealBuildTool/Platform/Windows/MicrosoftPlatformSDK.cs b/Engine/Source/Programs/UnrealBuildTool/Platform/Windows/MicrosoftPlatformSDK.cs
--- a/Engine/Source/Programs/UnrealBuildTool/Platform/Windows/MicrosoftPlatformSDK.cs
+++ b/Engine/Source/Programs/UnrealBuildTool/Platform/Windows/MicrosoftPlatformSDK.cs
@@ -659,27 +659,27 @@ namespace UnrealBuildTool
 					{
 						// Check for a manual installation to the default directory
 						DirectoryReference ManualInstallDir = DirectoryReference.Combine(DirectoryReference.GetSpecialFolder(Environment.SpecialFolder.ProgramFiles)!, "LLVM");
-						AddClangToolChain(ManualInstallDir, ToolChains, IsAutoSdk: false, Logger);
+						AddClangToolChain(Compiler, ManualInstallDir, ToolChains, IsAutoSdk: false, Logger);
 
 						// Check for a manual installation to a custom directory
 						string? LlvmPath = Environment.GetEnvironmentVariable("LLVM_PATH");
 						if (!String.IsNullOrEmpty(LlvmPath))
 						{
-							AddClangToolChain(new DirectoryReference(LlvmPath), ToolChains, IsAutoSdk: false, Logger);
+							AddClangToolChain(Compiler, new DirectoryReference(LlvmPath), ToolChains, IsAutoSdk: false, Logger);
 						}
 
 						// Check for installations bundled with Visual Studio 2019
 						foreach (VisualStudioInstallation Installation in FindVisualStudioInstallations(WindowsCompiler.VisualStudio2019, Logger))
 						{
-							AddClangToolChain(DirectoryReference.Combine(Installation.BaseDir, "VC", "Tools", "Llvm"), ToolChains, IsAutoSdk: false, Logger);
-							AddClangToolChain(DirectoryReference.Combine(Installation.BaseDir, "VC", "Tools", "Llvm", "x64"), ToolChains, IsAutoSdk: false, Logger);
+							AddClangToolChain(Compiler, DirectoryReference.Combine(Installation.BaseDir, "VC", "Tools", "Llvm"), ToolChains, IsAutoSdk: false, Logger);
+							AddClangToolChain(Compiler, DirectoryReference.Combine(Installation.BaseDir, "VC", "Tools", "Llvm", "x64"), ToolChains, IsAutoSdk: false, Logger);
 						}
 
 						// Check for installations bundled with Visual Studio 2022
 						foreach (VisualStudioInstallation Installation in FindVisualStudioInstallations(WindowsCompiler.VisualStudio2022, Logger))
 						{
-							AddClangToolChain(DirectoryReference.Combine(Installation.BaseDir, "VC", "Tools", "Llvm"), ToolChains, IsAutoSdk: false, Logger);
-							AddClangToolChain(DirectoryReference.Combine(Installation.BaseDir, "VC", "Tools", "Llvm", "x64"), ToolChains, IsAutoSdk: false, Logger);
+							AddClangToolChain(Compiler, DirectoryReference.Combine(Installation.BaseDir, "VC", "Tools", "Llvm"), ToolChains, IsAutoSdk: false, Logger);
+							AddClangToolChain(Compiler, DirectoryReference.Combine(Installation.BaseDir, "VC", "Tools", "Llvm", "x64"), ToolChains, IsAutoSdk: false, Logger);
 						}
 
 						// Check for AutoSDK paths
@@ -691,11 +691,17 @@ namespace UnrealBuildTool
 							{
 								foreach (DirectoryReference ToolChainDir in DirectoryReference.EnumerateDirectories(ClangBaseDir))
 								{
-									AddClangToolChain(ToolChainDir, ToolChains, IsAutoSdk: true, Logger);
+									AddClangToolChain(Compiler, ToolChainDir, ToolChains, IsAutoSdk: true, Logger);
 								}
 							}
 						}
 					}
+					else if (Compiler == WindowsCompiler.ClangRTFM)
+					{
+						// the AutoRTFM Clang compiler is in-tree
+						DirectoryReference AutoRTFMDir = DirectoryReference.Combine(Unreal.EngineDirectory, "Restricted", "NotForLicensees", "Binaries", "Win64", "AutoRTFM");
+						AddClangToolChain(Compiler, AutoRTFMDir, ToolChains, IsAutoSdk: false, Logger);
+					}
 					else if (Compiler == WindowsCompiler.Intel)
 					{
 						DirectoryReference InstallDir = DirectoryReference.Combine(DirectoryReference.GetSpecialFolder(Environment.SpecialFolder.ProgramFiles)!, "Intel", "oneAPI", "compiler");
@@ -971,13 +977,15 @@ namespace UnrealBuildTool
 		/// <summary>
 		/// Add a Clang toolchain
 		/// </summary>
+		/// <param name="Compiler"></param>
 		/// <param name="ToolChainDir"></param>
 		/// <param name="ToolChains"></param>
 		/// <param name="IsAutoSdk"></param>
 		/// <param name="Logger"></param>
-		static void AddClangToolChain(DirectoryReference ToolChainDir, List<ToolChainInstallation> ToolChains, bool IsAutoSdk, ILogger Logger)
+		static void AddClangToolChain(WindowsCompiler Compiler, DirectoryReference ToolChainDir, List<ToolChainInstallation> ToolChains, bool IsAutoSdk, ILogger Logger)
 		{
-			FileReference CompilerFile = FileReference.Combine(ToolChainDir, "bin", "clang-cl.exe");
+			string CompilerFilename = Compiler == WindowsCompiler.ClangRTFM ? "verse-clang-cl.exe" : "clang-cl.exe";
+			FileReference CompilerFile = FileReference.Combine(ToolChainDir, "bin", CompilerFilename);
 			if (FileReference.Exists(CompilerFile))
 			{
 				FileVersionInfo VersionInfo = FileVersionInfo.GetVersionInfo(CompilerFile.FullName);
diff --git a/Engine/Source/Programs/UnrealBuildTool/Platform/Windows/UEBuildWindows.cs b/Engine/Source/Programs/UnrealBuildTool/Platform/Windows/UEBuildWindows.cs
--- a/Engine/Source/Programs/UnrealBuildTool/Platform/Windows/UEBuildWindows.cs
+++ b/Engine/Source/Programs/UnrealBuildTool/Platform/Windows/UEBuildWindows.cs
@@ -106,6 +106,11 @@ namespace UnrealBuildTool
 		/// </summary>
 		Clang,
 
+		/// <summary>
+		/// Use the RTFM (Rollback Transactions on Failure Memory) Clang variant for Verse on Windows, using the verse-clang-cl driver.
+		/// </summary>
+		ClangRTFM,
+
 		/// <summary>
 		/// Use the Intel oneAPI C++ compiler
 		/// </summary>
@@ -135,7 +140,7 @@ namespace UnrealBuildTool
 		/// <returns>true if Clang based</returns>
 		public static bool IsClang(this WindowsCompiler Compiler)
 		{
-			return Compiler == WindowsCompiler.Clang || Compiler == WindowsCompiler.Intel;
+			return Compiler == WindowsCompiler.Clang || Compiler == WindowsCompiler.ClangRTFM || Compiler == WindowsCompiler.Intel;
 		}
 
 		/// <summary>
@@ -602,6 +607,7 @@ namespace UnrealBuildTool
 			switch (Compiler)
 			{
 				case WindowsCompiler.Clang:
+				case WindowsCompiler.ClangRTFM:
 				case WindowsCompiler.Intel:
 				case WindowsCompiler.VisualStudio2019:
 				case WindowsCompiler.VisualStudio2022:
diff --git a/Engine/Source/Programs/UnrealBuildTool/Platform/Windows/VCEnvironment.cs b/Engine/Source/Programs/UnrealBuildTool/Platform/Windows/VCEnvironment.cs
--- a/Engine/Source/Programs/UnrealBuildTool/Platform/Windows/VCEnvironment.cs
+++ b/Engine/Source/Programs/UnrealBuildTool/Platform/Windows/VCEnvironment.cs
@@ -214,6 +214,10 @@ namespace UnrealBuildTool
 			{
 				return FileReference.Combine(CompilerDir, "bin", "clang-cl.exe");
 			}
+			else if (Compiler == WindowsCompiler.ClangRTFM)
+			{
+				return FileReference.Combine(CompilerDir, "bin", "verse-clang-cl.exe");
+			}
 			else if (Compiler == WindowsCompiler.Intel)
 			{
 				return FileReference.Combine(CompilerDir, "windows", "bin", "icx.exe");
@@ -226,7 +230,7 @@ namespace UnrealBuildTool
 		/// </summary>
 		FileReference GetLinkerToolPath(WindowsCompiler Compiler, UnrealArch Architecture, DirectoryReference CompilerDir, DirectoryReference ToochainDir)
 		{
-			if (Compiler == WindowsCompiler.Clang && bAllowClangLinker)
+			if ((Compiler == WindowsCompiler.Clang || Compiler == WindowsCompiler.ClangRTFM) && bAllowClangLinker)
 			{
 				return FileReference.Combine(CompilerDir, "bin", "lld-link.exe");
 			}
@@ -243,7 +247,7 @@ namespace UnrealBuildTool
 		/// </summary>
 		FileReference GetLibraryLinkerToolPath(WindowsCompiler Compiler, UnrealArch Architecture, DirectoryReference CompilerDir, DirectoryReference ToochainDir)
 		{
-			if (Compiler == WindowsCompiler.Clang && bAllowClangLinker)
+			if ((Compiler == WindowsCompiler.Clang || Compiler == WindowsCompiler.ClangRTFM) && bAllowClangLinker)
 			{
 				// @todo: lld-link is not currently working for building .lib
 				//return FileReference.Combine(CompilerDir, "bin", "lld-link.exe");
