Allowing UBT to build individual plugins as though they were local.
* Also reverting a previous change that was building the whole project when adding new plugins.

diff --git a/Engine/Plugins/Developer/PluginUtils/Source/PluginUtils/Private/PluginUtils.cpp b/Engine/Plugins/Developer/PluginUtils/Source/PluginUtils/Private/PluginUtils.cpp
--- a/Engine/Plugins/Developer/PluginUtils/Source/PluginUtils/Private/PluginUtils.cpp
+++ b/Engine/Plugins/Developer/PluginUtils/Source/PluginUtils/Private/PluginUtils.cpp
@@ -638,9 +638,9 @@ TSharedPtr<IPlugin> FPluginUtils::CreateAndLoadNewPlugin(const FString& PluginNa
 			// Internal plugins will be found and built automatically, and using the -Plugin param will mark them as external plugins causing them to fail to load.
 			const bool bIsEnginePlugin = FPaths::IsUnderDirectory(PluginFilePath, FPaths::EnginePluginsDir());
 			const bool bIsProjectPlugin = FPaths::IsUnderDirectory(PluginFilePath, FPaths::ProjectPluginsDir());
-			const FString ExternalPluginArgument = !(bIsEnginePlugin || bIsProjectPlugin) ? FString::Printf(TEXT("-Plugin=\"%s\" "), *PluginFilePath) : FString();
+			const FString LocalPluginArgument = (bIsEnginePlugin || bIsProjectPlugin) ? FString::Printf(TEXT("-BuildPluginAsLocal ")) : FString();
 
-			const FString Arguments = FString::Printf(TEXT("%s %s %s %s-Project=\"%s\" -Progress -NoHotReloadFromIDE"), FPlatformMisc::GetUBTTargetName(), FModuleManager::Get().GetUBTConfiguration(), FPlatformMisc::GetUBTPlatform(), *ExternalPluginArgument, *ProjectFileName);
+			const FString Arguments = FString::Printf(TEXT("%s %s %s -Plugin=\"%s\" -Project=\"%s\" %s-Progress -NoHotReloadFromIDE"), FPlatformMisc::GetUBTTargetName(), FModuleManager::Get().GetUBTConfiguration(), FPlatformMisc::GetUBTPlatform(), *PluginFilePath, *ProjectFileName, *LocalPluginArgument);
 
 			if (!FDesktopPlatformModule::Get()->RunUnrealBuildTool(FText::Format(LOCTEXT("CompilingPlugin", "Compiling {0} plugin..."), FText::FromString(PluginName)), FPaths::RootDir(), Arguments, GWarn))
 			{
diff --git a/Engine/Source/Programs/UnrealBuildTool/Configuration/TargetDescriptor.cs b/Engine/Source/Programs/UnrealBuildTool/Configuration/TargetDescriptor.cs
--- a/Engine/Source/Programs/UnrealBuildTool/Configuration/TargetDescriptor.cs
+++ b/Engine/Source/Programs/UnrealBuildTool/Configuration/TargetDescriptor.cs
@@ -42,6 +42,12 @@ namespace UnrealBuildTool
 		[CommandLine("-Plugin=")]
 		public FileReference? ForeignPlugin = null;
 
+		/// <summary>
+		/// Whether we should treat the ForeignPlugin argument as a local plugin for building purposes
+		/// </summary>
+		[CommandLine("-BuildPluginAsLocal")]
+		public bool bBuildPluginAsLocal = false;
+
 		/// <summary>
 		/// Set of module names to compile.
 		/// </summary>
diff --git a/Engine/Source/Programs/UnrealBuildTool/Configuration/UEBuildTarget.cs b/Engine/Source/Programs/UnrealBuildTool/Configuration/UEBuildTarget.cs
--- a/Engine/Source/Programs/UnrealBuildTool/Configuration/UEBuildTarget.cs
+++ b/Engine/Source/Programs/UnrealBuildTool/Configuration/UEBuildTarget.cs
@@ -1147,7 +1147,7 @@ namespace UnrealBuildTool
 			RulesAssembly RulesAssembly;
 			using (GlobalTracer.Instance.BuildSpan("RulesCompiler.CreateTargetRulesAssembly()").StartActive())
 			{
-				RulesAssembly = RulesCompiler.CreateTargetRulesAssembly(Descriptor.ProjectFile, Descriptor.Name, bSkipRulesCompile, bForceRulesCompile, bUsePrecompiled, Descriptor.ForeignPlugin, Logger);
+				RulesAssembly = RulesCompiler.CreateTargetRulesAssembly(Descriptor.ProjectFile, Descriptor.Name, bSkipRulesCompile, bForceRulesCompile, bUsePrecompiled, Descriptor.ForeignPlugin, Descriptor.bBuildPluginAsLocal, Logger);
 			}
 
 			TargetRules RulesObject;
diff --git a/Engine/Source/Programs/UnrealBuildTool/Modes/CleanMode.cs b/Engine/Source/Programs/UnrealBuildTool/Modes/CleanMode.cs
--- a/Engine/Source/Programs/UnrealBuildTool/Modes/CleanMode.cs
+++ b/Engine/Source/Programs/UnrealBuildTool/Modes/CleanMode.cs
@@ -90,7 +90,7 @@ namespace UnrealBuildTool
 					TargetDescriptor TargetDescriptor = TargetDescriptors[Idx];
 
 					// Create the rules assembly
-					RulesAssembly RulesAssembly = RulesCompiler.CreateTargetRulesAssembly(TargetDescriptor.ProjectFile, TargetDescriptor.Name, bSkipRulesCompile, BuildConfiguration.bForceRulesCompile, BuildConfiguration.bUsePrecompiled, TargetDescriptor.ForeignPlugin, Logger);
+					RulesAssembly RulesAssembly = RulesCompiler.CreateTargetRulesAssembly(TargetDescriptor.ProjectFile, TargetDescriptor.Name, bSkipRulesCompile, BuildConfiguration.bForceRulesCompile, BuildConfiguration.bUsePrecompiled, TargetDescriptor.ForeignPlugin, TargetDescriptor.bBuildPluginAsLocal, Logger);
 
 					// Create the rules object
 					ReadOnlyTargetRules Target = new ReadOnlyTargetRules(RulesAssembly.CreateTargetRules(TargetDescriptor.Name, TargetDescriptor.Platform, TargetDescriptor.Configuration, TargetDescriptor.Architectures, TargetDescriptor.ProjectFile, TargetDescriptor.AdditionalArguments, Logger));
diff --git a/Engine/Source/Programs/UnrealBuildTool/System/RulesCompiler.cs b/Engine/Source/Programs/UnrealBuildTool/System/RulesCompiler.cs
--- a/Engine/Source/Programs/UnrealBuildTool/System/RulesCompiler.cs
+++ b/Engine/Source/Programs/UnrealBuildTool/System/RulesCompiler.cs
@@ -330,10 +330,11 @@ namespace UnrealBuildTool
 		/// <param name="bSkipCompile">Whether to skip compilation for this assembly</param>
 		/// <param name="bForceCompile">Whether to always compile this assembly</param>
 		/// <param name="Parent">The parent rules assembly</param>
-        /// <param name="bContainsEngineModules">Whether the plugin contains engine modules. Used to initialize the default value for ModuleRules.bTreatAsEngineModule.</param>
+		/// <param name="bBuildPluginAsLocal">Whether the plugin should be built as though it is a local plugin.</param>
+		/// <param name="bContainsEngineModules">Whether the plugin contains engine modules. Used to initialize the default value for ModuleRules.bTreatAsEngineModule.</param>
 		/// <param name="Logger">Logger for ouptut</param>
 		/// <returns>The new rules assembly</returns>
-        public static RulesAssembly CreatePluginRulesAssembly(FileReference PluginFileName, bool bSkipCompile, bool bForceCompile, RulesAssembly Parent, bool bContainsEngineModules, ILogger Logger)
+		public static RulesAssembly CreatePluginRulesAssembly(FileReference PluginFileName, bool bSkipCompile, bool bForceCompile, RulesAssembly Parent, bool bBuildPluginAsLocal, bool bContainsEngineModules, ILogger Logger)
 		{
 			// Check if there's an existing assembly for this project
 			RulesAssembly? PluginRulesAssembly;
@@ -366,6 +367,13 @@ namespace UnrealBuildTool
 				// Find all the modules
 				ModuleRulesContext PluginModuleContext = new ModuleRulesContext(Scope, PluginFileName.Directory);
 				PluginModuleContext.bClassifyAsGameModuleForUHT = !bContainsEngineModules;
+				if(bBuildPluginAsLocal)
+				{
+					PluginModuleContext.bCanBuildDebugGame = true;
+					PluginModuleContext.bCanHotReload = true;
+					PluginModuleContext.bClassifyAsGameModuleForUHT = true;
+					PluginModuleContext.bCanUseForSharedPCH = false;
+				}
 				FindModuleRulesForPlugins(ForeignPlugins, PluginModuleContext, ModuleFiles);
 
 				Logger.LogTrace(" Found {Count} Modules:", ModuleFiles.Count);
@@ -391,9 +399,10 @@ namespace UnrealBuildTool
 		/// <param name="bForceRulesCompile">Whether to always compile all rules assemblies</param>
 		/// <param name="bUsePrecompiled">Whether to use a precompiled engine build</param>
 		/// <param name="ForeignPlugin">Foreign plugin to be compiled</param>
+		/// <param name="bBuildPluginAsLocal">Whether the plugin should be built as though it is a local plugin</param>
 		/// <param name="Logger">Logger for output</param>
 		/// <returns>The compiled rules assembly</returns>
-		public static RulesAssembly CreateTargetRulesAssembly(FileReference? ProjectFile, string TargetName, bool bSkipRulesCompile, bool bForceRulesCompile, bool bUsePrecompiled, FileReference? ForeignPlugin, ILogger Logger)
+		public static RulesAssembly CreateTargetRulesAssembly(FileReference? ProjectFile, string TargetName, bool bSkipRulesCompile, bool bForceRulesCompile, bool bUsePrecompiled, FileReference? ForeignPlugin, bool bBuildPluginAsLocal, ILogger Logger)
 		{
 			RulesAssembly RulesAssembly;
 			if (ProjectFile != null)
@@ -406,7 +415,7 @@ namespace UnrealBuildTool
 			}
 			if (ForeignPlugin != null)
 			{
-				RulesAssembly = CreatePluginRulesAssembly(ForeignPlugin, bSkipRulesCompile, bForceRulesCompile, RulesAssembly, true, Logger);
+				RulesAssembly = CreatePluginRulesAssembly(ForeignPlugin, bSkipRulesCompile, bForceRulesCompile, RulesAssembly, bBuildPluginAsLocal, true, Logger);
 			}
 			return RulesAssembly;
 		}
diff --git a/Engine/Source/Programs/UnrealBuildTool/ToolChain/RemoteMac.cs b/Engine/Source/Programs/UnrealBuildTool/ToolChain/RemoteMac.cs
--- a/Engine/Source/Programs/UnrealBuildTool/ToolChain/RemoteMac.cs
+++ b/Engine/Source/Programs/UnrealBuildTool/ToolChain/RemoteMac.cs
@@ -430,7 +430,7 @@ namespace UnrealBuildTool
 		public bool Build(TargetDescriptor TargetDesc, FileReference RemoteLogFile, bool bSkipPreBuildTargets, ILogger Logger)
 		{
 			// Compile the rules assembly
-			RulesAssembly RulesAssembly = RulesCompiler.CreateTargetRulesAssembly(TargetDesc.ProjectFile, TargetDesc.Name, false, false, false, TargetDesc.ForeignPlugin, Logger);
+			RulesAssembly RulesAssembly = RulesCompiler.CreateTargetRulesAssembly(TargetDesc.ProjectFile, TargetDesc.Name, false, false, false, TargetDesc.ForeignPlugin, TargetDesc.bBuildPluginAsLocal, Logger);
 
 			// Create the target rules
 			TargetRules Rules = RulesAssembly.CreateTargetRules(TargetDesc.Name, TargetDesc.Platform, TargetDesc.Configuration, TargetDesc.Architectures, TargetDesc.ProjectFile, TargetDesc.AdditionalArguments, Logger);
