- Removing deprecated IOS bitcode support

diff --git a/Engine/Config/BaseEngine.ini b/Engine/Config/BaseEngine.ini
--- a/Engine/Config/BaseEngine.ini
+++ b/Engine/Config/BaseEngine.ini
@@ -2784,7 +2784,6 @@ bSupportsLandscapeRightOrientation=True
 PreferredLandscapeOrientation=LandscapeLeft
 bSupportsMetal=True
 bSupportsMetalMRT=False
-bShipForBitcode=True
 bGameSupportsMultipleActiveControllers=False
 bAllowRemoteRotation=True
 bUseAbsoluteDpadValues=False
diff --git a/Engine/Source/Developer/IOS/IOSTargetPlatform/Private/IOSTargetPlatform.h b/Engine/Source/Developer/IOS/IOSTargetPlatform/Private/IOSTargetPlatform.h
--- a/Engine/Source/Developer/IOS/IOSTargetPlatform/Private/IOSTargetPlatform.h
+++ b/Engine/Source/Developer/IOS/IOSTargetPlatform/Private/IOSTargetPlatform.h
@@ -96,7 +96,6 @@ public:
 		InBoolKeys.Add(TEXT("bGeneratedSYMFile"));
 		InBoolKeys.Add(TEXT("bGeneratedSYMBundle"));
 		InBoolKeys.Add(TEXT("bGenerateXCArchive"));
-		InBoolKeys.Add(TEXT("bShipForBitcode"));
 		if (bIsTVOS)
 		{
 			InStringKeys.Add(TEXT("MinimumTVOSVersion"));
diff --git a/Engine/Source/Programs/AutomationTool/IOS/IOSPlatform.Automation.cs b/Engine/Source/Programs/AutomationTool/IOS/IOSPlatform.Automation.cs
--- a/Engine/Source/Programs/AutomationTool/IOS/IOSPlatform.Automation.cs
+++ b/Engine/Source/Programs/AutomationTool/IOS/IOSPlatform.Automation.cs
@@ -1593,20 +1593,6 @@ public class IOSPlatform : ApplePlatform
 				CopyFile_NoExceptions(dSYMSrcPath, dSYMDstPath);
 			}
 
-			// copy in the bitcode symbol maps if found
-			string[] bcmapfiles = Directory.GetFiles(SC.ProjectBinariesFolder.FullName, "*.bcsymbolmap");
-			if (bcmapfiles.Length > 0)
-			{
-				// Create the dsyms archive folder
-				CreateDirectory(Path.Combine(ArchiveName, "BCSymbolMaps"));
-				foreach (string symbolSrcFilePath in bcmapfiles)
-				{
-					string symbolLeafFileName = Path.GetFileName(symbolSrcFilePath);
-					string bcDstFilePath = Path.Combine(ArchiveName, "BCSymbolMaps", symbolLeafFileName);
-					CopyFile_NoExceptions(symbolSrcFilePath, bcDstFilePath);
-				}
-			}
-
 			// get the settings from the app plist file
 			string AppPlist = Path.Combine(ArchiveName, "Products", "Applications", AppName, "Info.plist");
 			string OldPListData = File.Exists(AppPlist) ? File.ReadAllText(AppPlist) : "";
diff --git a/Engine/Source/Programs/UnrealBuildTool/Platform/IOS/IOSToolChain.cs b/Engine/Source/Programs/UnrealBuildTool/Platform/IOS/IOSToolChain.cs
--- a/Engine/Source/Programs/UnrealBuildTool/Platform/IOS/IOSToolChain.cs
+++ b/Engine/Source/Programs/UnrealBuildTool/Platform/IOS/IOSToolChain.cs
@@ -269,11 +269,6 @@ namespace UnrealBuildTool
 		{
 			base.GetCompileArguments_Global(CompileEnvironment, Arguments);
 
-			if (IsBitcodeCompilingEnabled(CompileEnvironment.Configuration))
-			{
-				Arguments.Add("-fembed-bitcode");
-			}
-
 			// What architecture(s) to build for
 			Arguments.Add(FormatArchitectureArg(CompileEnvironment.Architectures));
 
@@ -346,11 +341,6 @@ namespace UnrealBuildTool
 			RunExecutableAndWait("mkdir", String.Format("-p \"{0}\"", Path), out ResultsText);
 		}
 
-		bool IsBitcodeCompilingEnabled(CppConfiguration Configuration)
-		{
-			return Configuration == CppConfiguration.Shipping && ProjectSettings.bShipForBitcode;
-		}
-
 		public virtual string GetXcodeMinVersionParam()
 		{
 			return "iphoneos-version-min";
@@ -376,20 +366,9 @@ namespace UnrealBuildTool
 			Arguments.Add(String.Format(" -isysroot \\\"{0}Platforms/{1}.platform/Developer/SDKs/{1}{2}.sdk\\\"",
 				Settings.Value.XcodeDeveloperDir, bIsDevice ? Settings.Value.DevicePlatformName : Settings.Value.SimulatorPlatformName, Settings.Value.IOSSDKVersion));
 
-			if (IsBitcodeCompilingEnabled(LinkEnvironment.Configuration))
-			{
-				FileItem OutputFile = FileItem.GetItemByFileReference(LinkEnvironment.OutputFilePath);
-
-				Arguments.Add("-fembed-bitcode -Xlinker -bitcode_verify -Xlinker -bitcode_hide_symbols -Xlinker -bitcode_symbol_map ");
-				Arguments.Add("-Xlinker \\\"" + Path.GetDirectoryName(OutputFile.AbsolutePath) + "\\\"");
-			}
-
 			Arguments.Add("-dead_strip");
 			Arguments.Add("-m" + GetXcodeMinVersionParam() + "=" + ProjectSettings.RuntimeVersion);
-			if (!IsBitcodeCompilingEnabled(LinkEnvironment.Configuration))
-			{
-				Arguments.Add("-Wl-no_pie");
-			}
+			Arguments.Add("-Wl-no_pie");
 			Arguments.Add("-stdlib=libc++");
 			Arguments.Add("-ObjC");
 			// Arguments.Add("-v");
@@ -643,7 +622,6 @@ namespace UnrealBuildTool
 				string linkCommandArguments = "-c \"";
 
 				linkCommandArguments += string.Format("rm -f \\\"{0}\\\";", OutputFile.AbsolutePath);
-				linkCommandArguments += string.Format("rm -f \\\"{0}\\*.bcsymbolmap\\\";", Path.GetDirectoryName(OutputFile.AbsolutePath));
 				linkCommandArguments += "\\\"" + LinkerPath + "\\\" " + LinkCommandArguments + ";";
 
 				linkCommandArguments += "\"";
@@ -1195,50 +1173,52 @@ namespace UnrealBuildTool
 				// update version xcconfig file so xcode can update the .plist with our incrementing build version
 				OutputFiles.Add(UpdateVersionFile(BinaryLinkEnvironment, Executable, Graph));
 			}
-
-			// strip the debug info from the executable if needed. creates a dummy output file for the action graph to track that it's out of date.
-			if (Target.IOSPlatform.bStripSymbols || (Target.Configuration == UnrealTargetConfiguration.Shipping))
-			{
-				FileItem StripCompleteFile = FileItem.GetItemByFileReference(FileReference.Combine(BinaryLinkEnvironment.IntermediateDirectory!, Executable.Location.GetFileName() + ".stripped"));
-
-				// If building a framework we can only strip local symbols, need to leave global in place
-				string StripArguments = BinaryLinkEnvironment.bIsBuildingDLL ? "-x" : "";
-
-				Action StripAction = Graph.CreateAction(ActionType.CreateAppBundle);
-				StripAction.WorkingDirectory = GetMacDevSrcRoot();
-				StripAction.CommandPath = BuildHostPlatform.Current.Shell;
-				StripAction.CommandArguments = String.Format("-c \"\\\"{0}strip\\\" {1} \\\"{2}\\\" && touch \\\"{3}\\\"\"", Settings.Value.ToolchainDir, StripArguments, Executable.Location, StripCompleteFile);
-				StripAction.PrerequisiteItems.Add(Executable);
-				StripAction.PrerequisiteItems.UnionWith(OutputFiles);
-				StripAction.ProducedItems.Add(StripCompleteFile);
-				StripAction.StatusDescription = String.Format("Stripping symbols from {0}", Executable.AbsolutePath);
-				StripAction.bCanExecuteRemotely = false;
-
-				OutputFiles.Add(StripCompleteFile);
-			}
-
-			if (!BinaryLinkEnvironment.bIsBuildingDLL && !bUseModernXcode)
+			else
 			{
-				// generate the asset catalog
-				bool bUserImagesExist = false;
-				DirectoryReference ResourcesDir = GenerateAssetCatalog(ProjectFile, BinaryLinkEnvironment.Platform, ref bUserImagesExist);
-
-				// Get the output location for the asset catalog
-				FileItem AssetCatalogFile = FileItem.GetItemByFileReference(GetAssetCatalogFile(BinaryLinkEnvironment.Platform, Executable.Location));
-
-				// Make the compile action
-				Action CompileAssetAction = Graph.CreateAction(ActionType.CreateAppBundle);
-				CompileAssetAction.WorkingDirectory = GetMacDevSrcRoot();
-				CompileAssetAction.CommandPath = new FileReference("/usr/bin/xcrun");
-				CompileAssetAction.CommandArguments = GetAssetCatalogArgs(BinaryLinkEnvironment.Platform, ResourcesDir.FullName, Path.GetDirectoryName(AssetCatalogFile.AbsolutePath)!);
-				CompileAssetAction.PrerequisiteItems.Add(Executable);
-				CompileAssetAction.ProducedItems.Add(AssetCatalogFile);
-				CompileAssetAction.DeleteItems.Add(AssetCatalogFile);
-				CompileAssetAction.StatusDescription = CompileAssetAction.CommandArguments;// string.Format("Generating debug info for {0}", Path.GetFileName(Executable.AbsolutePath));
-				CompileAssetAction.bCanExecuteRemotely = false;
+				// strip the debug info from the executable if needed. creates a dummy output file for the action graph to track that it's out of date.
+				if (Target.IOSPlatform.bStripSymbols || (Target.Configuration == UnrealTargetConfiguration.Shipping))
+				{
+					FileItem StripCompleteFile = FileItem.GetItemByFileReference(FileReference.Combine(BinaryLinkEnvironment.IntermediateDirectory!, Executable.Location.GetFileName() + ".stripped"));
+
+					// If building a framework we can only strip local symbols, need to leave global in place
+					string StripArguments = BinaryLinkEnvironment.bIsBuildingDLL ? "-x" : "";
+
+					Action StripAction = Graph.CreateAction(ActionType.CreateAppBundle);
+					StripAction.WorkingDirectory = GetMacDevSrcRoot();
+					StripAction.CommandPath = BuildHostPlatform.Current.Shell;
+					StripAction.CommandArguments = String.Format("-c \"\\\"{0}strip\\\" {1} \\\"{2}\\\" && touch \\\"{3}\\\"\"", Settings.Value.ToolchainDir, StripArguments, Executable.Location, StripCompleteFile);
+					StripAction.PrerequisiteItems.Add(Executable);
+					StripAction.PrerequisiteItems.UnionWith(OutputFiles);
+					StripAction.ProducedItems.Add(StripCompleteFile);
+					StripAction.StatusDescription = String.Format("Stripping symbols from {0}", Executable.AbsolutePath);
+					StripAction.bCanExecuteRemotely = false;
+
+					OutputFiles.Add(StripCompleteFile);
+				}
 
-				// Add it to the output files so it's always built
-				OutputFiles.Add(AssetCatalogFile);
+				if (!BinaryLinkEnvironment.bIsBuildingDLL)
+				{
+					// generate the asset catalog
+					bool bUserImagesExist = false;
+					DirectoryReference ResourcesDir = GenerateAssetCatalog(ProjectFile, BinaryLinkEnvironment.Platform, ref bUserImagesExist);
+
+					// Get the output location for the asset catalog
+					FileItem AssetCatalogFile = FileItem.GetItemByFileReference(GetAssetCatalogFile(BinaryLinkEnvironment.Platform, Executable.Location));
+
+					// Make the compile action
+					Action CompileAssetAction = Graph.CreateAction(ActionType.CreateAppBundle);
+					CompileAssetAction.WorkingDirectory = GetMacDevSrcRoot();
+					CompileAssetAction.CommandPath = new FileReference("/usr/bin/xcrun");
+					CompileAssetAction.CommandArguments = GetAssetCatalogArgs(BinaryLinkEnvironment.Platform, ResourcesDir.FullName, Path.GetDirectoryName(AssetCatalogFile.AbsolutePath)!);
+					CompileAssetAction.PrerequisiteItems.Add(Executable);
+					CompileAssetAction.ProducedItems.Add(AssetCatalogFile);
+					CompileAssetAction.DeleteItems.Add(AssetCatalogFile);
+					CompileAssetAction.StatusDescription = CompileAssetAction.CommandArguments;// string.Format("Generating debug info for {0}", Path.GetFileName(Executable.AbsolutePath));
+					CompileAssetAction.bCanExecuteRemotely = false;
+
+					// Add it to the output files so it's always built
+					OutputFiles.Add(AssetCatalogFile);
+				}
 			}
 
 			// Generate the app bundle
diff --git a/Engine/Source/Programs/UnrealBuildTool/Platform/IOS/UEBuildIOS.cs b/Engine/Source/Programs/UnrealBuildTool/Platform/IOS/UEBuildIOS.cs
--- a/Engine/Source/Programs/UnrealBuildTool/Platform/IOS/UEBuildIOS.cs
+++ b/Engine/Source/Programs/UnrealBuildTool/Platform/IOS/UEBuildIOS.cs
@@ -176,11 +176,6 @@ namespace UnrealBuildTool
 			get { return Inner.bStripSymbols; }
 		}
 			
-		public bool bShipForBitcode
-		{
-			get { return Inner.ProjectSettings!.bShipForBitcode; }
-		}
-
 		public bool bGenerateFrameworkWrapperProject
 		{
 			get { return Inner.bGenerateFrameworkWrapperProject; }
@@ -335,13 +330,6 @@ namespace UnrealBuildTool
 		[ConfigFile(ConfigHierarchyType.Engine, "/Script/IOSRuntimeSettings.IOSRuntimeSettings", "SigningCertificate")]
         public readonly string SigningCertificate = "";
 
-
-		/// <summary>
-		/// true if bit code should be embedded
-		/// </summary>
-		[ConfigFile(ConfigHierarchyType.Engine, "/Script/IOSRuntimeSettings.IOSRuntimeSettings", "bShipForBitcode")]
-		public readonly bool bShipForBitcode = false;
-
         /// <summary>
         /// true if notifications are enabled
         /// </summary>
diff --git a/Engine/Source/Runtime/IOS/IOSRuntimeSettings/Classes/IOSRuntimeSettings.h b/Engine/Source/Runtime/IOS/IOSRuntimeSettings/Classes/IOSRuntimeSettings.h
--- a/Engine/Source/Runtime/IOS/IOSRuntimeSettings/Classes/IOSRuntimeSettings.h
+++ b/Engine/Source/Runtime/IOS/IOSRuntimeSettings/Classes/IOSRuntimeSettings.h
@@ -232,10 +232,6 @@ public:
 	UPROPERTY(GlobalConfig, EditAnywhere, Category = Build, meta = (DisplayName = "Generate xcode archive package"))
 	bool bGenerateXCArchive;	
 	
-	// Enable bitcode compiling?
-	UPROPERTY(GlobalConfig, EditAnywhere, Category = Build, meta = (DisplayName = "Support bitcode in Shipping"))
-	bool bShipForBitcode;
-
 	// Enable Advertising Identified
 	UPROPERTY(GlobalConfig, EditAnywhere, Category = Build, meta = (DisplayName = "Enable Advertising Identified (IDFA)"))
 	bool bEnableAdvertisingIdentifier;
diff --git a/Engine/Source/Runtime/IOS/IOSRuntimeSettings/Private/IOSRuntimeSettings.cpp b/Engine/Source/Runtime/IOS/IOSRuntimeSettings/Private/IOSRuntimeSettings.cpp
--- a/Engine/Source/Runtime/IOS/IOSRuntimeSettings/Private/IOSRuntimeSettings.cpp
+++ b/Engine/Source/Runtime/IOS/IOSRuntimeSettings/Private/IOSRuntimeSettings.cpp
@@ -36,7 +36,6 @@ UIOSRuntimeSettings::UIOSRuntimeSettings(const FObjectInitializer& ObjectInitial
 	bGeneratedSYMFile = false;
 	bGeneratedSYMBundle = false;
 	bGenerateXCArchive = false;
-	bShipForBitcode = true;
 	bSupportSecondaryMac = false;
 	bUseRSync = true;
 	bCustomLaunchscreenStoryboard = false;
