Modern Xcode changes
- Removed Payload from IOS Binaries subdir (this is needed for archiving envvar changes below)
- IOS .app's now have decorated names (QAGame-IOS-Debug.app)
- Archiving works now, with dSYMs, for all platforms, and can be uploaded directly to app store in xcode
- Updated some envvars, main thing is no longer overriding BUILD_CONFIGURATION_DIR because it breaks Archiving
- No longer force dsym generation in UBT when building shipping (in modern) because xcode will do it now when archiving (will make iterating on Shipping much faster now!)

diff --git a/Engine/Source/Programs/UnrealBuildTool/Platform/IOS/UEBuildIOS.cs b/Engine/Source/Programs/UnrealBuildTool/Platform/IOS/UEBuildIOS.cs
--- a/Engine/Source/Programs/UnrealBuildTool/Platform/IOS/UEBuildIOS.cs
+++ b/Engine/Source/Programs/UnrealBuildTool/Platform/IOS/UEBuildIOS.cs
@@ -255,6 +255,12 @@ namespace UnrealBuildTool
 		/// </summary>
 		public readonly FileReference? ProjectFile;
 
+		/// <summary>
+		/// Are we using modern xcode in this project?
+		/// </summary>
+		[ConfigFile(ConfigHierarchyType.Engine, "XcodeConfiguration", "bUseModernXcode")]
+		public bool bUseModernXcode = false;
+
 		/// <summary>
 		/// Whether to build the iOS project as a framework.
 		/// </summary>
@@ -773,17 +779,20 @@ namespace UnrealBuildTool
 			Target.bDeployAfterCompile = true;
 
 			Target.IOSPlatform.ProjectSettings = ((IOSPlatform)GetBuildPlatform(Target.Platform)).ReadProjectSettings(Target.ProjectFile);
-			
-			// always strip in shipping configuration (commandline could have set it also)
-			if (Target.Configuration == UnrealTargetConfiguration.Shipping)
-			{
-				Target.IOSPlatform.bStripSymbols = true;	
-			}
-			
-			// if we are stripping the executable, or if the project requested it, or if it's a buildmachine, generate the dsym
-			if (Target.IOSPlatform.bStripSymbols || Target.IOSPlatform.ProjectSettings.bGeneratedSYMFile || Environment.GetEnvironmentVariable("IsBuildMachine") == "1")
+
+			if (!Target.IOSPlatform.ProjectSettings.bUseModernXcode)
 			{
-				Target.IOSPlatform.bGeneratedSYM = true;
+				// always strip in shipping configuration (commandline could have set it also)
+				if (Target.Configuration == UnrealTargetConfiguration.Shipping)
+				{
+					Target.IOSPlatform.bStripSymbols = true;
+				}
+
+				// if we are stripping the executable, or if the project requested it, or if it's a buildmachine, generate the dsym
+				if (Target.IOSPlatform.bStripSymbols || Target.IOSPlatform.ProjectSettings.bGeneratedSYMFile || Environment.GetEnvironmentVariable("IsBuildMachine") == "1")
+				{
+					Target.IOSPlatform.bGeneratedSYM = true;
+				}
 			}
 
 			// Set bShouldCompileAsDLL when building as a framework
diff --git a/Engine/Source/Programs/UnrealBuildTool/ProjectFiles/Xcode/XcodeProject.Scheme.cs b/Engine/Source/Programs/UnrealBuildTool/ProjectFiles/Xcode/XcodeProject.Scheme.cs
--- a/Engine/Source/Programs/UnrealBuildTool/ProjectFiles/Xcode/XcodeProject.Scheme.cs
+++ b/Engine/Source/Programs/UnrealBuildTool/ProjectFiles/Xcode/XcodeProject.Scheme.cs
@@ -217,7 +217,7 @@ namespace UnrealBuildTool.XcodeProjectXcconfig
 				Content.WriteLine("      buildConfiguration = \"" + DefaultConfiguration + "\">");
 				Content.WriteLine("   </AnalyzeAction>");
 				Content.WriteLine("   <ArchiveAction");
-				Content.WriteLine("      buildConfiguration = \"" + DefaultConfiguration + "\"");
+				Content.WriteLine("      buildConfiguration = \"Shipping\"");
 				Content.WriteLine("      revealArchiveInOrganizer = \"YES\">");
 				Content.WriteLine("   </ArchiveAction>");
 				Content.WriteLine("</Scheme>");
diff --git a/Engine/Source/Programs/UnrealBuildTool/ProjectFiles/Xcode/XcodeProject.Utils.cs b/Engine/Source/Programs/UnrealBuildTool/ProjectFiles/Xcode/XcodeProject.Utils.cs
--- a/Engine/Source/Programs/UnrealBuildTool/ProjectFiles/Xcode/XcodeProject.Utils.cs
+++ b/Engine/Source/Programs/UnrealBuildTool/ProjectFiles/Xcode/XcodeProject.Utils.cs
@@ -55,6 +55,12 @@ namespace UnrealBuildTool.XcodeProjectXcconfig
 
 		public static bool ShouldIncludeProjectInWorkspace(ProjectFile Proj, ILogger Logger)
 		{
+			// since IOS/TVOS don't have the UnrealEditor project as valid, force it so that we get the source code
+			// this is likely temporary until we can put source code into UnrealGame
+			if (Proj.ProjectFilePath.GetFileNameWithoutAnyExtensions() == "UnrealEditor")
+			{
+				return true;
+			}
 			foreach (Project ProjectTarget in Proj.ProjectTargets)
 			{
 				foreach (UnrealTargetPlatform Platform in XcodeProjectFileGenerator.XcodePlatforms)
diff --git a/Engine/Source/Programs/UnrealBuildTool/ProjectFiles/Xcode/XcodeProject.cs b/Engine/Source/Programs/UnrealBuildTool/ProjectFiles/Xcode/XcodeProject.cs
--- a/Engine/Source/Programs/UnrealBuildTool/ProjectFiles/Xcode/XcodeProject.cs
+++ b/Engine/Source/Programs/UnrealBuildTool/ProjectFiles/Xcode/XcodeProject.cs
@@ -850,7 +850,6 @@ namespace UnrealBuildTool.XcodeProjectXcconfig
 				Content.WriteLine(3, Line);
 			}
 
-			Content.WriteLine(3, "runOnlyForDeploymentPostprocessing = 0;");
 			Content.WriteLine(2, "};");
 			Content.WriteLine($"/* End {IsAType} section */");
 		}
@@ -940,7 +939,7 @@ namespace UnrealBuildTool.XcodeProjectXcconfig
 
 	class XcodeShellScriptBuildPhase : XcodeBuildPhase
 	{
-		public XcodeShellScriptBuildPhase(string Name, IEnumerable<string> ScriptLines, IEnumerable<string> Inputs, IEnumerable<string> Outputs)
+		public XcodeShellScriptBuildPhase(string Name, IEnumerable<string> ScriptLines, IEnumerable<string> Inputs, IEnumerable<string> Outputs, bool bInstallOnly=false)
 			: base(Name, "PBXShellScriptBuildPhase")
 		{
 			MiscItems.Add($"name = \"{Name}\";");
@@ -963,6 +962,10 @@ namespace UnrealBuildTool.XcodeProjectXcconfig
 			string Script = string.Join("\\n", ScriptLines);
 			MiscItems.Add($"shellPath = /bin/sh;");
 			MiscItems.Add($"shellScript = \"{Script}\";");
+			if (bInstallOnly)
+			{
+				MiscItems.Add("runOnlyForDeploymentPostprocessing = 1;");
+			}
 		}
 	}
 
@@ -1223,31 +1226,105 @@ namespace UnrealBuildTool.XcodeProjectXcconfig
 				References.Add(ResourcesBuildPhase);
 
 				ProcessFrameworks(ResourcesBuildPhase, Project, ProjectFile, Logger);
+				ProcessScripts(ResourcesBuildPhase, Project);
+				ProcessAssets(ResourcesBuildPhase);
+			}
 
-				if (Project.Platform == UnrealTargetPlatform.IOS || Project.Platform == UnrealTargetPlatform.TVOS)
-				{
-					// UBT no longer copies the executable into the .app directory in PostBuild, so we do it here
-					string[] Script = {
-						"cp \\\"${CONFIGURATION_BUILD_DIR}/../${UE_IOS_EXECUTABLE_NAME}\\\" \\\"${CONFIGURATION_BUILD_DIR}/${PRODUCT_NAME}.app/${PRODUCT_NAME}\\\"",
+			if (BuildTarget != null)
+			{
+				AddDependency(BuildTarget, Project);
+			}
+
+			CreateXcconfigFile(Project, Platform, $"{Name}");
+			// create per-config Xcconfig files
+			foreach (XcodeBuildConfig Config in BuildConfigList.BuildConfigs)
+			{
+				Config.CreateXcconfigFile(Project, Platform, $"{Name}_{Config.Info.DisplayName.Replace(" ", "")}");
+			}
+		}
+
+
+		/// <summary>
+		/// Write some scripts to do some fixup with how UBT links files. There is currently a difference between Mac and IOS/TVOS:
+		/// Mac:
+		///   - UBT will link directly to Foo.app/Contents/MacOS/Foo, for .app bundled apps
+		///   - During normal processing, that's all that is needed, so we don't need to copy anything
+		///   - However, during Archiving, the .app is created in a intermediate location, so then here we copy from Binaries/Mac to the intermediate location
+		///     - Trying to have UBT link directly to the intermeidate location causes various issues, so we copy it like IOS does
+		/// IOS/TVOS:
+		///   - IOS will link to Binaries/IOS/Foo
+		///   - During normal operation, here we copy from Binaries/IOS/Foo to Binaries/IOS/Foo.app/Foo
+		///     - Note that IOS and Mac have different internal directory structures (which EXECUTABLE_PATH expresses)
+		///   - When Archiving, we copy from Binaries/IOS/Foo to the intermediate location's .app
+		/// All:
+		///   - At this point, the executable is in the correct spot, and so CONFIGURATION_BUILD_DIR/EXECUTABLE_PATH points to it
+		///   - So here we gneerate a dSYM from the executable, copying it to where Xcode wants it (DWARF_DSYM_FOLDER_PATH/DWARF_DSYM_FILE_NAME), and
+		///       then we strip the executable in place
+		/// </summary>
+		/// <param name="ResourcesBuildPhase"></param>
+		/// <param name="Project"></param>
+		protected void ProcessScripts(XcodeResourcesBuildPhase ResourcesBuildPhase, XcodeProject Project)
+		{
+			// @todo remove this if UBT links directly into the .app
+			if (Project.Platform == UnrealTargetPlatform.IOS || Project.Platform == UnrealTargetPlatform.TVOS)
+			{
+				// UBT no longer copies the executable into the .app directory in PostBuild, so we do it here
+				// EXECUTABLE_NAME is Foo, EXECUTABLE_PATH is Foo.app/Foo
+				// NOTE: We read from hardcoded location where UBT writes to, but we write to CONFIGURATION_BUILD_DIR because
+				// when Archiving, the .app is somewhere else
+				string[] Script = {
+						$"cp \\\"${{UE_PROJECT_DIR}}/Binaries/{Platform}/${{EXECUTABLE_NAME}}\\\" \\\"${{CONFIGURATION_BUILD_DIR}}/${{EXECUTABLE_PATH}}\\\"",
 					};
-					string ScriptInput = "$(CONFIGURATION_BUILD_DIR)/../$(UE_IOS_EXECUTABLE_NAME)";
-					string ScriptOutput = "$(CONFIGURATION_BUILD_DIR)/$(PRODUCT_NAME).app/$(PRODUCT_NAME)";
-					XcodeShellScriptBuildPhase ScriptPhase = new("Copy UE Executable into .app", Script, new string[] { ScriptInput }, new string[] { ScriptOutput });
-					BuildPhases.Add(ScriptPhase);
-					References.Add(ScriptPhase);
-
-					// cookeddata needs to be brought in from the Staged directory, (IOS for Game, IOSClient for Client, etc)
-					string TargetPlatformName = Project.Platform.ToString()!;
-					if (TargetType != TargetType.Game)
-					{
-						TargetPlatformName += TargetType.ToString();
-					}
+				string ScriptInput = $"$(UE_PROJECT_DIR)/Binaries/{Platform}/$(EXECUTABLE_NAME)";
+				string ScriptOutput = $"$(CONFIGURATION_BUILD_DIR)/$(EXECUTABLE_PATH)";
+				XcodeShellScriptBuildPhase ScriptPhase = new("Copy UE Executable into .app", Script, new string[] { ScriptInput }, new string[] { ScriptOutput });
+				BuildPhases.Add(ScriptPhase);
+				References.Add(ScriptPhase);
+			}
+
+			// always generate a dsym file when we archive, and by having Xcode do it, it will be put into the archive properly
+			// (note bInstallOnly which will make this onle run when archiving)
+			List<string> DsymScript = new();
+			if (Platform == UnrealTargetPlatform.Mac)
+			{
+				DsymScript.AddRange(new string[]
+				{
+					"# Copy the Mac executable from where it write to into the archiving working directory",
+					$"cp \\\"${{UE_PROJECT_DIR}}/Binaries/{Platform}/${{EXECUTABLE_PATH}}\\\" \\\"${{CONFIGURATION_BUILD_DIR}}/${{EXECUTABLE_PATH}}\\\"",
+					"",
+				});
+			}
 
-					DirectoryReference CookedData = DirectoryReference.Combine(UnrealData.ProductDirectory, "Saved", "StagedBuilds",TargetPlatformName, "cookeddata");
-					ResourcesBuildPhase.AddFolderResource(CookedData, "Resources");
+			DsymScript.AddRange(new string[]
+			{
+				"# Generate the dsym when making an archive, then stripping the executable",
+				"rm -rf \\\"${DWARF_DSYM_FOLDER_PATH}/${DWARF_DSYM_FILE_NAME}\\\"",
+				"dsymutil \\\"${CONFIGURATION_BUILD_DIR}/${EXECUTABLE_PATH}\\\" -o \\\"${DWARF_DSYM_FOLDER_PATH}/${DWARF_DSYM_FILE_NAME}\\\"",
+				"strip -D \\\"${CONFIGURATION_BUILD_DIR}/${EXECUTABLE_PATH}\\\"",
+			});
+			string DsymScriptInput = $"\\\"$(CONFIGURATION_BUILD_DIR)/$(EXECUTABLE_PATH)\\\"";
+			string DsymScriptOutput = $"\\\"$(DWARF_DSYM_FOLDER_PATH)/$(DWARF_DSYM_FILE_NAME)\\\"";
+			XcodeShellScriptBuildPhase DsymScriptPhase = new("Generate dsym for archive, and strip", DsymScript, new string[] { DsymScriptInput }, new string[] { DsymScriptOutput }, bInstallOnly:true);
+			BuildPhases.Add(DsymScriptPhase);
+			References.Add(DsymScriptPhase);
+		}
+		private void ProcessAssets(XcodeResourcesBuildPhase ResourcesBuildPhase)
+		{
+			if (Platform == UnrealTargetPlatform.IOS || Platform == UnrealTargetPlatform.TVOS)
+			{
+				// cookeddata needs to be brought in from the Staged directory, (IOS for Game, IOSClient for Client, etc)
+				string TargetPlatformName = Platform.ToString()!;
+				if (TargetType != TargetType.Game)
+				{
+					TargetPlatformName += TargetType.ToString();
 				}
 
-				List<string> StoryboardPaths = new List<string>()
+				// @todo do this for Mac?
+				DirectoryReference CookedData = DirectoryReference.Combine(UnrealData.ProductDirectory, "Saved", "StagedBuilds", TargetPlatformName, "cookeddata");
+				ResourcesBuildPhase.AddFolderResource(CookedData, "Resources");
+			}
+
+			List<string> StoryboardPaths = new List<string>()
 				{
 					"$(Project)/Build/$(Platform)/Resources/Interface/LaunchScreen.storyboardc",
 					"$(Project)/Build/$(Platform)/Resources/Interface/LaunchScreen.storyboard",
@@ -1259,40 +1336,27 @@ namespace UnrealBuildTool.XcodeProjectXcconfig
 					"$(Engine)/Build/Apple/Resources/Interface/LaunchScreen.storyboard",
 				};
 
-				// look for Assets
-				string? StoryboardPath;
-				if (Platform == UnrealTargetPlatform.Mac)
-				{
-					string AssetsPath = UnrealData.ProjectOrEnginePath("Build/Mac/Resources/Assets.xcassets", false);
-					ResourcesBuildPhase.AddResource(new FileReference(AssetsPath));
-					StoryboardPath = UnrealData.FindFile(StoryboardPaths, UnrealTargetPlatform.Mac, false);
-				}
-				else
-				{
-					string AssetsPath = UnrealData.ProjectOrEnginePath("Build/IOS/Resources/Assets.xcassets", false);
-					ResourcesBuildPhase.AddResource(new FileReference(AssetsPath));
-					StoryboardPath = UnrealData.FindFile(StoryboardPaths, UnrealTargetPlatform.IOS, false);
-				}
-
-				if (StoryboardPath != null)
-				{
-					ResourcesBuildPhase.AddResource(new FileReference(StoryboardPath));
-				}
+			// look for Assets
+			string? StoryboardPath;
+			if (Platform == UnrealTargetPlatform.Mac)
+			{
+				string AssetsPath = UnrealData.ProjectOrEnginePath("Build/Mac/Resources/Assets.xcassets", false);
+				ResourcesBuildPhase.AddResource(new FileReference(AssetsPath));
+				StoryboardPath = UnrealData.FindFile(StoryboardPaths, UnrealTargetPlatform.Mac, false);
 			}
-
-			if (BuildTarget != null)
+			else
 			{
-				AddDependency(BuildTarget, Project);
+				string AssetsPath = UnrealData.ProjectOrEnginePath("Build/IOS/Resources/Assets.xcassets", false);
+				ResourcesBuildPhase.AddResource(new FileReference(AssetsPath));
+				StoryboardPath = UnrealData.FindFile(StoryboardPaths, UnrealTargetPlatform.IOS, false);
 			}
 
-			CreateXcconfigFile(Project, Platform, $"{Name}");
-			// create per-config Xcconfig files
-			foreach (XcodeBuildConfig Config in BuildConfigList.BuildConfigs)
+			if (StoryboardPath != null)
 			{
-				Config.CreateXcconfigFile(Project, Platform, $"{Name}_{Config.Info.DisplayName.Replace(" ", "")}");
+				ResourcesBuildPhase.AddResource(new FileReference(StoryboardPath));
 			}
-		}
 
+		}
 		protected void ProcessFrameworks(XcodeResourcesBuildPhase ResourcesBuildPhase, XcodeProject Project, XcodeProjectFile ProjectFile, ILogger Logger)
 		{
 			// look up to see if we had cached any Frameworks
@@ -1470,7 +1534,7 @@ namespace UnrealBuildTool.XcodeProjectXcconfig
 					SDKRoot = "iphoneos";
 					SupportedPlatforms = "iphoneos"; // iphonesimulator
 					DeploymentTargetKey = "IPHONEOS_DEPLOYMENT_TARGET";
-					ConfigBuildDir = DirectoryReference.Combine(BuildConfig.IOSExecutablePath!.Directory, "Payload");
+					ConfigBuildDir = BuildConfig.IOSExecutablePath!.Directory;
 					SupportedDevices = UnrealData.IOSProjectSettings!.RuntimeDevices;
 					DeploymentTarget = UnrealData.IOSProjectSettings.RuntimeVersion;
 
@@ -1489,7 +1553,7 @@ namespace UnrealBuildTool.XcodeProjectXcconfig
 					SDKRoot = "appletvos";
 					SupportedPlatforms = "appletvos"; // appletvsimulator
 					DeploymentTargetKey = "TVOS_DEPLOYMENT_TARGET";
-					ConfigBuildDir = DirectoryReference.Combine(BuildConfig.TVOSExecutablePath!.Directory, "Payload");
+					ConfigBuildDir = BuildConfig.TVOSExecutablePath!.Directory;
 					SupportedDevices = UnrealData.TVOSProjectSettings!.RuntimeDevices;
 					DeploymentTarget = UnrealData.TVOSProjectSettings.RuntimeVersion;
 
@@ -1520,6 +1584,7 @@ namespace UnrealBuildTool.XcodeProjectXcconfig
 			Xcconfig.AppendLine($"UE_DISPLAY_NAME = {UnrealData.DisplayName}");
 			Xcconfig.AppendLine($"UE_SIGNING_PREFIX = {SigningPrefix}");
 			Xcconfig.AppendLine($"UE_ENGINE_DIR = {Unreal.EngineDirectory}");
+			Xcconfig.AppendLine($"UE_PROJECT_DIR = {ProjectOrEngineDir}");
 			Xcconfig.AppendLine($"UE_PLATFORM_NAME = {Platform}");
 			Xcconfig.AppendLine($"UE_TARGET_NAME = {UnrealData.TargetRules.Name}");
 			Xcconfig.AppendLine($"UE_TARGET_PLATFORM_NAME = {TargetPlatformName}");
@@ -1533,15 +1598,26 @@ namespace UnrealBuildTool.XcodeProjectXcconfig
 			Xcconfig.AppendLine("USE_HEADERMAP = NO");
 			Xcconfig.AppendLine("ONLY_ACTIVE_ARCH = YES");
 
-
 			Xcconfig.AppendLine("");
 			Xcconfig.AppendLine("// Platform settings");
 			Xcconfig.AppendLine($"SUPPORTED_PLATFORMS = {SupportedPlatforms}");
 			Xcconfig.AppendLine($"SDKROOT = {SDKRoot}");
-			Xcconfig.AppendLine($"CONFIGURATION_BUILD_DIR = {ConfigBuildDir}");
-			foreach (string Line in ExtraConfigLines)
+
+			// Xcode creates the Build Dir (where the .app is) by combining {SYMROOT}/{CONFIGURATION}{EFFECTIVE_PLATFORM_NAME}
+			Xcconfig.AppendLine("");
+			Xcconfig.AppendLine($"// These settings combined will tell Xcode to write to Binaries/{Platform} (instead of something like Binaries/Development-iphoneos)");
+			Xcconfig.AppendLine($"SYMROOT = {ConfigBuildDir.ParentDirectory}");
+			Xcconfig.AppendLine($"CONFIGURATION = ");
+			Xcconfig.AppendLine($"EFFECTIVE_PLATFORM_NAME = {Platform}");
+
+			if (ExtraConfigLines.Count > 0)
 			{
-				Xcconfig.AppendLine(Line);
+				Xcconfig.AppendLine("");
+				Xcconfig.AppendLine("// Misc settings");
+				foreach (string Line in ExtraConfigLines)
+				{
+					Xcconfig.AppendLine(Line);
+				}
 			}
 
 			Xcconfig.AppendLine("");
@@ -1612,8 +1688,7 @@ namespace UnrealBuildTool.XcodeProjectXcconfig
 				else // IOS,TVOS
 				{
 					ExecutableName = ((Platform == UnrealTargetPlatform.IOS) ? Config.IOSExecutablePath : Config.TVOSExecutablePath)!.GetFileName();
-					// @todo: change to Executable when we stop using payload!
-					ProductName = Config.BuildTarget;
+					ProductName = ExecutableName;
 				}
 
 
@@ -1690,7 +1765,7 @@ namespace UnrealBuildTool.XcodeProjectXcconfig
 			Xcconfig.AppendLine("GCC_PRECOMPILE_PREFIX_HEADER = YES");
 			Xcconfig.AppendLine("GCC_OPTIMIZATION_LEVEL = 0");
 			Xcconfig.AppendLine($"PRODUCT_NAME = {Name}");
-			Xcconfig.AppendLine("CONFIGURATION_BUILD_DIR = build");
+			Xcconfig.AppendLine("SYMROOT = build");
 			Xcconfig.Write();
 		}
 	}
