UnrealBuildTool: Update -Deterministic to not have a per module setting as this doesn't work due to PCH requirements and instead add a setting to control warning level

diff --git a/Engine/Source/Programs/UnrealBuildTool/Configuration/ModuleRules.cs b/Engine/Source/Programs/UnrealBuildTool/Configuration/ModuleRules.cs
--- a/Engine/Source/Programs/UnrealBuildTool/Configuration/ModuleRules.cs
+++ b/Engine/Source/Programs/UnrealBuildTool/Configuration/ModuleRules.cs
@@ -757,11 +757,6 @@ namespace UnrealBuildTool
 		}
 		private EngineIncludeOrderVersion? IncludeOrderVersionPrivate;
 
-		/// <summary>
-		/// Set flags for determinstic compiles (experimental, may not be fully supported). Deterministic linking is controlled via TargetRules.
-		/// </summary>
-		public bool bDeterministic = false;
-
 		/// <summary>
 		/// Use run time type information
 		/// </summary>
@@ -794,6 +789,18 @@ namespace UnrealBuildTool
 		/// Not doing this will result in a compile errors because shared PCHs were compiled with different flags than consumer
 		/// </summary>
 		public bool bEnableObjCAutomaticReferenceCounting = false;
+
+		/// <summary>
+		/// How to treat deterministic warnings (experimental).
+		/// </summary>
+		public WarningLevel DeterministicWarningLevel
+		{
+			get { return (DeterministicWarningLevelPrivate == WarningLevel.Default) ? (Target.bDeterministic ? WarningLevel.Warning : WarningLevel.Off) : DeterministicWarningLevelPrivate; }
+			set { DeterministicWarningLevelPrivate = value; }
+		}
+
+		/// <inheritdoc cref="DeterministicWarningLevelPrivate"/>
+		private WarningLevel DeterministicWarningLevelPrivate;
 		
 		/// <summary>
 		/// How to treat shadow variable warnings
diff --git a/Engine/Source/Programs/UnrealBuildTool/Configuration/TargetRules.cs b/Engine/Source/Programs/UnrealBuildTool/Configuration/TargetRules.cs
--- a/Engine/Source/Programs/UnrealBuildTool/Configuration/TargetRules.cs
+++ b/Engine/Source/Programs/UnrealBuildTool/Configuration/TargetRules.cs
@@ -1659,21 +1659,12 @@ namespace UnrealBuildTool
 		public bool bUsePCHFiles = true;
 
 		/// <summary>
-		/// Set flags require for deterministic linking (experimental, may not be fully supported).
-		/// Deterministic compiling is controlled via ModuleRules.
-		/// </summary>
-		[XmlConfigFile(Category = "BuildConfiguration")]
-		public bool bDeterministic = false;
-
-		/// <summary>
-		/// Force set flags require for deterministic compiling and linking (experimental, may not be fully supported).
-		/// This setting is only recommended for testing, instead:
-		/// * Set bDeterministic on a per module basis in ModuleRules to control deterministic compiling.
-		/// * Set bDeterministic on a per target basis in TargetRules to control deterministic linking.
+		/// Set flags require for deterministic compiling and linking (experimental).
 		/// </summary>
 		[CommandLine("-Deterministic")]
 		[XmlConfigFile(Category = "BuildConfiguration")]
-		public bool bForceDeterministic = false;
+		[RequiresUniqueBuildEnvironment]
+		public bool bDeterministic = false;
 
 		/// <summary>
 		/// Whether PCHs should be chained when compiling with clang.
@@ -3500,12 +3491,6 @@ namespace UnrealBuildTool
 		{
 			get { return Inner.bDeterministic; }
 		}
-
-		public bool bForceDeterministic
-		{
-			get { return Inner.bForceDeterministic; }
-		}
-
 		public bool bChainPCHs
 		{
 			get { return Inner.bChainPCHs; }
diff --git a/Engine/Source/Programs/UnrealBuildTool/Configuration/UEBuildModuleCPP.cs b/Engine/Source/Programs/UnrealBuildTool/Configuration/UEBuildModuleCPP.cs
--- a/Engine/Source/Programs/UnrealBuildTool/Configuration/UEBuildModuleCPP.cs
+++ b/Engine/Source/Programs/UnrealBuildTool/Configuration/UEBuildModuleCPP.cs
@@ -1689,7 +1689,7 @@ namespace UnrealBuildTool
 			Result.StaticAnalyzerAdditionalCheckers = Rules.StaticAnalyzerAdditionalCheckers;
 			Result.bEnableUndefinedIdentifierWarnings = Rules.bEnableUndefinedIdentifierWarnings;
 			Result.IncludeOrderVersion = Rules.IncludeOrderVersion;
-			Result.bDeterministic |= Rules.bDeterministic;
+			Result.DeterministicWarningLevel = Rules.DeterministicWarningLevel;
 
 			if (Result.OptimizationLevel != Rules.OptimizationLevel)
 			{
diff --git a/Engine/Source/Programs/UnrealBuildTool/Configuration/UEBuildTarget.cs b/Engine/Source/Programs/UnrealBuildTool/Configuration/UEBuildTarget.cs
--- a/Engine/Source/Programs/UnrealBuildTool/Configuration/UEBuildTarget.cs
+++ b/Engine/Source/Programs/UnrealBuildTool/Configuration/UEBuildTarget.cs
@@ -4426,7 +4426,7 @@ namespace UnrealBuildTool
 			GlobalCompileEnvironment.CppStandard = Rules.CppStandard;
 			GlobalCompileEnvironment.CStandard = Rules.CStandard;
 			GlobalCompileEnvironment.AdditionalArguments = Rules.AdditionalCompilerArguments ?? String.Empty;
-			GlobalCompileEnvironment.bDeterministic = Rules.bForceDeterministic;
+			GlobalCompileEnvironment.bDeterministic = Rules.bDeterministic;
 			GlobalCompileEnvironment.CrashDiagnosticDirectory = Rules.CrashDiagnosticDirectory;
 
 			GlobalLinkEnvironment.bUseDebugCRT = bUseDebugCRT;
@@ -4448,7 +4448,7 @@ namespace UnrealBuildTool
 			GlobalLinkEnvironment.bPGOProfile = Rules.bPGOProfile;
 			GlobalLinkEnvironment.bUseIncrementalLinking = Rules.bUseIncrementalLinking;
 			GlobalLinkEnvironment.bUseFastPDBLinking = Rules.bUseFastPDBLinking ?? false;
-			GlobalLinkEnvironment.bDeterministic = Rules.bDeterministic || Rules.bForceDeterministic;
+			GlobalLinkEnvironment.bDeterministic = Rules.bDeterministic;
 			GlobalLinkEnvironment.bPrintTimingInfo = Rules.bPrintToolChainTimingInfo;
 			GlobalLinkEnvironment.bUsePIE = Rules.bEnablePIE;
 			GlobalLinkEnvironment.AdditionalArguments = Rules.AdditionalLinkerArguments ?? String.Empty;
diff --git a/Engine/Source/Programs/UnrealBuildTool/Platform/Windows/VCToolChain.cs b/Engine/Source/Programs/UnrealBuildTool/Platform/Windows/VCToolChain.cs
--- a/Engine/Source/Programs/UnrealBuildTool/Platform/Windows/VCToolChain.cs
+++ b/Engine/Source/Programs/UnrealBuildTool/Platform/Windows/VCToolChain.cs
@@ -542,6 +542,19 @@ namespace UnrealBuildTool
 				if (Target.WindowsPlatform.Compiler.IsMSVC())
 				{
 					Arguments.Add("/experimental:deterministic");
+
+					if (CompileEnvironment.DeterministicWarningLevel == WarningLevel.Error)
+					{
+						Arguments.Add("/we5048");
+					}
+					else if (CompileEnvironment.DeterministicWarningLevel == WarningLevel.Warning)
+					{
+						Arguments.Add("/w5048");
+					}
+					else
+					{
+						Arguments.Add("/wd5048");
+					}
 				}
 				else if (Target.WindowsPlatform.Compiler.IsClang())
 				{
@@ -1021,7 +1034,14 @@ namespace UnrealBuildTool
 				// Warn if __DATE__ or __TIME__ are used as they prevent reproducible builds
 				if (CompileEnvironment.bDeterministic)
 				{
-					Arguments.Add("-Wdate-time -Wno-error=date-time");
+					if (CompileEnvironment.DeterministicWarningLevel == WarningLevel.Error)
+					{
+						Arguments.Add("-Wdate-time");
+					}
+					else if (CompileEnvironment.DeterministicWarningLevel == WarningLevel.Warning)
+					{
+						Arguments.Add("-Wdate-time -Wno-error=date-time");
+					}
 				}
 
 				// This is disabled because clang explicitly warns about changing pack alignment in a header and not
diff --git a/Engine/Source/Programs/UnrealBuildTool/System/CppCompileEnvironment.cs b/Engine/Source/Programs/UnrealBuildTool/System/CppCompileEnvironment.cs
--- a/Engine/Source/Programs/UnrealBuildTool/System/CppCompileEnvironment.cs
+++ b/Engine/Source/Programs/UnrealBuildTool/System/CppCompileEnvironment.cs
@@ -617,6 +617,11 @@ namespace UnrealBuildTool
 		/// </summary>
 		public bool bDeterministic = false;
 
+		/// <summary>
+		/// Set flags for determinstic warnings (experimental).
+		/// </summary>
+		public WarningLevel DeterministicWarningLevel = WarningLevel.Off;
+
 		/// <summary>
 		/// Directory where to put crash report files for platforms that support it
 		/// </summary>
@@ -722,6 +727,7 @@ namespace UnrealBuildTool
 			bEnableCoroutines = Other.bEnableCoroutines;
 			IncludeOrderVersion = Other.IncludeOrderVersion;
 			bDeterministic = Other.bDeterministic;
+			DeterministicWarningLevel = Other.DeterministicWarningLevel;
 			CrashDiagnosticDirectory = Other.CrashDiagnosticDirectory;
 		}
 
diff --git a/Engine/Source/Programs/UnrealBuildTool/ToolChain/ClangToolChain.cs b/Engine/Source/Programs/UnrealBuildTool/ToolChain/ClangToolChain.cs
--- a/Engine/Source/Programs/UnrealBuildTool/ToolChain/ClangToolChain.cs
+++ b/Engine/Source/Programs/UnrealBuildTool/ToolChain/ClangToolChain.cs
@@ -605,7 +605,15 @@ namespace UnrealBuildTool
 			// Warn if __DATE__ or __TIME__ are used as they prevent reproducible builds
 			if (CompileEnvironment.bDeterministic)
 			{
-				Arguments.Add("-Wdate-time -Wno-error=date-time"); // https://clang.llvm.org/docs/DiagnosticsReference.html#wdate-time
+				// https://clang.llvm.org/docs/DiagnosticsReference.html#wdate-time
+				if (CompileEnvironment.DeterministicWarningLevel == WarningLevel.Error)
+				{
+					Arguments.Add("-Wdate-time");
+				}
+				else if (CompileEnvironment.DeterministicWarningLevel == WarningLevel.Warning)
+				{
+					Arguments.Add("-Wdate-time -Wno-error=date-time");
+				}
 			}
 
 			// https://clang.llvm.org/docs/DiagnosticsReference.html#wshadow
diff --git a/Engine/Source/Runtime/BuildSettings/BuildSettings.Build.cs b/Engine/Source/Runtime/BuildSettings/BuildSettings.Build.cs
--- a/Engine/Source/Runtime/BuildSettings/BuildSettings.Build.cs
+++ b/Engine/Source/Runtime/BuildSettings/BuildSettings.Build.cs
@@ -8,6 +8,7 @@ public class BuildSettings : ModuleRules
 {
 	public BuildSettings(ReadOnlyTargetRules Target) : base(Target)
 	{
+		DeterministicWarningLevel = WarningLevel.Off; // This module intentionally uses __DATE__ and __TIME__ macros
 		PrivateIncludePathModuleNames.Add("Core");
 
 		bRequiresImplementModule = false;
