UnrealBuildTool: Remove warning from exceptions while scanning for MSVS installations

When invoked under Wine, this generates unnecessary warnings. Talked to Joe and decided to remove the warning instead.

diff --git a/Engine/Source/Programs/UnrealBuildTool/Platform/Windows/MicrosoftPlatformSDK.cs b/Engine/Source/Programs/UnrealBuildTool/Platform/Windows/MicrosoftPlatformSDK.cs
--- a/Engine/Source/Programs/UnrealBuildTool/Platform/Windows/MicrosoftPlatformSDK.cs
+++ b/Engine/Source/Programs/UnrealBuildTool/Platform/Windows/MicrosoftPlatformSDK.cs
@@ -842,8 +842,7 @@ namespace UnrealBuildTool
 				}
 				catch (Exception Ex)
 				{
-					Logger.LogWarning("Unable to enumerate Visual Studio installations");
-					Logger.LogDebug(Ex.ToString());
+					Logger.LogDebug(Ex, "Unable to enumerate Visual Studio installations");
 				}
 			}
 			CachedVisualStudioInstallations = Installations;
