[UnrealBuildTool]
* Added bCanCache to actions and enabled bCanCache to vc compile/link if /experimental:deterministic is set (because it can actually be cached then :-))

diff --git a/Engine/Source/Programs/UnrealBuildTool/Platform/Windows/VCCompileAction.cs b/Engine/Source/Programs/UnrealBuildTool/Platform/Windows/VCCompileAction.cs
--- a/Engine/Source/Programs/UnrealBuildTool/Platform/Windows/VCCompileAction.cs
+++ b/Engine/Source/Programs/UnrealBuildTool/Platform/Windows/VCCompileAction.cs
@@ -156,6 +156,9 @@ namespace UnrealBuildTool
 		/// </summary>
 		public List<FileItem> DeleteItems { get; } = new List<FileItem>();
 
+		/// <inheritdoc/>
+		public bool bCanCache { get; set; }
+
 		/// <inheritdoc/>
 		public bool bCanExecuteRemotely { get; set; }
 
diff --git a/Engine/Source/Programs/UnrealBuildTool/Platform/Windows/VCToolChain.cs b/Engine/Source/Programs/UnrealBuildTool/Platform/Windows/VCToolChain.cs
--- a/Engine/Source/Programs/UnrealBuildTool/Platform/Windows/VCToolChain.cs
+++ b/Engine/Source/Programs/UnrealBuildTool/Platform/Windows/VCToolChain.cs
@@ -542,16 +542,11 @@ namespace UnrealBuildTool
 				if (Target.WindowsPlatform.Compiler.IsMSVC())
 				{
 					Arguments.Add("/experimental:deterministic");
-
 					if (CompileEnvironment.DeterministicWarningLevel == WarningLevel.Error)
 					{
 						Arguments.Add("/we5048");
 					}
-					else if (CompileEnvironment.DeterministicWarningLevel == WarningLevel.Warning)
-					{
-						Arguments.Add("/w5048");
-					}
-					else
+					else if (CompileEnvironment.DeterministicWarningLevel == WarningLevel.Off)
 					{
 						Arguments.Add("/wd5048");
 					}
@@ -1537,6 +1532,8 @@ namespace UnrealBuildTool
 					}
 				}
 
+				CompileAction.bCanCache = CompileEnvironment.bDeterministic;
+
 				// Don't farm out creation of precompiled headers as it is the critical path task.
 				CompileAction.bCanExecuteRemotely =
 					CompileEnvironment.PrecompiledHeaderAction != PrecompiledHeaderAction.Create ||
@@ -1667,6 +1664,8 @@ namespace UnrealBuildTool
 					CompileAction.WriteResponseFile(Graph, Logger);
 				}
 
+				CompileAction.bCanCache = CompileEnvironment.bDeterministic;
+
 				// When compiling with SN-DBS, modules that contain a #import must be built locally
 				CompileAction.bCanExecuteRemotelyWithSNDBS = CompileAction.bCanExecuteRemotely;
 				if (CompileEnvironment.bBuildLocallyWithSNDBS == true)
diff --git a/Engine/Source/Programs/UnrealBuildTool/System/Action.cs b/Engine/Source/Programs/UnrealBuildTool/System/Action.cs
--- a/Engine/Source/Programs/UnrealBuildTool/System/Action.cs
+++ b/Engine/Source/Programs/UnrealBuildTool/System/Action.cs
@@ -103,6 +103,11 @@ namespace UnrealBuildTool
 		/// </summary>
 		bool bCanExecuteRemotely { get; }
 
+		/// <summary>
+		/// True if this action can be cached. Cache will store all touched inputs and store outputs
+		/// </summary>
+		bool bCanCache { get; }
+
 		/// <summary>
 		/// True if this action is allowed to be run on a remote machine with SNDBS. Files with #import directives must be compiled locally. Also requires bCanExecuteRemotely = true.
 		/// </summary>
@@ -205,6 +210,11 @@ namespace UnrealBuildTool
 		/// </summary>
 		public string StatusDescription { get; set; } = "...";
 
+		/// <summary>
+		/// True if this action's result can be cached based on its inputs
+		/// </summary>
+		public bool bCanCache { get; set; }
+
 		/// <summary>
 		/// True if this action is allowed to be run on a remote machine when a distributed build system is being used, such as XGE
 		/// </summary>
@@ -597,6 +607,7 @@ namespace UnrealBuildTool
 		public string CommandVersion => Inner.CommandVersion;
 		public string CommandDescription => Inner.CommandDescription;
 		public string StatusDescription => Inner.StatusDescription;
+		public bool bCanCache => Inner.bCanCache;
 		public bool bCanExecuteRemotely => Inner.bCanExecuteRemotely;
 		public bool bCanExecuteRemotelyWithSNDBS => Inner.bCanExecuteRemotelyWithSNDBS;
 		public bool bIsGCCCompiler => Inner.bIsGCCCompiler;
