UnrealBuildTool: Only -SingleFile compile files if they are included in the target. Requires storing the header list in the makefile

#rnx

diff --git a/Engine/Source/Programs/UnrealBuildTool/Configuration/UEBuildModuleCPP.cs b/Engine/Source/Programs/UnrealBuildTool/Configuration/UEBuildModuleCPP.cs
--- a/Engine/Source/Programs/UnrealBuildTool/Configuration/UEBuildModuleCPP.cs
+++ b/Engine/Source/Programs/UnrealBuildTool/Configuration/UEBuildModuleCPP.cs
@@ -388,6 +388,9 @@ namespace UnrealBuildTool
 				Graph.AddSourceFiles(Pair.Key, Pair.Value);
 			}
 
+			Graph.AddHeaderFiles(InputFiles.HeaderFiles.ToArray());
+			Graph.AddHeaderFiles(InputFiles.ISPCHeaderFiles.ToArray());
+
 			// We are building with IWYU and thismodule does not support it, early out
 			if (Target.bIWYU && Rules.IWYUSupport == IWYUSupport.None)
 			{
diff --git a/Engine/Source/Programs/UnrealBuildTool/Modes/BuildMode.cs b/Engine/Source/Programs/UnrealBuildTool/Modes/BuildMode.cs
--- a/Engine/Source/Programs/UnrealBuildTool/Modes/BuildMode.cs
+++ b/Engine/Source/Programs/UnrealBuildTool/Modes/BuildMode.cs
@@ -449,22 +449,31 @@ namespace UnrealBuildTool
 				for (int Idx = 0; Idx < Makefiles.Length; Idx++)
 				{
 					TargetDescriptor TargetDescriptor = TargetDescriptors[Idx];
-					List<FileReference> FilesToBuild = TargetDescriptor.SpecificFilesToCompile;
+					TargetMakefile Makefile = Makefiles[Idx];
 
-					QueuedActions[Idx] = Makefiles[Idx].Actions.ConvertAll(x => new LinkedAction(x, TargetDescriptors[Idx]));
+					QueuedActions[Idx] = Makefile.Actions.ConvertAll(x => new LinkedAction(x, TargetDescriptors[Idx]));
 
-					// If we have specific files to build we need to put those as MergedOutputItems (and create special single file actions if needed)
-					if (FilesToBuild.Count != 0)
+					if (TargetDescriptor.SpecificFilesToCompile.Count != 0)
 					{
-						// If there are headers in the list, expand the FilesToBuild to also include all files that include those headers
-						if (TargetDescriptor.bSingleFileBuildDependents)
+						// Filter out SpecificFilesToCompile so only source files and headers included in the target are considered
+						HashSet<FileReference> TargetFiles = Makefile.SourceAndHeaderFiles.Select(x => x.Location).ToHashSet();
+						List<FileReference> FilesToBuild = TargetDescriptor.SpecificFilesToCompile
+							.Where(x => TargetFiles.Contains(x))
+							.ToList();
+
+						// If we have specific files to build we need to put those as MergedOutputItems (and create special single file actions if needed)
+						if (FilesToBuild.Count != 0)
 						{
-							FilesToBuild = GetAllSourceFilesIncludingHeader(FilesToBuild, TargetDescriptor.ProjectFile, Makefiles[Idx].Actions, Logger);
-						}
+							// If there are headers in the list, expand the FilesToBuild to also include all files that include those headers
+							if (TargetDescriptor.bSingleFileBuildDependents)
+							{
+								FilesToBuild = GetAllSourceFilesIncludingHeader(FilesToBuild, TargetDescriptor.ProjectFile, Makefile.Actions, Logger);
+							}
 
-						// We have specific files to compile so we will only queue up those files.
-						List<FileItem> ProducedItems = CreateLinkedActionsFromFileList(TargetDescriptor, BuildConfiguration, FilesToBuild, QueuedActions[Idx], Logger);
-						MergedOutputItems.UnionWith(ProducedItems);
+							// We have specific files to compile so we will only queue up those files.
+							List<FileItem> ProducedItems = CreateLinkedActionsFromFileList(TargetDescriptor, BuildConfiguration, FilesToBuild, QueuedActions[Idx], Logger);
+							MergedOutputItems.UnionWith(ProducedItems);
+						}
 					}
 				}
 
diff --git a/Engine/Source/Programs/UnrealBuildTool/Modes/GenerateClangDatabase.cs b/Engine/Source/Programs/UnrealBuildTool/Modes/GenerateClangDatabase.cs
--- a/Engine/Source/Programs/UnrealBuildTool/Modes/GenerateClangDatabase.cs
+++ b/Engine/Source/Programs/UnrealBuildTool/Modes/GenerateClangDatabase.cs
@@ -78,6 +78,11 @@ namespace UnrealBuildTool
 			{
 			}
 
+			/// <inheritdoc/>
+			public void AddHeaderFiles(FileItem[] HeaderFiles)
+			{
+			}
+
 			/// <inheritdoc/>
 			public void AddFileToWorkingSet(FileItem File)
 			{
diff --git a/Engine/Source/Programs/UnrealBuildTool/System/ActionGraphBuilder.cs b/Engine/Source/Programs/UnrealBuildTool/System/ActionGraphBuilder.cs
--- a/Engine/Source/Programs/UnrealBuildTool/System/ActionGraphBuilder.cs
+++ b/Engine/Source/Programs/UnrealBuildTool/System/ActionGraphBuilder.cs
@@ -66,6 +66,12 @@ namespace UnrealBuildTool
 		/// <param name="SourceFiles">Contents of the directory</param>
 		void AddSourceFiles(DirectoryItem SourceDir, FileItem[] SourceFiles);
 
+		/// <summary>
+		/// Adds a list of known header files
+		/// </summary>
+		/// <param name="HeaderFiles">List of header files to track</param>
+		void AddHeaderFiles(FileItem[] HeaderFiles);
+
 		/// <summary>
 		/// Sets the output items which belong to a particular module
 		/// </summary>
@@ -123,6 +129,11 @@ namespace UnrealBuildTool
 		{
 		}
 
+		/// <inheritdoc/>
+		public void AddHeaderFiles(FileItem[] HeaderFiles)
+		{
+		}
+
 		/// <inheritdoc/>
 		public void AddFileToWorkingSet(FileItem File)
 		{
@@ -193,6 +204,12 @@ namespace UnrealBuildTool
 			Inner.AddSourceFiles(SourceDir, SourceFiles);
 		}
 
+		/// <inheritdoc/>
+		public virtual void AddHeaderFiles(FileItem[] HeaderFiles)
+		{
+			Inner.AddHeaderFiles(HeaderFiles);
+		}
+
 		/// <inheritdoc/>
 		public virtual void AddFileToWorkingSet(FileItem File)
 		{
diff --git a/Engine/Source/Programs/UnrealBuildTool/System/TargetMakefile.cs b/Engine/Source/Programs/UnrealBuildTool/System/TargetMakefile.cs
--- a/Engine/Source/Programs/UnrealBuildTool/System/TargetMakefile.cs
+++ b/Engine/Source/Programs/UnrealBuildTool/System/TargetMakefile.cs
@@ -1,5 +1,6 @@
 // Copyright Epic Games, Inc. All Rights Reserved.
 using System;
+using System.Collections;
 using System.Collections.Concurrent;
 using System.Collections.Generic;
 using System.IO;
@@ -131,6 +132,12 @@ namespace UnrealBuildTool
 			Makefile.DirectoryToSourceFiles[SourceDir] = SourceFiles.Select(fi => new TargetMakefileSourceFileInfo(fi)).ToArray();
 		}
 
+		/// <inheritdoc/>
+		public virtual void AddHeaderFiles(FileItem[] HeaderFiles)
+		{
+			Makefile.HeaderFiles.UnionWith(HeaderFiles);
+		}
+
 		/// <inheritdoc/>
 		public void AddDiagnostic(string Message)
 		{
@@ -305,6 +312,11 @@ namespace UnrealBuildTool
 		/// </summary>
 		public Dictionary<DirectoryItem, TargetMakefileSourceFileInfo[]> DirectoryToSourceFiles;
 
+		/// <summary>
+		/// Set of all known header files, changes will not invalidate the makefile but are used to determine if a header is included in the target.
+		/// </summary>
+		public HashSet<FileItem> HeaderFiles;
+
 		/// <summary>
 		/// The set of source files that UnrealBuildTool determined to be part of the programmer's "working set". Used for adaptive non-unity builds.
 		/// </summary>
@@ -357,6 +369,28 @@ namespace UnrealBuildTool
 		/// </summary>
 		public double MemoryPerActionGB = 0.0;
 
+		/// <summary>
+		/// Enumerable of all source and header files tracked by this makefile.
+		/// </summary>
+		public IEnumerable<FileItem> SourceAndHeaderFiles
+		{
+			get
+			{
+				foreach (KeyValuePair<DirectoryItem, TargetMakefileSourceFileInfo[]> Item in DirectoryToSourceFiles)
+				{
+					foreach (TargetMakefileSourceFileInfo File in Item.Value.Where(x => x.SourceFileItem != null))
+					{
+						yield return File.SourceFileItem!;
+					}
+				}
+
+				foreach (FileItem File in HeaderFiles)
+				{
+					yield return File;
+				}
+			}
+		}
+
 		/// <summary>
 		/// Constructor
 		/// </summary>
@@ -396,6 +430,7 @@ namespace UnrealBuildTool
 			this.HotReloadModuleNames = new HashSet<string>(StringComparer.OrdinalIgnoreCase);
 			this.SourceDirectories = new List<DirectoryItem>();
 			this.DirectoryToSourceFiles = new();
+			this.HeaderFiles = new();
 			this.WorkingSet = new HashSet<FileItem>();
 			this.CandidatesForWorkingSet = new HashSet<FileItem>();
 			this.UObjectModules = new List<UHTModuleInfo>();
@@ -441,6 +476,7 @@ namespace UnrealBuildTool
 			HotReloadModuleNames = Reader.ReadHashSet(() => Reader.ReadString(), StringComparer.OrdinalIgnoreCase)!;
 			SourceDirectories = Reader.ReadList(() => Reader.ReadDirectoryItem())!;
 			DirectoryToSourceFiles = Reader.ReadDictionary(() => Reader.ReadDirectoryItem()!, () => Reader.ReadArray(() =>  new TargetMakefileSourceFileInfo(Reader)))!;
+			HeaderFiles = Reader.ReadHashSet(() => Reader.ReadFileItem())!;
 			WorkingSet = Reader.ReadHashSet(() => Reader.ReadFileItem())!;
 			CandidatesForWorkingSet = Reader.ReadHashSet(() => Reader.ReadFileItem())!;
 			UObjectModules = Reader.ReadList(() => new UHTModuleInfo(Reader))!;
@@ -485,6 +521,7 @@ namespace UnrealBuildTool
 			Writer.WriteHashSet(HotReloadModuleNames, x => Writer.WriteString(x));
 			Writer.WriteList(SourceDirectories, x => Writer.WriteDirectoryItem(x));
 			Writer.WriteDictionary(DirectoryToSourceFiles, k => Writer.WriteDirectoryItem(k), v => Writer.WriteArray(v, e => e.Write(Writer)));
+			Writer.WriteHashSet(HeaderFiles, x => Writer.WriteFileItem(x));
 			Writer.WriteHashSet(WorkingSet, x => Writer.WriteFileItem(x));
 			Writer.WriteHashSet(CandidatesForWorkingSet, x => Writer.WriteFileItem(x));
 			Writer.WriteList(UObjectModules, e => e.Write(Writer));
