Use direct include in Mac indexing response files, this resolves majority of indexing compile fails.

diff --git a/Engine/Source/Programs/UnrealBuildTool/ProjectFiles/Xcode/XcodeProject.cs b/Engine/Source/Programs/UnrealBuildTool/ProjectFiles/Xcode/XcodeProject.cs
--- a/Engine/Source/Programs/UnrealBuildTool/ProjectFiles/Xcode/XcodeProject.cs
+++ b/Engine/Source/Programs/UnrealBuildTool/ProjectFiles/Xcode/XcodeProject.cs
@@ -85,7 +85,7 @@ namespace UnrealBuildTool.XcodeProjectXcconfig
 	class UnrealBatchedFiles
 	{
 		// build settings that cause uniqueness
-		public string Definitions = "";
+		public IEnumerable<String>? ForceIncludeFiles = null;
 		// @todo can we actually use this effectively with indexing other than fotced include?
 		public FileReference? PCHFile = null;
 		public bool bEnableRTTI = false;
@@ -96,12 +96,13 @@ namespace UnrealBuildTool.XcodeProjectXcconfig
 		public HashSet<DirectoryReference> UserIncludePaths = new();
 
 		public List<XcodeSourceFile> Files = new();
-		public List<UEBuildModuleCPP> Modules = new List<UEBuildModuleCPP>();
+		public UEBuildModuleCPP Module;
 
 		public FileReference ResponseFile;
 
-		public UnrealBatchedFiles(UnrealData UnrealData, int Index)
+		public UnrealBatchedFiles(UnrealData UnrealData, int Index, UEBuildModuleCPP Module)
 		{
+			this.Module = Module;
 			ResponseFile = FileReference.Combine(UnrealData.XcodeProjectFileLocation.ParentDirectory!, "ResponseFiles", $"{UnrealData.ProductName}{Index}.response");
 		}
 
@@ -112,7 +113,11 @@ namespace UnrealBuildTool.XcodeProjectXcconfig
 			ResponseFileContents.AppendJoin(" -I", SystemIncludePaths.Select(x => x.FullName.Contains(' ') ? $"\"{x.FullName}\"" : x.FullName));
 			ResponseFileContents.Append(" -I");
 			ResponseFileContents.AppendJoin(" -I", UserIncludePaths.Select(x => x.FullName.Contains(' ') ? $"\"{x.FullName}\"" : x.FullName));
-
+			if (ForceIncludeFiles != null)
+			{
+				ResponseFileContents.Append(" -include ");
+				ResponseFileContents.AppendJoin(" -include ", ForceIncludeFiles.Select(x => x.Contains(' ') ? $"\"{x}\"" : x));
+			}
 			ResponseFileContents.Append(" -D");
 			ResponseFileContents.AppendJoin(" -D", AllDefines);
 
@@ -483,99 +488,27 @@ namespace UnrealBuildTool.XcodeProjectXcconfig
 
 		public void AddModule(UEBuildModuleCPP Module, CppCompileEnvironment CompileEnvironment)
 		{
-			// we need to keep all _API defines, but we can gather them up to append at the end, instead of using them to compute sameness
+			// one batched files per module
+			UnrealBatchedFiles FileBatch = new UnrealBatchedFiles(this, BatchedFiles.Count + 1, Module);
+			BatchedFiles.Add(FileBatch);
 
-			// remove some extraneous defines that will cause every module to be unique, but we can do without
-			List<string> Defines = new();
-			List<string> APIDefines = new();
-			Regex Regex = new Regex("#define ([A-Z0-9_]*) ?(.*)?");
-
-			string DefinesString = "";
-			FileReference? PCHFile = null;
 			if (CompileEnvironment.ForceIncludeFiles.Count == 0)
 			{
 				// if there are no ForceInclude files, then that means it's a module that forces the includes to come from a generated PCH file
 				// and so we will use this for definitions and uniqueness
 				if (CompileEnvironment.PrecompiledHeaderIncludeFilename != null)
 				{
-					PCHFile = FileReference.Combine(XcodeProjectFileLocation.ParentDirectory!, "PCHFiles", CompileEnvironment.PrecompiledHeaderIncludeFilename.GetFileName());
-					DirectoryReference.CreateDirectory(PCHFile.Directory);
-					FileReference.Copy(CompileEnvironment.PrecompiledHeaderIncludeFilename, PCHFile, true);
+					FileBatch.PCHFile = FileReference.Combine(XcodeProjectFileLocation.ParentDirectory!, "PCHFiles", CompileEnvironment.PrecompiledHeaderIncludeFilename.GetFileName());
+					DirectoryReference.CreateDirectory(FileBatch.PCHFile.Directory);
+					FileReference.Copy(CompileEnvironment.PrecompiledHeaderIncludeFilename, FileBatch.PCHFile, true);
 				}
 			}
 			else
 			{
-				foreach (string Line in File.ReadAllLines(CompileEnvironment.ForceIncludeFiles.First(x => x.FullName.Contains("Definitions")).FullName))
-				{
-					Match Match = Regex.Match(Line);
-					if (!Match.Success)
-					{
-						continue;
-					}
-
-					string Key = Match.Groups[1].Value;
-					string Value = Match.Groups[2].Value;
-					// if no value, then just define with no = stuff
-					if (!Match.Groups[2].Success)
-					{
-						Defines.Add(Key);
-						continue;
-					}
-
-					// skip some known per-module defines we don't need
-					if (Key.StartsWith("UE_MODULE_NAME") || Key.StartsWith("UE_PLUGIN_NAME") || Key.StartsWith("UBT_MODULE_MANIFEST"))
-					{
-						continue;
-					}
-					// these API ones are per module but still need to be defined, can be defined to nothing
-					else if (Key.Contains("_API"))
-					{
-						APIDefines.Add($"{Key}=");
-					}
-					else
-					{
-						// if the value is normal, just add the define as is
-						if (Value.All(x => char.IsLetterOrDigit(x) || x == '_'))
-						{
-							Defines.Add($"{Key}={Value}");
-						}
-						else
-						{
-							// escape any quotes in the define, then quote the whole thing
-							Value = Value.Replace("\"", "\\\"");
-							Defines.Add($"{Key}=\"{Value}\"");
-						}
-					}
-				}
-
-				// sort and joing them into a single string to act as a key (and happily string to put into .xcconfig)
-				Defines.Sort();
-				DefinesString = string.Join(" ", Defines);
-			}
-
-			// now find a matching SubTarget, and make a new one if needed
-			UnrealBatchedFiles? FileBatch = null;
-			foreach (UnrealBatchedFiles Search in BatchedFiles)
-			{
-				if (Search.Definitions == DefinesString && Search.PCHFile == PCHFile &&
-					Search.bEnableRTTI == CompileEnvironment.bUseRTTI)
-				{
-					FileBatch = Search;
-					break;
-				}
-			}
-			if (FileBatch == null)
-			{
-				FileBatch = new UnrealBatchedFiles(this, BatchedFiles.Count + 1);
-				BatchedFiles.Add(FileBatch);
-				FileBatch.Definitions = DefinesString;
-				FileBatch.PCHFile = PCHFile;
-				FileBatch.bEnableRTTI = CompileEnvironment.bUseRTTI;
+				FileBatch.ForceIncludeFiles = CompileEnvironment.ForceIncludeFiles.Select(x => x.FullName);
 			}
 
-			FileBatch.Modules.Add(Module);
-			FileBatch.AllDefines.UnionWith(Defines);
-			FileBatch.AllDefines.UnionWith(APIDefines);
+			FileBatch.bEnableRTTI = CompileEnvironment.bUseRTTI;
 			FileBatch.SystemIncludePaths.UnionWith(CompileEnvironment.SystemIncludePaths);
 			FileBatch.UserIncludePaths.UnionWith(CompileEnvironment.UserIncludePaths);
 		}
@@ -1992,14 +1925,11 @@ namespace UnrealBuildTool.XcodeProjectXcconfig
 		{
 			foreach (UnrealBatchedFiles Batch in UnrealData.BatchedFiles)
 			{
-				foreach (UEBuildModuleCPP Module in Batch.Modules)
+				if (Batch.Module.ContainsFile(File.Reference))
 				{
-					if (Module.ContainsFile(File.Reference))
-					{
-						Batch.Files.Add(File);
-						FileCollection.BuildableFilesToResponseFile[File] = Batch.ResponseFile;
-						return;
-					}
+					Batch.Files.Add(File);
+					FileCollection.BuildableFilesToResponseFile[File] = Batch.ResponseFile;
+					return;
 				}
 			}
 		}
