PR #8239: Support for project-specific BuildConfiguration.xml and Programs during project file generation (Contributed by [REDACTED_PERSON])

Contains some additional changes from Joe.Kirchoff:
* Fix project dependency for rules parent chain
* Always add UnrealBuildTool and AutomationTool .csproj for game targets -except-
* Filter out other non-critial .csproj unless engine programs are included

diff --git a/Engine/Source/Programs/UnrealBuildTool/Modes/GenerateProjectFilesMode.cs b/Engine/Source/Programs/UnrealBuildTool/Modes/GenerateProjectFilesMode.cs
--- a/Engine/Source/Programs/UnrealBuildTool/Modes/GenerateProjectFilesMode.cs
+++ b/Engine/Source/Programs/UnrealBuildTool/Modes/GenerateProjectFilesMode.cs
@@ -83,6 +83,10 @@ namespace UnrealBuildTool
 			FileReference? ProjectFile;
 			TryParseProjectFileArgument(Arguments, Logger, out ProjectFile);
 
+			// Apply the XML config again with a project specific BuildConfiguration.xml 
+			XmlConfig.ReadConfigFiles(null, ProjectFile?.Directory, Logger);
+			XmlConfig.ApplyTo(this);
+			
 			// Warn if there are explicit project file formats specified
 			if (ProjectFileFormats.Count > 0 && !bAutomated)
 			{
diff --git a/Engine/Source/Programs/UnrealBuildTool/ProjectFiles/ProjectFileGenerator.cs b/Engine/Source/Programs/UnrealBuildTool/ProjectFiles/ProjectFileGenerator.cs
--- a/Engine/Source/Programs/UnrealBuildTool/ProjectFiles/ProjectFileGenerator.cs
+++ b/Engine/Source/Programs/UnrealBuildTool/ProjectFiles/ProjectFileGenerator.cs
@@ -13,6 +13,7 @@ using System.Xml.Linq;
 using System.Xml;
 using UnrealBuildBase;
 using Microsoft.Extensions.Logging;
+using EpicGames.Serialization;
 
 namespace UnrealBuildTool
 {
@@ -537,6 +538,11 @@ namespace UnrealBuildTool
 						PrimaryProjectFolder Folder = ProgramsFolder.AddSubFolder("Shared");
 						foreach (FileInfo ProjectFile in ProjectFiles)
 						{
+							// Don't add Shared Test projects unless all engine programs are included
+							if (!bIncludeEnginePrograms && ProjectFile.Name.EndsWith(".Tests.csproj", StringComparison.OrdinalIgnoreCase))
+							{
+								continue;
+							}
 							VCSharpProjectFile Project = new VCSharpProjectFile(new FileReference(ProjectFile), Logger);
 							Project.ShouldBuildForAllSolutionTargets = false;
 							AddExistingProjectFile(Project, bForceDevelopmentConfiguration: true);
@@ -575,14 +581,17 @@ namespace UnrealBuildTool
 			List<FileReference> FoundProjects = new List<FileReference>();
 			List<DirectoryReference> ProjectDirs = new List<DirectoryReference>();
 
-			DirectoryReference EngineExtras = DirectoryReference.Combine(Unreal.EngineDirectory, "Extras");
-			ProjectDirs.Add(EngineExtras);
-			DiscoverCSharpProgramProjectsRecursively(EngineExtras, FoundProjects);
-			
-			ProjectDirs = ProjectDirs.Union(AllEngineDirectories).ToList();
-			foreach (DirectoryReference EngineDir in AllEngineDirectories)
+			if (bIncludeEnginePrograms)
 			{
-				DiscoverCSharpProgramProjectsRecursively(EngineDir, FoundProjects);
+				DirectoryReference EngineExtras = DirectoryReference.Combine(Unreal.EngineDirectory, "Extras");
+				ProjectDirs.Add(EngineExtras);
+				DiscoverCSharpProgramProjectsRecursively(EngineExtras, FoundProjects);
+			
+				ProjectDirs = ProjectDirs.Union(AllEngineDirectories).ToList();
+				foreach (DirectoryReference EngineDir in AllEngineDirectories)
+				{
+					DiscoverCSharpProgramProjectsRecursively(EngineDir, FoundProjects);
+				}
 			}
 
 			foreach (FileReference GameProjectFile in AllGameProjects)
@@ -638,12 +647,28 @@ namespace UnrealBuildTool
 			}
 		}
 
-		private VCSharpProjectFile CreateRulesAssemblyProject(RulesAssembly RulesAssembly, DirectoryReference FSPathBase, ILogger Logger)
+		private VCSharpProjectFile CreateRulesAssemblyProject(Dictionary<RulesAssembly, DirectoryReference> RulesAssemblies, RulesAssembly RulesAssembly, DirectoryReference FSPathBase, ILogger Logger)
 		{
 			DirectoryReference ProjectFilesDirectory = DirectoryReference.Combine(FSPathBase, "Intermediate/Build/BuildRulesProjects", RulesAssembly.GetSimpleAssemblyName()!);
 			DirectoryReference.CreateDirectory(ProjectFilesDirectory);
 			FileReference ModuleFilesProjectLocation = FileReference.Combine(ProjectFilesDirectory, RulesAssembly.GetSimpleAssemblyName() + ".csproj");
 
+			SortedSet<FileReference> ReferencedProjects = new() {
+				FileReference.Combine(Unreal.EngineDirectory, "Source", "Programs", "Shared", "EpicGames.Build", "EpicGames.Build.csproj"),
+				FileReference.Combine(Unreal.EngineDirectory, "Source", "Programs", "UnrealBuildTool", "UnrealBuildTool.csproj")
+			};
+
+			var CurrentParent = RulesAssembly.Parent;
+			while (CurrentParent != null)
+			{
+				if (RulesAssemblies.TryGetValue(CurrentParent, out DirectoryReference? FSParentPathBase))
+				{
+					DirectoryReference ProjectFilesParentDirectory = DirectoryReference.Combine(FSParentPathBase, "Intermediate", "Build", "BuildRulesProjects", CurrentParent.GetSimpleAssemblyName()!);
+					ReferencedProjects.Add(FileReference.Combine(ProjectFilesParentDirectory, CurrentParent.GetSimpleAssemblyName() + ".csproj"));
+				}
+				CurrentParent = CurrentParent.Parent;
+			}
+
 			{
 				using FileStream Stream = FileReference.Open(ModuleFilesProjectLocation, FileMode.Create, FileAccess.Write, FileShare.Read);
 				using StreamWriter Writer = new StreamWriter(Stream, Encoding.UTF8);
@@ -652,7 +677,7 @@ namespace UnrealBuildTool
 				Writer.WriteLine("<Project Sdk=\"Microsoft.NET.Sdk\">");
 
 				Writer.WriteLine("  <Import Project=\"" + 
-					FileReference.Combine(Unreal.EngineDirectory, "Source/Programs/Shared/UnrealEngine.csproj.props").MakeRelativeTo(ProjectFilesDirectory) + 
+					FileReference.Combine(Unreal.EngineDirectory, "Source", "Programs", "Shared", "UnrealEngine.csproj.props").MakeRelativeTo(ProjectFilesDirectory) + 
 					"\" />");
 
 				Writer.WriteLine("  <PropertyGroup>");
@@ -663,16 +688,15 @@ namespace UnrealBuildTool
 				Writer.WriteLine("  </PropertyGroup>");
 				
 				Writer.WriteLine("  <ItemGroup>");
-				Writer.WriteLine("    <ProjectReference Include=\"" + 
-					FileReference.Combine(Unreal.EngineDirectory, "Source/Programs/Shared/EpicGames.Build/EpicGames.Build.csproj").MakeRelativeTo(ProjectFilesDirectory) + 
-					"\"><Private>false</Private></ProjectReference>"); // <Private>false</Private> to suppress copying of dependencies into output directory
-				Writer.WriteLine("    <ProjectReference Include=\"" + 
-					FileReference.Combine(Unreal.EngineDirectory, "Source", "Programs", "UnrealBuildTool", "UnrealBuildTool.csproj").MakeRelativeTo(ProjectFilesDirectory) + 
-					"\"><Private>false</Private></ProjectReference>");
+				foreach (FileReference Reference in ReferencedProjects)
+				{
+					// <Private>false</Private> suppress copying of dependencies into output directory
+					Writer.WriteLine($"    <ProjectReference Include=\"{Reference.MakeRelativeTo(ProjectFilesDirectory)}\"><Private>false</Private></ProjectReference>");
+				}
 				Writer.WriteLine("  </ItemGroup>");
 
 				Writer.WriteLine("  <ItemGroup>");
-				foreach (FileReference ModuleFile in RulesAssembly.AssemblySourceFiles!)
+				foreach (FileReference ModuleFile in RulesAssembly.AssemblySourceFiles!.OrderBy(x => x.FullName))
 				{
 					string CsprojRelativePath = ModuleFile.MakeRelativeTo(ProjectFilesDirectory);
 					string ProjectFolder = ModuleFile.MakeRelativeTo(FSPathBase);
@@ -1115,6 +1139,8 @@ namespace UnrealBuildTool
 			AddProjectsForAllModules(AllGameProjects, ProjectFileToUProjectFile, ProgramProjects, ModProjects, AllModuleFiles, AdditionalSearchPaths, bGatherThirdPartySource, Logger);
 
 			{
+				PrimaryProjectFolder ProgramsFolder = RootFolder.AddSubFolder("Programs");
+
 				if (bIncludeDotNetPrograms)
 				{
 					// gather engine directories across extension dirs
@@ -1130,15 +1156,15 @@ namespace UnrealBuildTool
 							return true;
 						}).ToList();
 
-					PrimaryProjectFolder ProgramsFolder = RootFolder.AddSubFolder("Programs");
-
 					// Add UnrealBuildTool to the primary project
 					AddUnrealBuildToolProject(ProgramsFolder, Logger);
 
 					// Add AutomationTool to the primary project
 					VCSharpProjectFile AutomationToolProject = AddSimpleCSharpProject("AutomationTool", Logger, bShouldBuildForAllSolutionTargets: true, bForceDevelopmentConfiguration: true);
 					if (AutomationToolProject != null)
+					{
 						ProgramsFolder.ChildProjects.Add(AutomationToolProject);
+					}
 
 					// Add automation.csproj files to the primary project
 					AddRulesModules(Rules.RulesFileType.AutomationModule, "Automation", AutomationProjectFiles, AllGameProjects, RootFolder, ProgramsFolder, Logger);
@@ -1149,18 +1175,20 @@ namespace UnrealBuildTool
 					// Add shared projects
 					AddSharedDotNetModules(AllEngineDirectories, ProgramsFolder, Logger);
 
-					// Discover C# programs which should additionally be included in the solution.
-					DiscoverCSharpProgramProjects(AllEngineDirectories, AllGameProjects, ProgramsFolder, Logger);
+					if (bIncludeEnginePrograms)
+					{
+						// Discover C# programs which should additionally be included in the solution.
+						DiscoverCSharpProgramProjects(AllEngineDirectories, AllGameProjects, ProgramsFolder, Logger);
+					}
 
 					foreach (KeyValuePair<RulesAssembly, DirectoryReference> RulesAssemblyEntry in RulesAssemblies)
 					{
-						VCSharpProjectFile RulesProject = CreateRulesAssemblyProject(RulesAssemblyEntry.Key, RulesAssemblyEntry.Value, Logger);
+						VCSharpProjectFile RulesProject = CreateRulesAssemblyProject(RulesAssemblies, RulesAssemblyEntry.Key, RulesAssemblyEntry.Value, Logger);
 						AddExistingProjectFile(RulesProject, bForceDevelopmentConfiguration: false);
 						RootFolder.AddSubFolder("Rules").ChildProjects.Add(RulesProject);
 					}
 				}
-
-
+				
 				// Eliminate all redundant project folders.  E.g., folders which contain only one project and that project
 				// has the same name as the folder itself.  To the user, projects "feel like" folders already in the IDE, so we
 				// want to collapse them down where possible.
@@ -1301,8 +1329,6 @@ namespace UnrealBuildTool
 				}
 			}
 
-			bool bAlwaysIncludeEngineModules = false;
-			bool bAlwaysIncludeDotNetPrograms = false;
 			foreach (string CurArgument in Arguments)
 			{
 				if (CurArgument.StartsWith("-"))
@@ -1365,14 +1391,12 @@ namespace UnrealBuildTool
 
 							case "-GAME":
 								// Generates project files for a single game
-								bIncludeDotNetPrograms = false;
 								bIncludeEnginePrograms = false;
 								bGeneratingGameProjectFiles = true;
 								break;
 
 							case "-ENGINE":
-								// Forces engine modules and targets to be included in game-specific project files
-								bAlwaysIncludeEngineModules = true;
+								// -Engine is no longer needed as the engine module is always included now
 								break;
 
 							case "-NOCPP":
@@ -1412,7 +1436,7 @@ namespace UnrealBuildTool
 								break;
 
 							case "-DOTNET":
-								bAlwaysIncludeDotNetPrograms = true;
+								bIncludeDotNetPrograms = true;
 								break;
 
 							case "-NODOTNET":
@@ -1455,8 +1479,6 @@ namespace UnrealBuildTool
 				bIncludeShaderSource = true;
 				bIncludeTemplateFiles = false;
 				bIncludeConfigFiles = true;
-				bIncludeEnginePrograms = bAlwaysIncludeEngineModules;
-				bIncludeDotNetPrograms = bIncludeDotNetPrograms || bAlwaysIncludeDotNetPrograms;
 			}
 			else
 			{
@@ -1895,12 +1917,15 @@ namespace UnrealBuildTool
 				ProgramsFolder.ChildProjects.Add(UnrealBuildToolProject);
 			}
 
-			FileReference TestsProjectFileName = FileReference.Combine(ProgramsDirectory, "UnrealBuildTool.Tests", "UnrealBuildTool.Tests.csproj");
-			if (FileReference.Exists(TestsProjectFileName))
+			if (bIncludeEnginePrograms)
 			{
-				VCSharpProjectFile UbtTestsProject = new VCSharpProjectFile(TestsProjectFileName, Logger);
-				AddExistingProjectFile(UbtTestsProject, bNeedsAllPlatformAndConfigurations: true, bForceDevelopmentConfiguration: true);
-				ProgramsFolder.ChildProjects.Add(UbtTestsProject);
+				FileReference TestsProjectFileName = FileReference.Combine(ProgramsDirectory, "UnrealBuildTool.Tests", "UnrealBuildTool.Tests.csproj");
+				if (FileReference.Exists(TestsProjectFileName))
+				{
+					VCSharpProjectFile UbtTestsProject = new VCSharpProjectFile(TestsProjectFileName, Logger);
+					AddExistingProjectFile(UbtTestsProject, bNeedsAllPlatformAndConfigurations: true, bForceDevelopmentConfiguration: true);
+					ProgramsFolder.ChildProjects.Add(UbtTestsProject);
+				}
 			}
 		}
 
diff --git a/Engine/Source/Programs/UnrealBuildTool/System/PlatformExports.cs b/Engine/Source/Programs/UnrealBuildTool/System/PlatformExports.cs
--- a/Engine/Source/Programs/UnrealBuildTool/System/PlatformExports.cs
+++ b/Engine/Source/Programs/UnrealBuildTool/System/PlatformExports.cs
@@ -154,7 +154,7 @@ namespace UnrealBuildTool
 		public static bool Initialize(ILogger Logger)
 		{
 			// Read the XML configuration files
-			XmlConfig.ReadConfigFiles(null, Logger);
+			XmlConfig.ReadConfigFiles(null, null, Logger);
 
 			// Register all the platform classes
 			UEBuildPlatform.RegisterPlatforms(false, false, Logger);
diff --git a/Engine/Source/Programs/UnrealBuildTool/System/XmlConfig.cs b/Engine/Source/Programs/UnrealBuildTool/System/XmlConfig.cs
--- a/Engine/Source/Programs/UnrealBuildTool/System/XmlConfig.cs
+++ b/Engine/Source/Programs/UnrealBuildTool/System/XmlConfig.cs
@@ -68,8 +68,9 @@ namespace UnrealBuildTool
 		/// Initialize the config system with the given types
 		/// </summary>
 		/// <param name="OverrideCacheFile">Force use of the cached XML config without checking if it's valid (useful for remote builds)</param>
+		/// <param name="ProjectRootDirectory">Read XML configuration with a Project directory</param>
 		/// <param name="Logger">Logger for output</param>
-		public static void ReadConfigFiles(FileReference? OverrideCacheFile, ILogger Logger)
+		public static void ReadConfigFiles(FileReference? OverrideCacheFile, DirectoryReference? ProjectRootDirectory, ILogger Logger)
 		{
 			// Find all the configurable types
 			List<Type> ConfigTypes = FindConfigurableTypes();
@@ -88,22 +89,43 @@ namespace UnrealBuildTool
 			}
 			else
 			{
-				// Get the default cache file
-				CacheFile = FileReference.Combine(Unreal.EngineDirectory, "Intermediate", "Build", "XmlConfigCache.bin");
-				if(Unreal.IsEngineInstalled())
+				if (ProjectRootDirectory == null)
 				{
-					DirectoryReference? UserSettingsDir = Unreal.UserSettingDirectory;
-					if(UserSettingsDir != null)
+					// Get the default cache file
+					CacheFile = FileReference.Combine(Unreal.EngineDirectory, "Intermediate", "Build", "XmlConfigCache.bin");
+					if (Unreal.IsEngineInstalled())
 					{
-						CacheFile = FileReference.Combine(UserSettingsDir, "UnrealEngine", String.Format("XmlConfigCache-{0}.bin", Unreal.RootDirectory.FullName.Replace(":", "").Replace(Path.DirectorySeparatorChar, '+')));
+						DirectoryReference? UserSettingsDir = Unreal.UserSettingDirectory;
+						if(UserSettingsDir != null)
+						{
+							CacheFile = FileReference.Combine(UserSettingsDir, "UnrealEngine", String.Format("XmlConfigCache-{0}.bin", Unreal.RootDirectory.FullName.Replace(":", "").Replace(Path.DirectorySeparatorChar, '+')));
+						}
 					}
 				}
-
+				else
+				{
+					CacheFile = FileReference.Combine(ProjectRootDirectory, "Intermediate", "Build", "XmlConfigCache.bin");
+					Values = null;
+				}
+				
 				// Find all the input files
-				FileReference[] InputFileLocations = InputFiles.Select(x => x.Location).ToArray();
+				List<FileReference> TemporaryInputFileLocations = InputFiles.Select(x => x.Location).ToList();
+				
+				if (ProjectRootDirectory != null)
+				{
+					FileReference ProjectRootConfigLocation = FileReference.Combine(ProjectRootDirectory, "Saved", "UnrealBuildTool", "BuildConfiguration.xml");
+					if(!FileReference.Exists(ProjectRootConfigLocation))
+					{
+						CreateDefaultConfigFile(ProjectRootConfigLocation);
+					}
+
+					TemporaryInputFileLocations.Add(ProjectRootConfigLocation);
+				}
+				
+				FileReference[] InputFileLocations = TemporaryInputFileLocations.ToArray();
 
 				// Get the path to the schema
-				FileReference SchemaFile = GetSchemaLocation();
+				FileReference SchemaFile = GetSchemaLocation(ProjectRootDirectory);
 
 				// Try to read the existing cache from disk
 				XmlConfigData? CachedValues;
@@ -223,9 +245,18 @@ namespace UnrealBuildTool
 		/// <summary>
 		/// Find the location of the XML config schema
 		/// </summary>
+		/// <param name="ProjectRootDirecory">Optional project root directory</param>
 		/// <returns>The location of the schema file</returns>
-		public static FileReference GetSchemaLocation()
+		public static FileReference GetSchemaLocation(DirectoryReference? ProjectRootDirecory = null)
 		{
+			if (ProjectRootDirecory != null)
+			{
+				FileReference ProjectSchema = FileReference.Combine(ProjectRootDirecory, "Saved", "UnrealBuildTool", "BuildConfiguration.Schema.xsd");
+				if (FileReference.Exists(ProjectSchema))
+				{
+					return ProjectSchema;
+				}
+			}
 			return FileReference.Combine(Unreal.EngineDirectory, "Saved", "UnrealBuildTool", "BuildConfiguration.Schema.xsd");
 		}
 
diff --git a/Engine/Source/Programs/UnrealBuildTool/UnrealBuildTool.cs b/Engine/Source/Programs/UnrealBuildTool/UnrealBuildTool.cs
--- a/Engine/Source/Programs/UnrealBuildTool/UnrealBuildTool.cs
+++ b/Engine/Source/Programs/UnrealBuildTool/UnrealBuildTool.cs
@@ -538,7 +538,7 @@ namespace UnrealBuildTool
 						using(SingleInstanceMutex XmlConfigMutex = new SingleInstanceMutex(XmlConfigMutexName, true))
 						{
 							FileReference? XmlConfigCache = Arguments.GetFileReferenceOrDefault("-XmlConfigCache=", null);
-							XmlConfig.ReadConfigFiles(XmlConfigCache, Logger);
+							XmlConfig.ReadConfigFiles(XmlConfigCache, null, Logger);
 						}
 					}
 				
