Fix typo that flipped -frtti and -fno-rtti

diff --git a/Engine/Source/Programs/UnrealBuildTool/ProjectFiles/Xcode/XcodeProject.cs b/Engine/Source/Programs/UnrealBuildTool/ProjectFiles/Xcode/XcodeProject.cs
--- a/Engine/Source/Programs/UnrealBuildTool/ProjectFiles/Xcode/XcodeProject.cs
+++ b/Engine/Source/Programs/UnrealBuildTool/ProjectFiles/Xcode/XcodeProject.cs
@@ -126,7 +126,7 @@ namespace UnrealBuildTool.XcodeProjectXcconfig
 				ResponseFileContents.Append($" -include {PCHFile.FullName}");
 			}
 
-			ResponseFileContents.Append(bEnableRTTI ? " -fno-rtti" : " -frtti");
+			ResponseFileContents.Append(bEnableRTTI ? " -frtti" : " -fno-rtti");
 
 			DirectoryReference.CreateDirectory(ResponseFile.Directory);
 			FileReference.WriteAllText(ResponseFile, ResponseFileContents.ToString());
