UnrealBuildTool: Fix null warning

diff --git a/Engine/Source/Programs/UnrealBuildTool/Configuration/UEBuildTarget.cs b/Engine/Source/Programs/UnrealBuildTool/Configuration/UEBuildTarget.cs
--- a/Engine/Source/Programs/UnrealBuildTool/Configuration/UEBuildTarget.cs
+++ b/Engine/Source/Programs/UnrealBuildTool/Configuration/UEBuildTarget.cs
@@ -2493,7 +2493,7 @@ namespace UnrealBuildTool
 			}
 
 			// Make sure all the checked headers were valid
-			List<string> InvalidIncludeDirectiveMessages = Modules.Values.OfType<UEBuildModuleCPP>().Where(x => x.InvalidIncludeDirectiveMessages != null).SelectMany(x => x.InvalidIncludeDirectiveMessages).ToList();
+			List<string> InvalidIncludeDirectiveMessages = Modules.Values.OfType<UEBuildModuleCPP>().Where(x => x.InvalidIncludeDirectiveMessages != null).SelectMany(x => x.InvalidIncludeDirectiveMessages!).ToList();
 			if (InvalidIncludeDirectiveMessages.Count > 0)
 			{
 				foreach (string InvalidIncludeDirectiveMessage in InvalidIncludeDirectiveMessages)
