[UnrealBuildTool]
* Added experimental executor "Box"

diff --git a/Engine/Source/Programs/UnrealBuildTool/Configuration/BuildConfiguration.cs b/Engine/Source/Programs/UnrealBuildTool/Configuration/BuildConfiguration.cs
--- a/Engine/Source/Programs/UnrealBuildTool/Configuration/BuildConfiguration.cs
+++ b/Engine/Source/Programs/UnrealBuildTool/Configuration/BuildConfiguration.cs
@@ -42,6 +42,15 @@ namespace UnrealBuildTool
 		[XmlConfigFile]
 		public bool bAllowHybridExecutor = false;
 
+#if __BOXEXECUTOR_AVAILABLE__
+		/// <summary>
+		/// Whether the experimental box executor will be used.
+		/// </summary>
+		[XmlConfigFile]
+		[CommandLine("-Box", Value = "true")]
+		public bool bAllowBoxExecutor = false;
+#endif // #if __BOXEXECUTOR_AVAILABLE__
+
 		/// <summary>
 		/// Whether XGE may be used.
 		/// </summary>
diff --git a/Engine/Source/Programs/UnrealBuildTool/Executors/ParallelExecutor.cs b/Engine/Source/Programs/UnrealBuildTool/Executors/ParallelExecutor.cs
--- a/Engine/Source/Programs/UnrealBuildTool/Executors/ParallelExecutor.cs
+++ b/Engine/Source/Programs/UnrealBuildTool/Executors/ParallelExecutor.cs
@@ -48,7 +48,7 @@ namespace UnrealBuildTool
 		/// Default: BelowNormal or Normal for an Asymmetrical processor as BelowNormal can cause scheduling issues.
 		/// </summary>
 		[XmlConfigFile]
-		private static ProcessPriorityClass ProcessPriority = Utils.IsAsymmetricalProcessor() ? ProcessPriorityClass.Normal : ProcessPriorityClass.BelowNormal;
+		protected static ProcessPriorityClass ProcessPriority = Utils.IsAsymmetricalProcessor() ? ProcessPriorityClass.Normal : ProcessPriorityClass.BelowNormal;
 
 		/// <summary>
 		/// When enabled, will stop compiling targets after a compile error occurs.
@@ -84,13 +84,13 @@ namespace UnrealBuildTool
 		/// Whether to take into account the Action's weight when determining to do more work or not.
 		/// </summary>
 		[XmlConfigFile]
-		private static bool bUseActionWeights = false;
+		protected static bool bUseActionWeights = false;
 
 		/// <summary>
 		/// Whether to show CPU utilization after the work is complete.
 		/// </summary>
 		[XmlConfigFile]
-		private static bool bShowCPUUtilization = false;
+		protected static bool bShowCPUUtilization = false;
 
 		/// <summary>
 		/// Collapse non-error output lines
@@ -158,12 +158,13 @@ namespace UnrealBuildTool
 			public TimeSpan ProcessorTime { get; private set; }
 			public string? AdditionalDescription { get; protected set; } = null;
 
-			public ExecuteResults(List<string> LogLines, int ExitCode, TimeSpan ExecutionTime, TimeSpan ProcessorTime)
+			public ExecuteResults(List<string> LogLines, int ExitCode, TimeSpan ExecutionTime, TimeSpan ProcessorTime, string? AdditionalDescription = null)
 			{
 				this.LogLines = LogLines;
 				this.ExitCode = ExitCode;
 				this.ProcessorTime = ProcessorTime;
 				this.ExecutionTime = ExecutionTime;
+				this.AdditionalDescription = AdditionalDescription;
 			}
 			public ExecuteResults(List<string> LogLines, int ExitCode)
 			{
@@ -172,7 +173,7 @@ namespace UnrealBuildTool
 			}
 		}
 
-		enum ActionStatus : byte
+		internal enum ActionStatus : byte
 		{
 			Queued,
 			Running,
@@ -413,7 +414,7 @@ namespace UnrealBuildTool
 			}
 			if (!string.IsNullOrEmpty(AdditionalDescription))
             {
-				Description = $"{AdditionalDescription} {Description}";
+				Description = $"{Description} {AdditionalDescription}";
 			}
 
 			lock (ProgressWriter)
diff --git a/Engine/Source/Programs/UnrealBuildTool/System/ActionGraph.cs b/Engine/Source/Programs/UnrealBuildTool/System/ActionGraph.cs
--- a/Engine/Source/Programs/UnrealBuildTool/System/ActionGraph.cs
+++ b/Engine/Source/Programs/UnrealBuildTool/System/ActionGraph.cs
@@ -353,6 +353,13 @@ namespace UnrealBuildTool
 		/// </summary>
 		private static ActionExecutor SelectExecutor(BuildConfiguration BuildConfiguration, int ActionCount, List<TargetDescriptor> TargetDescriptors, ILogger Logger)
 		{
+#if __BOXEXECUTOR_AVAILABLE__
+			if (BuildConfiguration.bAllowBoxExecutor)
+			{
+				return new BoxExecutor(BuildConfiguration.MaxParallelActions, BuildConfiguration.bAllCores, BuildConfiguration.bCompactOutput, Logger);
+			}
+#endif // #if __BOXEXECUTOR_AVAILABLE__
+
 			if (ActionCount > ParallelExecutor.GetDefaultNumParallelProcesses(BuildConfiguration.MaxParallelActions, BuildConfiguration.bAllCores, Logger))
 			{
 				if (BuildConfiguration.bAllowHybridExecutor && HybridExecutor.IsAvailable(Logger))
diff --git a/Engine/Source/Programs/UnrealBuildTool/UnrealBuildTool.csproj b/Engine/Source/Programs/UnrealBuildTool/UnrealBuildTool.csproj
--- a/Engine/Source/Programs/UnrealBuildTool/UnrealBuildTool.csproj
+++ b/Engine/Source/Programs/UnrealBuildTool/UnrealBuildTool.csproj
@@ -71,6 +71,29 @@
     </When>
   </Choose>
 
+  <Choose>
+    <When Condition="'$([System.Runtime.InteropServices.RuntimeInformation]::IsOSPlatform($([System.Runtime.InteropServices.OSPlatform]::Windows)))' == 'true'
+      AND Exists('..\..\..\Restricted\NotForLicensees\Source\Programs\UnrealBuildTool\Executors\Experimental\BoxExecutor.cs')
+      AND Exists('..\..\..\Restricted\NotForLicensees\Binaries\Win64\Box\Box.dll')
+      ">
+      <ItemGroup>
+        <Reference Include="..\..\..\Restricted\NotForLicensees\Binaries\Win64\Box\Box.dll" />
+        <Content Include="..\..\..\Restricted\NotForLicensees\Binaries\Win64\Box\BoxNative.dll">
+          <CopyToOutputDirectory>PreserveNewest</CopyToOutputDirectory>
+        </Content>
+        <Content Include="..\..\..\Restricted\NotForLicensees\Binaries\Win64\Box\BoxDetours.dll">
+          <CopyToOutputDirectory>PreserveNewest</CopyToOutputDirectory>
+        </Content>
+        <Content Include="..\..\..\Restricted\NotForLicensees\Binaries\Win64\Box\zstd.dll">
+          <CopyToOutputDirectory>PreserveNewest</CopyToOutputDirectory>
+        </Content>
+      </ItemGroup>
+      <PropertyGroup>
+        <DefineConstants>$(DefineConstants);__BOXEXECUTOR_AVAILABLE__</DefineConstants>
+      </PropertyGroup>
+    </When>
+  </Choose>
+
   <ItemGroup>
     <Compile Include="..\Shared\MetaData.cs">
       <Link>Properties\MetaData.cs</Link>
