[UnrealBuildTool]
* Backed out csproj file to temporarily fix UAT

diff --git a/Engine/Source/Programs/UnrealBuildTool/UnrealBuildTool.csproj b/Engine/Source/Programs/UnrealBuildTool/UnrealBuildTool.csproj
--- a/Engine/Source/Programs/UnrealBuildTool/UnrealBuildTool.csproj
+++ b/Engine/Source/Programs/UnrealBuildTool/UnrealBuildTool.csproj
@@ -71,29 +71,6 @@
     </When>
   </Choose>
 
-  <Choose>
-    <When Condition="'$([System.Runtime.InteropServices.RuntimeInformation]::IsOSPlatform($([System.Runtime.InteropServices.OSPlatform]::Windows)))' == 'true'
-      AND Exists('..\..\..\Restricted\NotForLicensees\Source\Programs\UnrealBuildTool\Executors\Experimental\BoxExecutor.cs')
-      AND Exists('..\..\..\Restricted\NotForLicensees\Binaries\Win64\Box\Box.dll')
-      ">
-      <ItemGroup>
-        <Reference Include="..\..\..\Restricted\NotForLicensees\Binaries\Win64\Box\Box.dll" />
-        <Content Include="..\..\..\Restricted\NotForLicensees\Binaries\Win64\Box\BoxNative.dll">
-          <CopyToOutputDirectory>PreserveNewest</CopyToOutputDirectory>
-        </Content>
-        <Content Include="..\..\..\Restricted\NotForLicensees\Binaries\Win64\Box\BoxDetours.dll">
-          <CopyToOutputDirectory>PreserveNewest</CopyToOutputDirectory>
-        </Content>
-        <Content Include="..\..\..\Restricted\NotForLicensees\Binaries\Win64\Box\zstd.dll">
-          <CopyToOutputDirectory>PreserveNewest</CopyToOutputDirectory>
-        </Content>
-      </ItemGroup>
-      <PropertyGroup>
-        <DefineConstants>$(DefineConstants);__BOXEXECUTOR_AVAILABLE__</DefineConstants>
-      </PropertyGroup>
-    </When>
-  </Choose>
-
   <ItemGroup>
     <Compile Include="..\Shared\MetaData.cs">
       <Link>Properties\MetaData.cs</Link>
