Disable Xcode header map and a specific C++11 warning. Now almost all files can be compiled for indexing, only 10 files still has errors

diff --git a/Engine/Source/Programs/UnrealBuildTool/ProjectFiles/Xcode/XcodeProject.cs b/Engine/Source/Programs/UnrealBuildTool/ProjectFiles/Xcode/XcodeProject.cs
--- a/Engine/Source/Programs/UnrealBuildTool/ProjectFiles/Xcode/XcodeProject.cs
+++ b/Engine/Source/Programs/UnrealBuildTool/ProjectFiles/Xcode/XcodeProject.cs
@@ -1699,6 +1699,8 @@ namespace UnrealBuildTool.XcodeProjectXcconfig
 			Xcconfig.AppendLine("GCC_OPTIMIZATION_LEVEL = 0");
 			Xcconfig.AppendLine($"PRODUCT_NAME = {Name}");
 			Xcconfig.AppendLine("SYMROOT = build");
+			Xcconfig.AppendLine("USE_HEADERMAP = NO");
+			Xcconfig.AppendLine("WARNING_CFLAGS = -Wno-c++11-narrowing");
 			Xcconfig.Write();
 		}
 	}
