Fixed UBT log event level check that resulted is massive spam to the logger.

diff --git a/Engine/Source/Programs/UnrealBuildTool/Configuration/UEBuildModule.cs b/Engine/Source/Programs/UnrealBuildTool/Configuration/UEBuildModule.cs
--- a/Engine/Source/Programs/UnrealBuildTool/Configuration/UEBuildModule.cs
+++ b/Engine/Source/Programs/UnrealBuildTool/Configuration/UEBuildModule.cs
@@ -1079,7 +1079,7 @@ namespace UnrealBuildTool
 			if(PrivateIncludePathModules == null)
 			{
 				// Log dependencies if required
-				if ( Log.OutputLevel >= LogEventType.VeryVerbose)
+				if (Logger.IsEnabled((LogLevel)LogEventType.VeryVerbose))
 				{
 					Logger.LogTrace("Module {0} dependencies:", this.Name);
 					LogDependencyNameList("Public:", Rules.PublicDependencyModuleNames, Logger);
