UnrealBuildTool: Protect against exception on shutdown

#rnx

diff --git a/Engine/Source/Programs/UnrealBuildTool/System/FileMetadataPrefetch.cs b/Engine/Source/Programs/UnrealBuildTool/System/FileMetadataPrefetch.cs
--- a/Engine/Source/Programs/UnrealBuildTool/System/FileMetadataPrefetch.cs
+++ b/Engine/Source/Programs/UnrealBuildTool/System/FileMetadataPrefetch.cs
@@ -158,7 +158,12 @@ namespace UnrealBuildTool
 		/// <param name="Directory">The directory which may contain plugin directories</param>
 		static void ScanPluginFolder(DirectoryItem Directory)
 		{
-			foreach(DirectoryItem SubDirectory in Directory.EnumerateDirectories())
+			if (CancelToken.IsCancellationRequested || Directory.TryGetFile(".ubtignore", out FileItem? _))
+			{
+				return;
+			}
+
+			foreach (DirectoryItem SubDirectory in Directory.EnumerateDirectories())
 			{
 				if (SubDirectory.EnumerateFiles().Any((fi) => fi.HasExtension(".uplugin")))
 				{
@@ -177,7 +182,7 @@ namespace UnrealBuildTool
 		/// <param name="Directory">Root of the directory tree</param>
 		static void ScanDirectoryTree(DirectoryItem Directory)
 		{
-			if (Directory.TryGetFile(".ubtignore", out FileItem? OutFile))
+			if (CancelToken.IsCancellationRequested || Directory.TryGetFile(".ubtignore", out FileItem? _))
 			{
 				return;
 			}
diff --git a/Engine/Source/Programs/UnrealBuildTool/UnrealBuildTool.cs b/Engine/Source/Programs/UnrealBuildTool/UnrealBuildTool.cs
--- a/Engine/Source/Programs/UnrealBuildTool/UnrealBuildTool.cs
+++ b/Engine/Source/Programs/UnrealBuildTool/UnrealBuildTool.cs
@@ -691,7 +691,13 @@ namespace UnrealBuildTool
 				// Cancel the prefetcher
 				using (GlobalTracer.Instance.BuildSpan("FileMetadataPrefetch.Stop()").StartActive())
 				{
-					FileMetadataPrefetch.Stop();
+					try
+					{
+						FileMetadataPrefetch.Stop();
+					}
+					catch
+					{
+					}
 				}
 
 				// Uncomment this to output a file that contains all files that UBT has scanned.
