PR #10269: Add TargetRules.bWarnAboutMonolithicHeadersIncluded to control UE_IS_? (Contributed by [REDACTED_PERSON])

diff --git a/Engine/Source/Programs/UnrealBuildTool/Configuration/TargetRules.cs b/Engine/Source/Programs/UnrealBuildTool/Configuration/TargetRules.cs
--- a/Engine/Source/Programs/UnrealBuildTool/Configuration/TargetRules.cs
+++ b/Engine/Source/Programs/UnrealBuildTool/Configuration/TargetRules.cs
@@ -1342,6 +1342,11 @@ namespace UnrealBuildTool
 		/// </summary>
 		public bool bEnforceIWYU = true;
 
+		/// <summary>
+		/// Emit a warning when an old-style monolithic header is included while compiling this target.
+		/// </summary>
+		public bool bWarnAboutMonolithicHeadersIncluded = false;
+
 		/// <summary>
 		/// Whether the final executable should export symbols.
 		/// </summary>
@@ -3312,6 +3317,11 @@ namespace UnrealBuildTool
 			get { return Inner.bEnforceIWYU; }
 		}
 
+		public bool bWarnAboutMonolithicHeadersIncluded
+		{
+			get { return Inner.bWarnAboutMonolithicHeadersIncluded; }
+		}
+
 		public bool bHasExports
 		{
 			get { return Inner.bHasExports; }
diff --git a/Engine/Source/Programs/UnrealBuildTool/Configuration/UEBuildModuleCPP.cs b/Engine/Source/Programs/UnrealBuildTool/Configuration/UEBuildModuleCPP.cs
--- a/Engine/Source/Programs/UnrealBuildTool/Configuration/UEBuildModuleCPP.cs
+++ b/Engine/Source/Programs/UnrealBuildTool/Configuration/UEBuildModuleCPP.cs
@@ -1730,13 +1730,13 @@ namespace UnrealBuildTool
 			}
 
 			// Set the macro used to check whether monolithic headers can be used
-			if (Rules.bTreatAsEngineModule && (Rules.IWYUSupport == IWYUSupport.None || !Target.bEnforceIWYU))
+			if ((Rules.bTreatAsEngineModule || Target.bWarnAboutMonolithicHeadersIncluded) && (Rules.IWYUSupport == IWYUSupport.None || !Target.bEnforceIWYU))
 			{
 				Result.Definitions.Add("SUPPRESS_MONOLITHIC_HEADER_WARNINGS=1");
 			}
 
 			// Add a macro for when we're compiling an engine module, to enable additional compiler diagnostics through code.
-			if (Rules.bTreatAsEngineModule)
+			if (Rules.bTreatAsEngineModule || Target.bWarnAboutMonolithicHeadersIncluded)
 			{
 				Result.Definitions.Add("UE_IS_ENGINE_MODULE=1");
 			}
