- Added a parameter to GenerateProjectFiles() to control if it generates the QueryTargets data for the editor (only the first project format needs to do this)
- Disabled it with IOS stub project gneeration that doesn't need the data (it can be slow)

diff --git a/Engine/Source/Programs/UnrealBuildTool/Modes/GenerateProjectFilesMode.cs b/Engine/Source/Programs/UnrealBuildTool/Modes/GenerateProjectFilesMode.cs
--- a/Engine/Source/Programs/UnrealBuildTool/Modes/GenerateProjectFilesMode.cs
+++ b/Engine/Source/Programs/UnrealBuildTool/Modes/GenerateProjectFilesMode.cs
@@ -240,6 +240,8 @@ namespace UnrealBuildTool
 
 			// Now generate project files
 			ProjectFileGenerator.bGenerateProjectFiles = true;
+			// perform anything that only needs to happen one time, in the first project genereator
+			bool bPerformOneTimeOperations = true;
 			foreach(KeyValuePair< ProjectFileFormat, ProjectFileGenerator> Pair in Generators)
 			{
 				Logger.LogInformation("");
@@ -247,13 +249,16 @@ namespace UnrealBuildTool
 
 				ProjectFileGenerator.Current = Pair.Value;
 				Arguments.ApplyTo(Pair.Value);
-				bool bGenerateSuccess = Pair.Value.GenerateProjectFiles(PlatformProjectGenerators, Arguments.GetRawArray(), Logger);
+				bool bGenerateSuccess = Pair.Value.GenerateProjectFiles(PlatformProjectGenerators, Arguments.GetRawArray(), bCacheDataForEditor: bPerformOneTimeOperations, Logger);
 				ProjectFileGenerator.Current = null;
 
 				if (!bGenerateSuccess)
 				{
 					return (int)CompilationResult.OtherCompilationError;
 				}
+
+				// any further generators can skip one-time operations
+				bPerformOneTimeOperations = false;
 			}
 			return (int)CompilationResult.Succeeded;
 		}
diff --git a/Engine/Source/Programs/UnrealBuildTool/Platform/IOS/IOSToolChain.cs b/Engine/Source/Programs/UnrealBuildTool/Platform/IOS/IOSToolChain.cs
--- a/Engine/Source/Programs/UnrealBuildTool/Platform/IOS/IOSToolChain.cs
+++ b/Engine/Source/Programs/UnrealBuildTool/Platform/IOS/IOSToolChain.cs
@@ -1414,7 +1414,8 @@ namespace UnrealBuildTool
 				PlatformProjectGenerators.RegisterPlatformProjectGenerator(UnrealTargetPlatform.TVOS, new TVOSProjectGenerator(CmdLine, Logger), Logger);
 
 				XcodeProjectFileGenerator Generator = new XcodeProjectFileGenerator(ProjectFile, CmdLine);
-				bool bSucces = Generator.GenerateProjectFiles(PlatformProjectGenerators, Arguments, Logger);
+				// don't need the editor data since these are stub projects
+				bool bSucces = Generator.GenerateProjectFiles(PlatformProjectGenerators, Arguments, bCacheDataForEditor: false, Logger);
 				XcodeProjectFile = Generator.XCWorkspace;
 				return bSucces;
 			}
diff --git a/Engine/Source/Programs/UnrealBuildTool/ProjectFiles/ProjectFileGenerator.cs b/Engine/Source/Programs/UnrealBuildTool/ProjectFiles/ProjectFileGenerator.cs
--- a/Engine/Source/Programs/UnrealBuildTool/ProjectFiles/ProjectFileGenerator.cs
+++ b/Engine/Source/Programs/UnrealBuildTool/ProjectFiles/ProjectFileGenerator.cs
@@ -829,8 +829,9 @@ namespace UnrealBuildTool
 		/// </summary>
 		/// <param name="PlatformProjectGenerators">The registered platform project generators</param>
 		/// <param name="Arguments">Command-line arguments</param>
+		/// <param name="bCacheDataForEditor">If true, write out target data for the editor</param>
 		/// <param name="Logger">Logger for output</param>
-		public virtual bool GenerateProjectFiles(PlatformProjectGeneratorCollection PlatformProjectGenerators, String[] Arguments, ILogger Logger)
+		public virtual bool GenerateProjectFiles(PlatformProjectGeneratorCollection PlatformProjectGenerators, String[] Arguments, bool bCacheDataForEditor, ILogger Logger)
 		{
 			bool bSuccess = true;
 
@@ -1257,6 +1258,12 @@ namespace UnrealBuildTool
 				// Write the project files
 				WriteProjectFiles(PlatformProjectGenerators, Logger);
 				Logger.LogDebug("Project generation complete ({GeneratedProjectFilesCount} generated, {OtherProjectFilesCount} imported)", GeneratedProjectFiles.Count, OtherProjectFiles.Count);
+			}
+
+			if (bCacheDataForEditor)
+			{
+				Logger.LogInformation("");
+				Logger.LogInformation("Generating QueryTargets data for editor...");
 
 				// Generate all the target info files for the editor
 				foreach (FileReference ProjectFile in AllGameProjects)
diff --git a/Engine/Source/Programs/UnrealBuildTool/ProjectFiles/Rider/RiderProjectFileGenerator.cs b/Engine/Source/Programs/UnrealBuildTool/ProjectFiles/Rider/RiderProjectFileGenerator.cs
--- a/Engine/Source/Programs/UnrealBuildTool/ProjectFiles/Rider/RiderProjectFileGenerator.cs
+++ b/Engine/Source/Programs/UnrealBuildTool/ProjectFiles/Rider/RiderProjectFileGenerator.cs
@@ -424,7 +424,7 @@ namespace UnrealBuildTool
 		}
 		
 		public override bool GenerateProjectFiles(PlatformProjectGeneratorCollection PlatformProjectGenerators,
-			String[] arguments, ILogger Logger)
+			String[] arguments, bool bCacheDataForEditor, ILogger Logger)
 		{
 			ConfigureProjectFileGeneration();
 
