- Updating making Mac builds with modern xcode to stage files into the .app via xcode
- Mac and IOS both now run a script in Xcode to rsync from StagedBuilds into the .app
- Mac now generates and runs the xcode project after building (PostBuildSync like IOS) to codesign and pull in staged data, using a new ApplePostBuildSync class to share with IOS
- Added a SingleTarget mode to XcodeProjectGenerator to make a minimal project for PostBuildSync

diff --git a/Engine/Source/Programs/AutomationTool/IOS/IOSPlatform.Automation.cs b/Engine/Source/Programs/AutomationTool/IOS/IOSPlatform.Automation.cs
--- a/Engine/Source/Programs/AutomationTool/IOS/IOSPlatform.Automation.cs
+++ b/Engine/Source/Programs/AutomationTool/IOS/IOSPlatform.Automation.cs
@@ -1103,7 +1103,7 @@ public class IOSPlatform : ApplePlatform
 		{
 			// we need an xcode project to continue
 			DirectoryReference GeneratedProject = null;
-			IOSExports.GenerateRunOnlyXcodeProject(RawProjectPath, TargetPlatformType, false, Logger, out GeneratedProject);
+			IOSExports.GenerateRunOnlyXcodeProject(RawProjectPath, TargetPlatformType, SchemeName ?? GameName, false, Logger, out GeneratedProject);
 
 			if (GeneratedProject == null || !DirectoryReference.Exists(GeneratedProject))
 			{
diff --git a/Engine/Source/Programs/AutomationTool/Mac/MacPlatform.Automation.cs b/Engine/Source/Programs/AutomationTool/Mac/MacPlatform.Automation.cs
--- a/Engine/Source/Programs/AutomationTool/Mac/MacPlatform.Automation.cs
+++ b/Engine/Source/Programs/AutomationTool/Mac/MacPlatform.Automation.cs
@@ -129,6 +129,10 @@ public class MacPlatform : ApplePlatform
 
 	public override void GetFilesToDeployOrStage(ProjectParams Params, DeploymentContext SC)
 	{
+		ConfigHierarchy Ini = ConfigCache.ReadHierarchy(ConfigHierarchyType.Engine, Params.RawProjectPath.Directory!, UnrealTargetPlatform.IOS);
+		bool bUseModernXcode;
+		Ini.TryGetValue("XcodeConfiguration", "bUseModernXcode", out bUseModernXcode);
+
 		// Stage all the build products
 		foreach (StageTarget Target in SC.StageTargets)
 		{
@@ -173,8 +177,8 @@ public class MacPlatform : ApplePlatform
 			SC.StageFile(StagedFileType.NonUFS, SplashImage);
 		}
 
-		// Stage the bootstrap executable
-		if (!Params.NoBootstrapExe)
+		// Stage the bootstrap executable (modern doesn't need it with full .apps)
+		if (!Params.NoBootstrapExe && !bUseModernXcode)
 		{
 			foreach (StageTarget Target in SC.StageTargets)
 			{
@@ -401,7 +405,7 @@ public class MacPlatform : ApplePlatform
 		if (Params.CreateAppBundle)
 		{
 			string ExeName = SC.StageExecutables[0];
-			string BundlePath = SC.IsCodeBasedProject ? CombinePaths(SC.ArchiveDirectory.FullName, ExeName + ".app") : BundlePath = CombinePaths(SC.ArchiveDirectory.FullName, SC.ShortProjectName + ".app");
+			string BundlePath = SC.IsCodeBasedProject ? CombinePaths(SC.ArchiveDirectory.FullName, ExeName + ".app") : CombinePaths(SC.ArchiveDirectory.FullName, SC.ShortProjectName + ".app");
 
 			if (SC.bIsCombiningMultiplePlatforms)
 			{
diff --git a/Engine/Source/Programs/UnrealBuildTool/Platform/IOS/IOSExports.cs b/Engine/Source/Programs/UnrealBuildTool/Platform/IOS/IOSExports.cs
--- a/Engine/Source/Programs/UnrealBuildTool/Platform/IOS/IOSExports.cs
+++ b/Engine/Source/Programs/UnrealBuildTool/Platform/IOS/IOSExports.cs
@@ -243,21 +243,25 @@ namespace UnrealBuildTool
 		/// </summary>
 		/// <param name="UProjectFile">Location of .uproject file (or null for the engine project</param>
 		/// <param name="Platform">The platform to generate a project for</param>
+		/// <param name="TargetName">The name of the target being built, so we can generate a more minimal project</param>
 		/// <param name="bForDistribution">True if this is making a bild for uploading to app store</param>
 		/// <param name="Logger">Logging object</param>
 		/// <param name="GeneratedProjectFile">Returns the .xcworkspace that was made</param>
-		public static void GenerateRunOnlyXcodeProject(FileReference? UProjectFile, UnrealTargetPlatform Platform, bool bForDistribution, ILogger Logger, out DirectoryReference? GeneratedProjectFile)
+		public static void GenerateRunOnlyXcodeProject(FileReference? UProjectFile, UnrealTargetPlatform Platform, string TargetName, bool bForDistribution, ILogger Logger, out DirectoryReference? GeneratedProjectFile)
 		{
 			List<string> Options = new()
 			{
 				$"-platforms={Platform}",
-				$"-{Platform}DeployOnly",
-				"-NoIntellisens",
-				"-IngnoreJunk",
+				"-DeployOnly",
+				"-NoIntellisense",
+				"-NoCPP",
+				"-NoDotNet",
+				"-IgnoreJunk",
 				bForDistribution ? "-distribution" : "-development",
 				"-IncludeTempTargets",
-				"-projectfileformat = XCode",
+				"-projectfileformat=XCode",
 				"-automated",
+				$"-singletarget={TargetName}"
 			};
 
 			if (UProjectFile == null || UProjectFile.IsUnderDirectory(Unreal.EngineDirectory))
@@ -316,7 +320,7 @@ namespace UnrealBuildTool
 				$"-workspace \"{XcodeProject.FullName}\"",
 				$"-scheme \"{SchemeName}\"",
 				$"-configuration \"{Configuration}\"",
-				$"-destination generic/platform=" + (Platform == UnrealTargetPlatform.TVOS ? "tvOS" : "iOS"),
+				$"-destination generic/platform=" + (Platform == UnrealTargetPlatform.TVOS ? "tvOS" : Platform == UnrealTargetPlatform.Mac ? "macOS" : "iOS"),
 				//$"-sdk {SDKName}",
 			};
 
@@ -371,8 +375,10 @@ namespace UnrealBuildTool
 				}
 			}
 
-			int ReturnCode = Utils.RunLocalProcessAndLogOutput("/usr/bin/env", string.Join(" ", Arguments), Logger);
-//			DirectoryReference.Delete(XcodeProject, true);
+			int ReturnCode;
+			string Output = Utils.RunLocalProcessAndReturnStdOut("/usr/bin/env", string.Join(" ", Arguments), null, out ReturnCode);
+			Logger.LogDebug(Output);
+			DirectoryReference.Delete(XcodeProject, true);
 			return ReturnCode;
 		}
 
diff --git a/Engine/Source/Programs/UnrealBuildTool/Platform/IOS/IOSToolChain.cs b/Engine/Source/Programs/UnrealBuildTool/Platform/IOS/IOSToolChain.cs
--- a/Engine/Source/Programs/UnrealBuildTool/Platform/IOS/IOSToolChain.cs
+++ b/Engine/Source/Programs/UnrealBuildTool/Platform/IOS/IOSToolChain.cs
@@ -1170,55 +1170,53 @@ namespace UnrealBuildTool
 
 			if (bUseModernXcode)
 			{
-				// update version xcconfig file so xcode can update the .plist with our incrementing build version
-				OutputFiles.Add(UpdateVersionFile(BinaryLinkEnvironment, Executable, Graph));
+				// base class AppleToolchain makes the shared postbuild sync object, so we are done
+				return OutputFiles;
 			}
-			else
+
+			// strip the debug info from the executable if needed. creates a dummy output file for the action graph to track that it's out of date.
+			if (Target.IOSPlatform.bStripSymbols || (Target.Configuration == UnrealTargetConfiguration.Shipping))
 			{
-				// strip the debug info from the executable if needed. creates a dummy output file for the action graph to track that it's out of date.
-				if (Target.IOSPlatform.bStripSymbols || (Target.Configuration == UnrealTargetConfiguration.Shipping))
-				{
-					FileItem StripCompleteFile = FileItem.GetItemByFileReference(FileReference.Combine(BinaryLinkEnvironment.IntermediateDirectory!, Executable.Location.GetFileName() + ".stripped"));
-
-					// If building a framework we can only strip local symbols, need to leave global in place
-					string StripArguments = BinaryLinkEnvironment.bIsBuildingDLL ? "-x" : "";
-
-					Action StripAction = Graph.CreateAction(ActionType.CreateAppBundle);
-					StripAction.WorkingDirectory = GetMacDevSrcRoot();
-					StripAction.CommandPath = BuildHostPlatform.Current.Shell;
-					StripAction.CommandArguments = String.Format("-c \"\\\"{0}strip\\\" {1} \\\"{2}\\\" && touch \\\"{3}\\\"\"", Settings.Value.ToolchainDir, StripArguments, Executable.Location, StripCompleteFile);
-					StripAction.PrerequisiteItems.Add(Executable);
-					StripAction.PrerequisiteItems.UnionWith(OutputFiles);
-					StripAction.ProducedItems.Add(StripCompleteFile);
-					StripAction.StatusDescription = String.Format("Stripping symbols from {0}", Executable.AbsolutePath);
-					StripAction.bCanExecuteRemotely = false;
-
-					OutputFiles.Add(StripCompleteFile);
-				}
+				FileItem StripCompleteFile = FileItem.GetItemByFileReference(FileReference.Combine(BinaryLinkEnvironment.IntermediateDirectory!, Executable.Location.GetFileName() + ".stripped"));
 
-				if (!BinaryLinkEnvironment.bIsBuildingDLL)
-				{
-					// generate the asset catalog
-					bool bUserImagesExist = false;
-					DirectoryReference ResourcesDir = GenerateAssetCatalog(ProjectFile, BinaryLinkEnvironment.Platform, ref bUserImagesExist);
-
-					// Get the output location for the asset catalog
-					FileItem AssetCatalogFile = FileItem.GetItemByFileReference(GetAssetCatalogFile(BinaryLinkEnvironment.Platform, Executable.Location));
-
-					// Make the compile action
-					Action CompileAssetAction = Graph.CreateAction(ActionType.CreateAppBundle);
-					CompileAssetAction.WorkingDirectory = GetMacDevSrcRoot();
-					CompileAssetAction.CommandPath = new FileReference("/usr/bin/xcrun");
-					CompileAssetAction.CommandArguments = GetAssetCatalogArgs(BinaryLinkEnvironment.Platform, ResourcesDir.FullName, Path.GetDirectoryName(AssetCatalogFile.AbsolutePath)!);
-					CompileAssetAction.PrerequisiteItems.Add(Executable);
-					CompileAssetAction.ProducedItems.Add(AssetCatalogFile);
-					CompileAssetAction.DeleteItems.Add(AssetCatalogFile);
-					CompileAssetAction.StatusDescription = CompileAssetAction.CommandArguments;// string.Format("Generating debug info for {0}", Path.GetFileName(Executable.AbsolutePath));
-					CompileAssetAction.bCanExecuteRemotely = false;
-
-					// Add it to the output files so it's always built
-					OutputFiles.Add(AssetCatalogFile);
-				}
+				// If building a framework we can only strip local symbols, need to leave global in place
+				string StripArguments = BinaryLinkEnvironment.bIsBuildingDLL ? "-x" : "";
+
+				Action StripAction = Graph.CreateAction(ActionType.CreateAppBundle);
+				StripAction.WorkingDirectory = GetMacDevSrcRoot();
+				StripAction.CommandPath = BuildHostPlatform.Current.Shell;
+				StripAction.CommandArguments = String.Format("-c \"\\\"{0}strip\\\" {1} \\\"{2}\\\" && touch \\\"{3}\\\"\"", Settings.Value.ToolchainDir, StripArguments, Executable.Location, StripCompleteFile);
+				StripAction.PrerequisiteItems.Add(Executable);
+				StripAction.PrerequisiteItems.UnionWith(OutputFiles);
+				StripAction.ProducedItems.Add(StripCompleteFile);
+				StripAction.StatusDescription = String.Format("Stripping symbols from {0}", Executable.AbsolutePath);
+				StripAction.bCanExecuteRemotely = false;
+
+				OutputFiles.Add(StripCompleteFile);
+			}
+
+			if (!BinaryLinkEnvironment.bIsBuildingDLL)
+			{
+				// generate the asset catalog
+				bool bUserImagesExist = false;
+				DirectoryReference ResourcesDir = GenerateAssetCatalog(ProjectFile, BinaryLinkEnvironment.Platform, ref bUserImagesExist);
+
+				// Get the output location for the asset catalog
+				FileItem AssetCatalogFile = FileItem.GetItemByFileReference(GetAssetCatalogFile(BinaryLinkEnvironment.Platform, Executable.Location));
+
+				// Make the compile action
+				Action CompileAssetAction = Graph.CreateAction(ActionType.CreateAppBundle);
+				CompileAssetAction.WorkingDirectory = GetMacDevSrcRoot();
+				CompileAssetAction.CommandPath = new FileReference("/usr/bin/xcrun");
+				CompileAssetAction.CommandArguments = GetAssetCatalogArgs(BinaryLinkEnvironment.Platform, ResourcesDir.FullName, Path.GetDirectoryName(AssetCatalogFile.AbsolutePath)!);
+				CompileAssetAction.PrerequisiteItems.Add(Executable);
+				CompileAssetAction.ProducedItems.Add(AssetCatalogFile);
+				CompileAssetAction.DeleteItems.Add(AssetCatalogFile);
+				CompileAssetAction.StatusDescription = CompileAssetAction.CommandArguments;// string.Format("Generating debug info for {0}", Path.GetFileName(Executable.AbsolutePath));
+				CompileAssetAction.bCanExecuteRemotely = false;
+
+				// Add it to the output files so it's always built
+				OutputFiles.Add(AssetCatalogFile);
 			}
 
 			// Generate the app bundle
@@ -1231,40 +1229,37 @@ namespace UnrealBuildTool
 
 				Dictionary<string, DirectoryReference> FrameworkNameToSourceDir = new Dictionary<string, DirectoryReference>();
 
-				if (!bUseModernXcode)
+				FileReference StagedExecutablePath = GetStagedExecutablePath(Executable.Location, Target.Name);
+				DirectoryReference BundleDirectory = Target.bShouldCompileAsDLL ? Executable.Directory.Location : StagedExecutablePath.Directory;
+				foreach (UEBuildFramework Framework in BinaryLinkEnvironment.AdditionalFrameworks)
 				{
-					FileReference StagedExecutablePath = GetStagedExecutablePath(Executable.Location, Target.Name);
-					DirectoryReference BundleDirectory = Target.bShouldCompileAsDLL ? Executable.Directory.Location : StagedExecutablePath.Directory;
-					foreach (UEBuildFramework Framework in BinaryLinkEnvironment.AdditionalFrameworks)
+					DirectoryReference? FrameworkDirectory = Framework.GetFrameworkDirectory(BinaryLinkEnvironment.Platform, BinaryLinkEnvironment.Architecture, Logger);
+					if (FrameworkDirectory != null)
 					{
-						DirectoryReference? FrameworkDirectory = Framework.GetFrameworkDirectory(BinaryLinkEnvironment.Platform, BinaryLinkEnvironment.Architecture, Logger);
-						if (FrameworkDirectory != null)
+						if (!String.IsNullOrEmpty(Framework.CopyBundledAssets))
 						{
-							if (!String.IsNullOrEmpty(Framework.CopyBundledAssets))
-							{
-								// For now, this is hard coded, but we need to loop over all modules, and copy bundled assets that need it
-								DirectoryReference LocalSource = DirectoryReference.Combine(FrameworkDirectory, Framework.CopyBundledAssets);
-								string BundleName = Framework.CopyBundledAssets.Substring(Framework.CopyBundledAssets.LastIndexOf('/') + 1);
-								FrameworkNameToSourceDir[BundleName] = LocalSource;
-							}
+							// For now, this is hard coded, but we need to loop over all modules, and copy bundled assets that need it
+							DirectoryReference LocalSource = DirectoryReference.Combine(FrameworkDirectory, Framework.CopyBundledAssets);
+							string BundleName = Framework.CopyBundledAssets.Substring(Framework.CopyBundledAssets.LastIndexOf('/') + 1);
+							FrameworkNameToSourceDir[BundleName] = LocalSource;
+						}
 
-							if (Framework.bCopyFramework)
+						if (Framework.bCopyFramework)
+						{
+							string FrameworkDir = FrameworkDirectory.FullName;
+							if (FrameworkDir.EndsWith(".framework"))
 							{
-								string FrameworkDir = FrameworkDirectory.FullName;
-								if (FrameworkDir.EndsWith(".framework"))
-								{
-									FrameworkDir = Path.GetDirectoryName(FrameworkDir)!;
-								}
-
-								OutputFiles.Add(CopyBundleResource(new UEBuildBundleResource(new ModuleRules.BundleResource(Path.Combine(FrameworkDir, Framework.Name + ".framework"), "Frameworks")), Executable, BundleDirectory, Graph));
+								FrameworkDir = Path.GetDirectoryName(FrameworkDir)!;
 							}
+
+							OutputFiles.Add(CopyBundleResource(new UEBuildBundleResource(new ModuleRules.BundleResource(Path.Combine(FrameworkDir, Framework.Name + ".framework"), "Frameworks")), Executable, BundleDirectory, Graph));
 						}
 					}
+				}
 
-					foreach (UEBuildBundleResource Resource in BinaryLinkEnvironment.AdditionalBundleResources)
-					{
-						OutputFiles.Add(CopyBundleResource(Resource, Executable, StagedExecutablePath.Directory, Graph));
-					}
+				foreach (UEBuildBundleResource Resource in BinaryLinkEnvironment.AdditionalBundleResources)
+				{
+					OutputFiles.Add(CopyBundleResource(Resource, Executable, StagedExecutablePath.Directory, Graph));
 				}
 
 				IOSPostBuildSyncTarget PostBuildSyncTarget = new IOSPostBuildSyncTarget(Target, Executable.Location, BinaryLinkEnvironment.IntermediateDirectory, UPLScripts, SdkVersion, FrameworkNameToSourceDir);
@@ -1410,6 +1405,7 @@ namespace UnrealBuildTool
 				CommandLineArguments CmdLine = new CommandLineArguments(Arguments);
 
 				PlatformProjectGeneratorCollection PlatformProjectGenerators = new PlatformProjectGeneratorCollection();
+				PlatformProjectGenerators.RegisterPlatformProjectGenerator(UnrealTargetPlatform.Mac, new MacProjectGenerator(CmdLine, Logger), Logger);
 				PlatformProjectGenerators.RegisterPlatformProjectGenerator(UnrealTargetPlatform.IOS, new IOSProjectGenerator(CmdLine, Logger), Logger);
 				PlatformProjectGenerators.RegisterPlatformProjectGenerator(UnrealTargetPlatform.TVOS, new TVOSProjectGenerator(CmdLine, Logger), Logger);
 
@@ -1590,7 +1586,7 @@ namespace UnrealBuildTool
 			{
 				// generate a run-only project file for codesigning, etc
 				DirectoryReference? GeneratedProjectFile;
-				IOSExports.GenerateRunOnlyXcodeProject(Target.ProjectFile, Target.Platform, Target.bForDistribution, Logger, out GeneratedProjectFile);
+				IOSExports.GenerateRunOnlyXcodeProject(Target.ProjectFile, Target.Platform, Target.TargetName, Target.bForDistribution, Logger, out GeneratedProjectFile);
 
 				// @todo - should we move this in to FinalizeAppWithXcode?
 				string ConfigName = Target.Configuration.ToString();
diff --git a/Engine/Source/Programs/UnrealBuildTool/Platform/IOS/UEBuildIOS.cs b/Engine/Source/Programs/UnrealBuildTool/Platform/IOS/UEBuildIOS.cs
--- a/Engine/Source/Programs/UnrealBuildTool/Platform/IOS/UEBuildIOS.cs
+++ b/Engine/Source/Programs/UnrealBuildTool/Platform/IOS/UEBuildIOS.cs
@@ -145,6 +145,12 @@ namespace UnrealBuildTool
 		/// </summary>
 		[CommandLine("-EnableUBSan")]
 		public bool bEnableUndefinedBehaviorSanitizer = false;
+
+		/// <summary>
+		/// Are we using modern xcode in this project?
+		/// </summary>
+		[ConfigFile(ConfigHierarchyType.Engine, "XcodeConfiguration", "bUseModernXcode")]
+		public bool bUseModernXcode = false;
 	}
 
 	/// <summary>
@@ -241,6 +247,11 @@ namespace UnrealBuildTool
 			get { return Inner.bEnableUndefinedBehaviorSanitizer; }
 		}
 
+		public bool bUseModernXcode
+		{
+			get { return Inner.bUseModernXcode; }
+		}
+
 #pragma warning restore CS1591
 		#endregion
 	}
@@ -255,12 +266,6 @@ namespace UnrealBuildTool
 		/// </summary>
 		public readonly FileReference? ProjectFile;
 
-		/// <summary>
-		/// Are we using modern xcode in this project?
-		/// </summary>
-		[ConfigFile(ConfigHierarchyType.Engine, "XcodeConfiguration", "bUseModernXcode")]
-		public bool bUseModernXcode = false;
-
 		/// <summary>
 		/// Whether to build the iOS project as a framework.
 		/// </summary>
@@ -780,7 +785,7 @@ namespace UnrealBuildTool
 
 			Target.IOSPlatform.ProjectSettings = ((IOSPlatform)GetBuildPlatform(Target.Platform)).ReadProjectSettings(Target.ProjectFile);
 
-			if (!Target.IOSPlatform.ProjectSettings.bUseModernXcode)
+			if (!Target.IOSPlatform.bUseModernXcode)
 			{
 				// always strip in shipping configuration (commandline could have set it also)
 				if (Target.Configuration == UnrealTargetConfiguration.Shipping)
diff --git a/Engine/Source/Programs/UnrealBuildTool/Platform/Mac/MacToolChain.cs b/Engine/Source/Programs/UnrealBuildTool/Platform/Mac/MacToolChain.cs
--- a/Engine/Source/Programs/UnrealBuildTool/Platform/Mac/MacToolChain.cs
+++ b/Engine/Source/Programs/UnrealBuildTool/Platform/Mac/MacToolChain.cs
@@ -1128,19 +1128,23 @@ namespace UnrealBuildTool
 				BuildProducts.Add(FileReference.Combine(BundleContentsDirectory!, "Info.plist"), BuildProductType.RequiredResource);
 				BuildProducts.Add(FileReference.Combine(BundleContentsDirectory!, "PkgInfo"), BuildProductType.RequiredResource);
 
-				if (Target.Type == TargetType.Editor)
+				// modern xcode doesn't use the bootstrap launcher because it can make full .app with staged data inside it
+				if (!Target.MacPlatform.bUseModernXcode)
 				{
-					BuildProducts.Add(FileReference.Combine(BundleContentsDirectory!, "Resources/UnrealEditor.icns"), BuildProductType.RequiredResource);
-					BuildProducts.Add(FileReference.Combine(BundleContentsDirectory!, "Resources/UProject.icns"), BuildProductType.RequiredResource);
-				}
-				else
-				{
-					string IconName = Target.Name;
-					if (IconName == "EpicGamesBootstrapLauncher")
+					if (Target.Type == TargetType.Editor)
+					{
+						BuildProducts.Add(FileReference.Combine(BundleContentsDirectory!, "Resources/UnrealEditor.icns"), BuildProductType.RequiredResource);
+						BuildProducts.Add(FileReference.Combine(BundleContentsDirectory!, "Resources/UProject.icns"), BuildProductType.RequiredResource);
+					}
+					else
 					{
-						IconName = "EpicGamesLauncher";
+						string IconName = Target.Name;
+						if (IconName == "EpicGamesBootstrapLauncher")
+						{
+							IconName = "EpicGamesLauncher";
+						}
+						BuildProducts.Add(FileReference.Combine(BundleContentsDirectory!, "Resources/" + IconName + ".icns"), BuildProductType.RequiredResource);
 					}
-					BuildProducts.Add(FileReference.Combine(BundleContentsDirectory!, "Resources/" + IconName + ".icns"), BuildProductType.RequiredResource);
 				}
 			}
 		}
diff --git a/Engine/Source/Programs/UnrealBuildTool/Platform/Mac/UEBuildMac.cs b/Engine/Source/Programs/UnrealBuildTool/Platform/Mac/UEBuildMac.cs
--- a/Engine/Source/Programs/UnrealBuildTool/Platform/Mac/UEBuildMac.cs
+++ b/Engine/Source/Programs/UnrealBuildTool/Platform/Mac/UEBuildMac.cs
@@ -59,6 +59,12 @@ namespace UnrealBuildTool
 		/// </summary>
 		[ConfigFile(ConfigHierarchyType.Engine, "/Script/MacTargetPlatform.MacTargetSettings", "bEnableRayTracing")]
 		public bool bEnableRayTracing = false;
+
+		/// <summary>
+		/// Whether or not to use Modernized xcode mode
+		/// </summary>
+		[ConfigFile(ConfigHierarchyType.Engine, "XcodeConfiguration", "bUseModernXcode")]
+		public bool bUseModernXcode = false;
 	}
 
 	/// <summary>
@@ -111,6 +117,11 @@ namespace UnrealBuildTool
 			get { return Inner.bEnableRayTracing; }
 		}
 
+		public bool bUseModernXcode
+		{
+			get { return Inner.bUseModernXcode; }
+		}
+
 #pragma warning restore CS1591
 		#endregion
 	}
diff --git a/Engine/Source/Programs/UnrealBuildTool/ProjectFiles/Xcode/XcodeProject.cs b/Engine/Source/Programs/UnrealBuildTool/ProjectFiles/Xcode/XcodeProject.cs
--- a/Engine/Source/Programs/UnrealBuildTool/ProjectFiles/Xcode/XcodeProject.cs
+++ b/Engine/Source/Programs/UnrealBuildTool/ProjectFiles/Xcode/XcodeProject.cs
@@ -1179,6 +1179,11 @@ namespace UnrealBuildTool.XcodeProjectXcconfig
 
 		/// <summary>
 		/// Write some scripts to do some fixup with how UBT links files. There is currently a difference between Mac and IOS/TVOS:
+		/// All:
+		///   - The staged files gets pulled into the .app for a self-contained app, created at Xcode build time, unless an envvar is set
+		///       that tells this script to skip the copy. This allows UAT to potentially skip copying staged files before we stage
+		///       ("BuildCookRun -build -cook -stage -package" can skip the copying until the -package step).
+		///       Also we don't copy anything if this is the engine, no-project, build, and an 
 		/// Mac:
 		///   - UBT will link directly to Foo.app/Contents/MacOS/Foo, for .app bundled apps
 		///   - During normal processing, that's all that is needed, so we don't need to copy anything
@@ -1198,23 +1203,63 @@ namespace UnrealBuildTool.XcodeProjectXcconfig
 		/// <param name="Project"></param>
 		protected void ProcessScripts(XcodeResourcesBuildPhase ResourcesBuildPhase, XcodeProject Project)
 		{
-			// @todo remove this if UBT links directly into the .app
+			// nothing to do for editors
+			if (Project.UnrealData.TargetRules.Type == TargetType.Editor)
+			{
+				return;
+			}
+
+			List<string> CopyScript = new();
+
 			if (Project.Platform == UnrealTargetPlatform.IOS || Project.Platform == UnrealTargetPlatform.TVOS)
 			{
 				// UBT no longer copies the executable into the .app directory in PostBuild, so we do it here
 				// EXECUTABLE_NAME is Foo, EXECUTABLE_PATH is Foo.app/Foo
 				// NOTE: We read from hardcoded location where UBT writes to, but we write to CONFIGURATION_BUILD_DIR because
 				// when Archiving, the .app is somewhere else
-				string[] Script = {
-						$"cp \\\"${{UE_PROJECT_DIR}}/Binaries/{Platform}/${{EXECUTABLE_NAME}}\\\" \\\"${{CONFIGURATION_BUILD_DIR}}/${{EXECUTABLE_PATH}}\\\"",
-					};
-				string ScriptInput = $"$(UE_PROJECT_DIR)/Binaries/{Platform}/$(EXECUTABLE_NAME)";
-				string ScriptOutput = $"$(CONFIGURATION_BUILD_DIR)/$(EXECUTABLE_PATH)";
-				XcodeShellScriptBuildPhase ScriptPhase = new("Copy UE Executable into .app", Script, new string[] { ScriptInput }, new string[] { ScriptOutput });
-				BuildPhases.Add(ScriptPhase);
-				References.Add(ScriptPhase);
+				CopyScript.AddRange(new string[]
+				{
+					"# Copy the IOS/TVOS executable into .app",
+					$"cp \\\"${{UE_BINARIES_DIR}}/${{EXECUTABLE_NAME}}\\\" \\\"${{CONFIGURATION_BUILD_DIR}}/${{EXECUTABLE_PATH}}\\\"",
+					"",
+				});
+			}
+
+			// rsync the Staged build into the .app, unless the UE_SKIP_STAGEDDATA_SYNC var is set to 1
+			string TargetPlatformName = Platform.ToString()!;
+			if (TargetType != TargetType.Game && TargetType != TargetType.Program)
+			{
+				TargetPlatformName += TargetType.ToString();
+			}
+
+			// editor builds don't need staged content in them
+			if (TargetType != TargetType.Editor)
+			{
+				bool bIsEngineBuild = Project.UnrealData.UProjectFileLocation == null;
+				string DefaultStageDir = bIsEngineBuild ? "${UE_OVERRIDE_STAGE_DIR}" : $"${{UE_PROJECT_DIR}}/Saved/StagedBuilds/{TargetPlatformName}";
+				string SyncSubdir = (Platform == UnrealTargetPlatform.Mac) ? "/UE" : "";
+				CopyScript.AddRange(new string[]
+				{
+					"# Skip syncing if desired",
+					"if [[ ${UE_SKIP_STAGEDDATA_SYNC} -eq 1 ]]; then exit 0; fi",
+					"",
+
+					"# When building engine projects, like UnrealGame, we don't have data to stage unless something has specified UE_OVERRIDE_STAGE_DIR",
+					$"STAGED_DIR={DefaultStageDir}",
+					"if [[ -z ${STAGED_DIR} ]]; then exit 0; fi",
+					"if [[ ! -e ${STAGED_DIR} ]]; then exit 0; fi",
+					"",
+					"echo \\\"Syncing ${STAGED_DIR} to ${CONFIGURATION_BUILD_DIR}/${CONTENTS_FOLDER_PATH}\\\"",
+					$"rsync -av --exclude=Info.plist --exclude=/Manifest_* --exclude=*.app \\\"${{STAGED_DIR}}/\\\" \\\"${{CONFIGURATION_BUILD_DIR}}/${{CONTENTS_FOLDER_PATH}}{SyncSubdir}\\\"",
+				});
 			}
 
+			// run this script every time, but xcode will show a warning if there isn't _some_ output
+			string ScriptOutput = $"/dev/null";
+			XcodeShellScriptBuildPhase CopyScriptPhase = new("Copy UE Executable into .app", CopyScript, new string[] { }, new string[] { ScriptOutput });
+			BuildPhases.Add(CopyScriptPhase);
+			References.Add(CopyScriptPhase);
+
 			// always generate a dsym file when we archive, and by having Xcode do it, it will be put into the archive properly
 			// (note bInstallOnly which will make this onle run when archiving)
 			List<string> DsymScript = new();
@@ -1243,20 +1288,6 @@ namespace UnrealBuildTool.XcodeProjectXcconfig
 		}
 		private void ProcessAssets(XcodeResourcesBuildPhase ResourcesBuildPhase)
 		{
-			if (Platform == UnrealTargetPlatform.IOS || Platform == UnrealTargetPlatform.TVOS)
-			{
-				// cookeddata needs to be brought in from the Staged directory, (IOS for Game, IOSClient for Client, etc)
-				string TargetPlatformName = Platform.ToString()!;
-				if (TargetType != TargetType.Game)
-				{
-					TargetPlatformName += TargetType.ToString();
-				}
-
-				// @todo do this for Mac?
-				DirectoryReference CookedData = DirectoryReference.Combine(UnrealData.ProductDirectory, "Saved", "StagedBuilds", TargetPlatformName, "cookeddata");
-				ResourcesBuildPhase.AddFolderResource(CookedData, "Resources");
-			}
-
 			List<string> StoryboardPaths = new List<string>()
 				{
 					"$(Project)/Build/$(Platform)/Resources/Interface/LaunchScreen.storyboardc",
@@ -1516,11 +1547,15 @@ namespace UnrealBuildTool.XcodeProjectXcconfig
 			Xcconfig.AppendLine($"UE_PRODUCT_NAME_STRIPPED = {UnrealData.ProductName.Replace("_", "").Replace(" ", "")}");
 			Xcconfig.AppendLine($"UE_DISPLAY_NAME = {UnrealData.DisplayName}");
 			Xcconfig.AppendLine($"UE_SIGNING_PREFIX = {SigningPrefix}");
-			Xcconfig.AppendLine($"UE_ENGINE_DIR = {Unreal.EngineDirectory}");
-			Xcconfig.AppendLine($"UE_PROJECT_DIR = {ProjectOrEngineDir}");
 			Xcconfig.AppendLine($"UE_PLATFORM_NAME = {Platform}");
 			Xcconfig.AppendLine($"UE_TARGET_NAME = {UnrealData.TargetRules.Name}");
 			Xcconfig.AppendLine($"UE_TARGET_PLATFORM_NAME = {TargetPlatformName}");
+			Xcconfig.AppendLine($"UE_ENGINE_DIR = {Unreal.EngineDirectory}");
+			Xcconfig.AppendLine($"UE_BINARIES_DIR = {DirectoryReference.Combine(ProjectOrEngineDir, "Binaries", Platform.ToString())}");
+			if (UnrealData.UProjectFileLocation != null)
+			{
+				Xcconfig.AppendLine($"UE_PROJECT_DIR = {UnrealData.UProjectFileLocation.Directory}");
+			}
 
 			Xcconfig.AppendLine("");
 			Xcconfig.AppendLine("// Constant settings (same for all platforms and targets)");
@@ -1813,7 +1848,8 @@ namespace UnrealBuildTool.XcodeProjectXcconfig
 		/// <param name="BundleID">Override option for bundle identifier</param>
 		/// <param name="AppName"></param>
 		/// <param name="bMakeProjectPerTarget"></param>
-		public XcodeProjectFile(FileReference InitFilePath, DirectoryReference BaseDir, bool bIsForDistribution, string BundleID, string AppName, bool bMakeProjectPerTarget)
+		/// <param name="SingleTargetName"></param>
+		public XcodeProjectFile(FileReference InitFilePath, DirectoryReference BaseDir, bool bIsForDistribution, string BundleID, string AppName, bool bMakeProjectPerTarget, string? SingleTargetName)
 			: base(InitFilePath, BaseDir)
 		{
 			UnrealData = new UnrealData(InitFilePath, bIsForDistribution, BundleID, AppName, bMakeProjectPerTarget, IsStubProject);
@@ -1821,6 +1857,7 @@ namespace UnrealBuildTool.XcodeProjectXcconfig
 			// create the container for all the files that will 
 			SharedFileCollection = new XcodeFileCollection(this);
 			FileCollection = SharedFileCollection;
+			this.SingleTargetName = SingleTargetName;
 		}
 
 		public UnrealData UnrealData;
@@ -1841,6 +1878,8 @@ namespace UnrealBuildTool.XcodeProjectXcconfig
 		private XcodeFileCollection SharedFileCollection;
 		public XcodeFileCollection FileCollection;
 
+		// if set, only this will be written
+		private string? SingleTargetName;
 
 		/// <summary>
 		/// Allocates a generator-specific source file object
@@ -1950,6 +1989,13 @@ namespace UnrealBuildTool.XcodeProjectXcconfig
 		/// Implements Project interface
 		public override bool WriteProjectFile(List<UnrealTargetPlatform> InPlatforms, List<UnrealTargetConfiguration> InConfigurations, PlatformProjectGeneratorCollection PlatformProjectGenerators, ILogger Logger)
 		{
+			// if we don't want this one, just skip
+			if (SingleTargetName != null && !ProjectFilePath.GetFileNameWithoutAnyExtensions().Equals(SingleTargetName, StringComparison.InvariantCultureIgnoreCase))
+			{
+				return true;
+			}
+
+
 			ConditionalCreateLegacyProject();
 
 			if (LegacyProjectFile != null)
diff --git a/Engine/Source/Programs/UnrealBuildTool/ProjectFiles/Xcode/XcodeProjectFileGenerator.cs b/Engine/Source/Programs/UnrealBuildTool/ProjectFiles/Xcode/XcodeProjectFileGenerator.cs
--- a/Engine/Source/Programs/UnrealBuildTool/ProjectFiles/Xcode/XcodeProjectFileGenerator.cs
+++ b/Engine/Source/Programs/UnrealBuildTool/ProjectFiles/Xcode/XcodeProjectFileGenerator.cs
@@ -49,6 +49,12 @@ namespace UnrealBuildTool
 		/// </summary>
 		string AppName = "";
 
+		/// <summary>
+		/// If set, only write out one target/project 
+		/// </summary>
+		string? SingleTargetName = null;
+		
+
 		public XcodeProjectFileGenerator(FileReference? InOnlyGameProject, CommandLineArguments CommandLine)
 			: base(InOnlyGameProject)
 		{
@@ -65,6 +71,12 @@ namespace UnrealBuildTool
 			{
 				AppName = CommandLine.GetString("-appname=");
 			}
+
+			if (CommandLine.HasValue("-SingleTarget="))
+			{
+				SingleTargetName = CommandLine.GetString("-SingleTarget=");
+			}
+
 		}
 
 		/// <summary>
@@ -129,7 +141,7 @@ namespace UnrealBuildTool
 			// unfortunately, we can't read the project configs now because we don't have enough information to 
 			// find the .uproject file for that would make this project (we could change the high level to pass it 
 			// down but it would touch all project generators - not worth it if we end up removing the legacy)
-			return new XcodeProjectXcconfig.XcodeProjectFile(InitFilePath, BaseDir, bForDistribution, BundleIdentifier, AppName, bMakeProjectPerTarget);
+			return new XcodeProjectXcconfig.XcodeProjectFile(InitFilePath, BaseDir, bForDistribution, BundleIdentifier, AppName, bMakeProjectPerTarget, SingleTargetName);
 		}
 
 		private bool WriteWorkspaceSettingsFile(string Path, ILogger Logger)
@@ -386,9 +398,13 @@ namespace UnrealBuildTool
 
 			if (bGenerateRunOnlyProject)
 			{
-				bIncludeEnginePrograms = false;
+				// if we have a single target, allow all targets so we can pick it out (helps with programs)
+				if (SingleTargetName != null)
+				{
+					bIncludeEnginePrograms = false;
+					bIncludeTemplateFiles = false;
+				}
 				//bIncludeEngineSource = false;
-				bIncludeTemplateFiles = false;
 				bIncludeConfigFiles = false;
 				bIncludeDocumentation = false;
 				bIncludeShaderSource = false;
@@ -414,6 +430,12 @@ namespace UnrealBuildTool
 			DateTime MainStart = DateTime.UtcNow;
 			foreach (var TargetPair in Targets)
 			{
+				// don't bother if we aren't interested in this target
+				if (SingleTargetName != null && !TargetPair.Item2.Name.Equals(SingleTargetName, StringComparison.InvariantCultureIgnoreCase))
+				{
+					continue;
+				}
+
 				ProjectFile TargetProjectFile = TargetPair.Item1;
 				// don't do this for legacy projets, for speed
 				((XcodeProjectXcconfig.XcodeProjectFile)TargetProjectFile).ConditionalCreateLegacyProject();
diff --git a/Engine/Source/Programs/UnrealBuildTool/ToolChain/AppleToolChain.cs b/Engine/Source/Programs/UnrealBuildTool/ToolChain/AppleToolChain.cs
--- a/Engine/Source/Programs/UnrealBuildTool/ToolChain/AppleToolChain.cs
+++ b/Engine/Source/Programs/UnrealBuildTool/ToolChain/AppleToolChain.cs
@@ -463,5 +463,110 @@ namespace UnrealBuildTool
 			ExtraOptions = (bIsForLTOBuild && Major == 10 && Minor == 0 && Patch == 1) ? "-j 1" : "";
 			return DsymutilLocation.FullName;
 		}
+
+
+		public override ICollection<FileItem> PostBuild(ReadOnlyTargetRules Target, FileItem Executable, LinkEnvironment BinaryLinkEnvironment, IActionGraphBuilder Graph)
+		{
+			List<FileItem> OutputFiles = new List<FileItem>(base.PostBuild(Target, Executable, BinaryLinkEnvironment, Graph));
+
+			bool bIsBuildingAppBundle = !BinaryLinkEnvironment.bIsBuildingDLL && !BinaryLinkEnvironment.bIsBuildingLibrary && !BinaryLinkEnvironment.bIsBuildingConsoleApplication;
+			if (Target.MacPlatform.bUseModernXcode && bIsBuildingAppBundle)
+			{
+				Action PostBuildAction = ApplePostBuildSyncMode.CreatePostBuildSyncAction(Target, Executable, BinaryLinkEnvironment.IntermediateDirectory!, Graph);
+
+				PostBuildAction.PrerequisiteItems.UnionWith(OutputFiles);
+				OutputFiles.AddRange(PostBuildAction.ProducedItems);
+			}
+
+			return OutputFiles;
+		}
 	};
+
+	[Serializable]
+	class ApplePostBuildSyncTarget
+	{
+		public FileReference? ProjectFile;
+		public UnrealTargetPlatform Platform;
+		public UnrealTargetConfiguration Configuration;
+		public string TargetName;
+
+		public ApplePostBuildSyncTarget(ReadOnlyTargetRules Target)
+		{
+			Platform = Target.Platform;
+			Configuration = Target.Configuration;
+			ProjectFile = Target.ProjectFile;
+			TargetName = Target.Name;
+		}
+	}
+
+	[ToolMode("ApplePostBuildSync", ToolModeOptions.XmlConfig | ToolModeOptions.BuildPlatforms)]
+	class ApplePostBuildSyncMode : ToolMode
+	{
+		[CommandLine("-Input=", Required = true)]
+		public FileReference? InputFile = null;
+
+		[CommandLine("-XmlConfigCache=")]
+		public FileReference? XmlConfigCache = null;
+
+		public override int Execute(CommandLineArguments Arguments, ILogger Logger)
+		{
+			Arguments.ApplyTo(this);
+			Arguments.CheckAllArgumentsUsed();
+
+			// Run the PostBuildSync command
+			ApplePostBuildSyncTarget Target = BinaryFormatterUtils.Load<ApplePostBuildSyncTarget>(InputFile!);
+			PostBuildSync(Target, Logger);
+
+			return 0;
+		}
+
+		private void PostBuildSync(ApplePostBuildSyncTarget Target, ILogger Logger)
+		{
+			// generate a run-only project file for codesigning, etc
+			DirectoryReference? GeneratedProjectFile;
+			Logger.LogInformation("Generating temporary Xcode project...");
+			DateTime Start = DateTime.Now;
+			IOSExports.GenerateRunOnlyXcodeProject(Target.ProjectFile, Target.Platform, Target.TargetName, false, Logger, out GeneratedProjectFile);
+
+			// run xcodebuild on the generated project to make the .app
+			bool bForDistribution = false;
+			Logger.LogInformation($"Finalizing .app via Xcode... [generation took {(DateTime.Now - Start).TotalMilliseconds / 1000.0}s]");
+			IOSExports.FinalizeAppWithModernXcode(GeneratedProjectFile!, Target.Platform, Target.TargetName, Target.Configuration.ToString(), bForDistribution, Logger);
+		}
+
+
+
+
+		private static FileItem GetPostBuildOutputFile(FileReference Executable, string TargetName, UnrealTargetPlatform Platform)
+		{
+			FileReference StagedExe;
+			if (Platform == UnrealTargetPlatform.Mac)
+			{
+				StagedExe = FileReference.Combine(Executable.Directory.ParentDirectory!, "PkgInfo");
+			}
+			else
+			{
+				StagedExe = FileReference.Combine(Executable.Directory, TargetName + ".app", TargetName);
+			}
+			return FileItem.GetItemByFileReference(StagedExe);
+		}
+
+		public static Action CreatePostBuildSyncAction(ReadOnlyTargetRules Target, FileItem Executable, DirectoryReference IntermediateDir, IActionGraphBuilder Graph)
+		{
+			ApplePostBuildSyncTarget PostBuildSync = new(Target);
+			FileReference PostBuildSyncFile = FileReference.Combine(IntermediateDir!, "PostBuildSync.dat");
+			BinaryFormatterUtils.Save(PostBuildSyncFile, PostBuildSync);
+
+			string PostBuildSyncArguments = String.Format("-Input=\"{0}\" -XmlConfigCache=\"{1}\" -remoteini=\"{2}\"", PostBuildSyncFile, XmlConfig.CacheFile, UnrealBuildTool.GetRemoteIniPath());
+			Action PostBuildSyncAction = Graph.CreateRecursiveAction<ApplePostBuildSyncMode>(ActionType.CreateAppBundle, PostBuildSyncArguments);
+
+			PostBuildSyncAction.WorkingDirectory = Unreal.EngineSourceDirectory;
+			PostBuildSyncAction.PrerequisiteItems.Add(Executable);
+			PostBuildSyncAction.ProducedItems.Add(GetPostBuildOutputFile(Executable.Location, Target.Name, Target.Platform));
+			PostBuildSyncAction.StatusDescription = $"Executing PostBuildSync [{Executable.Location}]";
+			PostBuildSyncAction.bCanExecuteRemotely = false;
+
+			return PostBuildSyncAction;
+		}
+	}
 }
