Added FORCE_ANSI_ALLOCATOR=1 to Android ASan builds

diff --git a/Engine/Source/Programs/UnrealBuildTool/Platform/Android/AndroidToolChain.cs b/Engine/Source/Programs/UnrealBuildTool/Platform/Android/AndroidToolChain.cs
--- a/Engine/Source/Programs/UnrealBuildTool/Platform/Android/AndroidToolChain.cs
+++ b/Engine/Source/Programs/UnrealBuildTool/Platform/Android/AndroidToolChain.cs
@@ -601,7 +601,7 @@ namespace UnrealBuildTool
 
 				if (Sanitizer == ClangSanitizer.Address || Sanitizer == ClangSanitizer.HwAddress)
 				{
-					Arguments.Add("-fno-omit-frame-pointer -DRUNNING_WITH_ASAN=1");
+					Arguments.Add("-fno-omit-frame-pointer -DRUNNING_WITH_ASAN=1 -DFORCE_ANSI_ALLOCATOR=1");
 				}
 			}
 
diff --git a/Engine/Source/Runtime/Core/Private/Android/AndroidPlatformMemory.cpp b/Engine/Source/Runtime/Core/Private/Android/AndroidPlatformMemory.cpp
--- a/Engine/Source/Runtime/Core/Private/Android/AndroidPlatformMemory.cpp
+++ b/Engine/Source/Runtime/Core/Private/Android/AndroidPlatformMemory.cpp
@@ -547,7 +547,7 @@ FMalloc* FAndroidPlatformMemory::BaseAllocator()
 	FLowLevelMemTracker::Get().SetProgramSize(Stats.UsedPhysical);
 #endif
 
-#if RUNNING_WITH_ASAN
+#if FORCE_ANSI_ALLOCATOR
 	return (Instance = new FMallocAnsi());
 #endif
 
