UnrealBuildTool: Add SupportedTargetTypes attribute, allowing for ModuleRules to be marked as supported for only specific TargetTypes (Editor, Game, Client, Server).
* Add SupportedPlatformGroup attribute, which will resolve to a SupportedPlatform attribute that contains all the platforms in that group
* Allow SupportedConfigurations attribute on ModuleRules
* Filter out -AllModules plugin modules that are unsupported
* Print warning with reference chain when a module is referenced but unsupported

diff --git a/Engine/Source/Programs/UnrealBuildTool/Configuration/ModuleRules.cs b/Engine/Source/Programs/UnrealBuildTool/Configuration/ModuleRules.cs
--- a/Engine/Source/Programs/UnrealBuildTool/Configuration/ModuleRules.cs
+++ b/Engine/Source/Programs/UnrealBuildTool/Configuration/ModuleRules.cs
@@ -5,6 +5,7 @@ using System.Collections.Generic;
 using System.Diagnostics.CodeAnalysis;
 using System.IO;
 using System.Linq;
+using System.Reflection;
 using System.Text;
 using System.Threading.Tasks;
 using EpicGames.Core;
@@ -1540,6 +1541,45 @@ namespace UnrealBuildTool
 				});
 		}
 
+		/// <summary>
+		/// Determines if a module type is valid for a target, based on custom attributes
+		/// </summary>
+		/// <param name="ModuleType">The type of the module to check</param>
+		/// <param name="TargetRules">The target to check against</param>
+		/// <param name="InvalidReason">Out, reason this module was invalid</param>
+		/// <returns>True if the module is valid, false otherwise</returns>
+		internal static bool IsValidForTarget(Type ModuleType, ReadOnlyTargetRules TargetRules, [NotNullWhen(false)] out string? InvalidReason)
+		{
+			IEnumerable<TargetType> SupportedTargetTypes = ModuleType.GetCustomAttributes<SupportedTargetTypesAttribute>().SelectMany(x => x.TargetTypes).Distinct();
+			if (SupportedTargetTypes.Any() && !SupportedTargetTypes.Contains(TargetRules.Type))
+			{
+				InvalidReason = $"TargetType '{TargetRules.Type}'";
+				return false;
+			}
+
+			IEnumerable<UnrealTargetConfiguration> SupportedConfigurations = ModuleType.GetCustomAttributes<SupportedConfigurationsAttribute>().SelectMany(x => x.Configurations).Distinct();
+			if (SupportedConfigurations.Any() && !SupportedConfigurations.Contains(TargetRules.Configuration))
+			{
+				InvalidReason = $"Configuration '{TargetRules.Configuration}'";
+				return false;
+			}
+
+			// Skip platform extension modules. We only care about the base modules, not the platform overrides.
+			// The platform overrides get applied at a later stage when we actually come to build the module.
+			if (!UEBuildPlatform.GetPlatformFolderNames().Any(Name => ModuleType.Name.EndsWith("_" + Name)))
+			{
+				IEnumerable<UnrealTargetPlatform> SupportedPlatforms = ModuleType.GetCustomAttributes<SupportedPlatformsAttribute>().SelectMany(x => x.Platforms).Distinct();
+				if (SupportedPlatforms.Any() && !SupportedPlatforms.Contains(TargetRules.Platform))
+				{
+					InvalidReason = $"Platform '{TargetRules.Platform}'";
+					return false;
+				}
+			}
+
+			InvalidReason = null;
+			return true;
+		}
+
 		/// <summary>
 		/// Determines if this module can be precompiled for the current target.
 		/// </summary>
@@ -1547,7 +1587,7 @@ namespace UnrealBuildTool
 		/// <returns>True if the module can be precompiled, false otherwise</returns>
 		internal bool IsValidForTarget(FileReference RulesFile)
 		{
-			if(Type == ModuleRules.ModuleType.CPlusPlus)
+			if (Type == ModuleRules.ModuleType.CPlusPlus)
 			{
 				switch (PrecompileForTargets)
 				{
diff --git a/Engine/Source/Programs/UnrealBuildTool/Configuration/UEBuildModule.cs b/Engine/Source/Programs/UnrealBuildTool/Configuration/UEBuildModule.cs
--- a/Engine/Source/Programs/UnrealBuildTool/Configuration/UEBuildModule.cs
+++ b/Engine/Source/Programs/UnrealBuildTool/Configuration/UEBuildModule.cs
@@ -1027,6 +1027,12 @@ namespace UnrealBuildTool
 			// Add us to the reference stack - note do this before checking for cycles as we allow them if an element in the stack
 			// declares the next element at leading to a dependency cycle.
 			ReferenceStack.Add(this);
+			
+			// Check if this module is invalid for the current target, based on Module attributes
+			if (!ModuleRules.IsValidForTarget(Rules.GetType(), Rules.Target, out string? InvalidReason))
+			{
+				Logger.LogWarning("Warning: Referenced module '{Module}' does not support {InvalidReason} via {ReferenceChain}", ThisRefName, InvalidReason, NextReferenceChain);
+			}
 
 			if (CheckForCycles)
 			{
diff --git a/Engine/Source/Programs/UnrealBuildTool/Configuration/UEBuildTarget.cs b/Engine/Source/Programs/UnrealBuildTool/Configuration/UEBuildTarget.cs
--- a/Engine/Source/Programs/UnrealBuildTool/Configuration/UEBuildTarget.cs
+++ b/Engine/Source/Programs/UnrealBuildTool/Configuration/UEBuildTarget.cs
@@ -3563,30 +3563,30 @@ namespace UnrealBuildTool
 				foreach (string ModuleName in ModuleNames)
 				{
 					FileReference ModuleFileName = RulesAssembly.GetModuleFileName(ModuleName)!;
-					if (Directories.Any(BaseDir => ModuleFileName.IsUnderDirectory(BaseDir)))
+					if (!Directories.Any(BaseDir => ModuleFileName.IsUnderDirectory(BaseDir)))
 					{
-						Type RulesType = RulesAssembly.GetModuleRulesType(ModuleName)!;
+						Logger.LogDebug("Excluding module {Module}: Not under compatible directories", ModuleName);
+						continue;
+					}
+
+					Type RulesType = RulesAssembly.GetModuleRulesType(ModuleName)!;
+
+					if (!ModuleRules.IsValidForTarget(RulesType, Rules, out string? InvalidReason))
+					{
+						Logger.LogDebug("Excluding module {Module}: Does not support {InvalidReason}", ModuleName, InvalidReason);
+						continue;
+					}
 
-						// Skip platform extension modules. We only care about the base modules, not the platform overrides.
-						// The platform overrides get applied at a later stage when we actually come to build the module.
-						if (!UEBuildPlatform.GetPlatformFolderNames().Any(Name => RulesType.Name.EndsWith("_" + Name)))
+					// Skip platform extension modules. We only care about the base modules, not the platform overrides.
+					// The platform overrides get applied at a later stage when we actually come to build the module.
+					if (!UEBuildPlatform.GetPlatformFolderNames().Any(Name => RulesType.Name.EndsWith("_" + Name)))
+					{
+						if (ModuleFileName.ContainsAnyNames(ExcludeFolders, Unreal.EngineDirectory))
 						{
-							SupportedPlatformsAttribute? SupportedPlatforms = RulesType.GetCustomAttribute<SupportedPlatformsAttribute>();
-							if (SupportedPlatforms != null)
-							{
-								if (SupportedPlatforms.Platforms.Contains(Platform))
-								{
-									FilteredModuleNames.Add(ModuleName);
-								}
-							}
-							else
-							{
-								if (!ModuleFileName.ContainsAnyNames(ExcludeFolders, Unreal.EngineDirectory))
-								{
-									FilteredModuleNames.Add(ModuleName);
-								}
-							}
+							Logger.LogDebug("Excluding module {Module}: Excluded folder name", ModuleName);
+							continue;
 						}
+						FilteredModuleNames.Add(ModuleName);
 					}
 				}
 			}
@@ -3635,10 +3635,18 @@ namespace UnrealBuildTool
 						{
 							throw new BuildException("Unable to find module '{0}' referenced by {1}", ModuleDescriptor.Name, Plugin.File);
 						}
-						if (!ModuleFileName.ContainsAnyNames(ExcludeFolders, Plugin.Directory))
+						Type RulesType = RulesAssembly.GetModuleRulesType(ModuleDescriptor.Name)!;
+						if (!ModuleRules.IsValidForTarget(RulesType, Rules, out string? InvalidReason))
 						{
-							FilteredModuleNames.Add(ModuleDescriptor.Name);
+							Logger.LogDebug("Excluding plugin {Plugin} module {Module}: Does not support {InvalidReason}", Plugin.Name, ModuleDescriptor.Name, InvalidReason);
+							continue;
 						}
+						if (ModuleFileName.ContainsAnyNames(ExcludeFolders, Plugin.Directory))
+						{
+							Logger.LogDebug("Excluding plugin {Plugin} module {Module}: Excluded folder name", Plugin.Name, ModuleDescriptor.Name);
+							continue;
+						}
+						FilteredModuleNames.Add(ModuleDescriptor.Name);
 					}
 				}
 			}
@@ -3658,8 +3666,13 @@ namespace UnrealBuildTool
 				}
 
 				// Figure out if it can be precompiled
-				if (ModuleRules != null && ModuleRules.IsValidForTarget(ModuleRules.File))
+				if (ModuleRules != null)
 				{
+					if (!ModuleRules.IsValidForTarget(ModuleRules.File))
+					{
+						Logger.LogDebug("Excluding module {Module}: Not precompiled for target '{TargetType}'", ModuleRules.Name, Rules.Type);
+						continue;
+					}
 					ValidModuleNames.Add(FilteredModuleName);
 				}
 			}
diff --git a/Engine/Source/Programs/UnrealBuildTool/System/SupportedPlatformGroupsAttribute.cs b/Engine/Source/Programs/UnrealBuildTool/System/SupportedPlatformGroupsAttribute.cs
new file mode 100644
--- /dev/null
+++ b/Engine/Source/Programs/UnrealBuildTool/System/SupportedPlatformGroupsAttribute.cs
@@ -0,0 +1,48 @@
+// Copyright Epic Games, Inc. All Rights Reserved.
+
+using System;
+using System.Collections.Generic;
+using System.Linq;
+
+namespace UnrealBuildTool
+{
+	/// <summary>
+	/// Attribute which can be applied to a TargetRules-dervied class to indicate which platforms it supports
+	/// </summary>
+	[AttributeUsage(AttributeTargets.Class, AllowMultiple = true)]
+	public class SupportedPlatformGroupsAttribute : SupportedPlatformsAttribute
+	{
+		/// <summary>
+		/// Initialize the attribute with a list of platform groups
+		/// </summary>
+		/// <param name="PlatformGroups">Variable-length array of platform group arguments</param>
+		public SupportedPlatformGroupsAttribute(params string[] PlatformGroups) : base(GetPlatformsForGroups(PlatformGroups))
+		{
+		}
+
+		private static string[] GetPlatformsForGroups(params string[] PlatformGroups)
+		{
+			HashSet<UnrealTargetPlatform> SupportedPlatforms = new();
+			try
+			{
+				foreach (string Name in PlatformGroups)
+				{
+					if (UnrealPlatformGroup.TryParse(Name, out UnrealPlatformGroup Group))
+					{
+						SupportedPlatforms.UnionWith(UnrealTargetPlatform.GetValidPlatforms().Where(x => x.IsInGroup(Group)));
+						continue;
+					}
+					throw new BuildException(string.Format("The platform group name {0} is not a valid platform group name. Valid names are ({1})", Name,
+						string.Join(",", UnrealPlatformGroup.GetValidGroupNames())));
+				}
+			}
+			catch (BuildException Ex)
+			{
+				EpicGames.Core.ExceptionUtils.AddContext(Ex, $"while parsing a SupportedPlatformGroups attribute '{String.Join(',', PlatformGroups)}'");
+				throw;
+			}
+
+			return SupportedPlatforms.Select(x => x.ToString()).ToArray();
+		}
+	}
+}
diff --git a/Engine/Source/Programs/UnrealBuildTool/System/SupportedTargetTypesAttribute.cs b/Engine/Source/Programs/UnrealBuildTool/System/SupportedTargetTypesAttribute.cs
new file mode 100644
--- /dev/null
+++ b/Engine/Source/Programs/UnrealBuildTool/System/SupportedTargetTypesAttribute.cs
@@ -0,0 +1,32 @@
+// Copyright Epic Games, Inc. All Rights Reserved.
+
+using System;
+using System.Collections.Generic;
+using System.Linq;
+using System.Text;
+using System.Threading.Tasks;
+using UnrealBuildBase;
+
+namespace UnrealBuildTool
+{
+	/// <summary>
+	/// Attribute which can be applied to a TargetRules-dervied class to indicate which target types it supports
+	/// </summary>
+	[AttributeUsage(AttributeTargets.Class, AllowMultiple = true)]
+	public class SupportedTargetTypesAttribute : Attribute
+	{
+		/// <summary>
+		/// Array of supported target types
+		/// </summary>
+		public readonly TargetType[] TargetTypes;
+
+		/// <summary>
+		/// Initialize the attribute with a list of target types
+		/// </summary>
+		/// <param name="TargetTypes">Variable-length array of target type arguments</param>
+		public SupportedTargetTypesAttribute(params TargetType[] TargetTypes)
+		{
+			this.TargetTypes = TargetTypes;
+		}
+	}
+}
