fix system include which caused some files fail to index, also disable all warnings for index

diff --git a/Engine/Source/Programs/UnrealBuildTool/ProjectFiles/Xcode/XcodeProject.cs b/Engine/Source/Programs/UnrealBuildTool/ProjectFiles/Xcode/XcodeProject.cs
--- a/Engine/Source/Programs/UnrealBuildTool/ProjectFiles/Xcode/XcodeProject.cs
+++ b/Engine/Source/Programs/UnrealBuildTool/ProjectFiles/Xcode/XcodeProject.cs
@@ -109,8 +109,8 @@ namespace UnrealBuildTool.XcodeProjectXcconfig
 		public void GenerateResponseFile()
 		{
 			StringBuilder ResponseFileContents = new();
-			ResponseFileContents.Append("-I");
-			ResponseFileContents.AppendJoin(" -I", SystemIncludePaths.Select(x => x.FullName.Contains(' ') ? $"\"{x.FullName}\"" : x.FullName));
+			ResponseFileContents.Append("-isystem");
+			ResponseFileContents.AppendJoin(" -isystem", SystemIncludePaths.Select(x => x.FullName.Contains(' ') ? $"\"{x.FullName}\"" : x.FullName));
 			ResponseFileContents.Append(" -I");
 			ResponseFileContents.AppendJoin(" -I", UserIncludePaths.Select(x => x.FullName.Contains(' ') ? $"\"{x.FullName}\"" : x.FullName));
 			if (ForceIncludeFiles != null)
@@ -1729,13 +1729,12 @@ namespace UnrealBuildTool.XcodeProjectXcconfig
 		{
 			// write out settings that for compiling natively
 			Xcconfig!.AppendLine("CLANG_CXX_LANGUAGE_STANDARD = c++17");
-			Xcconfig.AppendLine("GCC_WARN_CHECK_SWITCH_STATEMENTS = NO");
+			Xcconfig.AppendLine("GCC_WARN_INHIBIT_ALL_WARNINGS = YES");
 			Xcconfig.AppendLine("GCC_PRECOMPILE_PREFIX_HEADER = YES");
 			Xcconfig.AppendLine("GCC_OPTIMIZATION_LEVEL = 0");
 			Xcconfig.AppendLine($"PRODUCT_NAME = {Name}");
 			Xcconfig.AppendLine("SYMROOT = build");
 			Xcconfig.AppendLine("USE_HEADERMAP = NO");
-			Xcconfig.AppendLine("WARNING_CFLAGS = -Wno-c++11-narrowing");
 			Xcconfig.Write();
 		}
 	}
