UnrealBuildTool: Fix issue where -SingleFile was writing a makefile without linked outputs

diff --git a/Engine/Source/Programs/UnrealBuildTool/Configuration/UEBuildTarget.cs b/Engine/Source/Programs/UnrealBuildTool/Configuration/UEBuildTarget.cs
--- a/Engine/Source/Programs/UnrealBuildTool/Configuration/UEBuildTarget.cs
+++ b/Engine/Source/Programs/UnrealBuildTool/Configuration/UEBuildTarget.cs
@@ -1241,10 +1241,11 @@ namespace UnrealBuildTool
 				RulesObject.bDisableLinking = true;
 			}
 
-			// Don't link if we're compiling specific files and ignore build outputs
+			// Ignore build outputs when compiling specific files, as most actions are skipped.
+			// Do not disable linking! Those actions will be skipped and compiling specific files doesn't invalidate the makefile
+			// so it will be saved with no binary outputs. That would be bad!
 			if (Descriptor.SpecificFilesToCompile.Any())
 			{
-				RulesObject.bDisableLinking = true;
 				RulesObject.bIgnoreBuildOutputs = true;
 			}
 
diff --git a/Engine/Source/Programs/UnrealBuildTool/Modes/BuildMode.cs b/Engine/Source/Programs/UnrealBuildTool/Modes/BuildMode.cs
--- a/Engine/Source/Programs/UnrealBuildTool/Modes/BuildMode.cs
+++ b/Engine/Source/Programs/UnrealBuildTool/Modes/BuildMode.cs
@@ -774,7 +774,8 @@ namespace UnrealBuildTool
 					// Run the deployment steps
 					foreach(TargetMakefile Makefile in Makefiles)
 					{
-						if (Makefile.bDeployAfterCompile)
+						// Receipt file may not exist when compiling specific files
+						if (Makefile.bDeployAfterCompile && FileReference.Exists(Makefile.ReceiptFile))
 						{
 							TargetReceipt Receipt = TargetReceipt.Read(Makefile.ReceiptFile);
 							Logger.LogInformation("Deploying {ReceiptTargetName} {ReceiptPlatform} {ReceiptConfiguration}...", Receipt.TargetName, Receipt.Platform, Receipt.Configuration);
