Remove redundant saving of the AutomationTool.csproj.References file.

diff --git a/Engine/Source/Programs/UnrealBuildTool/ProjectFiles/VisualStudio/VCProjectFileGenerator.cs b/Engine/Source/Programs/UnrealBuildTool/ProjectFiles/VisualStudio/VCProjectFileGenerator.cs
--- a/Engine/Source/Programs/UnrealBuildTool/ProjectFiles/VisualStudio/VCProjectFileGenerator.cs
+++ b/Engine/Source/Programs/UnrealBuildTool/ProjectFiles/VisualStudio/VCProjectFileGenerator.cs
@@ -10,6 +10,8 @@ using System.Xml.Linq;
 using EpicGames.Core;
 using UnrealBuildBase;
 using Microsoft.Extensions.Logging;
+using System.Reflection.Metadata;
+using System.Xml;
 
 namespace UnrealBuildTool
 {
@@ -498,7 +500,7 @@ namespace UnrealBuildTool
 				XNamespace NS = XNamespace.Get("http://schemas.microsoft.com/developer/msbuild/2003");
 
 				DirectoryReference AutomationToolDir = DirectoryReference.Combine(Unreal.EngineSourceDirectory, "Programs", "AutomationTool");
-				new XDocument(
+				XDocument AutomationToolDocument = new XDocument(
 					new XElement(NS + "Project",
 						new XAttribute("ToolsVersion", VCProjectFileGenerator.GetProjectFileToolVersionString(Settings.ProjectFileFormat)),
 						new XAttribute("DefaultTargets", "Build"),
@@ -509,7 +511,21 @@ namespace UnrealBuildTool
 							)
 						)
 					)
-				).Save(FileReference.Combine(IntermediateProjectFilesPath, "AutomationTool.csproj.References").FullName);
+				);
+				
+				StringBuilder Output = new StringBuilder();
+				Output.Append("<?xml version=\"1.0\" encoding=\"utf-8\"?>\n");
+
+				XmlWriterSettings XmlSettings = new XmlWriterSettings();
+				XmlSettings.Encoding = new UTF8Encoding(false);
+				XmlSettings.Indent = true;
+				XmlSettings.OmitXmlDeclaration = true;
+
+				using (XmlWriter Writer = XmlWriter.Create(Output, XmlSettings))
+				{
+					AutomationToolDocument.Save(Writer);
+				}
+				Utils.WriteFileIfChanged(FileReference.Combine(IntermediateProjectFilesPath, "AutomationTool.csproj.References"), Output.ToString(), Logger);
 			}
 
 			return true;
