Fixed an issue where vs project platform properties would leak into other platforms.

diff --git a/Engine/Source/Programs/UnrealBuildTool/ProjectFiles/VisualStudio/VCProject.cs b/Engine/Source/Programs/UnrealBuildTool/ProjectFiles/VisualStudio/VCProject.cs
--- a/Engine/Source/Programs/UnrealBuildTool/ProjectFiles/VisualStudio/VCProject.cs
+++ b/Engine/Source/Programs/UnrealBuildTool/ProjectFiles/VisualStudio/VCProject.cs
@@ -966,6 +966,7 @@ namespace UnrealBuildTool
 			List<UnrealTargetPlatform> ProjectPlatforms = new List<UnrealTargetPlatform>();
 			List<Tuple<string, UnrealTargetPlatform>> ProjectPlatformNameAndPlatforms = new List<Tuple<string, UnrealTargetPlatform>>();    // ProjectPlatformName, Platform
 			List<Tuple<string, UnrealTargetConfiguration>> ProjectConfigurationNameAndConfigurations = new List<Tuple<string, UnrealTargetConfiguration>>();    // ProjectConfigurationName, Configuration
+			List<ProjectConfigAndTargetCombination> ValidProjectConfigAndTargetCombinations = new List<ProjectConfigAndTargetCombination>();
 			foreach (ProjectConfigAndTargetCombination Combination in ProjectConfigAndTargetCombinations!)
 			{
 				if (Combination.Platform == null)
@@ -988,6 +989,7 @@ namespace UnrealBuildTool
 				{
 					ProjectConfigurationNameAndConfigurations.Add(Tuple.Create(Combination.ProjectConfigurationName, Combination.Configuration));
 				}
+				ValidProjectConfigAndTargetCombinations.Add(Combination);
 			}
 
 			// Add the "invalid" configuration for each platform. We use this when the solution configuration does not match any project configuration.
@@ -1059,15 +1061,11 @@ namespace UnrealBuildTool
 			}
 
 			// Write each project configuration PreDefaultProps section
-			foreach (Tuple<string, UnrealTargetConfiguration> ConfigurationTuple in ProjectConfigurationNameAndConfigurations)
+			foreach (ProjectConfigAndTargetCombination Combination in ValidProjectConfigAndTargetCombinations)
 			{
-				string ProjectConfigurationName = ConfigurationTuple.Item1;
-				UnrealTargetConfiguration TargetConfiguration = ConfigurationTuple.Item2;
-				foreach (Tuple<string, UnrealTargetPlatform> PlatformTuple in ProjectPlatformNameAndPlatforms)
+				if (Combination.Platform != null)
 				{
-					string ProjectPlatformName = PlatformTuple.Item1;
-					UnrealTargetPlatform TargetPlatform = PlatformTuple.Item2;
-					WritePreDefaultPropsConfiguration(TargetPlatform, TargetConfiguration, ProjectPlatformName, ProjectConfigurationName, PlatformProjectGenerators, VCProjectFileContent);
+					WritePreDefaultPropsConfiguration(Combination.Platform.Value, Combination.Configuration, Combination.ProjectPlatformName, Combination.ProjectConfigurationName, PlatformProjectGenerators, VCProjectFileContent);
 				}
 			}
 
