Updated several places in UBT to be thread safe.

diff --git a/Engine/Source/Programs/Shared/EpicGames.Core/AssemblyUtils.cs b/Engine/Source/Programs/Shared/EpicGames.Core/AssemblyUtils.cs
--- a/Engine/Source/Programs/Shared/EpicGames.Core/AssemblyUtils.cs
+++ b/Engine/Source/Programs/Shared/EpicGames.Core/AssemblyUtils.cs
@@ -1,6 +1,7 @@
 // Copyright Epic Games, Inc. All Rights Reserved.
 
 using System;
+using System.Collections.Concurrent;
 using System.Collections.Generic;
 using System.Diagnostics;
 using System.IO;
@@ -100,20 +101,23 @@ namespace EpicGames.Core
 
 			string assemblyName = Path.GetFileNameWithoutExtension(assemblyPath);
 			DateTime assemblyLastWriteTime = File.GetLastWriteTimeUtc(assemblyPath);
-			if (s_assemblyLocationCache.ContainsKey(assemblyName))
+			lock(s_assemblyLocationCache)
 			{
-				// We already have this assembly in our cache. Only replace it if the discovered file is newer (to avoid stale assemblies breaking stuff).
-				if (assemblyLastWriteTime > s_assemblyWriteTimes[assemblyName])
+				if (s_assemblyLocationCache.ContainsKey(assemblyName))
 				{
-					s_assemblyLocationCache[assemblyName] = assemblyPath;
-					s_assemblyWriteTimes[assemblyName] = assemblyLastWriteTime;
+					// We already have this assembly in our cache. Only replace it if the discovered file is newer (to avoid stale assemblies breaking stuff).
+					if (assemblyLastWriteTime > s_assemblyWriteTimes[assemblyName])
+					{
+						s_assemblyLocationCache[assemblyName] = assemblyPath;
+						s_assemblyWriteTimes[assemblyName] = assemblyLastWriteTime;
+					}
+				}
+				else
+				{
+					// This is the first copy of this assembly ... add it to our cache.
+					s_assemblyLocationCache.Add(assemblyName, assemblyPath);
+					s_assemblyWriteTimes.Add(assemblyName, assemblyLastWriteTime);
 				}
-			}
-			else
-			{
-				// This is the first copy of this assembly ... add it to our cache.
-				s_assemblyLocationCache.Add(assemblyName, assemblyPath);
-				s_assemblyWriteTimes.Add(assemblyName, assemblyLastWriteTime);
 			}
 
 			if (!s_addedToAssemblyResolver)
@@ -122,18 +126,21 @@ namespace EpicGames.Core
 				{
 					// Name is fully qualified assembly definition - e.g. "p4dn, Version=1.0.0.0, Culture=neutral, PublicKeyToken=ff968dc1933aba6f"
 					string assemblyName = args.Name!.Split(',')[0];
-					if (s_assemblyLocationCache.TryGetValue(assemblyName, out string? assemblyLocation))
+					lock (s_assemblyLocationCache)
 					{
-						// We have this assembly in our folder.					
-						if (File.Exists(assemblyLocation))
-						{
-							// The assembly still exists, so load it.
-							return Assembly.LoadFile(assemblyLocation);
-						}
-						else
+						if (s_assemblyLocationCache.TryGetValue(assemblyName, out string? assemblyLocation))
 						{
-							// The assembly no longer exists on disk, so remove it from our cache.
-							s_assemblyLocationCache.Remove(assemblyName);
+							// We have this assembly in our folder.					
+							if (File.Exists(assemblyLocation))
+							{
+								// The assembly still exists, so load it.
+								return Assembly.LoadFile(assemblyLocation);
+							}
+							else
+							{
+								// The assembly no longer exists on disk, so remove it from our cache.
+								s_assemblyLocationCache.Remove(assemblyName);
+							}
 						}
 					}
 
diff --git a/Engine/Source/Programs/UnrealBuildTool/System/ConfigCache.cs b/Engine/Source/Programs/UnrealBuildTool/System/ConfigCache.cs
--- a/Engine/Source/Programs/UnrealBuildTool/System/ConfigCache.cs
+++ b/Engine/Source/Programs/UnrealBuildTool/System/ConfigCache.cs
@@ -11,6 +11,7 @@ using System.Threading.Tasks;
 using EpicGames.Core;
 using UnrealBuildTool;
 using Microsoft.Extensions.Logging;
+using System.Collections.Concurrent;
 
 namespace UnrealBuildTool
 {
@@ -226,17 +227,17 @@ namespace UnrealBuildTool
 		/// <summary>
 		/// Cache of individual config files
 		/// </summary>
-		static Dictionary<FileReference, ConfigFile?> LocationToConfigFile = new Dictionary<FileReference, ConfigFile?>();
+		static ConcurrentDictionary<FileReference, ConfigFile?> LocationToConfigFile = new ConcurrentDictionary<FileReference, ConfigFile?>();
 
 		/// <summary>
 		/// Cache of config hierarchies by project
 		/// </summary>
-		static Dictionary<ConfigHierarchyKey, ConfigHierarchy> HierarchyKeyToHierarchy = new Dictionary<ConfigHierarchyKey, ConfigHierarchy>();
+		static ConcurrentDictionary<ConfigHierarchyKey, ConfigHierarchy> HierarchyKeyToHierarchy = new ConcurrentDictionary<ConfigHierarchyKey, ConfigHierarchy>();
 
 		/// <summary>
 		/// Cache of config fields by type
 		/// </summary>
-		static Dictionary<Type, List<ConfigMember>> TypeToConfigMembers = new Dictionary<Type, List<ConfigMember>>();
+		static ConcurrentDictionary<Type, List<ConfigMember>> TypeToConfigMembers = new ConcurrentDictionary<Type, List<ConfigMember>>();
 
 		/// <summary>
 		/// Attempts to read a config file (or retrieve it from the cache)
@@ -246,24 +247,16 @@ namespace UnrealBuildTool
 		/// <returns>True if the file exists and was read, false otherwise</returns>
 		internal static bool TryReadFile(FileReference Location, [NotNullWhen(true)] out ConfigFile? ConfigFile)
 		{
-			lock (LocationToConfigFile)
+			ConfigFile = LocationToConfigFile.GetOrAdd(Location, _ =>
 			{
-				if (!LocationToConfigFile.TryGetValue(Location, out ConfigFile))
+				if (FileReference.Exists(Location))
 				{
-					if (FileReference.Exists(Location))
-					{
-						ConfigFile = new ConfigFile(Location);
-					}
-					else
-					{
-						ConfigFile = null;
-					}
-
-					// Note: it's important to keep the null value here so we don't keep checking for files that don't exist
-					LocationToConfigFile.Add(Location, ConfigFile);
+					return new ConfigFile(Location);
 				}
-			}
 
+				// Note: it's important to keep the null value here so we don't keep checking for files that don't exist
+				return null;
+			});
 			return ConfigFile != null;
 		}
 
@@ -318,34 +311,28 @@ namespace UnrealBuildTool
 			ConfigHierarchyKey Key = new ConfigHierarchyKey(Type, ProjectDir, Platform, CustomConfig, OverrideStrings);
 
 			// Try to get the cached hierarchy with this key
-			ConfigHierarchy? Hierarchy;
-			lock (HierarchyKeyToHierarchy)
+			return HierarchyKeyToHierarchy.GetOrAdd(Key, _ =>
 			{
-				if (!HierarchyKeyToHierarchy.TryGetValue(Key, out Hierarchy))
+				// Find all the input files
+				List<ConfigFile> Files = new List<ConfigFile>();
+				foreach (FileReference IniFileName in ConfigHierarchy.EnumerateConfigFileLocations(Type, ProjectDir, Platform, CustomConfig))
 				{
-					// Find all the input files
-					List<ConfigFile> Files = new List<ConfigFile>();
-					foreach (FileReference IniFileName in ConfigHierarchy.EnumerateConfigFileLocations(Type, ProjectDir, Platform, CustomConfig))
+					ConfigFile? File;
+					if (TryReadFile(IniFileName, out File))
 					{
-						ConfigFile? File;
-						if (TryReadFile(IniFileName, out File))
-						{
-							Files.Add(File);
-						}
-					}
-
-					foreach (string OverrideString in OverrideStrings)
-					{
-						ConfigFile OverrideFile = new ConfigFile(OverrideString);
-						Files.Add(OverrideFile);
+						Files.Add(File);
 					}
+				}
 
-					// Create the hierarchy
-					Hierarchy = new ConfigHierarchy(Files);
-					HierarchyKeyToHierarchy.Add(Key, Hierarchy);
+				foreach (string OverrideString in OverrideStrings)
+				{
+					ConfigFile OverrideFile = new ConfigFile(OverrideString);
+					Files.Add(OverrideFile);
 				}
-			}
-			return Hierarchy;
+
+				// Create the hierarchy
+				return new ConfigHierarchy(Files);
+			});
 		}
 	
 		/// <summary>
@@ -355,28 +342,23 @@ namespace UnrealBuildTool
 		/// <returns>List of config fields for the given type</returns>
 		static List<ConfigMember> FindConfigMembersForType(Type TargetObjectType)
 		{
-			List<ConfigMember>? Members;
-			lock(TypeToConfigMembers)
+			return TypeToConfigMembers.GetOrAdd(TargetObjectType, _ =>
 			{
-				if (!TypeToConfigMembers.TryGetValue(TargetObjectType, out Members))
+				List<ConfigMember> Members = new List<ConfigMember>();
+				if (TargetObjectType.BaseType != null)
 				{
-					Members = new List<ConfigMember>();
-					if(TargetObjectType.BaseType != null)
-					{
-						Members.AddRange(FindConfigMembersForType(TargetObjectType.BaseType));
-					}
-					foreach (FieldInfo FieldInfo in TargetObjectType.GetFields(BindingFlags.Instance | BindingFlags.GetField | BindingFlags.GetProperty | BindingFlags.Public | BindingFlags.NonPublic | BindingFlags.DeclaredOnly))
-					{
-						ProcessConfigTypeMember<FieldInfo>(TargetObjectType, FieldInfo, Members, (FieldInfo, Attribute) => new ConfigField(FieldInfo, Attribute));
-					}
-					foreach (PropertyInfo PropertyInfo in TargetObjectType.GetProperties(BindingFlags.Instance | BindingFlags.GetField | BindingFlags.GetProperty | BindingFlags.Public | BindingFlags.NonPublic | BindingFlags.DeclaredOnly))
-					{
-						ProcessConfigTypeMember<PropertyInfo>(TargetObjectType, PropertyInfo, Members, (PropertyInfo, Attribute) => new ConfigProperty(PropertyInfo, Attribute));
-					}
-					TypeToConfigMembers.Add(TargetObjectType, Members);
+					Members.AddRange(FindConfigMembersForType(TargetObjectType.BaseType));
 				}
-			}
-			return Members;
+				foreach (FieldInfo FieldInfo in TargetObjectType.GetFields(BindingFlags.Instance | BindingFlags.GetField | BindingFlags.GetProperty | BindingFlags.Public | BindingFlags.NonPublic | BindingFlags.DeclaredOnly))
+				{
+					ProcessConfigTypeMember<FieldInfo>(TargetObjectType, FieldInfo, Members, (FieldInfo, Attribute) => new ConfigField(FieldInfo, Attribute));
+				}
+				foreach (PropertyInfo PropertyInfo in TargetObjectType.GetProperties(BindingFlags.Instance | BindingFlags.GetField | BindingFlags.GetProperty | BindingFlags.Public | BindingFlags.NonPublic | BindingFlags.DeclaredOnly))
+				{
+					ProcessConfigTypeMember<PropertyInfo>(TargetObjectType, PropertyInfo, Members, (PropertyInfo, Attribute) => new ConfigProperty(PropertyInfo, Attribute));
+				}
+				return Members;
+			});
 		}
 
 		static void ProcessConfigTypeMember<MEMBER>(Type TargetType, MEMBER MemberInfo, List<ConfigMember> Members, Func<MEMBER, ConfigFileAttribute, ConfigMember> CreateConfigMember)
diff --git a/Engine/Source/Programs/UnrealBuildTool/System/ConfigHierarchy.cs b/Engine/Source/Programs/UnrealBuildTool/System/ConfigHierarchy.cs
--- a/Engine/Source/Programs/UnrealBuildTool/System/ConfigHierarchy.cs
+++ b/Engine/Source/Programs/UnrealBuildTool/System/ConfigHierarchy.cs
@@ -195,11 +195,27 @@ namespace UnrealBuildTool
         /// </summary>
         System.Threading.ReaderWriterLockSlim NameToSectionLock = new System.Threading.ReaderWriterLockSlim();
 
-        /// <summary>
-        /// Construct a config hierarchy from the given files
-        /// </summary>
-        /// <param name="Files">Set of files to include (in order)</param>
-        public ConfigHierarchy(IEnumerable<ConfigFile> Files)
+		static readonly string? PersonalConfigFolder = null;
+
+		static ConfigHierarchy()
+		{
+			// Some user accounts (eg. SYSTEM on Windows) don't have a home directory. Ignore them if Environment.GetFolderPath() returns an empty string.
+			string PersonalFolder = Environment.GetFolderPath(Environment.SpecialFolder.Personal);
+			if (!String.IsNullOrEmpty(PersonalFolder))
+			{
+				PersonalConfigFolder = PersonalFolder;
+				if (RuntimePlatform.IsMac || RuntimePlatform.IsLinux)
+				{
+					PersonalConfigFolder = System.IO.Path.Combine(PersonalConfigFolder, "Documents");
+				}
+			}
+		}
+
+		/// <summary>
+		/// Construct a config hierarchy from the given files
+		/// </summary>
+		/// <param name="Files">Set of files to include (in order)</param>
+		public ConfigHierarchy(IEnumerable<ConfigFile> Files)
 		{
 			this.Files = Files.ToArray();
 		}
@@ -968,22 +984,9 @@ namespace UnrealBuildTool
 		// Match FPlatformProcess::UserDir()
 		private static string? GetUserDir()
 		{
-			// Some user accounts (eg. SYSTEM on Windows) don't have a home directory. Ignore them if Environment.GetFolderPath() returns an empty string.
-			string PersonalFolder = Environment.GetFolderPath(Environment.SpecialFolder.Personal);
-			string? PersonalConfigFolder = null;
-			if (!String.IsNullOrEmpty(PersonalFolder))
-			{
-				PersonalConfigFolder = PersonalFolder;
-				if (RuntimePlatform.IsMac || RuntimePlatform.IsLinux)
-				{
-					PersonalConfigFolder = System.IO.Path.Combine(PersonalConfigFolder, "Documents");
-				}
-			}
-
 			return PersonalConfigFolder;
 		}
 
-
 		private static string PerformBasicReplacements(string InString, string BaseIniName, string CustomConfig)
 		{
 			string OutString = InString.Replace("{TYPE}", BaseIniName);
