Updating XGE coordinator event matchers

Coordinator message changed, also adding a post build summary regex to catch messages as well to tag as systemic

#rnx

diff --git a/Engine/Source/Programs/UnrealBuildTool/Matchers/XoreaxEventMatcher.cs b/Engine/Source/Programs/UnrealBuildTool/Matchers/XoreaxEventMatcher.cs
--- a/Engine/Source/Programs/UnrealBuildTool/Matchers/XoreaxEventMatcher.cs
+++ b/Engine/Source/Programs/UnrealBuildTool/Matchers/XoreaxEventMatcher.cs
@@ -25,6 +25,9 @@ namespace UnrealBuildTool.Matchers
 		static readonly Regex s_buildSystemWarning = new Regex(@"^\s*--------------------Build System Warning[- ]");
 		static readonly Regex s_buildSystemWarningNext = new Regex(@"^(\s*)([^ ].*):");
 
+		static readonly Regex s_buildSystemWarningSummary = new Regex(@"^\d+ build system warning\(s\):");
+		static readonly Regex s_buildSystemWarningSummaryNext = new Regex(@"^(\s*-\s*)([^ ].*)");
+
 		static readonly Regex s_cacheWarning = new Regex(@"^\s*WARNING: \d+ items \([^\)]*\) removed from the cache due to reaching the cache size limit");
 		static readonly Regex s_cacheWarning2 = new Regex(@"^\s*WARNING: Several items removed from the cache due to reaching the cache size limit");
 		static readonly Regex s_cacheWarning3 = new Regex(@"^\s*WARNING: The Build Cache is close to full");
@@ -76,7 +79,32 @@ namespace UnrealBuildTool.Matchers
 					string message = prefix.Groups[2].Value;
 
 					EventId eventId = KnownLogEvents.Systemic_Xge;
-					if (Regex.IsMatch(message, "Failed to connect to Coordinator"))
+					if (Regex.IsMatch(message, "Failed to connect to Coordinator") || Regex.IsMatch(message, "Failed to connect to Primary Coordinator"))
+					{
+						eventId = KnownLogEvents.Systemic_Xge_Standalone;
+					}
+
+					Regex prefixPattern = new Regex($"^{Regex.Escape(prefix.Groups[1].Value)}\\s");
+					while (builder.Next.IsMatch(prefixPattern))
+					{
+						builder.MoveNext();
+					}
+
+					return builder.ToMatch(LogEventPriority.High, LogLevel.Information, eventId);
+				}
+			}
+
+			if (cursor.IsMatch(s_buildSystemWarningSummary))
+			{
+				LogEventBuilder builder = new LogEventBuilder(cursor);
+				if (builder.Next.TryMatch(s_buildSystemWarningSummaryNext, out Match? prefix))
+				{
+					builder.MoveNext();
+
+					string message = prefix.Groups[2].Value;
+
+					EventId eventId = KnownLogEvents.Systemic_Xge;
+					if (Regex.IsMatch(message, "Failed to connect to Primary Coordinator"))
 					{
 						eventId = KnownLogEvents.Systemic_Xge_Standalone;
 					}
