UnrealBuildTool: Don't pass /Zm unless specifically overridden, documentation and testing shows that this flag should not be passed except in very limited circumstances

diff --git a/Engine/Source/Programs/UnrealBuildTool/Platform/Windows/VCToolChain.cs b/Engine/Source/Programs/UnrealBuildTool/Platform/Windows/VCToolChain.cs
--- a/Engine/Source/Programs/UnrealBuildTool/Platform/Windows/VCToolChain.cs
+++ b/Engine/Source/Programs/UnrealBuildTool/Platform/Windows/VCToolChain.cs
@@ -430,24 +430,12 @@ namespace UnrealBuildTool
 			// Separate functions for linker.
 			Arguments.Add("/Gy");
 
-			// Allow 750% of the default memory allocation limit when using the static analyzer, and 1000% at other times.
-			if(Target.WindowsPlatform.PCHMemoryAllocationFactor == 0)
-			{
-				if (Target.StaticAnalyzer == StaticAnalyzer.Default)
-				{
-					Arguments.Add("/Zm750");
-				}
-				else
-				{
-					Arguments.Add("/Zm1000");
-				}
-			}
-			else
+
+			// Microsoft recommends not passing /Zm except in very limited circumstances, please see:
+			// https://learn.microsoft.com/en-us/cpp/build/reference/zm-specify-precompiled-header-memory-allocation-limit
+			if (Target.WindowsPlatform.PCHMemoryAllocationFactor > 0)
 			{
-				if(Target.WindowsPlatform.PCHMemoryAllocationFactor > 0)
-				{
-					Arguments.Add($"/Zm{Target.WindowsPlatform.PCHMemoryAllocationFactor}");
-				}
+				Arguments.Add($"/Zm{Target.WindowsPlatform.PCHMemoryAllocationFactor}");
 			}
 
 			// Disable "The file contains a character that cannot be represented in the current code page" warning for non-US windows.
