EpicGames.Core: Obsolete ReadOnlyHashSet in favor of IReadOnlySet

diff --git a/Engine/Source/Programs/Horde/Horde.Server/Compute/RedisTaskScheduler.cs b/Engine/Source/Programs/Horde/Horde.Server/Compute/RedisTaskScheduler.cs
--- a/Engine/Source/Programs/Horde/Horde.Server/Compute/RedisTaskScheduler.cs
+++ b/Engine/Source/Programs/Horde/Horde.Server/Compute/RedisTaskScheduler.cs
@@ -7,7 +7,6 @@ using System.Linq;
 using System.Threading;
 using System.Threading.Channels;
 using System.Threading.Tasks;
-using EpicGames.Core;
 using EpicGames.Redis;
 using Microsoft.Extensions.Logging;
 using StackExchange.Redis;
@@ -89,7 +88,7 @@ namespace Horde.Server.Compute
 		readonly RedisKey _baseKey;
 		readonly RedisSetKey<TQueueId> _queueIndex;
 		readonly RedisHashKey<TQueueId, DateTime> _activeQueues; // Queues which are actively being dequeued from
-		ReadOnlyHashSet<TQueueId> _localActiveQueues = new HashSet<TQueueId>();
+		IReadOnlySet<TQueueId> _localActiveQueues = new HashSet<TQueueId>();
 		readonly Stopwatch _resetActiveQueuesTimer = Stopwatch.StartNew();
 		readonly List<Listener> _listeners = new List<Listener>();
 		readonly RedisChannel<TQueueId> _newQueueChannel;
@@ -332,7 +331,7 @@ namespace Horde.Server.Compute
 			// readers can thus access it without the need for any locking.
 			for(; ;)
 			{
-				ReadOnlyHashSet<TQueueId> localActiveQueuesCopy = _localActiveQueues;
+				IReadOnlySet<TQueueId> localActiveQueuesCopy = _localActiveQueues;
 				if (localActiveQueuesCopy.Contains(queueId))
 				{
 					break;
diff --git a/Engine/Source/Programs/Shared/EpicGames.Core/ReadOnlyHashSet.cs b/Engine/Source/Programs/Shared/EpicGames.Core/ReadOnlyHashSet.cs
--- a/Engine/Source/Programs/Shared/EpicGames.Core/ReadOnlyHashSet.cs
+++ b/Engine/Source/Programs/Shared/EpicGames.Core/ReadOnlyHashSet.cs
@@ -1,5 +1,6 @@
 // Copyright Epic Games, Inc. All Rights Reserved.
 
+using System;
 using System.Collections;
 using System.Collections.Generic;
 
@@ -9,6 +10,7 @@ namespace EpicGames.Core
 	/// Wrapper around the HashSet container that only allows read operations
 	/// </summary>
 	/// <typeparam name="T">Type of element for the hashset</typeparam>
+	[Obsolete("ReadOnlyHashSet<T> is obsolete, please replace with IReadOnlySet<T>")]
 	public class ReadOnlyHashSet<T> : IReadOnlyCollection<T>
 	{
 		/// <summary>
diff --git a/Engine/Source/Programs/UnrealBuildTool/Configuration/TargetRules.cs b/Engine/Source/Programs/UnrealBuildTool/Configuration/TargetRules.cs
--- a/Engine/Source/Programs/UnrealBuildTool/Configuration/TargetRules.cs
+++ b/Engine/Source/Programs/UnrealBuildTool/Configuration/TargetRules.cs
@@ -2769,7 +2769,7 @@ namespace UnrealBuildTool
 			get { return Inner.TargetSourceFile!; }
 		}
 
-		internal ReadOnlyHashSet<FileReference> TargetFiles
+		internal IReadOnlySet<FileReference> TargetFiles
 		{
 			get { return Inner.TargetFiles!; }
 		}
diff --git a/Engine/Source/Programs/UnrealBuildTool/Configuration/UEBuildModuleCPP.cs b/Engine/Source/Programs/UnrealBuildTool/Configuration/UEBuildModuleCPP.cs
--- a/Engine/Source/Programs/UnrealBuildTool/Configuration/UEBuildModuleCPP.cs
+++ b/Engine/Source/Programs/UnrealBuildTool/Configuration/UEBuildModuleCPP.cs
@@ -257,7 +257,7 @@ namespace UnrealBuildTool
 				{
 					PublicIncludePaths.Add(PublicDirectory);
 
-					ReadOnlyHashSet<string> ExcludeNames = UEBuildPlatform.GetBuildPlatform(Rules.Target.Platform).GetExcludedFolderNames();
+					IReadOnlySet<string> ExcludeNames = UEBuildPlatform.GetBuildPlatform(Rules.Target.Platform).GetExcludedFolderNames();
 					EnumerateLegacyIncludePaths(DirectoryItem.GetItemByDirectoryReference(PublicDirectory), ExcludeNames, LegacyPublicIncludePaths);
 				}
 
@@ -283,7 +283,7 @@ namespace UnrealBuildTool
 		/// <param name="BaseDirectory">The directory to start from. This directory is not added to the output list.</param>
 		/// <param name="ExcludeNames">Set of folder names to exclude from the search.</param>
 		/// <param name="LegacyPublicIncludePaths">List populated with the discovered directories</param>
-		static void EnumerateLegacyIncludePaths(DirectoryItem BaseDirectory, ReadOnlyHashSet<string> ExcludeNames, HashSet<DirectoryReference> LegacyPublicIncludePaths)
+		static void EnumerateLegacyIncludePaths(DirectoryItem BaseDirectory, IReadOnlySet<string> ExcludeNames, HashSet<DirectoryReference> LegacyPublicIncludePaths)
 		{
 			foreach(DirectoryItem SubDirectory in BaseDirectory.EnumerateDirectories())
 			{
@@ -1829,7 +1829,7 @@ namespace UnrealBuildTool
 		/// <returns>Set of source files that should be built</returns>
 		public InputFileCollection FindInputFiles(UnrealTargetPlatform Platform, Dictionary<DirectoryItem, FileItem[]> DirectoryToSourceFiles, ILogger Logger)
 		{
-			ReadOnlyHashSet<string> ExcludedNames = UEBuildPlatform.GetBuildPlatform(Platform).GetExcludedFolderNames();
+			IReadOnlySet<string> ExcludedNames = UEBuildPlatform.GetBuildPlatform(Platform).GetExcludedFolderNames();
 
 			InputFileCollection InputFiles = new InputFileCollection();
 
@@ -1852,7 +1852,7 @@ namespace UnrealBuildTool
 		/// <param name="DirectoryToSourceFiles">Map from directory to source files inside it</param>
 		/// <param name="InputFiles">Collection of source files, categorized by type</param>
 		/// <param name="Logger">Logger for output</param>
-		static void FindInputFilesFromDirectoryRecursive(DirectoryItem BaseDirectory, ReadOnlyHashSet<string> ExcludedNames, HashSet<DirectoryReference> SourceDirectories, Dictionary<DirectoryItem, FileItem[]> DirectoryToSourceFiles, InputFileCollection InputFiles, ILogger Logger)
+		static void FindInputFilesFromDirectoryRecursive(DirectoryItem BaseDirectory, IReadOnlySet<string> ExcludedNames, HashSet<DirectoryReference> SourceDirectories, Dictionary<DirectoryItem, FileItem[]> DirectoryToSourceFiles, InputFileCollection InputFiles, ILogger Logger)
 		{
 			bool bIgnoreFileFound;
 			FileItem[] SourceFiles = FindInputFilesFromDirectory(BaseDirectory, InputFiles, out bIgnoreFileFound, Logger);
diff --git a/Engine/Source/Programs/UnrealBuildTool/Configuration/UEBuildPlatform.cs b/Engine/Source/Programs/UnrealBuildTool/Configuration/UEBuildPlatform.cs
--- a/Engine/Source/Programs/UnrealBuildTool/Configuration/UEBuildPlatform.cs
+++ b/Engine/Source/Programs/UnrealBuildTool/Configuration/UEBuildPlatform.cs
@@ -213,12 +213,12 @@ namespace UnrealBuildTool
 		/// <summary>
 		/// Cached copy of the list of folders to include for this platform
 		/// </summary>
-		private ReadOnlyHashSet<string>? CachedIncludedFolderNames;
+		private IReadOnlySet<string>? CachedIncludedFolderNames;
 
 		/// <summary>
 		/// Cached copy of the list of folders to exclude for this platform
 		/// </summary>
-		private ReadOnlyHashSet<string>? CachedExcludedFolderNames;
+		private IReadOnlySet<string>? CachedExcludedFolderNames;
 
 		/// <summary>
 		/// Constructor.
@@ -320,7 +320,7 @@ namespace UnrealBuildTool
 		/// <summary>
 		/// Finds a list of folder names to include when building for this platform
 		/// </summary>
-		public ReadOnlyHashSet<string> GetIncludedFolderNames()
+		public IReadOnlySet<string> GetIncludedFolderNames()
 		{
 			if(CachedIncludedFolderNames == null)
 			{
@@ -332,7 +332,7 @@ namespace UnrealBuildTool
 					Names.Add(Group.ToString());
 				}
 
-				CachedIncludedFolderNames = new ReadOnlyHashSet<string>(Names, DirectoryReference.Comparer);
+				CachedIncludedFolderNames = Names;
 			}
 			return CachedIncludedFolderNames;
 		}
@@ -340,11 +340,11 @@ namespace UnrealBuildTool
 		/// <summary>
 		/// Finds a list of folder names to exclude when building for this platform
 		/// </summary>
-		public ReadOnlyHashSet<string> GetExcludedFolderNames()
+		public IReadOnlySet<string> GetExcludedFolderNames()
 		{
 			if(CachedExcludedFolderNames == null)
 			{
-				CachedExcludedFolderNames = new ReadOnlyHashSet<string>(GetPlatformFolderNames().Except(GetIncludedFolderNames()), DirectoryReference.Comparer);
+				CachedExcludedFolderNames = new HashSet<string>(GetPlatformFolderNames().Except(GetIncludedFolderNames()), DirectoryReference.Comparer);
 			}
 			return CachedExcludedFolderNames;
 		}
diff --git a/Engine/Source/Programs/UnrealBuildTool/Configuration/UEBuildTarget.cs b/Engine/Source/Programs/UnrealBuildTool/Configuration/UEBuildTarget.cs
--- a/Engine/Source/Programs/UnrealBuildTool/Configuration/UEBuildTarget.cs
+++ b/Engine/Source/Programs/UnrealBuildTool/Configuration/UEBuildTarget.cs
@@ -3533,7 +3533,7 @@ namespace UnrealBuildTool
 			}
 
 			// Find all the platform folders to exclude from the list of valid modules
-			ReadOnlyHashSet<string> ExcludeFolders = UEBuildPlatform.GetBuildPlatform(Platform).GetExcludedFolderNames();
+			IReadOnlySet<string> ExcludeFolders = UEBuildPlatform.GetBuildPlatform(Platform).GetExcludedFolderNames();
 
 			// Set of module names to build
 			HashSet<string> FilteredModuleNames = new HashSet<string>();
diff --git a/Engine/Source/Programs/UnrealBuildTool/Modes/UnrealHeaderToolMode.cs b/Engine/Source/Programs/UnrealBuildTool/Modes/UnrealHeaderToolMode.cs
--- a/Engine/Source/Programs/UnrealBuildTool/Modes/UnrealHeaderToolMode.cs
+++ b/Engine/Source/Programs/UnrealBuildTool/Modes/UnrealHeaderToolMode.cs
@@ -37,18 +37,18 @@ namespace UnrealBuildTool.Modes
 		/// <summary>
 		/// Supported units in the game
 		/// </summary>
-		private readonly ReadOnlyHashSet<StringView> _units;
+		private readonly IReadOnlySet<StringView> _units;
 
 		/// <summary>
 		/// Special parsed struct names that do not require a prefix
 		/// </summary>
 		[System.Diagnostics.CodeAnalysis.SuppressMessage("CodeQuality", "IDE0052:Remove unread private members", Justification = "<Pending>")]
-		private readonly ReadOnlyHashSet<StringView> _structsWithNoPrefix;
+		private readonly IReadOnlySet<StringView> _structsWithNoPrefix;
 
 		/// <summary>
 		/// Special parsed struct names that have a 'T' prefix
 		/// </summary>
-		private readonly ReadOnlyHashSet<StringView> _structsWithTPrefix;
+		private readonly IReadOnlySet<StringView> _structsWithTPrefix;
 
 		/// <summary>
 		/// Mapping from 'human-readable' macro substring to # of parameters for delegate declarations
@@ -401,7 +401,7 @@ namespace UnrealBuildTool.Modes
 			return List;
 		}
 
-		private ReadOnlyHashSet<StringView> GetHashSet(string Section, string Key, StringViewComparer Comparer)
+		private IReadOnlySet<StringView> GetHashSet(string Section, string Key, StringViewComparer Comparer)
 		{
 			HashSet<StringView> Set = new(Comparer);
 
diff --git a/Engine/Source/Programs/UnrealBuildTool/System/ExternalExecution.cs b/Engine/Source/Programs/UnrealBuildTool/System/ExternalExecution.cs
--- a/Engine/Source/Programs/UnrealBuildTool/System/ExternalExecution.cs
+++ b/Engine/Source/Programs/UnrealBuildTool/System/ExternalExecution.cs
@@ -354,7 +354,7 @@ namespace UnrealBuildTool
 		/// <param name="BaseDirectory">Base directory to search</param>
 		/// <param name="ExcludeFolders">Array of folders to exclude</param>
 		/// <param name="Headers">Receives the list of headers that was found</param>
-		static void FindHeaders(DirectoryItem BaseDirectory, ReadOnlyHashSet<string> ExcludeFolders, List<FileItem> Headers)
+		static void FindHeaders(DirectoryItem BaseDirectory, IReadOnlySet<string> ExcludeFolders, List<FileItem> Headers)
 		{
 			if (BaseDirectory.TryGetFile(".ubtignore", out FileItem? OutIgnoreFile))
 			{
@@ -427,7 +427,7 @@ namespace UnrealBuildTool
 			UHTModuleInfo[] ModuleInfoArray = new UHTModuleInfo[ModulesSortedByType.Count];
 			using(ThreadPoolWorkQueue Queue = new ThreadPoolWorkQueue())
 			{
-				ReadOnlyHashSet<string> ExcludedFolders = UEBuildPlatform.GetBuildPlatform(Platform).GetExcludedFolderNames();
+				IReadOnlySet<string> ExcludedFolders = UEBuildPlatform.GetBuildPlatform(Platform).GetExcludedFolderNames();
 				for(int Idx = 0; Idx < ModulesSortedByType.Count; Idx++)
 				{
 					UEBuildModuleCPP Module = ModulesSortedByType[Idx];
@@ -506,7 +506,7 @@ namespace UnrealBuildTool
 				Dictionary<UEBuildModuleCPP, UHTModuleInfo> IncludePathInfo = new();
 				using (ThreadPoolWorkQueue Queue = new ThreadPoolWorkQueue())
 				{
-					ReadOnlyHashSet<string> ExcludedFolders = UEBuildPlatform.GetBuildPlatform(Platform).GetExcludedFolderNames();
+					IReadOnlySet<string> ExcludedFolders = UEBuildPlatform.GetBuildPlatform(Platform).GetExcludedFolderNames();
 					foreach (var Module in IncludePathModules)
 					{
 						Queue.Enqueue(() =>
@@ -528,7 +528,7 @@ namespace UnrealBuildTool
 			}
 		}
 
-		static void SetupUObjectModule(UHTModuleInfo ModuleInfo, ReadOnlyHashSet<string> ExcludedFolders, SourceFileMetadataCache MetadataCache, ThreadPoolWorkQueue Queue)
+		static void SetupUObjectModule(UHTModuleInfo ModuleInfo, IReadOnlySet<string> ExcludedFolders, SourceFileMetadataCache MetadataCache, ThreadPoolWorkQueue Queue)
 		{
 			foreach (DirectoryReference ModuleDirectory in ModuleInfo.ModuleDirectories)
 			{
diff --git a/Engine/Source/Programs/UnrealBuildTool/System/TargetMakefile.cs b/Engine/Source/Programs/UnrealBuildTool/System/TargetMakefile.cs
--- a/Engine/Source/Programs/UnrealBuildTool/System/TargetMakefile.cs
+++ b/Engine/Source/Programs/UnrealBuildTool/System/TargetMakefile.cs
@@ -773,7 +773,7 @@ namespace UnrealBuildTool
 			using (GlobalTracer.Instance.BuildSpan("TargetMakefile.IsValidForSourceFiles()").StartActive())
 			{
 				// Get the list of excluded folder names for this platform
-				ReadOnlyHashSet<string> ExcludedFolderNames = UEBuildPlatform.GetBuildPlatform(Platform).GetExcludedFolderNames();
+				IReadOnlySet<string> ExcludedFolderNames = UEBuildPlatform.GetBuildPlatform(Platform).GetExcludedFolderNames();
 
 				// Load the metadata cache
 				SourceFileMetadataCache MetadataCache = SourceFileMetadataCache.CreateHierarchy(ProjectFile, Logger);
@@ -927,7 +927,7 @@ namespace UnrealBuildTool
 		/// <param name="ExcludedFolderNames">Set of directory names to exclude</param>
 		/// <param name="Logger">Logger for output diagnostics</param>
 		/// <returns>True if the directory contains any source files</returns>
-		static bool ContainsSourceFiles(DirectoryItem Directory, ReadOnlyHashSet<string> ExcludedFolderNames, ILogger Logger)
+		static bool ContainsSourceFiles(DirectoryItem Directory, IReadOnlySet<string> ExcludedFolderNames, ILogger Logger)
 		{
 			// Check this directory isn't ignored
 			if(!ExcludedFolderNames.Contains(Directory.Name))
@@ -959,7 +959,7 @@ namespace UnrealBuildTool
 		/// <param name="ExcludedFolderNames">Set of folder names to ignore when recursing the directory tree</param>
 		/// <param name="FilesWithMarkup">Receives the set of files which contain reflection markup</param>
 		/// <param name="Queue">Queue to add sub-tasks to</param>
-		static void FindFilesWithMarkup(DirectoryItem Directory, SourceFileMetadataCache MetadataCache, ReadOnlyHashSet<string> ExcludedFolderNames, ConcurrentBag<FileItem> FilesWithMarkup, ThreadPoolWorkQueue Queue)
+		static void FindFilesWithMarkup(DirectoryItem Directory, SourceFileMetadataCache MetadataCache, IReadOnlySet<string> ExcludedFolderNames, ConcurrentBag<FileItem> FilesWithMarkup, ThreadPoolWorkQueue Queue)
 		{
 			if (Directory.TryGetFile(".ubtignore", out FileItem? OutFile))
 			{
