Removed the redundant configuration data from the generated projects.

diff --git a/Engine/Source/Programs/UnrealBuildTool/Platform/Android/AndroidProjectGenerator.cs b/Engine/Source/Programs/UnrealBuildTool/Platform/Android/AndroidProjectGenerator.cs
--- a/Engine/Source/Programs/UnrealBuildTool/Platform/Android/AndroidProjectGenerator.cs
+++ b/Engine/Source/Programs/UnrealBuildTool/Platform/Android/AndroidProjectGenerator.cs
@@ -117,41 +117,6 @@ namespace UnrealBuildTool
 			}
 		}
 
-		/// <summary>
-		/// Return any custom property group lines
-		/// </summary>
-		/// <param name="InPlatform">  The UnrealTargetPlatform being built</param>
-		/// <param name="ProjectFileFormat"></param>
-		/// <returns>string    The custom property import lines for the project file; Empty string if it doesn't require one</returns>
-		public override string GetVisualStudioPlatformConfigurationType(UnrealTargetPlatform InPlatform, VCProjectFileFormat ProjectFileFormat)
-		{
-			string ConfigurationType = "";
-
-			if (AGDEInstalled)
-			{
-				ConfigurationType = "Makefile";
-			}
-			else
-			{
-				ConfigurationType = base.GetVisualStudioPlatformConfigurationType(InPlatform, ProjectFileFormat);
-			}
-
-			return ConfigurationType;
-		}
-
-		/// <summary>
-		/// Return the platform toolset string to write into the project configuration
-		/// </summary>
-		/// <param name="InPlatform">  The UnrealTargetPlatform being built</param>
-		/// <param name="InConfiguration"> The UnrealTargetConfiguration being built</param>
-		/// <param name="InProjectFileFormat">The version of Visual Studio to target</param>
-		/// <param name="ProjectFileBuilder">String builder for the project file</param>
-		/// <returns>string    The custom configuration section for the project file; Empty string if it doesn't require one</returns>
-		public override void GetVisualStudioPlatformToolsetString(UnrealTargetPlatform InPlatform, UnrealTargetConfiguration InConfiguration, VCProjectFileFormat InProjectFileFormat, StringBuilder ProjectFileBuilder)
-		{
-			VCProjectFileGenerator.AppendPlatformToolsetProperty(ProjectFileBuilder, InProjectFileFormat);
-		}
-
 		/// <summary>
 		/// Return any custom paths for VisualStudio this platform requires
 		/// This include ReferencePath, LibraryPath, LibraryWPath, IncludePath and ExecutablePath.
diff --git a/Engine/Source/Programs/UnrealBuildTool/ProjectFiles/PlatformProjectGenerator.cs b/Engine/Source/Programs/UnrealBuildTool/ProjectFiles/PlatformProjectGenerator.cs
--- a/Engine/Source/Programs/UnrealBuildTool/ProjectFiles/PlatformProjectGenerator.cs
+++ b/Engine/Source/Programs/UnrealBuildTool/ProjectFiles/PlatformProjectGenerator.cs
@@ -13,6 +13,7 @@ namespace UnrealBuildTool
 	/// </summary>
 	abstract class PlatformProjectGenerator
 	{
+		static public readonly string DefaultPlatformConfigurationType  = "Makefile";
 		protected readonly ILogger Logger;
 
 		/// <summary>
@@ -142,7 +143,7 @@ namespace UnrealBuildTool
 		/// <returns>string    The platform configuration type.  Defaults to "Makefile" unless overridden</returns>
 		public virtual string GetVisualStudioPlatformConfigurationType(UnrealTargetPlatform InPlatform, VCProjectFileFormat InProjectFileFormat)
 		{
-			return "Makefile";
+			return DefaultPlatformConfigurationType;
 		}
 
 		/// <summary>
diff --git a/Engine/Source/Programs/UnrealBuildTool/ProjectFiles/VisualStudio/VCProject.cs b/Engine/Source/Programs/UnrealBuildTool/ProjectFiles/VisualStudio/VCProject.cs
--- a/Engine/Source/Programs/UnrealBuildTool/ProjectFiles/VisualStudio/VCProject.cs
+++ b/Engine/Source/Programs/UnrealBuildTool/ProjectFiles/VisualStudio/VCProject.cs
@@ -1071,15 +1071,13 @@ namespace UnrealBuildTool
 
 			VCProjectFileContent.AppendLine("  <Import Project=\"$(VCTargetsPath)\\Microsoft.Cpp.Default.props\" />");
 
-			// Write the invalid configuration data
-			foreach (string InvalidConfigPlatformName in InvalidConfigPlatformNames)
-			{
-				VCProjectFileContent.AppendLine("  <PropertyGroup Condition=\"'$(Configuration)|$(Platform)'=='Invalid|{0}'\" Label=\"Configuration\">", InvalidConfigPlatformName);
-				VCProjectFileContent.AppendLine("    <ConfigurationType>Makefile</ConfigurationType>");
-				VCProjectFileContent.AppendLine("  </PropertyGroup>");
-			}
+			// Write the default configuration info
+			VCProjectFileContent.AppendLine("  <PropertyGroup Label=\"Configuration\">");
+			VCProjectFileContent.AppendLine($"    <ConfigurationType>{PlatformProjectGenerator.DefaultPlatformConfigurationType}</ConfigurationType>");
+			VCProjectFileGenerator.AppendPlatformToolsetProperty(VCProjectFileContent, ProjectFileFormat);
+			VCProjectFileContent.AppendLine("  </PropertyGroup>");
 
-			// Write each project configuration PreDefaultProps section
+			// Write the per platform/config configuration info
 			foreach (Tuple<string, UnrealTargetConfiguration> ConfigurationTuple in ProjectConfigurationNameAndConfigurations)
 			{
 				string ProjectConfigurationName = ConfigurationTuple.Item1;
@@ -1088,7 +1086,38 @@ namespace UnrealBuildTool
 				{
 					string ProjectPlatformName = PlatformTuple.Item1;
 					UnrealTargetPlatform TargetPlatform = PlatformTuple.Item2;
-					WritePostDefaultPropsConfiguration(TargetPlatform, TargetConfiguration, ProjectPlatformName, ProjectConfigurationName, PlatformProjectGenerators, VCProjectFileContent, Logger);
+
+					PlatformProjectGenerator? ProjGenerator = PlatformProjectGenerators.GetPlatformProjectGenerator(TargetPlatform, true);
+					if (ProjGenerator == null)
+					{
+						continue;
+					}
+
+					StringBuilder PlatformToolsetString = new StringBuilder();
+					ProjGenerator.GetVisualStudioPlatformToolsetString(TargetPlatform, TargetConfiguration, ProjectFileFormat, PlatformToolsetString);
+					string PlatformConfigurationType = ProjGenerator.GetVisualStudioPlatformConfigurationType(TargetPlatform, ProjectFileFormat);
+
+					// if we are using the defaults set earlier then skip writing this
+					if (PlatformConfigurationType == PlatformProjectGenerator.DefaultPlatformConfigurationType && PlatformToolsetString.Length == 0)
+					{
+						continue;
+					}
+
+					string ProjectConfigurationAndPlatformName = ProjectConfigurationName + "|" + ProjectPlatformName;
+					string ConditionString = "Condition=\"'$(Configuration)|$(Platform)'=='" + ProjectConfigurationAndPlatformName + "'\"";
+					VCProjectFileContent.AppendLine("  <PropertyGroup {0} Label=\"Configuration\">", ConditionString);
+					VCProjectFileContent.AppendLine("    <ConfigurationType>{0}</ConfigurationType>", PlatformConfigurationType);
+
+					if (PlatformToolsetString.Length == 0)
+					{
+						VCProjectFileGenerator.AppendPlatformToolsetProperty(VCProjectFileContent, ProjectFileFormat);
+					}
+					else
+					{
+						VCProjectFileContent.Append(PlatformToolsetString);
+					}
+
+					VCProjectFileContent.AppendLine("  </PropertyGroup>");
 				}
 			}
 
@@ -1620,36 +1649,6 @@ namespace UnrealBuildTool
 			}
 		}
 
-		// Anonymous function that writes post-Default.props configuration data
-		private void WritePostDefaultPropsConfiguration(UnrealTargetPlatform TargetPlatform, UnrealTargetConfiguration TargetConfiguration, string ProjectPlatformName, string ProjectConfigurationName, PlatformProjectGeneratorCollection PlatformProjectGenerators, StringBuilder VCProjectFileContent, ILogger Logger)
-		{
-			PlatformProjectGenerator? ProjGenerator = PlatformProjectGenerators.GetPlatformProjectGenerator(TargetPlatform, true);
-
-			string ProjectConfigurationAndPlatformName = ProjectConfigurationName + "|" + ProjectPlatformName;
-			string ConditionString = "Condition=\"'$(Configuration)|$(Platform)'=='" + ProjectConfigurationAndPlatformName + "'\"";
-
-			StringBuilder PlatformToolsetString = new StringBuilder();
-			if (ProjGenerator != null)
-			{
-				ProjGenerator.GetVisualStudioPlatformToolsetString(TargetPlatform, TargetConfiguration, ProjectFileFormat, PlatformToolsetString);
-			}
-
-			string PlatformConfigurationType = (ProjGenerator == null) ? "Makefile" : ProjGenerator.GetVisualStudioPlatformConfigurationType(TargetPlatform, ProjectFileFormat);
-			VCProjectFileContent.AppendLine("  <PropertyGroup {0} Label=\"Configuration\">", ConditionString);
-			VCProjectFileContent.AppendLine("    <ConfigurationType>{0}</ConfigurationType>", PlatformConfigurationType);
-
-			if (PlatformToolsetString.Length == 0)
-			{
-				VCProjectFileGenerator.AppendPlatformToolsetProperty(VCProjectFileContent, ProjectFileFormat);
-			}
-			else
-			{
-				VCProjectFileContent.Append(PlatformToolsetString);
-			}
-
-			VCProjectFileContent.AppendLine("  </PropertyGroup>");
-		}
-
 		// Helper class to generate NMake build commands and arguments
 		public class BuildCommandBuilder
 		{
