Set buildToolsVersion in subprojects to match project for Android packaging

#android

diff --git a/Engine/Source/Programs/UnrealBuildTool/Platform/Android/UEDeployAndroid.cs b/Engine/Source/Programs/UnrealBuildTool/Platform/Android/UEDeployAndroid.cs
--- a/Engine/Source/Programs/UnrealBuildTool/Platform/Android/UEDeployAndroid.cs
+++ b/Engine/Source/Programs/UnrealBuildTool/Platform/Android/UEDeployAndroid.cs
@@ -5875,7 +5875,8 @@ namespace UnrealBuildTool
 				StringBuilder BuildGradleContent = new StringBuilder();
 				BuildGradleContent.AppendLine("apply plugin: 'com.android.library'");
 				BuildGradleContent.AppendLine("android {");
-				BuildGradleContent.AppendLine(string.Format("\tcompileSdkVersion {0}", CompileSDKVersion));
+				BuildGradleContent.AppendLine("\tcompileSdkVersion = COMPILE_SDK_VERSION.toInteger()");
+				BuildGradleContent.AppendLine("\tbuildToolsVersion = BUILD_TOOLS_VERSION");
 				BuildGradleContent.AppendLine("\tdefaultConfig {");
 
 				// Try to get the SDK target from the AndroidManifest.xml
