Removed more redundant import, configuration and system settings from the generated projects.

diff --git a/Engine/Platforms/Hololens/Source/Programs/UnrealBuildTool/HoloLensProjectGenerator.cs b/Engine/Platforms/Hololens/Source/Programs/UnrealBuildTool/HoloLensProjectGenerator.cs
--- a/Engine/Platforms/Hololens/Source/Programs/UnrealBuildTool/HoloLensProjectGenerator.cs
+++ b/Engine/Platforms/Hololens/Source/Programs/UnrealBuildTool/HoloLensProjectGenerator.cs
@@ -100,7 +100,7 @@ namespace UnrealBuildTool
 			VCProjectFileContent.Append("		<NMakeForcedUsingAssemblies>$(NMakeForcedUsingAssemblies);" + FoundationWinMDPath + ";" + UniversalWinMDPath + ";platform.winmd</NMakeForcedUsingAssemblies>" + ProjectFileGenerator.NewLine);
 		}
 
-		public override void GetVisualStudioPreDefaultString(UnrealTargetPlatform InPlatform, UnrealTargetConfiguration InConfiguration, StringBuilder VCProjectFileContent)
+		public override void GetVisualStudioPreDefaultString(UnrealTargetPlatform InPlatform, StringBuilder VCProjectFileContent)
 		{
 			// VS2017 expects WindowsTargetPlatformVersion to be set in conjunction with these other properties, otherwise the projects
 			// will fail to load when the solution is in a HoloLens configuration.
@@ -118,12 +118,12 @@ namespace UnrealBuildTool
 				SDKVersion = version.ToString();
 			}
 
-			VCProjectFileContent.AppendLine("		<AppContainerApplication>true</AppContainerApplication>");
-			VCProjectFileContent.AppendLine("		<ApplicationType>Windows Store</ApplicationType>");
-			VCProjectFileContent.AppendLine("		<ApplicationTypeRevision>10.0</ApplicationTypeRevision>");
-			VCProjectFileContent.AppendLine("		<WindowsAppContainer>true</WindowsAppContainer>");
-			VCProjectFileContent.AppendLine("		<AppxPackage>true</AppxPackage>");
-			VCProjectFileContent.AppendLine("		<WindowsTargetPlatformVersion>{0}</WindowsTargetPlatformVersion>", SDKVersion.ToString()); 
+			VCProjectFileContent.AppendLine("    <AppContainerApplication>true</AppContainerApplication>");
+			VCProjectFileContent.AppendLine("    <ApplicationType>Windows Store</ApplicationType>");
+			VCProjectFileContent.AppendLine("    <ApplicationTypeRevision>10.0</ApplicationTypeRevision>");
+			VCProjectFileContent.AppendLine("    <WindowsAppContainer>true</WindowsAppContainer>");
+			VCProjectFileContent.AppendLine("    <AppxPackage>true</AppxPackage>");
+			VCProjectFileContent.AppendLine("    <WindowsTargetPlatformVersion>{0}</WindowsTargetPlatformVersion>", SDKVersion.ToString()); 
 		}
 
 		public override string GetVisualStudioLayoutDirSection(UnrealTargetPlatform InPlatform, UnrealTargetConfiguration InConfiguration, string InConditionString, TargetType TargetType, FileReference TargetRulesPath, FileReference ProjectFilePath, FileReference NMakeOutputPath, VCProjectFileFormat InProjectFileFormat)
diff --git a/Engine/Source/Programs/UnrealBuildTool/ProjectFiles/PlatformProjectGenerator.cs b/Engine/Source/Programs/UnrealBuildTool/ProjectFiles/PlatformProjectGenerator.cs
--- a/Engine/Source/Programs/UnrealBuildTool/ProjectFiles/PlatformProjectGenerator.cs
+++ b/Engine/Source/Programs/UnrealBuildTool/ProjectFiles/PlatformProjectGenerator.cs
@@ -101,6 +101,16 @@ namespace UnrealBuildTool
 			return false;
 		}
 
+		/// <summary>
+		/// Return project configuration settings that must be included before the default props file for all configurations
+		/// </summary>
+		/// <param name="Platform">The UnrealTargetPlatform being built</param>
+		/// <param name="ProjectFileBuilder">String builder for the project file</param>
+		/// <returns>string    The custom configuration section for the project file; Empty string if it doesn't require one</returns>
+		public virtual void GetVisualStudioPreDefaultString(UnrealTargetPlatform Platform, StringBuilder ProjectFileBuilder)
+		{
+		}
+
 		/// <summary>
 		/// Return project configuration settings that must be included before the default props file
 		/// </summary>
diff --git a/Engine/Source/Programs/UnrealBuildTool/ProjectFiles/VisualStudio/VCProject.cs b/Engine/Source/Programs/UnrealBuildTool/ProjectFiles/VisualStudio/VCProject.cs
--- a/Engine/Source/Programs/UnrealBuildTool/ProjectFiles/VisualStudio/VCProject.cs
+++ b/Engine/Source/Programs/UnrealBuildTool/ProjectFiles/VisualStudio/VCProject.cs
@@ -1061,11 +1061,45 @@ namespace UnrealBuildTool
 			}
 
 			// Write each project configuration PreDefaultProps section
+			HashSet<UnrealTargetPlatform> CommonPlatformsWritten = new HashSet<UnrealTargetPlatform>();
 			foreach (ProjectConfigAndTargetCombination Combination in ValidProjectConfigAndTargetCombinations)
 			{
 				if (Combination.Platform != null)
 				{
-					WritePreDefaultPropsConfiguration(Combination.Platform.Value, Combination.Configuration, Combination.ProjectPlatformName, Combination.ProjectConfigurationName, PlatformProjectGenerators, VCProjectFileContent);
+					UnrealTargetPlatform TargetPlatform = Combination.Platform.Value;
+					PlatformProjectGenerator? ProjGenerator = PlatformProjectGenerators.GetPlatformProjectGenerator(TargetPlatform, true);
+					if (ProjGenerator != null)
+					{
+						if (!CommonPlatformsWritten.Contains(Combination.Platform.Value))
+						{
+							StringBuilder CommonPlatformToolsetString = new StringBuilder();
+							ProjGenerator.GetVisualStudioPreDefaultString(TargetPlatform, CommonPlatformToolsetString);
+
+							if (CommonPlatformToolsetString.Length > 0)
+							{
+								string ConditionString = "Condition=\"'$(Platform)'=='" + Combination.ProjectPlatformName + "'\"";
+								VCProjectFileContent.AppendLine("  <PropertyGroup " + ConditionString + " Label=\"Configuration\">");
+								VCProjectFileContent.Append(CommonPlatformToolsetString);
+								VCProjectFileContent.AppendLine("  </PropertyGroup>");
+							}
+							CommonPlatformsWritten.Add(TargetPlatform);
+						}
+
+						if (ProjGenerator != null)
+						{
+							StringBuilder PlatformToolsetString = new StringBuilder();
+							ProjGenerator.GetVisualStudioPreDefaultString(TargetPlatform, Combination.Configuration, PlatformToolsetString);
+
+							if (PlatformToolsetString.Length > 0)
+							{
+								string ProjectConfigurationAndPlatformName = Combination.ProjectConfigurationName + "|" + Combination.ProjectPlatformName;
+								string ConditionString = "Condition=\"'$(Configuration)|$(Platform)'=='" + ProjectConfigurationAndPlatformName + "'\"";
+								VCProjectFileContent.AppendLine("  <PropertyGroup " + ConditionString + " Label=\"Configuration\">");
+								VCProjectFileContent.Append(PlatformToolsetString);
+								VCProjectFileContent.AppendLine("  </PropertyGroup>");
+							}
+						}
+					}
 				}
 			}
 
@@ -1126,13 +1160,12 @@ namespace UnrealBuildTool
 			VCProjectFileContent.AppendLine("  <PropertyGroup Label=\"UserMacros\" />");
 
 			// Write the invalid configuration
-			foreach (string InvalidConfigPlatformName in InvalidConfigPlatformNames)
 			{
 				const string InvalidMessage = "echo The selected platform/configuration is not valid for this target.";
 
 				string ProjectRelativeUnusedDirectory = NormalizeProjectPath(DirectoryReference.Combine(Unreal.EngineDirectory, "Intermediate", "Build", "Unused"));
 
-				VCProjectFileContent.AppendLine("  <PropertyGroup Condition=\"'$(Configuration)|$(Platform)'=='Invalid|{0}'\">", InvalidConfigPlatformName);
+				VCProjectFileContent.AppendLine("  <PropertyGroup Condition=\"'$(Configuration)'=='Invalid'\">");
 				VCProjectFileContent.AppendLine("    <NMakeBuildCommandLine>{0}</NMakeBuildCommandLine>", InvalidMessage);
 				VCProjectFileContent.AppendLine("    <NMakeReBuildCommandLine>{0}</NMakeReBuildCommandLine>", InvalidMessage);
 				VCProjectFileContent.AppendLine("    <NMakeCleanCommandLine>{0}</NMakeCleanCommandLine>", InvalidMessage);
@@ -1142,6 +1175,11 @@ namespace UnrealBuildTool
 				VCProjectFileContent.AppendLine("  </PropertyGroup>");
 			}
 
+			// Write default import group
+			VCProjectFileContent.AppendLine("  <ImportGroup Label=\"PropertySheets\">");
+			VCProjectFileContent.AppendLine("    <Import Project=\"$(UserRootDir)\\Microsoft.Cpp.$(Platform).user.props\" Condition=\"exists('$(UserRootDir)\\Microsoft.Cpp.$(Platform).user.props')\" Label=\"LocalAppDataPlatform\" />");
+			VCProjectFileContent.AppendLine("  </ImportGroup>");
+
 			// Write each project configuration
 			foreach (ProjectConfigAndTargetCombination Combination in ProjectConfigAndTargetCombinations)
 			{
@@ -1623,32 +1661,6 @@ namespace UnrealBuildTool
 			}
 		}
 
-		// Anonymous function that writes pre-Default.props configuration data
-		private void WritePreDefaultPropsConfiguration(UnrealTargetPlatform TargetPlatform, UnrealTargetConfiguration TargetConfiguration, string ProjectPlatformName, string ProjectConfigurationName, PlatformProjectGeneratorCollection PlatformProjectGenerators, StringBuilder VCProjectFileContent)
-		{
-			PlatformProjectGenerator? ProjGenerator = PlatformProjectGenerators.GetPlatformProjectGenerator(TargetPlatform, true);
-			if (ProjGenerator == null)
-			{
-				return;
-			}
-
-			string ProjectConfigurationAndPlatformName = ProjectConfigurationName + "|" + ProjectPlatformName;
-			string ConditionString = "Condition=\"'$(Configuration)|$(Platform)'=='" + ProjectConfigurationAndPlatformName + "'\"";
-
-			if (ProjGenerator != null)
-			{
-				StringBuilder PlatformToolsetString = new StringBuilder();
-				ProjGenerator.GetVisualStudioPreDefaultString(TargetPlatform, TargetConfiguration, PlatformToolsetString);
-
-				if (PlatformToolsetString.Length > 0)
-				{
-					VCProjectFileContent.AppendLine("  <PropertyGroup " + ConditionString + " Label=\"Configuration\">", ConditionString);
-					VCProjectFileContent.Append(PlatformToolsetString);
-					VCProjectFileContent.AppendLine("  </PropertyGroup>");
-				}
-			}
-		}
-
 		// Helper class to generate NMake build commands and arguments
 		public class BuildCommandBuilder
 		{
@@ -1794,13 +1806,18 @@ namespace UnrealBuildTool
 			string ConditionString = "Condition=\"'$(Configuration)|$(Platform)'=='" + Combination.ProjectConfigurationAndPlatformName + "'\"";
 
 			{
-				VCProjectFileContent.AppendLine("  <ImportGroup {0} Label=\"PropertySheets\">", ConditionString);
-				VCProjectFileContent.AppendLine("    <Import Project=\"$(UserRootDir)\\Microsoft.Cpp.$(Platform).user.props\" Condition=\"exists('$(UserRootDir)\\Microsoft.Cpp.$(Platform).user.props')\" Label=\"LocalAppDataPlatform\" />");
+				// Add custom import info
 				if (ProjGenerator != null)
 				{
-					ProjGenerator.GetVisualStudioImportGroupProperties(Combination.Platform!.Value, VCProjectFileContent);
+					StringBuilder CustomImportGroupInfo = new StringBuilder();
+					ProjGenerator.GetVisualStudioImportGroupProperties(Combination.Platform!.Value, CustomImportGroupInfo);
+					if (CustomImportGroupInfo.Length != 0)
+					{
+						VCProjectFileContent.AppendLine("  <ImportGroup {0} Label=\"PropertySheets\">", ConditionString);
+						VCProjectFileContent.Append(CustomImportGroupInfo);
+						VCProjectFileContent.AppendLine("  </ImportGroup>");
+					}
 				}
-				VCProjectFileContent.AppendLine("  </ImportGroup>");
 
 				DirectoryReference ProjectDirectory = ProjectFilePath.Directory;
 
