Fixed an issue where some common properties from a different platform were stomping on top of another.
Lots of minor cleanup.

diff --git a/Engine/Source/Programs/UnrealBuildTool/ProjectFiles/PlatformProjectGenerator.cs b/Engine/Source/Programs/UnrealBuildTool/ProjectFiles/PlatformProjectGenerator.cs
--- a/Engine/Source/Programs/UnrealBuildTool/ProjectFiles/PlatformProjectGenerator.cs
+++ b/Engine/Source/Programs/UnrealBuildTool/ProjectFiles/PlatformProjectGenerator.cs
@@ -104,21 +104,21 @@ namespace UnrealBuildTool
 		/// <summary>
 		/// Return project configuration settings that must be included before the default props file for all configurations
 		/// </summary>
-		/// <param name="Platform">The UnrealTargetPlatform being built</param>
+		/// <param name="InPlatform">The UnrealTargetPlatform being built</param>
 		/// <param name="ProjectFileBuilder">String builder for the project file</param>
 		/// <returns>string    The custom configuration section for the project file; Empty string if it doesn't require one</returns>
-		public virtual void GetVisualStudioPreDefaultString(UnrealTargetPlatform Platform, StringBuilder ProjectFileBuilder)
+		public virtual void GetVisualStudioPreDefaultString(UnrealTargetPlatform InPlatform, StringBuilder ProjectFileBuilder)
 		{
 		}
 
 		/// <summary>
 		/// Return project configuration settings that must be included before the default props file
 		/// </summary>
-		/// <param name="Platform">The UnrealTargetPlatform being built</param>
-		/// <param name="Configuration">The UnrealTargetConfiguration being built</param>
+		/// <param name="InPlatform">The UnrealTargetPlatform being built</param>
+		/// <param name="InConfiguration">The UnrealTargetConfiguration being built</param>
 		/// <param name="ProjectFileBuilder">String builder for the project file</param>
 		/// <returns>string    The custom configuration section for the project file; Empty string if it doesn't require one</returns>
-		public virtual void GetVisualStudioPreDefaultString(UnrealTargetPlatform Platform, UnrealTargetConfiguration Configuration, StringBuilder ProjectFileBuilder)
+		public virtual void GetVisualStudioPreDefaultString(UnrealTargetPlatform InPlatform, UnrealTargetConfiguration InConfiguration, StringBuilder ProjectFileBuilder)
 		{
 		}
 
diff --git a/Engine/Source/Programs/UnrealBuildTool/ProjectFiles/VisualStudio/VCProject.cs b/Engine/Source/Programs/UnrealBuildTool/ProjectFiles/VisualStudio/VCProject.cs
--- a/Engine/Source/Programs/UnrealBuildTool/ProjectFiles/VisualStudio/VCProject.cs
+++ b/Engine/Source/Programs/UnrealBuildTool/ProjectFiles/VisualStudio/VCProject.cs
@@ -12,6 +12,7 @@ using EpicGames.Core;
 using UnrealBuildBase;
 using System.Diagnostics.CodeAnalysis;
 using Microsoft.Extensions.Logging;
+using System.Reflection.Metadata;
 
 namespace UnrealBuildTool
 {
@@ -967,6 +968,7 @@ namespace UnrealBuildTool
 			List<Tuple<string, UnrealTargetPlatform>> ProjectPlatformNameAndPlatforms = new List<Tuple<string, UnrealTargetPlatform>>();    // ProjectPlatformName, Platform
 			List<Tuple<string, UnrealTargetConfiguration>> ProjectConfigurationNameAndConfigurations = new List<Tuple<string, UnrealTargetConfiguration>>();    // ProjectConfigurationName, Configuration
 			List<ProjectConfigAndTargetCombination> ValidProjectConfigAndTargetCombinations = new List<ProjectConfigAndTargetCombination>();
+			Dictionary<string, HashSet<UnrealTargetPlatform>> ProjectPlatformNameToPlatform = new Dictionary<string, HashSet<UnrealTargetPlatform>>();
 			foreach (ProjectConfigAndTargetCombination Combination in ProjectConfigAndTargetCombinations!)
 			{
 				if (Combination.Platform == null)
@@ -989,7 +991,17 @@ namespace UnrealBuildTool
 				{
 					ProjectConfigurationNameAndConfigurations.Add(Tuple.Create(Combination.ProjectConfigurationName, Combination.Configuration));
 				}
+
 				ValidProjectConfigAndTargetCombinations.Add(Combination);
+
+				if (ProjectPlatformNameToPlatform.TryGetValue(Combination.ProjectPlatformName, out HashSet<UnrealTargetPlatform>? PlatformList))
+				{
+					PlatformList.Add(Combination.Platform.Value);
+				}
+				else
+				{
+					ProjectPlatformNameToPlatform.Add(Combination.ProjectPlatformName, new HashSet<UnrealTargetPlatform>() { Combination.Platform.Value });
+				}
 			}
 
 			// Add the "invalid" configuration for each platform. We use this when the solution configuration does not match any project configuration.
@@ -1070,7 +1082,15 @@ namespace UnrealBuildTool
 					PlatformProjectGenerator? ProjGenerator = PlatformProjectGenerators.GetPlatformProjectGenerator(TargetPlatform, true);
 					if (ProjGenerator != null)
 					{
-						if (!CommonPlatformsWritten.Contains(Combination.Platform.Value))
+						bool ProjectPlatformRequiresConfigurationName = false;
+						if (ProjectPlatformNameToPlatform.TryGetValue(Combination.ProjectPlatformName, out HashSet<UnrealTargetPlatform>? PlatformList))
+						{
+							ProjectPlatformRequiresConfigurationName = PlatformList.Count > 1;
+						}
+
+						// Properties that are common for the unique platform
+						// Note that if the platform name is not unique then the common properties will be added below
+						if (!ProjectPlatformRequiresConfigurationName && !CommonPlatformsWritten.Contains(Combination.Platform.Value))
 						{
 							StringBuilder CommonPlatformToolsetString = new StringBuilder();
 							ProjGenerator.GetVisualStudioPreDefaultString(TargetPlatform, CommonPlatformToolsetString);
@@ -1085,19 +1105,21 @@ namespace UnrealBuildTool
 							CommonPlatformsWritten.Add(TargetPlatform);
 						}
 
-						if (ProjGenerator != null)
+						// Properties that require the configuration and platform name
+						StringBuilder PlatformToolsetString = new StringBuilder();
+						if (ProjectPlatformRequiresConfigurationName)
 						{
-							StringBuilder PlatformToolsetString = new StringBuilder();
-							ProjGenerator.GetVisualStudioPreDefaultString(TargetPlatform, Combination.Configuration, PlatformToolsetString);
+							ProjGenerator.GetVisualStudioPreDefaultString(TargetPlatform, PlatformToolsetString);
+						}
 
-							if (PlatformToolsetString.Length > 0)
-							{
-								string ProjectConfigurationAndPlatformName = Combination.ProjectConfigurationName + "|" + Combination.ProjectPlatformName;
-								string ConditionString = "Condition=\"'$(Configuration)|$(Platform)'=='" + ProjectConfigurationAndPlatformName + "'\"";
-								VCProjectFileContent.AppendLine("  <PropertyGroup " + ConditionString + " Label=\"Configuration\">");
-								VCProjectFileContent.Append(PlatformToolsetString);
-								VCProjectFileContent.AppendLine("  </PropertyGroup>");
-							}
+						ProjGenerator.GetVisualStudioPreDefaultString(TargetPlatform, Combination.Configuration, PlatformToolsetString);
+						if (PlatformToolsetString.Length > 0)
+						{
+							string ProjectConfigurationAndPlatformName = Combination.ProjectConfigurationName + "|" + Combination.ProjectPlatformName;
+							string ConditionString = "Condition=\"'$(Configuration)|$(Platform)'=='" + ProjectConfigurationAndPlatformName + "'\"";
+							VCProjectFileContent.AppendLine("  <PropertyGroup " + ConditionString + " Label=\"Configuration\">");
+							VCProjectFileContent.Append(PlatformToolsetString);
+							VCProjectFileContent.AppendLine("  </PropertyGroup>");
 						}
 					}
 				}
