Created a platform specific arch config because this platform doesn't support the same architectures as windows.

diff --git a/Engine/Source/Programs/UnrealBuildTool/Platform/Windows/UEBuildWindows.cs b/Engine/Source/Programs/UnrealBuildTool/Platform/Windows/UEBuildWindows.cs
--- a/Engine/Source/Programs/UnrealBuildTool/Platform/Windows/UEBuildWindows.cs
+++ b/Engine/Source/Programs/UnrealBuildTool/Platform/Windows/UEBuildWindows.cs
@@ -976,7 +976,7 @@ namespace UnrealBuildTool
 
 		public override UnrealArchitectures ActiveArchitectures(FileReference? ProjectFile, string? TargetName)
 		{
-			// for now alwways ocmpile X64 unless overridden on commandline
+			// for now always compile X64 unless overridden on commandline
 			return new UnrealArchitectures(UnrealArch.X64);
 		}
 
@@ -984,8 +984,6 @@ namespace UnrealBuildTool
 		{
 			return Architectures.SingleArchitecture != UnrealArch.X64;
 		}
-
-
 	}
 
 	class WindowsPlatform : UEBuildPlatform
