Updated the PrimaryProjectName.txt not be written out unless it changed.

diff --git a/Engine/Source/Programs/UnrealBuildTool/ProjectFiles/ProjectFileGenerator.cs b/Engine/Source/Programs/UnrealBuildTool/ProjectFiles/ProjectFileGenerator.cs
--- a/Engine/Source/Programs/UnrealBuildTool/ProjectFiles/ProjectFileGenerator.cs
+++ b/Engine/Source/Programs/UnrealBuildTool/ProjectFiles/ProjectFileGenerator.cs
@@ -883,8 +883,7 @@ namespace UnrealBuildTool
 
 				// Write out the name of the primary project file, so the runtime knows to use it
 				FileReference PrimaryProjectNameLocation = FileReference.Combine(Unreal.EngineDirectory, "Intermediate", "ProjectFiles", "PrimaryProjectName.txt");
-				DirectoryReference.CreateDirectory(PrimaryProjectNameLocation.Directory);
-				FileReference.WriteAllText(PrimaryProjectNameLocation, PrimaryProjectName);
+				Utils.WriteFileIfChanged(PrimaryProjectNameLocation, PrimaryProjectName, Logger);
 			}
 
 			// Modify the name if specific platforms were given
@@ -907,8 +906,7 @@ namespace UnrealBuildTool
 
 				// the primary project name is always read from our intermediate directory and not the overriden one for this set of platforms
 				FileReference PrimaryProjectNameLocation = FileReference.Combine(Unreal.EngineDirectory, "Intermediate", "ProjectFiles", "PrimaryProjectName.txt");
-				DirectoryReference.CreateDirectory(PrimaryProjectNameLocation.Directory);
-				FileReference.WriteAllText(PrimaryProjectNameLocation, PrimaryProjectName);
+				Utils.WriteFileIfChanged(PrimaryProjectNameLocation, PrimaryProjectName, Logger);
 			}
 
 			bool bCleanProjectFiles = Arguments.Any(x => x.Equals("-CleanProjects", StringComparison.InvariantCultureIgnoreCase));
