Updated the vc projects to be created in parallel.

diff --git a/Engine/Source/Programs/UnrealBuildTool/ProjectFiles/ProjectFileGenerator.cs b/Engine/Source/Programs/UnrealBuildTool/ProjectFiles/ProjectFileGenerator.cs
--- a/Engine/Source/Programs/UnrealBuildTool/ProjectFiles/ProjectFileGenerator.cs
+++ b/Engine/Source/Programs/UnrealBuildTool/ProjectFiles/ProjectFileGenerator.cs
@@ -14,6 +14,7 @@ using System.Xml;
 using UnrealBuildBase;
 using Microsoft.Extensions.Logging;
 using EpicGames.Serialization;
+using System.Threading;
 
 namespace UnrealBuildTool
 {
@@ -2932,15 +2933,30 @@ namespace UnrealBuildTool
 			{
 				int TotalProjectFileCount = GeneratedProjectFiles.Count + 1;    // +1 for the primary project file, which we'll save next
 
-				for (int ProjectFileIndex = 0; ProjectFileIndex < GeneratedProjectFiles.Count; ++ProjectFileIndex)
+				int ProjectsFinished = 0;
+				bool ProjectCreationFailed = false;
+				System.Threading.Tasks.Parallel.ForEach(GeneratedProjectFiles, CurProject =>
 				{
-					ProjectFile CurProject = GeneratedProjectFiles[ProjectFileIndex];
+					if (ProjectCreationFailed)
+					{
+						return;
+					}
+
 					if (!CurProject.WriteProjectFile(SupportedPlatforms, SupportedConfigurations, PlatformProjectGenerators, Logger))
 					{
-						return false;
+						ProjectCreationFailed = true;
+						return;
 					}
 
-					Progress.Write(ProjectFileIndex + 1, TotalProjectFileCount);
+					lock(Progress)
+					{
+						Progress.Write(++ProjectsFinished, TotalProjectFileCount);
+					}
+				});
+
+				if (ProjectCreationFailed)
+				{
+					return false;
 				}
 
 				WritePrimaryProjectFile(UBTProject, PlatformProjectGenerators, Logger);
