- Fixed Mac codesigning for InstalledBuilds (and in general)
- Now codesign with ad-hoc/dummy cert in FInalizeAppBundle.sh

diff --git a/Engine/Build/InstalledEngineBuild.xml b/Engine/Build/InstalledEngineBuild.xml
--- a/Engine/Build/InstalledEngineBuild.xml
+++ b/Engine/Build/InstalledEngineBuild.xml
@@ -396,7 +396,7 @@
 			<Strip BaseDir="$(RootDir)" OutputDir="$(SavedOutput)" Platform="Mac" Files="#UnrealGame Mac Unstripped" Tag="#UnrealGame Mac Stripped"/>
 
 			<Do If="'$(SignExecutables)' == true">
-				<Tag Files="#UnrealGame Mac" Filter="$(MacSignFilter)" With="#UnrealGame Mac Unsigned"/>
+				<Tag Files="#UnrealGame Mac" Filter="$(MacSignFilter)" Except=";Engine/Binaries/Mac/*.app/Contents/MacOS/*" With="#UnrealGame Mac Unsigned"/>
 				<Copy Files="#UnrealGame Mac Unsigned" From="$(RootDir)" To="$(SavedOutput)" Tag="#Game_ToSign_Mac"/>
 				<Sign Files="#Game_ToSign_Mac" Tag="#UnrealGame Mac Signed"/>
 			</Do>
diff --git a/Engine/Source/Programs/UnrealBuildTool/Platform/Mac/MacToolChain.cs b/Engine/Source/Programs/UnrealBuildTool/Platform/Mac/MacToolChain.cs
--- a/Engine/Source/Programs/UnrealBuildTool/Platform/Mac/MacToolChain.cs
+++ b/Engine/Source/Programs/UnrealBuildTool/Platform/Mac/MacToolChain.cs
@@ -692,6 +692,13 @@ namespace UnrealBuildTool
 				// Make sure OS X knows the bundle was updated
 				AppendMacLine(FinalizeAppBundleScript, "touch -c \"{0}.app\"", ExeName);
 
+
+				// codesign with ad-hoc signature for when building outside of Xcode the there will be at least some signature
+				// (and it can be a BuildProduct down in ModifyBuildProducts)
+				AppendMacLine(FinalizeAppBundleScript, "codesign -f -s - \"{0}.app\"", ExeName);
+				AppendMacLine(FinalizeAppBundleScript, "codesign -f -s \"Developer ID Application\" \"{0}.app\"", ExeName);
+				AppendMacLine(FinalizeAppBundleScript, "echo done > /dev/null");
+
 				FinalizeAppBundleScript.Close();
 			}
 
@@ -1128,6 +1135,10 @@ namespace UnrealBuildTool
 				BuildProducts.Add(FileReference.Combine(BundleContentsDirectory!, "Info.plist"), BuildProductType.RequiredResource);
 				BuildProducts.Add(FileReference.Combine(BundleContentsDirectory!, "PkgInfo"), BuildProductType.RequiredResource);
 
+				// when we codesign, we need to copy the signature around on build machines, etc. this will put the CodeResources (that were created by post-build signing)
+				// is in the .target receipt file
+				BuildProducts.Add(FileReference.Combine(BundleContentsDirectory!, "_CodeSignature", "CodeResources"), BuildProductType.RequiredResource);
+
 				// modern xcode doesn't use the bootstrap launcher because it can make full .app with staged data inside it
 				if (!Target.MacPlatform.bUseModernXcode)
 				{
