- Fixed crash caused by parallelizing project file writing

diff --git a/Engine/Source/Programs/UnrealBuildTool/Platform/IOS/UEBuildIOS.cs b/Engine/Source/Programs/UnrealBuildTool/Platform/IOS/UEBuildIOS.cs
--- a/Engine/Source/Programs/UnrealBuildTool/Platform/IOS/UEBuildIOS.cs
+++ b/Engine/Source/Programs/UnrealBuildTool/Platform/IOS/UEBuildIOS.cs
@@ -984,11 +984,15 @@ namespace UnrealBuildTool
 			string ProvisionKey = ProjectSettings.BundleIdentifier + " " + bForDistribution.ToString();
 
             IOSProvisioningData? ProvisioningData;
-			if(!ProvisionCache.TryGetValue(ProvisionKey, out ProvisioningData))
-            {
-				ProvisioningData = CreateProvisioningData(ProjectSettings, bForDistribution);
-                ProvisionCache.Add(ProvisionKey, ProvisioningData);
-            }
+ 
+			lock(ProvisionCache)
+			{
+				if (!ProvisionCache.TryGetValue(ProvisionKey, out ProvisioningData))
+				{
+					ProvisioningData = CreateProvisioningData(ProjectSettings, bForDistribution);
+					ProvisionCache.Add(ProvisionKey, ProvisioningData);
+				}
+			}
 			return ProvisioningData;
         }
 
