Updating code analysis to 4.2 and text encoding to 6.0.

diff --git a/Engine/Source/Programs/AutomationTool/AutomationUtils/AutomationUtils.Automation.csproj b/Engine/Source/Programs/AutomationTool/AutomationUtils/AutomationUtils.Automation.csproj
--- a/Engine/Source/Programs/AutomationTool/AutomationUtils/AutomationUtils.Automation.csproj
+++ b/Engine/Source/Programs/AutomationTool/AutomationUtils/AutomationUtils.Automation.csproj
@@ -1,4 +1,4 @@
-﻿<Project Sdk="Microsoft.NET.Sdk">
+<Project Sdk="Microsoft.NET.Sdk">
   <Import Project="..\..\Shared\UnrealEngine.csproj.props" />
   
   <PropertyGroup>
@@ -63,7 +63,7 @@
 
     <!--Required by Ionic -->
     <PackageReference Include="System.Security.Permissions" Version="4.7.0" />
-    <PackageReference Include="System.Text.Encoding.CodePages" Version="4.7.1" />
+    <PackageReference Include="System.Text.Encoding.CodePages" Version="6.0.0" />
 
     <PackageReference Include="Google.Apis.Drive.v3" Version="1.33.1.1225" PrivateAssets="all" />
     <PackageReference Include="Google.Apis.Sheets.v4" Version="1.33.1.1229" PrivateAssets="all" />
diff --git a/Engine/Source/Programs/UnrealBuildTool/UnrealBuildTool.csproj b/Engine/Source/Programs/UnrealBuildTool/UnrealBuildTool.csproj
--- a/Engine/Source/Programs/UnrealBuildTool/UnrealBuildTool.csproj
+++ b/Engine/Source/Programs/UnrealBuildTool/UnrealBuildTool.csproj
@@ -1,4 +1,4 @@
-<Project Sdk="Microsoft.NET.Sdk">
+﻿<Project Sdk="Microsoft.NET.Sdk">
   <Import Project="..\Shared\UnrealEngine.csproj.props" />
   
   <PropertyGroup>
@@ -127,7 +127,7 @@
     <ProjectReference Include="..\Shared\EpicGames.UHT\EpicGames.UHT.csproj" />
   </ItemGroup>
   <ItemGroup>
-    <PackageReference Include="Microsoft.CodeAnalysis.CSharp" Version="3.11.0" />
+    <PackageReference Include="Microsoft.CodeAnalysis.CSharp" Version="4.2.0" />
     <PackageReference Include="Microsoft.Win32.Registry" Version="5.0.0" />
     <PackageReference Include="OpenTracing" Version="0.12.1" />
     <PackageReference Include="System.Management" Version="4.7.0" />
@@ -136,7 +136,7 @@
 
     <!-- Required by Ionic -->
     <PackageReference Include="System.Security.Permissions" Version="4.7.0" />
-    <PackageReference Include="System.Text.Encoding.CodePages" Version="4.7.1" />
+    <PackageReference Include="System.Text.Encoding.CodePages" Version="6.0.0" />
 
     <!-- Required for Hololens -->
     <PackageReference Include="System.Reflection.MetadataLoadContext" Version="4.7.2" /> 
