Build: SNDBS: Applied a patch changing the way we look for the SNDBS installation.

Functionally, it should have no effect. It's a preparation step as in the future SNDBS might be installed relative to Program Files rather than the ..._ROOT_DIR environment variable. The updated code looks in both places.

#rnx

diff --git a/Engine/Source/Programs/UnrealBuildTool/Executors/SNDBS.cs b/Engine/Source/Programs/UnrealBuildTool/Executors/SNDBS.cs
--- a/Engine/Source/Programs/UnrealBuildTool/Executors/SNDBS.cs
+++ b/Engine/Source/Programs/UnrealBuildTool/Executors/SNDBS.cs
@@ -24,9 +24,25 @@ namespace UnrealBuildTool
 	{
 		public override string Name => "SNDBS";
 
+		private static readonly string? ProgramFilesx86 = Environment.GetEnvironmentVariable("ProgramFiles(x86)");
 		private static readonly string? SCERoot = Environment.GetEnvironmentVariable("SCE_ROOT_DIR");
-		private static readonly string SNDBSBuildExe = Path.Combine(SCERoot ?? string.Empty, "Common", "SN-DBS", "bin", "dbsbuild.exe");
-		private static readonly string SNDBSUtilExe = Path.Combine(SCERoot ?? string.Empty, "Common", "SN-DBS", "bin", "dbsutil.exe");
+
+		private static string FindDbsExe(string ExeName)
+		{
+			string InstallPath = Path.Combine(ProgramFilesx86 ?? string.Empty, "SCE", "Common", "SN-DBS", "bin", ExeName);
+			if (File.Exists(InstallPath))
+			{
+				return InstallPath;
+			}
+			else
+			{
+				// Legacy install location using SCE_ROOT_DIR
+				return Path.Combine(SCERoot ?? string.Empty, "Common", "SN-DBS", "bin", ExeName);
+			}
+		}
+
+		private static string SNDBSBuildExe => FindDbsExe("dbsbuild.exe");
+		private static string SNDBSUtilExe => FindDbsExe("dbsutil.exe");
 
 		private static readonly DirectoryReference IntermediateDir = DirectoryReference.Combine(Unreal.EngineDirectory, "Intermediate", "Build", "SNDBS");
 		private static readonly FileReference IncludeRewriteRulesFile = FileReference.Combine(IntermediateDir, "include-rewrite-rules.ini");
@@ -174,7 +190,7 @@ namespace UnrealBuildTool
 		public static bool IsAvailable(ILogger Logger)
 		{
 			// Check the executable exists on disk
-			if (SCERoot == null || !File.Exists(SNDBSBuildExe))
+			if (!File.Exists(SNDBSBuildExe))
 			{
 				return false;
 			}
