Move to adding a TargetRules bool for using AutoRTFM Compiler

diff --git a/Engine/Source/Programs/AutoRTFMTests/AutoRTFMTests.Target.cs b/Engine/Source/Programs/AutoRTFMTests/AutoRTFMTests.Target.cs
--- a/Engine/Source/Programs/AutoRTFMTests/AutoRTFMTests.Target.cs
+++ b/Engine/Source/Programs/AutoRTFMTests/AutoRTFMTests.Target.cs
@@ -31,7 +31,7 @@ public class AutoRTFMTestsTarget : TargetRules
 		// Set the RTFM clang compiler
 		if (!bGenerateProjectFiles)
 		{
-		     WindowsPlatform.Compiler = WindowsCompiler.ClangRTFM;
+			 bUseAutoRTFMCompiler = true;
 		}
 	}
 }
diff --git a/Engine/Source/Programs/UnrealBuildTool/Configuration/TargetRules.cs b/Engine/Source/Programs/UnrealBuildTool/Configuration/TargetRules.cs
--- a/Engine/Source/Programs/UnrealBuildTool/Configuration/TargetRules.cs
+++ b/Engine/Source/Programs/UnrealBuildTool/Configuration/TargetRules.cs
@@ -743,6 +743,15 @@ namespace UnrealBuildTool
 		[XmlConfigFile(Category = "BuildConfiguration")]
 		public bool bUseVerse = true;
 
+		/// <summary>
+		/// Whether to use the AutoRTFM Clang compiler.
+		/// </summary>
+		[RequiresUniqueBuildEnvironment]
+		[CommandLine("-NoUseAutoRTFM", Value = "false")]
+		[CommandLine("-UseAutoRTFM", Value = "true")]
+		[XmlConfigFile(Category = "BuildConfiguration")]
+		public bool bUseAutoRTFMCompiler = false;
+
 		/// <summary>
 		/// Whether to compile the Chaos physics plugin.
 		/// </summary>
@@ -2950,6 +2959,11 @@ namespace UnrealBuildTool
 			get { return Inner.bUseVerse; }
 		}
 
+		public bool bUseAutoRTFMCompiler
+		{
+			get { return Inner.bUseAutoRTFMCompiler; }
+		}
+
 		[Obsolete("Deprecated in UE5.1 - No longer used as Chaos is always enabled.")]
 		public bool bCompileChaos
 		{
diff --git a/Engine/Source/Programs/UnrealBuildTool/Platform/Linux/LinuxToolChain.cs b/Engine/Source/Programs/UnrealBuildTool/Platform/Linux/LinuxToolChain.cs
--- a/Engine/Source/Programs/UnrealBuildTool/Platform/Linux/LinuxToolChain.cs
+++ b/Engine/Source/Programs/UnrealBuildTool/Platform/Linux/LinuxToolChain.cs
@@ -125,6 +125,16 @@ namespace UnrealBuildTool
 				FileReference LlvmArPath = FileReference.Combine(BaseLinuxPath, "bin", $"llvm-ar{BuildHostPlatform.Current.BinarySuffix}");
 				FileReference ObjcopyPath = FileReference.Combine(BaseLinuxPath, "bin", $"llvm-objcopy{BuildHostPlatform.Current.BinarySuffix}");
 
+				// if we have RTFMCompiler enabled switch the compiler but leave the sysroot to use our toolchain
+				// this *shouldnt* cause issues as these compiler should ideally be the same version
+				if (Options.HasFlag(ClangToolChainOptions.UseAutoRTFMCompiler))
+				{
+					DirectoryReference AutoRTFMDir = DirectoryReference.Combine(Unreal.EngineDirectory, "Restricted", "NotForLicensees", "Binaries", BuildHostPlatform.Current.Platform.ToString(), "AutoRTFM", "bin");
+
+					// set up the path to our toolchain
+					ClangPath = FileReference.Combine(AutoRTFMDir, $"verse-clang-cl");
+				}
+
 				// When cross-compiling on Windows, use old FixDeps. It is slow, but it does not have timing issues
 				bUseFixdeps = BuildHostPlatform.Current.Platform == UnrealTargetPlatform.Win64;
 
@@ -452,6 +462,8 @@ namespace UnrealBuildTool
 
 			// build up the commandline common to C and C++
 
+			// always select the driver g++ in-case we are using a different binary for clang, such as clang/clang-cl
+			Arguments.Add("--driver-mode=g++");
 			if (ShouldUseLibcxx())
 			{
 				Arguments.Add("-nostdinc++");
@@ -552,6 +564,8 @@ namespace UnrealBuildTool
 
 		protected virtual void GetLinkArguments(LinkEnvironment LinkEnvironment, List<string> Arguments)
 		{
+			// always select the driver g++ in-case we are using a different binary for clang, such as clang/clang-cl
+			Arguments.Add("--driver-mode=g++");
 			Arguments.Add((BuildHostPlatform.Current.Platform == UnrealTargetPlatform.Win64) ? "-fuse-ld=lld.exe" : "-fuse-ld=lld");
 
 			// debugging symbols
diff --git a/Engine/Source/Programs/UnrealBuildTool/Platform/Linux/UEBuildLinux.cs b/Engine/Source/Programs/UnrealBuildTool/Platform/Linux/UEBuildLinux.cs
--- a/Engine/Source/Programs/UnrealBuildTool/Platform/Linux/UEBuildLinux.cs
+++ b/Engine/Source/Programs/UnrealBuildTool/Platform/Linux/UEBuildLinux.cs
@@ -605,6 +605,11 @@ namespace UnrealBuildTool
 				Options |= ClangToolChainOptions.EnableThinLTO;
 			}
 
+			if (Target.bUseAutoRTFMCompiler)
+			{
+				Options |= ClangToolChainOptions.UseAutoRTFMCompiler;
+			}
+
 			// When building a monolithic editor we have to avoid using objcopy.exe as it cannot handle files
 			// larger then 4GB. This is only an issue with our binutils objcopy.exe.
 			// llvm-objcopy.exe does not have this issue and once we switch over to using that in clang 10.0.1 we can remove this!
diff --git a/Engine/Source/Programs/UnrealBuildTool/Platform/Windows/UEBuildWindows.cs b/Engine/Source/Programs/UnrealBuildTool/Platform/Windows/UEBuildWindows.cs
--- a/Engine/Source/Programs/UnrealBuildTool/Platform/Windows/UEBuildWindows.cs
+++ b/Engine/Source/Programs/UnrealBuildTool/Platform/Windows/UEBuildWindows.cs
@@ -1087,6 +1087,11 @@ namespace UnrealBuildTool
 				Logger.LogInformation("Defaulting static analyzer output type to text");
 			}
 
+			if (Target.bUseAutoRTFMCompiler)
+			{
+				Target.WindowsPlatform.Compiler = WindowsCompiler.ClangRTFM;
+			}
+
 			// Set the compiler version if necessary
 			if (Target.WindowsPlatform.Compiler == WindowsCompiler.Default)
 			{
diff --git a/Engine/Source/Programs/UnrealBuildTool/ToolChain/ClangToolChain.cs b/Engine/Source/Programs/UnrealBuildTool/ToolChain/ClangToolChain.cs
--- a/Engine/Source/Programs/UnrealBuildTool/ToolChain/ClangToolChain.cs
+++ b/Engine/Source/Programs/UnrealBuildTool/ToolChain/ClangToolChain.cs
@@ -108,6 +108,11 @@ namespace UnrealBuildTool
 		/// Disable Dump Syms step for faster iteration
 		/// </summary>
 		DisableDumpSyms = 1 << 16,
+
+		/// <summary>
+		/// Indicates that the AutoRTFM Clang compiler should be used instead of the standard clang compiler
+		/// </summary>
+		UseAutoRTFMCompiler = 1 << 17,
 	}
 
 	abstract class ClangToolChain : ISPCToolChain
@@ -1258,4 +1263,4 @@ namespace UnrealBuildTool
 		{
 		}
 	}
-}
\ No newline at end of file
+}
