Link natvis files from module directories into binary debug info.

diff --git a/Engine/Source/Programs/UnrealBuildTool/Configuration/UEBuildBinary.cs b/Engine/Source/Programs/UnrealBuildTool/Configuration/UEBuildBinary.cs
--- a/Engine/Source/Programs/UnrealBuildTool/Configuration/UEBuildBinary.cs
+++ b/Engine/Source/Programs/UnrealBuildTool/Configuration/UEBuildBinary.cs
@@ -758,6 +758,12 @@ namespace UnrealBuildTool
 				BinaryDependency.SetupDependentLinkEnvironment(BinaryLinkEnvironment);
 			}
 
+			// Gather debug visualizers for every relevant module 
+			foreach (UEBuildModule VisitedModule in LinkEnvironmentVisitedModules)
+			{
+				VisitedModule.LinkDebuggerVisualizers(BinaryLinkEnvironment.DebuggerVisualizerFiles, ToolChain, Logger);
+			}
+
 			// Set the link output file.
 			BinaryLinkEnvironment.OutputFilePaths = OutputFilePaths.ToList();
 
diff --git a/Engine/Source/Programs/UnrealBuildTool/Configuration/UEBuildModule.cs b/Engine/Source/Programs/UnrealBuildTool/Configuration/UEBuildModule.cs
--- a/Engine/Source/Programs/UnrealBuildTool/Configuration/UEBuildModule.cs
+++ b/Engine/Source/Programs/UnrealBuildTool/Configuration/UEBuildModule.cs
@@ -210,6 +210,8 @@ namespace UnrealBuildTool
 		/// </summary>
 		public bool bDependsOnVerse = false;
 
+		public FileItem? NatvisSourceFile;
+
 		/// <summary>
 		/// Constructor
 		/// </summary>
@@ -294,6 +296,11 @@ namespace UnrealBuildTool
 			AliasRestrictedFolders = new Dictionary<string, string>(Rules.AliasRestrictedFolders);
 
 			DependenciesToSkipPerArchitecture = Rules.DependenciesToSkipPerArchitecture;
+			FileReference PossibleNatvisFile = FileReference.Combine(Rules.File.Directory, $"{Name}.natvis");
+			if (FileReference.Exists(PossibleNatvisFile))
+			{
+				NatvisSourceFile = FileItem.GetItemByFileReference(PossibleNatvisFile);
+			}
 
 			// get the module directories from the module
 			ModuleDirectories = Rules.GetAllModuleDirectories();
@@ -936,6 +943,27 @@ namespace UnrealBuildTool
 
 			return new List<FileItem>();
 		}
+		
+
+		public void CopyDebuggerVisualizers(UEToolChain ToolChain, IActionGraphBuilder Graph, ILogger Logger)
+		{
+			if (NatvisSourceFile != null)
+			{
+				ToolChain.CopyDebuggerVisualizer(NatvisSourceFile, IntermediateDirectory, Graph);
+			}	
+		}
+		
+		public void LinkDebuggerVisualizers(List<FileItem> OutFiles, UEToolChain ToolChain, ILogger Logger)
+		{
+			if (NatvisSourceFile != null) 
+			{
+				FileItem? Item = ToolChain.LinkDebuggerVisualizer(NatvisSourceFile, IntermediateDirectory);
+				if (Item != null)
+				{
+					OutFiles.Add(Item);
+				}
+			}	
+		}
 
 		// Object interface.
 		public override string ToString()
diff --git a/Engine/Source/Programs/UnrealBuildTool/Configuration/UEBuildTarget.cs b/Engine/Source/Programs/UnrealBuildTool/Configuration/UEBuildTarget.cs
--- a/Engine/Source/Programs/UnrealBuildTool/Configuration/UEBuildTarget.cs
+++ b/Engine/Source/Programs/UnrealBuildTool/Configuration/UEBuildTarget.cs
@@ -2411,6 +2411,12 @@ namespace UnrealBuildTool
 					}
 				}
 			}
+			
+			// Copy debugger visualizer files for each module to their intermediate directory 
+			foreach (UEBuildModule Module in Modules.Values) 
+			{
+				Module.CopyDebuggerVisualizers(TargetToolChain, MakefileBuilder, Logger);
+			}
 
 			// Build the target's binaries.
 			DirectoryReference ExeDir = GetExecutableDir();
diff --git a/Engine/Source/Programs/UnrealBuildTool/Platform/Windows/VCToolChain.cs b/Engine/Source/Programs/UnrealBuildTool/Platform/Windows/VCToolChain.cs
--- a/Engine/Source/Programs/UnrealBuildTool/Platform/Windows/VCToolChain.cs
+++ b/Engine/Source/Programs/UnrealBuildTool/Platform/Windows/VCToolChain.cs
@@ -1437,6 +1437,27 @@ namespace UnrealBuildTool
 			return BaseCompileAction;
 		}
 
+		public override FileItem? CopyDebuggerVisualizer(FileItem SourceFile, DirectoryReference IntermediateDirectory, IActionGraphBuilder Graph)
+		{
+			if (SourceFile.HasExtension(".natvis"))
+			{
+				FileItem IntermediateFile = FileItem.GetItemByFileReference(FileReference.Combine(IntermediateDirectory, SourceFile.Name));
+				Graph.CreateCopyAction(SourceFile, IntermediateFile);
+				return IntermediateFile;
+			}
+			return null;
+		}
+
+		public override FileItem? LinkDebuggerVisualizer(FileItem SourceFile, DirectoryReference IntermediateDirectory)
+		{
+			if (SourceFile.HasExtension(".natvis"))
+			{
+				FileItem IntermediateFile = FileItem.GetItemByFileReference(FileReference.Combine(IntermediateDirectory, SourceFile.Name));
+				return IntermediateFile;
+			}
+			return null;
+		}
+
 		protected override CPPOutput CompileCPPFiles(CppCompileEnvironment CompileEnvironment, List<FileItem> InputFiles, DirectoryReference OutputDir, string ModuleName, IActionGraphBuilder Graph)
 		{
 			VCCompileAction BaseCompileAction = CreateBaseCompileAction(CompileEnvironment);
@@ -2358,6 +2379,13 @@ namespace UnrealBuildTool
 				{
 					InputFileNames.Add(string.Format("\"{0}\"", SystemLibrary));
 				}
+				
+				foreach (FileItem NatvisFile in LinkEnvironment.DebuggerVisualizerFiles)
+				{
+					PrerequisiteItems.Add(NatvisFile);
+					Arguments.Add(string.Format("/NATVIS:\"{0}\"", 
+						NormalizeCommandLinePath(NatvisFile)));
+				}
 
 				if (Target.WindowsPlatform.ManifestFile != null)
 				{
diff --git a/Engine/Source/Programs/UnrealBuildTool/System/LinkEnvironment.cs b/Engine/Source/Programs/UnrealBuildTool/System/LinkEnvironment.cs
--- a/Engine/Source/Programs/UnrealBuildTool/System/LinkEnvironment.cs
+++ b/Engine/Source/Programs/UnrealBuildTool/System/LinkEnvironment.cs
@@ -332,6 +332,11 @@ namespace UnrealBuildTool
 		/// List of functions that should be exported from this module
 		/// </summary>
 		public List<string> IncludeFunctions = new List<string>();
+		
+		/// <summary>
+		/// Debugger visualizer files to build into debug info for this binary. 
+		/// </summary>
+		public List<FileItem> DebuggerVisualizerFiles = new List<FileItem>();
 
 		/// <summary>
 		/// Provides a Module Definition File (.def) to the linker to describe various attributes of a DLL.
diff --git a/Engine/Source/Programs/UnrealBuildTool/ToolChain/UEToolChain.cs b/Engine/Source/Programs/UnrealBuildTool/ToolChain/UEToolChain.cs
--- a/Engine/Source/Programs/UnrealBuildTool/ToolChain/UEToolChain.cs
+++ b/Engine/Source/Programs/UnrealBuildTool/ToolChain/UEToolChain.cs
@@ -44,6 +44,16 @@ namespace UnrealBuildTool
 		{
 			return null;
 		}
+		
+		public virtual FileItem? CopyDebuggerVisualizer(FileItem SourceFile, DirectoryReference IntermediateDirectory, IActionGraphBuilder Graph)
+		{
+			return null;
+		}
+
+		public virtual FileItem? LinkDebuggerVisualizer(FileItem SourceFile, DirectoryReference IntermediateDirectory)
+		{
+			return null;
+		}
 
 		protected abstract CPPOutput CompileCPPFiles(CppCompileEnvironment CompileEnvironment, List<FileItem> InputFiles, DirectoryReference OutputDir, string ModuleName, IActionGraphBuilder Graph);
 		 
