Updated the intellisense data to be created in parallel.

diff --git a/Engine/Source/Programs/UnrealBuildTool/ProjectFiles/ProjectFileGenerator.cs b/Engine/Source/Programs/UnrealBuildTool/ProjectFiles/ProjectFileGenerator.cs
--- a/Engine/Source/Programs/UnrealBuildTool/ProjectFiles/ProjectFileGenerator.cs
+++ b/Engine/Source/Programs/UnrealBuildTool/ProjectFiles/ProjectFileGenerator.cs
@@ -2008,7 +2008,12 @@ namespace UnrealBuildTool
 				string ProgressInfoText = RuntimePlatform.IsWindows ? "Binding IntelliSense data..." : "Generating data for project indexing...";
 				using (ProgressWriter Progress = new ProgressWriter(ProgressInfoText, true, Logger))
 				{
-					for (int TargetIndex = 0; TargetIndex < Targets.Count; ++TargetIndex)
+					int NumTargets = Targets.Count;
+					int NumTasks = NumTargets * 2;
+					int TasksFinished = 0;
+					TargetDescriptor[] CreatedTargetDesc = new TargetDescriptor[NumTargets];
+					UEBuildTarget[] CreatedTargets = new UEBuildTarget[NumTargets];
+					System.Threading.Tasks.Parallel.For(0, NumTargets, TargetIndex =>
 					{
 						ProjectFile TargetProjectFile = Targets[TargetIndex].Item1;
 						ProjectTarget CurTarget = Targets[TargetIndex].Item2;
@@ -2017,7 +2022,12 @@ namespace UnrealBuildTool
 						UnrealTargetPlatform IntellisensePlatform = BuildHostPlatform.Current.Platform;
 						if (!CurTarget.SupportedPlatforms.Any(x => x == IntellisensePlatform))
 						{
-							continue;
+							lock (Progress)
+							{
+								TasksFinished += 2;
+								Progress.Write(TasksFinished, NumTasks);
+							}
+							return;
 						}
 
 						Logger.LogDebug("Found target: {Target}", CurTarget.Name);
@@ -2035,10 +2045,37 @@ namespace UnrealBuildTool
 							UnrealArchitectures DefaultArchitecture = UnrealArchitectureConfig.ForPlatform(IntellisensePlatform).ActiveArchitectures(CurTarget.UnrealProjectFilePath, CurTarget.Name);
 
 							// Create the target descriptor
-							TargetDescriptor TargetDesc = new TargetDescriptor(CurTarget.UnrealProjectFilePath, CurTarget.Name, IntellisensePlatform, UnrealTargetConfiguration.Development, DefaultArchitecture, new CommandLineArguments(NewArguments.ToArray()));
+							CreatedTargetDesc[TargetIndex] = new TargetDescriptor(CurTarget.UnrealProjectFilePath, CurTarget.Name, IntellisensePlatform, UnrealTargetConfiguration.Development, DefaultArchitecture, new CommandLineArguments(NewArguments.ToArray()));
 
 							// Create the target
-							UEBuildTarget Target = UEBuildTarget.Create(TargetDesc, false, false, bUsePrecompiled, Logger);
+							CreatedTargets[TargetIndex] = UEBuildTarget.Create(CreatedTargetDesc[TargetIndex], false, false, bUsePrecompiled, Logger);
+						}
+						catch (Exception Ex)
+						{
+							Logger.LogWarning("Exception while generating include data for {Target}: {Ex}", CurTarget.Name, Ex.ToString());
+						}
+
+						lock(Progress)
+						{
+							Progress.Write(++TasksFinished, NumTasks);
+						}
+					});
+
+					for (int TargetIndex = 0; TargetIndex < Targets.Count; ++TargetIndex)
+					{
+						ProjectFile TargetProjectFile = Targets[TargetIndex].Item1;
+						ProjectTarget CurTarget = Targets[TargetIndex].Item2;
+
+						try
+						{
+							UEBuildTarget Target = CreatedTargets[TargetIndex];
+							TargetDescriptor TargetDesc = CreatedTargetDesc[TargetIndex];
+
+							if (TargetDesc == null || Target == null) 
+							{
+								continue;
+							}
+
 							AddTargetForIntellisense(Target, Logger);
 
 							// If the project generator just cares about the result of UEBuildTarget.Create, skip generating the compile environments.
@@ -2079,8 +2116,7 @@ namespace UnrealBuildTool
 							Logger.LogWarning("Exception while generating include data for {Target}: {Ex}", CurTarget.Name, Ex.ToString());
 						}
 
-						// Display progress
-						Progress.Write(TargetIndex + 1, Targets.Count);
+						Progress.Write(++TasksFinished, NumTasks);
 					}
 				}
 			}
