Revert the change of c++11 narrowing warning, it was still erroring

diff --git a/Engine/Source/Programs/UnrealBuildTool/ProjectFiles/Xcode/XcodeProject.cs b/Engine/Source/Programs/UnrealBuildTool/ProjectFiles/Xcode/XcodeProject.cs
--- a/Engine/Source/Programs/UnrealBuildTool/ProjectFiles/Xcode/XcodeProject.cs
+++ b/Engine/Source/Programs/UnrealBuildTool/ProjectFiles/Xcode/XcodeProject.cs
@@ -1735,6 +1735,7 @@ namespace UnrealBuildTool.XcodeProjectXcconfig
 			Xcconfig.AppendLine($"PRODUCT_NAME = {Name}");
 			Xcconfig.AppendLine("SYMROOT = build");
 			Xcconfig.AppendLine("USE_HEADERMAP = NO");
+			Xcconfig.AppendLine("WARNING_CFLAGS = -Wno-c++11-narrowing");
 			Xcconfig.Write();
 		}
 	}
