UnrealBuildTool: Option to write MSVC SARIF files for any errors and warnings when compiling or analyzing, disabled by default

diff --git a/Engine/Source/Programs/UnrealBuildTool/Platform/Windows/UEBuildWindows.cs b/Engine/Source/Programs/UnrealBuildTool/Platform/Windows/UEBuildWindows.cs
--- a/Engine/Source/Programs/UnrealBuildTool/Platform/Windows/UEBuildWindows.cs
+++ b/Engine/Source/Programs/UnrealBuildTool/Platform/Windows/UEBuildWindows.cs
@@ -364,6 +364,13 @@ namespace UnrealBuildTool
 		[CommandLine("-EnableASan")]
 		public bool bEnableAddressSanitizer = false;
 
+		/// <summary>
+		/// Whether .sarif files containing errors and warnings are written alongside each .obj, if supported
+		/// </summary>
+		[XmlConfigFile(Category = "BuildConfiguration", Name = "bWriteSarif")]
+		[CommandLine("-Sarif")]
+		public bool bWriteSarif = false;
+
 		/// <summary>
 		/// Whether we should export a file containing .obj to source file mappings.
 		/// </summary>
@@ -758,6 +765,11 @@ namespace UnrealBuildTool
 			get { return Inner.bEnableAddressSanitizer; }
 		}
 
+		public bool bWriteSarif
+		{
+			get { return Inner.bWriteSarif; }
+		}
+
 		public string? ObjSrcMapFile
 		{
 			get { return Inner.ObjSrcMapFile; }
diff --git a/Engine/Source/Programs/UnrealBuildTool/Platform/Windows/VCCompileAction.cs b/Engine/Source/Programs/UnrealBuildTool/Platform/Windows/VCCompileAction.cs
--- a/Engine/Source/Programs/UnrealBuildTool/Platform/Windows/VCCompileAction.cs
+++ b/Engine/Source/Programs/UnrealBuildTool/Platform/Windows/VCCompileAction.cs
@@ -64,6 +64,16 @@ namespace UnrealBuildTool
 		/// </summary>
 		public FileItem? PreprocessedFile { get; set; }
 
+		/// <summary>
+		/// The output analyze warning and error log file
+		/// </summary>
+		public FileItem? AnalyzeLogFile { get; set; }
+
+		/// <summary>
+		/// The output experimental warning and error log file
+		/// </summary>
+		public FileItem? ExperimentalLogFile { get; set; }
+
 		/// <summary>
 		/// The dependency list file
 		/// </summary>
@@ -226,6 +236,14 @@ namespace UnrealBuildTool
 				{
 					yield return PreprocessedFile;
 				}
+				if (AnalyzeLogFile != null)
+				{
+					yield return AnalyzeLogFile;
+				}
+				if (ExperimentalLogFile != null)
+				{
+					yield return ExperimentalLogFile;
+				}
 				if (DependencyListFile != null)
 				{
 					yield return DependencyListFile;
@@ -317,6 +335,8 @@ namespace UnrealBuildTool
 			ObjectFile = InAction.ObjectFile;
 			AssemblyFile = InAction.AssemblyFile;
 			PreprocessedFile = InAction.PreprocessedFile;
+			AnalyzeLogFile = InAction.AnalyzeLogFile;
+			ExperimentalLogFile = InAction.ExperimentalLogFile;
 			DependencyListFile = InAction.DependencyListFile;
 			CompiledModuleInterfaceFile = InAction.CompiledModuleInterfaceFile;
 			TimingFile = InAction.TimingFile;
@@ -354,6 +374,8 @@ namespace UnrealBuildTool
 			ObjectFile = Reader.ReadFileItem();
 			AssemblyFile = Reader.ReadFileItem();
 			PreprocessedFile = Reader.ReadFileItem();
+			AnalyzeLogFile = Reader.ReadFileItem();
+			ExperimentalLogFile = Reader.ReadFileItem();
 			DependencyListFile = Reader.ReadFileItem();
 			CompiledModuleInterfaceFile = Reader.ReadFileItem();
 			TimingFile = Reader.ReadFileItem();
@@ -386,6 +408,8 @@ namespace UnrealBuildTool
 			Writer.WriteFileItem(ObjectFile);
 			Writer.WriteFileItem(AssemblyFile);
 			Writer.WriteFileItem(PreprocessedFile);
+			Writer.WriteFileItem(AnalyzeLogFile);
+			Writer.WriteFileItem(ExperimentalLogFile);
 			Writer.WriteFileItem(DependencyListFile);
 			Writer.WriteFileItem(CompiledModuleInterfaceFile);
 			Writer.WriteFileItem(TimingFile);
@@ -478,6 +502,16 @@ namespace UnrealBuildTool
 				VCToolChain.AddAssemblyFile(Arguments, AssemblyFile);
 			}
 
+			if (AnalyzeLogFile != null)
+			{
+				VCToolChain.AddAnalyzeLogFile(Arguments, AnalyzeLogFile);
+			}
+
+			if (ExperimentalLogFile != null)
+			{
+				VCToolChain.AddExperimentalLogFile(Arguments, ExperimentalLogFile);
+			}
+
 			// A better way to express this? .json is used as output for /sourceDependencies), but .md.json is used as output for /sourceDependencies:directives)
 			if (DependencyListFile != null && DependencyListFile.HasExtension(".json") && !DependencyListFile.HasExtension(".md.json"))
 			{
diff --git a/Engine/Source/Programs/UnrealBuildTool/Platform/Windows/VCToolChain.cs b/Engine/Source/Programs/UnrealBuildTool/Platform/Windows/VCToolChain.cs
--- a/Engine/Source/Programs/UnrealBuildTool/Platform/Windows/VCToolChain.cs
+++ b/Engine/Source/Programs/UnrealBuildTool/Platform/Windows/VCToolChain.cs
@@ -255,6 +255,18 @@ namespace UnrealBuildTool
 			Arguments.Add($"/Fa\"{AssemblyFileString}\""); // Set the output patj for the asm file
 		}
 
+		public static void AddAnalyzeLogFile(List<string> Arguments, FileItem AnalyzeLogFile)
+		{
+			string AnalyzeLogFileString = NormalizeCommandLinePath(AnalyzeLogFile);
+			Arguments.Add("/analyze:log " + Utils.MakePathSafeToUseWithCommandLine(AnalyzeLogFileString));
+		}
+
+		public static void AddExperimentalLogFile(List<string> Arguments, FileItem ExperimentalLogFile)
+		{
+			string ExperimentalLogFileString = NormalizeCommandLinePath(ExperimentalLogFile);
+			Arguments.Add("/experimental:log " + Utils.MakePathSafeToUseWithCommandLine(ExperimentalLogFileString));
+		}
+
 		public static void AddSourceDependenciesFile(List<string> Arguments, FileItem SourceDependenciesFile)
 		{
 			string SourceDependenciesFileString = NormalizeCommandLinePath(SourceDependenciesFile);
@@ -1665,6 +1677,17 @@ namespace UnrealBuildTool
 					CompileAction.bShowIncludes = Target.WindowsPlatform.bShowIncludes;
 				}
 
+				// Write cl errors and warnings to a file if supported
+				if (Target.WindowsPlatform.Compiler.IsMSVC() && Target.WindowsPlatform.bWriteSarif && EnvVars.CompilerVersion >= new VersionNumber(14, 34, 31933))
+				{
+					if (Target.StaticAnalyzer == StaticAnalyzer.Default && !CompileEnvironment.bDisableStaticAnalysis)
+					{
+						CompileAction.AnalyzeLogFile = FileItem.GetItemByFileReference(FileReference.Combine(OutputDir, $"{SourceFile.Location.GetFileName()}.sa.sarif"));
+					}
+
+					CompileAction.ExperimentalLogFile = FileItem.GetItemByFileReference(FileReference.Combine(OutputDir, $"{SourceFile.Location.GetFileName()}.sarif"));
+				}
+
 				// Allow derived toolchains to make further changes
 				ModifyFinalCompileAction(CompileAction, CompileEnvironment, SourceFile, OutputDir, ModuleName );
 
