Build: SNDBS: Fixed ordering of jobs.

It turns out that job ids don't affect the build order as was expected by the previous change. To hint the desired order we need to generate "jobs" in sndbs.json as an array instead of the map that we used to use.

#rnx

diff --git a/Engine/Source/Programs/UnrealBuildTool/Executors/SNDBS.cs b/Engine/Source/Programs/UnrealBuildTool/Executors/SNDBS.cs
--- a/Engine/Source/Programs/UnrealBuildTool/Executors/SNDBS.cs
+++ b/Engine/Source/Programs/UnrealBuildTool/Executors/SNDBS.cs
@@ -240,10 +240,11 @@ namespace UnrealBuildTool
 			JsonOption.Encoder = System.Text.Encodings.Web.JavaScriptEncoder.UnsafeRelaxedJsonEscaping;
 			File.WriteAllText(ScriptFile.FullName, JsonSerializer.Serialize(new Dictionary<string, object>()
 			{
-				["jobs"] = Actions.ToDictionary(a => ActionIds[a], a =>
+				["jobs"] = Actions.Select(a =>
 				{
 					var Job = new Dictionary<string, object>()
 					{
+						["id"] = ActionIds[a],
 						["title"] = a.StatusDescription,
 						["command"] = $"\"{a.CommandPath}\" {a.CommandArguments}",
 						["working_directory"] = a.WorkingDirectory.FullName,
@@ -263,7 +264,7 @@ namespace UnrealBuildTool
 					Job["echo"] = $"{ProgressMarkupPrefix}{CommandDescription}:{a.StatusDescription}";
 
 					return Job;
-				})
+				}).ToArray()
 			}, JsonOption));
 
 			PrepareToolTemplates();
