Fixed the hashcode and equals for ConfigHierarchyKey. It was testing if the list was the same list and not the contents of the list and getting the hashcode of the list and again not the contents of the list.

diff --git a/Engine/Source/Programs/UnrealBuildTool/System/ConfigCache.cs b/Engine/Source/Programs/UnrealBuildTool/System/ConfigCache.cs
--- a/Engine/Source/Programs/UnrealBuildTool/System/ConfigCache.cs
+++ b/Engine/Source/Programs/UnrealBuildTool/System/ConfigCache.cs
@@ -206,7 +206,7 @@ namespace UnrealBuildTool
 					OtherKey.ProjectDir == ProjectDir &&
 					OtherKey.Platform == Platform &&
 					OtherKey.CustomConfig == CustomConfig &&
-					OtherKey.OverrideStrings == OverrideStrings;
+					OtherKey.OverrideStrings.SequenceEqual(OverrideStrings);
 			}
 
 			/// <summary>
@@ -215,12 +215,16 @@ namespace UnrealBuildTool
 			/// <returns>Hash value for this object</returns>
 			public override int GetHashCode()
 			{
-				return HashCode.Combine(
-					Type.GetHashCode,
-					(ProjectDir != null) ? ProjectDir.GetHashCode() : 0,
-					Platform.GetHashCode(),
-					CustomConfig.GetHashCode(),
-					OverrideStrings.GetHashCode());
+				int Hash = 17;
+				Hash = (Hash * 31) + Type.GetHashCode();
+				Hash = (Hash * 31) + ((ProjectDir == null) ? 0 : ProjectDir.GetHashCode());
+				Hash = (Hash * 31) + Platform.GetHashCode();
+				Hash = (Hash * 31) + CustomConfig.GetHashCode();
+				foreach (string OverrideString in OverrideStrings)
+				{
+					Hash = (Hash * 31) + OverrideString.GetHashCode();
+				}
+				return Hash;
 			}
 		}
 
