Updated some functions to take an IEnumerable instead of specific types. This is make sure containers don't get modified by the function being called.

diff --git a/Engine/Source/Programs/Shared/EpicGames.Build/System/CompileScriptModules.cs b/Engine/Source/Programs/Shared/EpicGames.Build/System/CompileScriptModules.cs
--- a/Engine/Source/Programs/Shared/EpicGames.Build/System/CompileScriptModules.cs
+++ b/Engine/Source/Programs/Shared/EpicGames.Build/System/CompileScriptModules.cs
@@ -275,7 +275,7 @@ namespace UnrealBuildBase
 		/// <param name="Logger"></param>
 		/// <returns>Collection of all the projects.  They will have been compiled.</returns>
 		public static Dictionary<FileReference, CsProjBuildRecordEntry> Build(Rules.RulesFileType RulesFileType,
-			HashSet<FileReference> FoundProjects, List<DirectoryReference> BaseDirectories, List<string>? DefineConstants, 
+			HashSet<FileReference> FoundProjects, IEnumerable<DirectoryReference> BaseDirectories, IEnumerable<string>? DefineConstants, 
 			BuildFlags BuildFlags, out bool bBuildSuccess, Action<int> OnBuildingProjects, ILogger Logger)
 		{
 			CsProjBuildHook Hook = new Hook(Logger, RulesFileType);
@@ -390,10 +390,10 @@ namespace UnrealBuildBase
 		/// </summary>
 		[MethodImpl(MethodImplOptions.NoInlining)]
 		static Dictionary<FileReference, CsProjBuildRecordEntry> Build(HashSet<FileReference> FoundProjects,
-			bool bForceCompile, out bool bBuildSuccess, CsProjBuildHook Hook, List<DirectoryReference> BaseDirectories,
-			List<string>? DefineConstants, Action<int> OnBuildingProjects, ILogger Logger)
+			bool bForceCompile, out bool bBuildSuccess, CsProjBuildHook Hook, IEnumerable<DirectoryReference> BaseDirectories,
+			IEnumerable<string>? DefineConstants, Action<int> OnBuildingProjects, ILogger Logger)
 		{
-			return CsProjBuilder.Build(FoundProjects, bForceCompile, out bBuildSuccess, Hook, BaseDirectories, DefineConstants ?? (new()), OnBuildingProjects, Logger);
+			return CsProjBuilder.Build(FoundProjects, bForceCompile, out bBuildSuccess, Hook, BaseDirectories, DefineConstants ?? (Enumerable.Empty<string>()), OnBuildingProjects, Logger);
 		}
 
 		/// <summary>
@@ -403,7 +403,7 @@ namespace UnrealBuildBase
 		/// <param name="BaseDirectories"></param>
 		/// <param name="Logger"></param>
 		/// <returns></returns>
-		static Dictionary<FileReference, CsProjBuildRecordEntry> LoadExistingBuildRecords(CsProjBuildHook Hook, List<DirectoryReference> BaseDirectories, ILogger Logger)
+		static Dictionary<FileReference, CsProjBuildRecordEntry> LoadExistingBuildRecords(CsProjBuildHook Hook, IEnumerable<DirectoryReference> BaseDirectories, ILogger Logger)
         {
 			Dictionary<FileReference, CsProjBuildRecordEntry> LoadedBuildRecords = new Dictionary<FileReference, CsProjBuildRecordEntry>();
 
diff --git a/Engine/Source/Programs/Shared/EpicGames.Build/System/EnumeratePlugins.cs b/Engine/Source/Programs/Shared/EpicGames.Build/System/EnumeratePlugins.cs
--- a/Engine/Source/Programs/Shared/EpicGames.Build/System/EnumeratePlugins.cs
+++ b/Engine/Source/Programs/Shared/EpicGames.Build/System/EnumeratePlugins.cs
@@ -144,7 +144,7 @@ namespace UnrealBuildBase
 		/// <param name="Plugins">Collection of built plugins</param>
 		/// <returns>True if the plugins compiled</returns>
 		public static bool BuildUbtPlugins(FileReference? ProjectFile, IEnumerable<FileReference> FoundPlugins,
-			List<string>? DefineConstants, ILogger Logger, out (FileReference ProjectFile, FileReference TargetAssembly)[]? Plugins)
+			IEnumerable<string>? DefineConstants, ILogger Logger, out (FileReference ProjectFile, FileReference TargetAssembly)[]? Plugins)
 		{
 			bool bBuildSuccess = true;
 			Plugins = null;
@@ -182,7 +182,7 @@ namespace UnrealBuildBase
 		}
 
 		private static bool BuildUbtPluginsInternal(PluginSet PluginSet, FileReference? ProjectFile, 
-			List<FileReference> Built, IEnumerable<FileReference> NotBuilt, List<string>? DefineConstants, ILogger Logger)
+			List<FileReference> Built, IEnumerable<FileReference> NotBuilt, IEnumerable<string>? DefineConstants, ILogger Logger)
 		{
 			// Collect the work to be done
 			IEnumerable<FileReference>? ToBuild = null;
diff --git a/Engine/Source/Programs/Shared/EpicGames.MsBuild/CsProjBuilder.cs b/Engine/Source/Programs/Shared/EpicGames.MsBuild/CsProjBuilder.cs
--- a/Engine/Source/Programs/Shared/EpicGames.MsBuild/CsProjBuilder.cs
+++ b/Engine/Source/Programs/Shared/EpicGames.MsBuild/CsProjBuilder.cs
@@ -125,7 +125,7 @@ namespace EpicGames.MsBuild
 			}
 		}
 
-		static FileReference ConstructBuildRecordPath(CsProjBuildHook Hook, FileReference ProjectPath, List<DirectoryReference> BaseDirectories)
+		static FileReference ConstructBuildRecordPath(CsProjBuildHook Hook, FileReference ProjectPath, IEnumerable<DirectoryReference> BaseDirectories)
 		{
 			DirectoryReference BasePath = null;
 
@@ -161,8 +161,8 @@ namespace EpicGames.MsBuild
 		/// <param name="OnBuildingProjects">Action invoked to notify caller regarding the number of projects being built</param>
 		/// <param name="Logger">Destination logger</param>
 		public static Dictionary<FileReference, CsProjBuildRecordEntry> Build(HashSet<FileReference> FoundProjects,
-			bool bForceCompile, out bool bBuildSuccess, CsProjBuildHook Hook, List<DirectoryReference> BaseDirectories, 
-			List<string> DefineConstants, Action<int> OnBuildingProjects, ILogger Logger)
+			bool bForceCompile, out bool bBuildSuccess, CsProjBuildHook Hook, IEnumerable<DirectoryReference> BaseDirectories,
+			IEnumerable<string> DefineConstants, Action<int> OnBuildingProjects, ILogger Logger)
 		{
 
 			// Register the MS build path prior to invoking the internal routine.  By not having the internal routine
@@ -184,7 +184,7 @@ namespace EpicGames.MsBuild
 		/// <param name="OnBuildingProjects">Action invoked to notify caller regarding the number of projects being built</param>
 		/// <param name="Logger">Destination logger</param>
 		private static Dictionary<FileReference, CsProjBuildRecordEntry> BuildInternal(HashSet<FileReference> FoundProjects,
-			bool bForceCompile, out bool bBuildSuccess, CsProjBuildHook Hook, List<DirectoryReference> BaseDirectories, List<string> DefineConstants,
+			bool bForceCompile, out bool bBuildSuccess, CsProjBuildHook Hook, IEnumerable<DirectoryReference> BaseDirectories, IEnumerable<string> DefineConstants,
 			Action<int> OnBuildingProjects, ILogger Logger)
 		{
 			Dictionary<string, string> GlobalProperties = new Dictionary<string, string>
@@ -197,7 +197,7 @@ namespace EpicGames.MsBuild
 #endif
 			};
 
-			if (DefineConstants.Count > 0)
+			if (DefineConstants.Any())
 			{
 				GlobalProperties.Add("DefineConstants", String.Join(';', DefineConstants));
 			}
diff --git a/Engine/Source/Programs/UnrealBuildTool/System/DynamicCompilation.cs b/Engine/Source/Programs/UnrealBuildTool/System/DynamicCompilation.cs
--- a/Engine/Source/Programs/UnrealBuildTool/System/DynamicCompilation.cs
+++ b/Engine/Source/Programs/UnrealBuildTool/System/DynamicCompilation.cs
@@ -36,7 +36,7 @@ namespace UnrealBuildTool
 		/// <param name="OutputAssemblyPath">Output path for the assembly</param>
 		/// <param name="Logger">Logger for output</param>
 		/// <returns>True if the assembly needs to be built</returns>
-		private static bool RequiresCompilation(HashSet<FileReference> SourceFiles, FileReference AssemblyManifestFilePath, FileReference OutputAssemblyPath, ILogger Logger)
+		private static bool RequiresCompilation(IEnumerable<FileReference> SourceFiles, FileReference AssemblyManifestFilePath, FileReference OutputAssemblyPath, ILogger Logger)
 		{
 			// Do not compile the file if it's installed
 			if (UnrealBuildTool.IsFileInstalled(OutputAssemblyPath))
@@ -165,7 +165,7 @@ namespace UnrealBuildTool
 			}
 		}
 
-		private static Assembly? CompileAssembly(FileReference OutputAssemblyPath, HashSet<FileReference> SourceFileNames, ILogger Logger, List<string>? ReferencedAssembies, List<string>? PreprocessorDefines = null, bool TreatWarningsAsErrors = false)
+		private static Assembly? CompileAssembly(FileReference OutputAssemblyPath, IEnumerable<FileReference> SourceFileNames, ILogger Logger, IEnumerable<string>? ReferencedAssembies, IEnumerable<string>? PreprocessorDefines = null, bool TreatWarningsAsErrors = false)
 		{
 			CSharpParseOptions ParseOptions = new CSharpParseOptions(
 				languageVersion:LanguageVersion.Latest, 
@@ -301,7 +301,7 @@ namespace UnrealBuildTool
 		/// <param name="ForceCompile"></param>
 		/// <param name="TreatWarningsAsErrors"></param>
 		/// <returns>The assembly that was loaded</returns>
-		public static Assembly? CompileAndLoadAssembly(FileReference OutputAssemblyPath, HashSet<FileReference> SourceFileNames, ILogger Logger, List<string>? ReferencedAssembies = null, List<string>? PreprocessorDefines = null, bool DoNotCompile = false, bool ForceCompile = false, bool TreatWarningsAsErrors = false)
+		public static Assembly? CompileAndLoadAssembly(FileReference OutputAssemblyPath, IEnumerable<FileReference> SourceFileNames, ILogger Logger, IEnumerable<string>? ReferencedAssembies = null, IEnumerable<string>? PreprocessorDefines = null, bool DoNotCompile = false, bool ForceCompile = false, bool TreatWarningsAsErrors = false)
 		{
 			// Check to see if the resulting assembly is compiled and up to date
 			FileReference AssemblyManifestFilePath = FileReference.Combine(OutputAssemblyPath.Directory, Path.GetFileNameWithoutExtension(OutputAssemblyPath.FullName) + "Manifest.json");
diff --git a/Engine/Source/Programs/UnrealBuildTool/System/RulesAssembly.cs b/Engine/Source/Programs/UnrealBuildTool/System/RulesAssembly.cs
--- a/Engine/Source/Programs/UnrealBuildTool/System/RulesAssembly.cs
+++ b/Engine/Source/Programs/UnrealBuildTool/System/RulesAssembly.cs
@@ -28,7 +28,7 @@ namespace UnrealBuildTool
 		/// <summary>
 		/// The compiled assembly
 		/// </summary>
-		private Assembly? CompiledAssembly;
+		private readonly Assembly? CompiledAssembly;
 
 		/// <summary>
 		/// Returns the simple name of the assembly e.g. "UE5ProgramRules"
@@ -46,43 +46,43 @@ namespace UnrealBuildTool
 		/// <summary>
 		/// The base directories for this assembly
 		/// </summary>
-		private List<DirectoryReference> BaseDirs;
+		private readonly IReadOnlyList<DirectoryReference> BaseDirs;
 
 		/// <summary>
 		/// All the plugins included in this assembly
 		/// </summary>
-		private IReadOnlyList<PluginInfo> Plugins;
+		private readonly IReadOnlyList<PluginInfo> Plugins;
 
 		/// <summary>
 		/// Maps module names to their actual xxx.Module.cs file on disk
 		/// </summary>
-		private Dictionary<string, FileReference> ModuleNameToModuleFile = new Dictionary<string, FileReference>(StringComparer.InvariantCultureIgnoreCase);
+		private readonly IReadOnlyDictionary<string, FileReference> ModuleNameToModuleFile;
 
 		/// <summary>
 		/// Maps target names to their actual xxx.Target.cs file on disk
 		/// </summary>
-		private Dictionary<string, FileReference> TargetNameToTargetFile = new Dictionary<string, FileReference>(StringComparer.InvariantCultureIgnoreCase);
+		private readonly IReadOnlyDictionary<string, FileReference> TargetNameToTargetFile;
 
 		/// <summary>
 		/// Mapping from module file to its context.
 		/// </summary>
-		private Dictionary<FileReference, ModuleRulesContext> ModuleFileToContext;
+		private readonly IReadOnlyDictionary<FileReference, ModuleRulesContext> ModuleFileToContext;
 
 		/// <summary>
 		/// Whether this assembly contains engine modules. Used to set default values for bTreatAsEngineModule.
 		/// </summary>
-		private bool bContainsEngineModules;
+		private readonly bool bContainsEngineModules;
 
 		/// <summary>
 		/// Whether to use backwards compatible default settings for module and target rules. This is enabled by default for game projects to support a simpler migration path, but
 		/// is disabled for engine modules.
 		/// </summary>
-		private BuildSettingsVersion? DefaultBuildSettings;
+		private readonly BuildSettingsVersion? DefaultBuildSettings;
 
 		/// <summary>
 		/// Whether the modules and targets in this assembly are read-only
 		/// </summary>
-		private bool bReadOnly;
+		private readonly bool bReadOnly;
 
 		/// <summary>
 		/// The parent rules assembly that this assembly inherits. Game assemblies inherit the engine assembly, and the engine assembly inherits nothing.
@@ -92,12 +92,12 @@ namespace UnrealBuildTool
 		/// <summary>
 		/// The set of files that were compiled to create this assembly
 		/// </summary>
-		public HashSet<FileReference>? AssemblySourceFiles { get; }
+		public IEnumerable<FileReference>? AssemblySourceFiles { get; }
 
 		/// <summary>
 		/// Any preprocessor defines that were set when this assembly was created
 		/// </summary>
-		public List<string>? PreprocessorDefines { get;  }
+		public IEnumerable<string>? PreprocessorDefines { get;  }
 
 		/// <summary>
 		/// Constructor. Compiles a rules assembly from the given source files.
@@ -127,12 +127,13 @@ namespace UnrealBuildTool
 			this.Parent = Parent;
 
 			// Find all the source files
-			AssemblySourceFiles = new HashSet<FileReference>();
-			AssemblySourceFiles.UnionWith(ModuleFileToContext.Keys);
-			AssemblySourceFiles.UnionWith(TargetFiles);
+			HashSet<FileReference> AssemblySourceFilesHashSet = new HashSet<FileReference>();
+			AssemblySourceFilesHashSet.UnionWith(ModuleFileToContext.Keys);
+			AssemblySourceFilesHashSet.UnionWith(TargetFiles);
+			AssemblySourceFiles = AssemblySourceFilesHashSet;
 
 			// Compile the assembly
-			if (AssemblySourceFiles.Count > 0)
+			if (AssemblySourceFiles.Any())
 			{
 				// Add the parent assemblies as references so they can be used
 				List<string>? ReferencedAssembies = null;
@@ -153,22 +154,26 @@ namespace UnrealBuildTool
 			}
 
 			// Setup the module map
+			Dictionary<string, FileReference> ModuleNameToModuleFileDict = new Dictionary<string, FileReference>(StringComparer.InvariantCultureIgnoreCase);
+			ModuleNameToModuleFile = ModuleNameToModuleFileDict;
 			foreach (FileReference ModuleFile in ModuleFileToContext.Keys)
 			{
 				string ModuleName = ModuleFile.GetFileNameWithoutAnyExtensions();
 				if (!ModuleNameToModuleFile.ContainsKey(ModuleName))
 				{
-					ModuleNameToModuleFile.Add(ModuleName, ModuleFile);
+					ModuleNameToModuleFileDict.Add(ModuleName, ModuleFile);
 				}
 			}
 
 			// Setup the target map
+			Dictionary<string, FileReference> TargetNameToTargetFileDict = new Dictionary<string, FileReference>(StringComparer.InvariantCultureIgnoreCase);
+			TargetNameToTargetFile = TargetNameToTargetFileDict;
 			foreach (FileReference TargetFile in TargetFiles)
 			{
 				string TargetName = TargetFile.GetFileNameWithoutAnyExtensions();
 				if (!TargetNameToTargetFile.ContainsKey(TargetName))
 				{
-					TargetNameToTargetFile.Add(TargetName, TargetFile);
+					TargetNameToTargetFileDict.Add(TargetName, TargetFile);
 				}
 			}
 
