Updated Timeline to be somewhat threadsafe.

diff --git a/Engine/Source/Programs/UnrealBuildTool/System/Timeline.cs b/Engine/Source/Programs/UnrealBuildTool/System/Timeline.cs
--- a/Engine/Source/Programs/UnrealBuildTool/System/Timeline.cs
+++ b/Engine/Source/Programs/UnrealBuildTool/System/Timeline.cs
@@ -35,12 +35,12 @@ namespace UnrealBuildTool
 			/// <summary>
 			/// Name of the marker
 			/// </summary>
-			public string Name;
+			public readonly string Name;
 
 			/// <summary>
 			/// Time at which the event ocurred
 			/// </summary>
-			public TimeSpan StartTime;
+			public readonly TimeSpan StartTime;
 
 			/// <summary>
 			/// Time at which the event ended
@@ -50,7 +50,7 @@ namespace UnrealBuildTool
 			/// <summary>
 			/// The trace span for external tracing
 			/// </summary>
-			public ITraceSpan Span;
+			public readonly ITraceSpan Span;
 
 			/// <summary>
 			/// Constructor
@@ -120,7 +120,10 @@ namespace UnrealBuildTool
 		public static void AddEvent(string Name)
 		{
 			TimeSpan Time = Stopwatch.Elapsed;
-			Events.Add(new Event(Name, Time, Time));
+			lock(Events)
+			{
+				Events.Add(new Event(Name, Time, Time));
+			}
 		}
 
 		/// <summary>
@@ -131,7 +134,10 @@ namespace UnrealBuildTool
 		public static ITimelineEvent ScopeEvent(string Name)
 		{
 			Event Event = new Event(Name, Stopwatch.Elapsed, null);
-			Events.Add(Event);
+			lock (Events)
+			{
+				Events.Add(Event);
+			}
 			return Event;
 		}
 
@@ -148,42 +154,47 @@ namespace UnrealBuildTool
 			// Create the root event
 			TimeSpan FinishTime = Stopwatch.Elapsed;
 
-			List<Event> OuterEvents = new List<Event>();
-			OuterEvents.Add(new Event("<Root>", TimeSpan.Zero, FinishTime));
+			List<Event> OuterEvents = new List<Event>
+			{
+				new Event("<Root>", TimeSpan.Zero, FinishTime)
+			};
 
 			// Print out all the child events
 			TimeSpan LastTime = TimeSpan.Zero;
-			for(int EventIdx = 0; EventIdx < Events.Count; EventIdx++)
+			lock (Events)
 			{
-				Event Event = Events[EventIdx];
-
-				// Pop events off the stack
-				for (; OuterEvents.Count > 1; OuterEvents.RemoveAt(OuterEvents.Count - 1))
+				for (int EventIdx = 0; EventIdx < Events.Count; EventIdx++)
 				{
-					Event OuterEvent = OuterEvents.Last();
-					if (Event.StartTime < OuterEvent.FinishTime!.Value)
+					Event Event = Events[EventIdx];
+
+					// Pop events off the stack
+					for (; OuterEvents.Count > 1; OuterEvents.RemoveAt(OuterEvents.Count - 1))
 					{
-						break;
+						Event OuterEvent = OuterEvents.Last();
+						if (Event.StartTime < OuterEvent.FinishTime!.Value)
+						{
+							break;
+						}
+						UpdateLastEventTime(ref LastTime, OuterEvent.FinishTime.Value, MaxUnknownTime, OuterEvents, Verbosity, Logger);
 					}
-					UpdateLastEventTime(ref LastTime, OuterEvent.FinishTime.Value, MaxUnknownTime, OuterEvents, Verbosity, Logger);
-				}
 
-				// If there's a gap since the last event, print an unknown marker
-				UpdateLastEventTime(ref LastTime, Event.StartTime, MaxUnknownTime, OuterEvents, Verbosity, Logger);
+					// If there's a gap since the last event, print an unknown marker
+					UpdateLastEventTime(ref LastTime, Event.StartTime, MaxUnknownTime, OuterEvents, Verbosity, Logger);
 
-				// Print this event
-				Print(Event.StartTime, Event.FinishTime, Event.Name, OuterEvents, Verbosity, Logger);
+					// Print this event
+					Print(Event.StartTime, Event.FinishTime, Event.Name, OuterEvents, Verbosity, Logger);
 
-				// Push it onto the stack
-				if(Event.FinishTime.HasValue)
-				{
-					if(EventIdx + 1 < Events.Count && Events[EventIdx + 1].StartTime < Event.FinishTime.Value)
-					{
-						OuterEvents.Add(Event);
-					}
-					else
+					// Push it onto the stack
+					if (Event.FinishTime.HasValue)
 					{
-						LastTime = Event.FinishTime.Value;
+						if (EventIdx + 1 < Events.Count && Events[EventIdx + 1].StartTime < Event.FinishTime.Value)
+						{
+							OuterEvents.Add(Event);
+						}
+						else
+						{
+							LastTime = Event.FinishTime.Value;
+						}
 					}
 				}
 			}
