Fixed an issue where other target test settings would leak into other tests.

diff --git a/Engine/Source/Programs/UnrealBuildTool/Configuration/ModuleRules.cs b/Engine/Source/Programs/UnrealBuildTool/Configuration/ModuleRules.cs
--- a/Engine/Source/Programs/UnrealBuildTool/Configuration/ModuleRules.cs
+++ b/Engine/Source/Programs/UnrealBuildTool/Configuration/ModuleRules.cs
@@ -1613,52 +1613,61 @@ namespace UnrealBuildTool
 		/// </summary>
 		internal void PrepareModuleForTests()
 		{
-			if (Name != "LowLevelTestsRunner" && System.IO.Directory.Exists(TestsDirectory))
+			TestTargetRules? TestTargetRules = Target.InnerTestTargetRules;
+			if (TestTargetRules == null) 
 			{
-				if (!PrivateIncludePathModuleNames.Contains("LowLevelTestsRunner"))
-				{
-					PrivateIncludePathModuleNames.Add("LowLevelTestsRunner");
-				}
-			}
-			else if (Name == "LowLevelTestsRunner")
-			{
-				TestTargetRules.LowLevelTestsRunnerModule = this;
+				return;
 			}
 
-			if (Name == "Engine" && !TestTargetRules.bNeverCompileAgainstEngine)
-			{
-				TestTargetRules.bTestsRequireEngine = true;
-			}
-			if (Name == "UnrealEd" && !TestTargetRules.bNeverCompileAgainstEditor)
-			{
-				TestTargetRules.bTestsRequireEditor = true;
-			}
-			if (Name == "ApplicationCore" && !TestTargetRules.bNeverCompileAgainstApplicationCore)
+			lock (TestTargetRules)
 			{
-				TestTargetRules.bTestsRequireApplicationCore = true;
-			}
-			if (Name == "CoreUObject" && !TestTargetRules.bNeverCompileAgainstCoreUObject)
-			{
-				TestTargetRules.bTestsRequireCoreUObject = true;
-			}
+				if (Name != "LowLevelTestsRunner" && System.IO.Directory.Exists(TestsDirectory))
+				{
+					if (!PrivateIncludePathModuleNames.Contains("LowLevelTestsRunner"))
+					{
+						PrivateIncludePathModuleNames.Add("LowLevelTestsRunner");
+					}
+				}
+				else if (Name == "LowLevelTestsRunner")
+				{
+					TestTargetRules.LowLevelTestsRunnerModule = this;
+				}
 
-			if (TestTargetRules.LowLevelTestsRunnerModule != null)
-			{
-				if (TestTargetRules.bTestsRequireEditor && !TestTargetRules.LowLevelTestsRunnerModule.PrivateDependencyModuleNames.Contains("UnrealEd"))
+				if (Name == "Engine" && !TestTargetRules.bNeverCompileAgainstEngine)
 				{
-					TestTargetRules.LowLevelTestsRunnerModule.PrivateDependencyModuleNames.Add("UnrealEd");
+					TestTargetRules.bTestsRequireEngine = true;
 				}
-				if (TestTargetRules.bTestsRequireEngine && !TestTargetRules.LowLevelTestsRunnerModule.PrivateDependencyModuleNames.Contains("Engine"))
+				if (Name == "UnrealEd" && !TestTargetRules.bNeverCompileAgainstEditor)
 				{
-					TestTargetRules.LowLevelTestsRunnerModule.PrivateDependencyModuleNames.Add("Engine");
+					TestTargetRules.bTestsRequireEditor = true;
 				}
-				if (TestTargetRules.bTestsRequireApplicationCore && !TestTargetRules.LowLevelTestsRunnerModule.PrivateDependencyModuleNames.Contains("ApplicationCore"))
+				if (Name == "ApplicationCore" && !TestTargetRules.bNeverCompileAgainstApplicationCore)
 				{
-					TestTargetRules.LowLevelTestsRunnerModule.PrivateDependencyModuleNames.Add("ApplicationCore");
+					TestTargetRules.bTestsRequireApplicationCore = true;
 				}
-				if (TestTargetRules.bTestsRequireCoreUObject && !TestTargetRules.LowLevelTestsRunnerModule.PrivateDependencyModuleNames.Contains("CoreUObject"))
+				if (Name == "CoreUObject" && !TestTargetRules.bNeverCompileAgainstCoreUObject)
 				{
-					TestTargetRules.LowLevelTestsRunnerModule.PrivateDependencyModuleNames.Add("CoreUObject");
+					TestTargetRules.bTestsRequireCoreUObject = true;
+				}
+
+				if (TestTargetRules.LowLevelTestsRunnerModule != null)
+				{
+					if (TestTargetRules.bTestsRequireEditor && !TestTargetRules.LowLevelTestsRunnerModule.PrivateDependencyModuleNames.Contains("UnrealEd"))
+					{
+						TestTargetRules.LowLevelTestsRunnerModule.PrivateDependencyModuleNames.Add("UnrealEd");
+					}
+					if (TestTargetRules.bTestsRequireEngine && !TestTargetRules.LowLevelTestsRunnerModule.PrivateDependencyModuleNames.Contains("Engine"))
+					{
+						TestTargetRules.LowLevelTestsRunnerModule.PrivateDependencyModuleNames.Add("Engine");
+					}
+					if (TestTargetRules.bTestsRequireApplicationCore && !TestTargetRules.LowLevelTestsRunnerModule.PrivateDependencyModuleNames.Contains("ApplicationCore"))
+					{
+						TestTargetRules.LowLevelTestsRunnerModule.PrivateDependencyModuleNames.Add("ApplicationCore");
+					}
+					if (TestTargetRules.bTestsRequireCoreUObject && !TestTargetRules.LowLevelTestsRunnerModule.PrivateDependencyModuleNames.Contains("CoreUObject"))
+					{
+						TestTargetRules.LowLevelTestsRunnerModule.PrivateDependencyModuleNames.Add("CoreUObject");
+					}
 				}
 			}
 		}
diff --git a/Engine/Source/Programs/UnrealBuildTool/Configuration/TargetRules.cs b/Engine/Source/Programs/UnrealBuildTool/Configuration/TargetRules.cs
--- a/Engine/Source/Programs/UnrealBuildTool/Configuration/TargetRules.cs
+++ b/Engine/Source/Programs/UnrealBuildTool/Configuration/TargetRules.cs
@@ -3956,6 +3956,11 @@ namespace UnrealBuildTool
 			get { return Inner.OptedInModulePlatforms; } 
 		}
 
+		public TestTargetRules? InnerTestTargetRules
+		{
+			get { return Inner as TestTargetRules; }
+		}
+
 #pragma warning restore C1591
 		#endregion
 
diff --git a/Engine/Source/Programs/UnrealBuildTool/Configuration/TestTargetRules.cs b/Engine/Source/Programs/UnrealBuildTool/Configuration/TestTargetRules.cs
--- a/Engine/Source/Programs/UnrealBuildTool/Configuration/TestTargetRules.cs
+++ b/Engine/Source/Programs/UnrealBuildTool/Configuration/TestTargetRules.cs
@@ -17,27 +17,27 @@ namespace UnrealBuildTool
 		/// <summary>
 		/// Keeps track if low level tests executable must build with the Editor.
 		/// </summary>
-		internal static bool bTestsRequireEditor = false;
+		internal bool bTestsRequireEditor = false;
 
 		/// <summary>
 		/// Keeps track if low level tests executable must build with the Engine.
 		/// </summary>
-		internal static bool bTestsRequireEngine = false;
+		internal bool bTestsRequireEngine = false;
 
 		/// <summary>
 		/// Keeps track if low level tests executable must build with the ApplicationCore.
 		/// </summary>
-		internal static bool bTestsRequireApplicationCore = false;
+		internal bool bTestsRequireApplicationCore = false;
 
 		/// <summary>
 		/// Keeps track if low level tests executable must build with the CoreUObject.
 		/// </summary>
-		internal static bool bTestsRequireCoreUObject = false;
+		internal bool bTestsRequireCoreUObject = false;
 
 		/// <summary>
 		/// Keep track of low level test runner module instance.
 		/// </summary>
-		internal static ModuleRules? LowLevelTestsRunnerModule;
+		internal ModuleRules? LowLevelTestsRunnerModule;
 
 		/// <summary>
 		/// Associated tested target of this test target, if defined.
diff --git a/Engine/Source/Programs/UnrealBuildTool/System/RulesAssembly.cs b/Engine/Source/Programs/UnrealBuildTool/System/RulesAssembly.cs
--- a/Engine/Source/Programs/UnrealBuildTool/System/RulesAssembly.cs
+++ b/Engine/Source/Programs/UnrealBuildTool/System/RulesAssembly.cs
@@ -115,7 +115,7 @@ namespace UnrealBuildTool
 		/// <param name="bForceCompile">Whether to always compile this assembly</param>
 		/// <param name="Parent">The parent rules assembly</param>
 		/// <param name="Logger"></param>
-		internal RulesAssembly(RulesScope Scope, List<DirectoryReference> BaseDirs, IReadOnlyList<PluginInfo> Plugins, Dictionary<FileReference, ModuleRulesContext> ModuleFileToContext, List<FileReference> TargetFiles, FileReference AssemblyFileName, bool bContainsEngineModules, BuildSettingsVersion? DefaultBuildSettings, bool bReadOnly, bool bSkipCompile, bool bForceCompile, RulesAssembly? Parent, ILogger Logger)
+		internal RulesAssembly(RulesScope Scope, IReadOnlyList<DirectoryReference> BaseDirs, IReadOnlyList<PluginInfo> Plugins, IReadOnlyDictionary<FileReference, ModuleRulesContext> ModuleFileToContext, IReadOnlyList<FileReference> TargetFiles, FileReference AssemblyFileName, bool bContainsEngineModules, BuildSettingsVersion? DefaultBuildSettings, bool bReadOnly, bool bSkipCompile, bool bForceCompile, RulesAssembly? Parent, ILogger Logger)
 		{
 			this.Scope = Scope;
 			this.BaseDirs = BaseDirs;
