Made code thread safe.

diff --git a/Engine/Source/Programs/UnrealBuildTool/Platform/Windows/MicrosoftPlatformSDK.cs b/Engine/Source/Programs/UnrealBuildTool/Platform/Windows/MicrosoftPlatformSDK.cs
--- a/Engine/Source/Programs/UnrealBuildTool/Platform/Windows/MicrosoftPlatformSDK.cs
+++ b/Engine/Source/Programs/UnrealBuildTool/Platform/Windows/MicrosoftPlatformSDK.cs
@@ -15,6 +15,7 @@ using System.Buffers.Binary;
 using UnrealBuildBase;
 using System.Runtime.Versioning;
 using Microsoft.Extensions.Logging;
+using System.Collections.Concurrent;
 
 ///////////////////////////////////////////////////////////////////
 // If you are looking for supported version numbers, look in the
@@ -91,7 +92,7 @@ namespace UnrealBuildTool
 		/// <summary>
 		/// Cache of Visual C++ installation directories
 		/// </summary>
-		private static Dictionary<WindowsCompiler, List<ToolChainInstallation>> CachedToolChainInstallations = new Dictionary<WindowsCompiler, List<ToolChainInstallation>>();
+		private static ConcurrentDictionary<WindowsCompiler, List<ToolChainInstallation>> CachedToolChainInstallations = new ConcurrentDictionary<WindowsCompiler, List<ToolChainInstallation>>();
 
 		/// <summary>
 		/// Cache of Windows SDK installation directories
@@ -111,7 +112,7 @@ namespace UnrealBuildTool
 		/// <summary>
 		/// Cache of DIA SDK installation directories
 		/// </summary>
-		private static Dictionary<WindowsCompiler, List<DirectoryReference>> CachedDiaSdkDirs = new Dictionary<WindowsCompiler, List<DirectoryReference>>();
+		private static ConcurrentDictionary<WindowsCompiler, List<DirectoryReference>> CachedDiaSdkDirs = new ConcurrentDictionary<WindowsCompiler, List<DirectoryReference>>();
 
 		#endregion
 
@@ -649,10 +650,9 @@ namespace UnrealBuildTool
 
 		static List<ToolChainInstallation> FindToolChainInstallations(WindowsCompiler Compiler, ILogger Logger)
 		{
-			List<ToolChainInstallation>? ToolChains;
-			if (!CachedToolChainInstallations.TryGetValue(Compiler, out ToolChains))
+			return CachedToolChainInstallations.GetOrAdd(Compiler, _ =>
 			{
-				ToolChains = new List<ToolChainInstallation>();
+				List<ToolChainInstallation>? ToolChains = new List<ToolChainInstallation>();
 				if (BuildHostPlatform.Current.Platform == UnrealTargetPlatform.Win64)
 				{
 					if (Compiler == WindowsCompiler.Clang)
@@ -758,9 +758,8 @@ namespace UnrealBuildTool
 						throw new BuildException("Unsupported compiler version ({0})", Compiler);
 					}
 				}
-				CachedToolChainInstallations.Add(Compiler, ToolChains);
-			}
-			return ToolChains;
+				return ToolChains;
+			});
 		}
 
 		/// <summary>
@@ -1278,10 +1277,9 @@ namespace UnrealBuildTool
 		/// <returns>Map of version number to directories</returns>
 		public static List<DirectoryReference> FindDiaSdkDirs(WindowsCompiler Compiler)
 		{
-			List<DirectoryReference>? DiaSdkDirs;
-			if (!CachedDiaSdkDirs.TryGetValue(Compiler, out DiaSdkDirs))
+			return CachedDiaSdkDirs.GetOrAdd(Compiler, _ =>
 			{
-				DiaSdkDirs = new List<DirectoryReference>();
+				List<DirectoryReference> DiaSdkDirs = new List<DirectoryReference>();
 
 				DirectoryReference? PlatformDir;
 				if (UEBuildPlatformSDK.TryGetHostPlatformAutoSDKDir(out PlatformDir))
@@ -1312,8 +1310,8 @@ namespace UnrealBuildTool
 						DiaSdkDirs.Add(DiaSdkDir);
 					}
 				}
-			}
-			return DiaSdkDirs;
+				return DiaSdkDirs;
+			});
 		}
 
 		/// <summary>
