Did a pass through the code to make it thread safe.

diff --git a/Engine/Source/Programs/AutomationTool/Scripts/GeneratePlatformReport.cs b/Engine/Source/Programs/AutomationTool/Scripts/GeneratePlatformReport.cs
--- a/Engine/Source/Programs/AutomationTool/Scripts/GeneratePlatformReport.cs
+++ b/Engine/Source/Programs/AutomationTool/Scripts/GeneratePlatformReport.cs
@@ -81,20 +81,24 @@ namespace AutomationTool
 		private PlatformReport GenerateDDPIReport(string[] Platforms, string DDPISection)
 		{
 			// Collect the list of DDPIs to use
-			Dictionary<string, DataDrivenPlatformInfo.ConfigDataDrivenPlatformInfo> PlatformInfos = new Dictionary<string, DataDrivenPlatformInfo.ConfigDataDrivenPlatformInfo>(); 
+			IReadOnlyDictionary<string, DataDrivenPlatformInfo.ConfigDataDrivenPlatformInfo> PlatformInfos; 
 			if (Platforms == null || Platforms.Length == 0)
 			{
 				PlatformInfos = DataDrivenPlatformInfo.GetAllPlatformInfos();
 			}
-			else foreach( string PlatformName in Platforms )
+			else
 			{
-				DataDrivenPlatformInfo.ConfigDataDrivenPlatformInfo? Config = DataDrivenPlatformInfo.GetDataDrivenInfoForPlatform(PlatformName);
-				if (Config != null)
+				Dictionary<string, DataDrivenPlatformInfo.ConfigDataDrivenPlatformInfo> PlatformInfosDict = new Dictionary<string, DataDrivenPlatformInfo.ConfigDataDrivenPlatformInfo>();
+				foreach (string PlatformName in Platforms)
 				{
-					PlatformInfos.Add( PlatformName, Config );
+					DataDrivenPlatformInfo.ConfigDataDrivenPlatformInfo? Config = DataDrivenPlatformInfo.GetDataDrivenInfoForPlatform(PlatformName);
+					if (Config != null)
+					{
+						PlatformInfosDict.Add(PlatformName, Config);
+					}
 				}
-			}
-
+				PlatformInfos = PlatformInfosDict;
+			} 	
 
 			// collect all possible config values
 			SortedSet<string> ColumnNames = new SortedSet<string>();
diff --git a/Engine/Source/Programs/Shared/EpicGames.Build/System/EnumeratePlugins.cs b/Engine/Source/Programs/Shared/EpicGames.Build/System/EnumeratePlugins.cs
--- a/Engine/Source/Programs/Shared/EpicGames.Build/System/EnumeratePlugins.cs
+++ b/Engine/Source/Programs/Shared/EpicGames.Build/System/EnumeratePlugins.cs
@@ -3,6 +3,7 @@
 using EpicGames.Core;
 using Microsoft.Extensions.Logging;
 using System;
+using System.Collections.Concurrent;
 using System.Collections.Generic;
 using System.Linq;
 using System.Text;
@@ -14,7 +15,7 @@ namespace UnrealBuildBase
 		/// <summary>
 		/// Cache of plugin filenames under each directory
 		/// </summary>
-		static Dictionary<DirectoryReference, List<FileReference>> PluginFileCache = new Dictionary<DirectoryReference, List<FileReference>>();
+		static ConcurrentDictionary<DirectoryReference, List<FileReference>> PluginFileCache = new ConcurrentDictionary<DirectoryReference, List<FileReference>>();
 
 		/// <summary>
 		/// Cache of all UBT plugins that have been built.  When generating projects, there is a good chance that a UBT project
@@ -28,7 +29,7 @@ namespace UnrealBuildBase
 		/// </summary>
 		public static void InvalidateCache_SLOW()
 		{
-			PluginFileCache = new Dictionary<DirectoryReference, List<FileReference>>();
+			PluginFileCache = new ConcurrentDictionary<DirectoryReference, List<FileReference>>();
 			DirectoryItem.ResetAllCachedInfo_SLOW();
 		}
 
@@ -55,15 +56,14 @@ namespace UnrealBuildBase
 		/// <param name="ParentDirectory">Parent directory to look in. Plugins will be found in any *subfolders* of this directory.</param>
 		public static IEnumerable<FileReference> EnumeratePlugins(DirectoryReference ParentDirectory)
 		{
-			List<FileReference>? FileNames;
-			if (!PluginFileCache.TryGetValue(ParentDirectory, out FileNames))
+			return PluginFileCache.GetOrAdd(ParentDirectory, _ =>
 			{
-				FileNames = new List<FileReference>();
+				List<FileReference> FileNames = new List<FileReference>();
 
 				DirectoryItem ParentDirectoryItem = DirectoryItem.GetItemByDirectoryReference(ParentDirectory);
 				if (ParentDirectoryItem.Exists)
 				{
-					using(ThreadPoolWorkQueue Queue = new ThreadPoolWorkQueue())
+					using (ThreadPoolWorkQueue Queue = new ThreadPoolWorkQueue())
 					{
 						EnumeratePluginsInternal(ParentDirectoryItem, FileNames, Queue);
 					}
@@ -71,10 +71,8 @@ namespace UnrealBuildBase
 
 				// Sort the filenames to ensure that the plugin order is deterministic; otherwise response files will change with each build.
 				FileNames = FileNames.OrderBy(x => x.FullName, StringComparer.OrdinalIgnoreCase).ToList();
-
-				PluginFileCache.Add(ParentDirectory, FileNames);
-			}
-			return FileNames;
+				return FileNames;
+			});
 		}
 
 		/// <summary>
diff --git a/Engine/Source/Programs/Shared/EpicGames.Build/Unreal.cs b/Engine/Source/Programs/Shared/EpicGames.Build/Unreal.cs
--- a/Engine/Source/Programs/Shared/EpicGames.Build/Unreal.cs
+++ b/Engine/Source/Programs/Shared/EpicGames.Build/Unreal.cs
@@ -2,6 +2,7 @@
 
 using EpicGames.Core;
 using System;
+using System.Collections.Concurrent;
 using System.Collections.Generic;
 using System.IO;
 using System.Linq;
@@ -257,7 +258,7 @@ namespace UnrealBuildBase
 		}
 
 		// cached dictionary of BaseDir to extension directories
-		private static Dictionary<DirectoryReference, Tuple<List<DirectoryReference>, List<DirectoryReference>>> CachedExtensionDirectories = new Dictionary<DirectoryReference, Tuple<List<DirectoryReference>, List<DirectoryReference>>>();
+		private static ConcurrentDictionary<DirectoryReference, Tuple<List<DirectoryReference>, List<DirectoryReference>>> CachedExtensionDirectories = new ConcurrentDictionary<DirectoryReference, Tuple<List<DirectoryReference>, List<DirectoryReference>>>();
 
 		/// <summary>
 		/// Finds all the extension directories for the given base directory. This includes platform extensions and restricted folders.
@@ -269,24 +270,21 @@ namespace UnrealBuildBase
 		/// <returns>List of extension directories, including the given base directory</returns>
 		public static List<DirectoryReference> GetExtensionDirs(DirectoryReference BaseDir, bool bIncludePlatformDirectories = true, bool bIncludeRestrictedDirectories = true, bool bIncludeBaseDirectory = true)
 		{
-			Tuple<List<DirectoryReference>, List<DirectoryReference>>? CachedDirs;
-			if (!CachedExtensionDirectories.TryGetValue(BaseDir, out CachedDirs))
+			Tuple<List<DirectoryReference>, List<DirectoryReference>> CachedDirs = CachedExtensionDirectories.GetOrAdd(BaseDir, _ =>
 			{
-				CachedDirs = Tuple.Create(new List<DirectoryReference>(), new List<DirectoryReference>());
-
-				CachedExtensionDirectories[BaseDir] = CachedDirs;
+				Tuple<List<DirectoryReference>, List<DirectoryReference>> NewCachedDirs = Tuple.Create(new List<DirectoryReference>(), new List<DirectoryReference>());
 
 				DirectoryReference PlatformExtensionBaseDir = DirectoryReference.Combine(BaseDir, "Platforms");
 				if (DirectoryReference.Exists(PlatformExtensionBaseDir))
 				{
-					CachedDirs.Item1.AddRange(DirectoryReference.EnumerateDirectories(PlatformExtensionBaseDir));
+					NewCachedDirs.Item1.AddRange(DirectoryReference.EnumerateDirectories(PlatformExtensionBaseDir));
 				}
 
 				DirectoryReference RestrictedBaseDir = DirectoryReference.Combine(BaseDir, "Restricted");
 				if (DirectoryReference.Exists(RestrictedBaseDir))
 				{
 					IEnumerable<DirectoryReference> RestrictedDirs = DirectoryReference.EnumerateDirectories(RestrictedBaseDir);
-					CachedDirs.Item2.AddRange(RestrictedDirs);
+					NewCachedDirs.Item2.AddRange(RestrictedDirs);
 
 					// also look for nested platforms in the restricted
 					foreach (DirectoryReference RestrictedDir in RestrictedDirs)
@@ -294,18 +292,19 @@ namespace UnrealBuildBase
 						DirectoryReference RestrictedPlatformExtensionBaseDir = DirectoryReference.Combine(RestrictedDir, "Platforms");
 						if (DirectoryReference.Exists(RestrictedPlatformExtensionBaseDir))
 						{
-							CachedDirs.Item1.AddRange(DirectoryReference.EnumerateDirectories(RestrictedPlatformExtensionBaseDir));
+							NewCachedDirs.Item1.AddRange(DirectoryReference.EnumerateDirectories(RestrictedPlatformExtensionBaseDir));
 						}
 					}
 				}
 
 				// remove any platform directories in non-engine locations if the engine doesn't have the platform 
-				if (BaseDir != Unreal.EngineDirectory && CachedDirs.Item1.Count > 0)
+				if (BaseDir != Unreal.EngineDirectory && NewCachedDirs.Item1.Count > 0)
 				{
 					// if the DDPI.ini file doesn't exist, we haven't synced the platform, so just skip this directory
-					CachedDirs.Item1.RemoveAll(x => !DataDrivenPlatformInfoIniIsPresent(x.GetDirectoryName()));
+					NewCachedDirs.Item1.RemoveAll(x => !DataDrivenPlatformInfoIniIsPresent(x.GetDirectoryName()));
 				}
-			}
+				return NewCachedDirs;
+			});
 
 			// now return what the caller wanted (always include BaseDir)
 			List<DirectoryReference> ExtensionDirs = new List<DirectoryReference>();
diff --git a/Engine/Source/Programs/UnrealBuildTool/Platform/IOS/UEBuildIOS.cs b/Engine/Source/Programs/UnrealBuildTool/Platform/IOS/UEBuildIOS.cs
--- a/Engine/Source/Programs/UnrealBuildTool/Platform/IOS/UEBuildIOS.cs
+++ b/Engine/Source/Programs/UnrealBuildTool/Platform/IOS/UEBuildIOS.cs
@@ -27,7 +27,7 @@ namespace UnrealBuildTool
 		public static UnrealArch TVOSSimulator = FindOrAddByName("tvossimulator", bIsX64: false);
 
 
-		private static Dictionary<UnrealArch, string> AppleToolchainArchitectures = new()
+		private static IReadOnlyDictionary<UnrealArch, string> AppleToolchainArchitectures = new Dictionary<UnrealArch, string>()
 		{
 			{ UnrealArch.Arm64,         "arm64" },
 			{ UnrealArch.X64,           "x86_64" },
diff --git a/Engine/Source/Programs/UnrealBuildTool/Platform/Linux/UEBuildLinux.cs b/Engine/Source/Programs/UnrealBuildTool/Platform/Linux/UEBuildLinux.cs
--- a/Engine/Source/Programs/UnrealBuildTool/Platform/Linux/UEBuildLinux.cs
+++ b/Engine/Source/Programs/UnrealBuildTool/Platform/Linux/UEBuildLinux.cs
@@ -14,7 +14,7 @@ namespace UnrealBuildTool
 {
 	partial struct UnrealArch
 	{
-		private static Dictionary<UnrealArch, string> LinuxToolchainArchitectures = new()
+		private static IReadOnlyDictionary<UnrealArch, string> LinuxToolchainArchitectures = new Dictionary<UnrealArch, string>()
 		{
 			{ UnrealArch.Arm64,         "aarch64-unknown-linux-gnueabi" },
 			{ UnrealArch.X64,           "x86_64-unknown-linux-gnu" },
diff --git a/Engine/Source/Programs/UnrealBuildTool/Platform/Windows/UEBuildWindows.cs b/Engine/Source/Programs/UnrealBuildTool/Platform/Windows/UEBuildWindows.cs
--- a/Engine/Source/Programs/UnrealBuildTool/Platform/Windows/UEBuildWindows.cs
+++ b/Engine/Source/Programs/UnrealBuildTool/Platform/Windows/UEBuildWindows.cs
@@ -35,7 +35,7 @@ namespace UnrealBuildTool
 			}
 		}
 
-		private static Dictionary<UnrealArch, WindowsArchInfo> WindowsToolchainArchitectures = new()
+		private static IReadOnlyDictionary<UnrealArch, WindowsArchInfo> WindowsToolchainArchitectures = new Dictionary<UnrealArch, WindowsArchInfo>()
 		{
 			{ UnrealArch.Arm64,         new WindowsArchInfo("arm64", "arm64", "arm64") },
 			{ UnrealArch.Arm64ec,       new WindowsArchInfo("arm64", "x64", "arm64") },
diff --git a/Engine/Source/Programs/UnrealBuildTool/System/ActionSerializer.cs b/Engine/Source/Programs/UnrealBuildTool/System/ActionSerializer.cs
--- a/Engine/Source/Programs/UnrealBuildTool/System/ActionSerializer.cs
+++ b/Engine/Source/Programs/UnrealBuildTool/System/ActionSerializer.cs
@@ -74,12 +74,12 @@ namespace UnrealBuildTool
 		/// <summary>
 		/// Map from type name to deserializing constructor
 		/// </summary>
-		static Dictionary<Type, IActionSerializer> TypeToSerializer;
+		static IReadOnlyDictionary<Type, IActionSerializer> TypeToSerializer;
 
 		/// <summary>
 		/// Map from serializer name to instance
 		/// </summary>
-		static Dictionary<string, IActionSerializer> NameToSerializer;
+		static IReadOnlyDictionary<string, IActionSerializer> NameToSerializer;
 
 		/// <summary>
 		/// Creates a map of type name to constructor
@@ -87,8 +87,8 @@ namespace UnrealBuildTool
 		/// <returns></returns>
 		static ActionSerialization()
 		{
-			TypeToSerializer = new Dictionary<Type, IActionSerializer>();
-			NameToSerializer = new Dictionary<string, IActionSerializer>(StringComparer.Ordinal);
+			var TypeToSerializerDict = new Dictionary<Type, IActionSerializer>();
+			var NameToSerializerDict = new Dictionary<string, IActionSerializer>(StringComparer.Ordinal);
 
 			Type[] Types = Assembly.GetExecutingAssembly().GetTypes();
 			foreach (Type Type in Types)
@@ -96,10 +96,13 @@ namespace UnrealBuildTool
 				if (Type.IsClass && !Type.IsAbstract && typeof(IActionSerializer).IsAssignableFrom(Type))
 				{
 					IActionSerializer Serializer = (IActionSerializer)Activator.CreateInstance(Type)!;
-					TypeToSerializer[Serializer.Type] = Serializer;
-					NameToSerializer[Type.Name] = Serializer;
+					TypeToSerializerDict[Serializer.Type] = Serializer;
+					NameToSerializerDict[Type.Name] = Serializer;
 				}
 			}
+
+			TypeToSerializer = TypeToSerializerDict;
+			NameToSerializer = NameToSerializerDict;
 		}
 
 		/// <summary>
diff --git a/Engine/Source/Programs/UnrealBuildTool/System/DataDrivenPlatformInfo.cs b/Engine/Source/Programs/UnrealBuildTool/System/DataDrivenPlatformInfo.cs
--- a/Engine/Source/Programs/UnrealBuildTool/System/DataDrivenPlatformInfo.cs
+++ b/Engine/Source/Programs/UnrealBuildTool/System/DataDrivenPlatformInfo.cs
@@ -62,132 +62,136 @@ namespace UnrealBuildTool
 			public ConfigFile? PlatformConfig = null;
 		};
 
-		static Dictionary<string, ConfigDataDrivenPlatformInfo>? PlatformInfos = null;
+		static IReadOnlyDictionary<string, ConfigDataDrivenPlatformInfo>? PlatformInfos = null;
+		static object LockObject = new object();
 
 		/// <summary>
 		/// Return all data driven infos found
 		/// </summary>
 		/// <returns></returns>
-		public static Dictionary<string, ConfigDataDrivenPlatformInfo> GetAllPlatformInfos()
+		public static IReadOnlyDictionary<string, ConfigDataDrivenPlatformInfo> GetAllPlatformInfos()
 		{
 			// need to init?
-			if (PlatformInfos == null)
+			lock (LockObject)
 			{
-				PlatformInfos = new Dictionary<string, ConfigDataDrivenPlatformInfo>(StringComparer.OrdinalIgnoreCase);
-				Dictionary<string, string> IniParents = new Dictionary<string, string>(StringComparer.OrdinalIgnoreCase);
-
-				// find all platform directories (skipping NFL/NoRedist)
-				foreach (DirectoryReference EngineConfigDir in Unreal.GetExtensionDirs(Unreal.EngineDirectory, "Config", bIncludeRestrictedDirectories:false))
+				if (PlatformInfos == null)
 				{
-					// look through all config dirs looking for the data driven ini file
-					foreach (string FilePath in Directory.EnumerateFiles(EngineConfigDir.FullName, "DataDrivenPlatformInfo.ini", SearchOption.AllDirectories))
-					{
-						FileReference FileRef = new FileReference(FilePath);
-
-						// get the platform name from the path
-						string IniPlatformName;
-						if (FileRef.IsUnderDirectory(DirectoryReference.Combine(Unreal.EngineDirectory, "Config")))
-						{
-							// Foo/Engine/Config/<Platform>/DataDrivenPlatformInfo.ini
-							IniPlatformName = Path.GetFileName(Path.GetDirectoryName(FilePath))!;
-						}
-						else
-						{
-							// Foo/Engine/Platforms/<Platform>/Config/DataDrivenPlatformInfo.ini
-							IniPlatformName = Path.GetFileName(Path.GetDirectoryName(Path.GetDirectoryName(FilePath)))!;
-						}
-
-						// load the DataDrivenPlatformInfo from the path (with Add support in a file that doesn't use +'s in the arrays for C++ usage)
-						ConfigFile Config = new ConfigFile(FileRef, ConfigLineAction.Add);
-						ConfigDataDrivenPlatformInfo NewInfo = new ConfigDataDrivenPlatformInfo();
+					var PlatformInfosDict = new Dictionary<string, ConfigDataDrivenPlatformInfo>(StringComparer.OrdinalIgnoreCase);
+					PlatformInfos = PlatformInfosDict;
+					Dictionary<string, string> IniParents = new Dictionary<string, string>(StringComparer.OrdinalIgnoreCase);
 
-
-						// we must have the key section 
-						ConfigFileSection? Section = null;
-						if (Config.TryGetSection("DataDrivenPlatformInfo", out Section))
+					// find all platform directories (skipping NFL/NoRedist)
+					foreach (DirectoryReference EngineConfigDir in Unreal.GetExtensionDirs(Unreal.EngineDirectory, "Config", bIncludeRestrictedDirectories: false))
+					{
+						// look through all config dirs looking for the data driven ini file
+						foreach (string FilePath in Directory.EnumerateFiles(EngineConfigDir.FullName, "DataDrivenPlatformInfo.ini", SearchOption.AllDirectories))
 						{
-							ConfigHierarchySection ParsedSection = new ConfigHierarchySection(new List<ConfigFileSection>() { Section });
+							FileReference FileRef = new FileReference(FilePath);
 
-							// get string values
-							string? IniParent;
-							if (ParsedSection.TryGetValue("IniParent", out IniParent))
-							{
-								IniParents[IniPlatformName] = IniParent;
-							}
-
-							// slightly nasty bool parsing for bool values
-							string? Temp;
-							if (ParsedSection.TryGetValue("bIsConfidential", out Temp))
-							{
-								ConfigHierarchy.TryParse(Temp, out NewInfo.bIsConfidential);
-							}
-							// unspecified means true for this property
-							NewInfo.bCanUseCrashReporter = true;
-							if (ParsedSection.TryGetValue("bCanUseCrashReporter", out Temp))
+							// get the platform name from the path
+							string IniPlatformName;
+							if (FileRef.IsUnderDirectory(DirectoryReference.Combine(Unreal.EngineDirectory, "Config")))
 							{
-								ConfigHierarchy.TryParse(Temp, out NewInfo.bCanUseCrashReporter);
+								// Foo/Engine/Config/<Platform>/DataDrivenPlatformInfo.ini
+								IniPlatformName = Path.GetFileName(Path.GetDirectoryName(FilePath))!;
 							}
-
-							string HostKey = ConfigHierarchy.GetIniPlatformName(BuildHostPlatform.Current.Platform) + ":bIsEnabled";
-							string NormalKey = "bIsEnabled";
-							if (ParsedSection.TryGetValue(HostKey, out Temp) == true || ParsedSection.TryGetValue(NormalKey, out Temp) == true)
+							else
 							{
-								ConfigHierarchy.TryParse(Temp, out NewInfo.bIsEnabled);
+								// Foo/Engine/Platforms/<Platform>/Config/DataDrivenPlatformInfo.ini
+								IniPlatformName = Path.GetFileName(Path.GetDirectoryName(Path.GetDirectoryName(FilePath)))!;
 							}
 
-							if (ParsedSection.TryGetValue("HardwareCompressionFormat", out NewInfo.HardwareCompressionFormat) == false)
-							{
-								NewInfo.HardwareCompressionFormat = null;
-							}
+							// load the DataDrivenPlatformInfo from the path (with Add support in a file that doesn't use +'s in the arrays for C++ usage)
+							ConfigFile Config = new ConfigFile(FileRef, ConfigLineAction.Add);
+							ConfigDataDrivenPlatformInfo NewInfo = new ConfigDataDrivenPlatformInfo();
 
-							NewInfo.bNoAccountOverride = false;
-							if (ParsedSection.TryGetValue("bNoAccountOverride", out Temp))
-							{
-								ConfigHierarchy.TryParse(Temp, out NewInfo.bNoAccountOverride);
-							}
 
-							// get a list of additional restricted folders
-							IReadOnlyList<string>? AdditionalRestrictedFolders;
-							if(ParsedSection.TryGetValues("AdditionalRestrictedFolders", out AdditionalRestrictedFolders) && AdditionalRestrictedFolders.Count > 0)
+							// we must have the key section 
+							ConfigFileSection? Section = null;
+							if (Config.TryGetSection("DataDrivenPlatformInfo", out Section))
 							{
-								NewInfo.AdditionalRestrictedFolders = AdditionalRestrictedFolders.Select(x => x.Trim()).Where(x => x.Length > 0).ToArray();
+								ConfigHierarchySection ParsedSection = new ConfigHierarchySection(new List<ConfigFileSection>() { Section });
+
+								// get string values
+								string? IniParent;
+								if (ParsedSection.TryGetValue("IniParent", out IniParent))
+								{
+									IniParents[IniPlatformName] = IniParent;
+								}
+
+								// slightly nasty bool parsing for bool values
+								string? Temp;
+								if (ParsedSection.TryGetValue("bIsConfidential", out Temp))
+								{
+									ConfigHierarchy.TryParse(Temp, out NewInfo.bIsConfidential);
+								}
+								// unspecified means true for this property
+								NewInfo.bCanUseCrashReporter = true;
+								if (ParsedSection.TryGetValue("bCanUseCrashReporter", out Temp))
+								{
+									ConfigHierarchy.TryParse(Temp, out NewInfo.bCanUseCrashReporter);
+								}
+
+								string HostKey = ConfigHierarchy.GetIniPlatformName(BuildHostPlatform.Current.Platform) + ":bIsEnabled";
+								string NormalKey = "bIsEnabled";
+								if (ParsedSection.TryGetValue(HostKey, out Temp) == true || ParsedSection.TryGetValue(NormalKey, out Temp) == true)
+								{
+									ConfigHierarchy.TryParse(Temp, out NewInfo.bIsEnabled);
+								}
+
+								if (ParsedSection.TryGetValue("HardwareCompressionFormat", out NewInfo.HardwareCompressionFormat) == false)
+								{
+									NewInfo.HardwareCompressionFormat = null;
+								}
+
+								NewInfo.bNoAccountOverride = false;
+								if (ParsedSection.TryGetValue("bNoAccountOverride", out Temp))
+								{
+									ConfigHierarchy.TryParse(Temp, out NewInfo.bNoAccountOverride);
+								}
+
+								// get a list of additional restricted folders
+								IReadOnlyList<string>? AdditionalRestrictedFolders;
+								if (ParsedSection.TryGetValues("AdditionalRestrictedFolders", out AdditionalRestrictedFolders) && AdditionalRestrictedFolders.Count > 0)
+								{
+									NewInfo.AdditionalRestrictedFolders = AdditionalRestrictedFolders.Select(x => x.Trim()).Where(x => x.Length > 0).ToArray();
+								}
+
+								// create cache it
+								NewInfo.PlatformConfig = Config;
+								PlatformInfosDict[IniPlatformName] = NewInfo;
 							}
-
-							// create cache it
-							NewInfo.PlatformConfig = Config;
-							PlatformInfos[IniPlatformName] = NewInfo;
 						}
 					}
-				}
 
-				// now that all are read in, calculate the ini parent chain, starting with parent-most
-				foreach (KeyValuePair<string, ConfigDataDrivenPlatformInfo> Pair in PlatformInfos)
-				{
-					string? CurrentPlatform;
-
-					// walk up the chain and build up the ini chain
-					List<string> Chain = new List<string>();
-					if (IniParents.TryGetValue(Pair.Key, out CurrentPlatform))
+					// now that all are read in, calculate the ini parent chain, starting with parent-most
+					foreach (KeyValuePair<string, ConfigDataDrivenPlatformInfo> Pair in PlatformInfos)
 					{
-						while (!string.IsNullOrEmpty(CurrentPlatform))
+						string? CurrentPlatform;
+
+						// walk up the chain and build up the ini chain
+						List<string> Chain = new List<string>();
+						if (IniParents.TryGetValue(Pair.Key, out CurrentPlatform))
 						{
-							// insert at 0 to reverse the order
-							Chain.Insert(0, CurrentPlatform);
-							if (IniParents.TryGetValue(CurrentPlatform, out CurrentPlatform) == false)
+							while (!string.IsNullOrEmpty(CurrentPlatform))
 							{
-								break;
+								// insert at 0 to reverse the order
+								Chain.Insert(0, CurrentPlatform);
+								if (IniParents.TryGetValue(CurrentPlatform, out CurrentPlatform) == false)
+								{
+									break;
+								}
 							}
 						}
-					}
 
-					// bake it into the info
-					if (Chain.Count > 0)
-					{
-						Pair.Value.IniParentChain = Chain.ToArray();
+						// bake it into the info
+						if (Chain.Count > 0)
+						{
+							Pair.Value.IniParentChain = Chain.ToArray();
+						}
 					}
 				}
 			}
-
 			return PlatformInfos;
 		}
 
diff --git a/Engine/Source/Programs/UnrealBuildTool/System/RulesCompiler.cs b/Engine/Source/Programs/UnrealBuildTool/System/RulesCompiler.cs
--- a/Engine/Source/Programs/UnrealBuildTool/System/RulesCompiler.cs
+++ b/Engine/Source/Programs/UnrealBuildTool/System/RulesCompiler.cs
@@ -9,6 +9,7 @@ using EpicGames.Core;
 using OpenTracing.Util;
 using UnrealBuildBase;
 using Microsoft.Extensions.Logging;
+using System.Collections.Concurrent;
 
 namespace UnrealBuildTool
 {
@@ -69,7 +70,7 @@ namespace UnrealBuildTool
 		/// Map of assembly names we've already compiled and loaded to their Assembly and list of game folders.  This is used to prevent
 		/// trying to recompile the same assembly when ping-ponging between different types of targets
 		/// </summary>
-		private static Dictionary<FileReference, RulesAssembly> LoadedAssemblyMap = new Dictionary<FileReference, RulesAssembly>();
+		private static ConcurrentDictionary<FileReference, RulesAssembly> LoadedAssemblyMap = new ConcurrentDictionary<FileReference, RulesAssembly>();
 
 		/// <summary>
 		/// Creates the engine rules assembly
@@ -232,8 +233,7 @@ namespace UnrealBuildTool
 		public static RulesAssembly CreateProjectRulesAssembly(FileReference ProjectFileName, bool bUsePrecompiled, bool bSkipCompile, bool bForceCompile, ILogger Logger)
 		{
 			// Check if there's an existing assembly for this project
-			RulesAssembly? ProjectRulesAssembly;
-			if (!LoadedAssemblyMap.TryGetValue(ProjectFileName, out ProjectRulesAssembly))
+			return LoadedAssemblyMap.GetOrAdd(ProjectFileName, _ =>
 			{
 				Logger.LogTrace("Creating project rules assembly for {Project}...", ProjectFileName);
 
@@ -308,19 +308,19 @@ namespace UnrealBuildTool
 					TargetFiles.AddRange(Rules.FindAllRulesFiles(ProjectIntermediateSourceDirectory, Rules.RulesFileType.Target));
 				}
 
+				RulesAssembly ProjectRulesAssembly;
 				// Compile the assembly. If there are no module or target files, just use the parent assembly.
 				FileReference AssemblyFileName = FileReference.Combine(MainProjectDirectory, "Intermediate", "Build", "BuildRules", ProjectFileName.GetFileNameWithoutExtension() + "ModuleRules" + FrameworkAssemblyExtension);
-				if(ModuleFiles.Count == 0 && TargetFiles.Count == 0 && ProjectPlugins.Count == 0)  
+				if (ModuleFiles.Count == 0 && TargetFiles.Count == 0 && ProjectPlugins.Count == 0)
 				{
-					ProjectRulesAssembly = Parent; 
+					ProjectRulesAssembly = Parent;
 				}
 				else
 				{
 					ProjectRulesAssembly = new RulesAssembly(Scope, new List<DirectoryReference> { MainProjectDirectory }, ProjectPlugins, ModuleFiles, TargetFiles, AssemblyFileName, bContainsEngineModules: false, DefaultBuildSettings: null, bReadOnly: UnrealBuildTool.IsProjectInstalled(), bSkipCompile: bSkipCompile, bForceCompile: bForceCompile, Parent: Parent, Logger: Logger);
 				}
-				LoadedAssemblyMap.Add(ProjectFileName, ProjectRulesAssembly);
-			}
-			return ProjectRulesAssembly;
+				return ProjectRulesAssembly;
+			});
 		}
 
 		/// <summary>
@@ -337,8 +337,7 @@ namespace UnrealBuildTool
 		public static RulesAssembly CreatePluginRulesAssembly(FileReference PluginFileName, bool bSkipCompile, bool bForceCompile, RulesAssembly Parent, bool bBuildPluginAsLocal, bool bContainsEngineModules, ILogger Logger)
 		{
 			// Check if there's an existing assembly for this project
-			RulesAssembly? PluginRulesAssembly;
-			if (!LoadedAssemblyMap.TryGetValue(PluginFileName, out PluginRulesAssembly))
+			return LoadedAssemblyMap.GetOrAdd(PluginFileName, _ =>
 			{
 				Logger.LogTrace("Creating plugin rules assembly for {Plugin}...", PluginFileName);
 
@@ -367,7 +366,7 @@ namespace UnrealBuildTool
 				// Find all the modules
 				ModuleRulesContext PluginModuleContext = new ModuleRulesContext(Scope, PluginFileName.Directory);
 				PluginModuleContext.bClassifyAsGameModuleForUHT = !bContainsEngineModules;
-				if(bBuildPluginAsLocal)
+				if (bBuildPluginAsLocal)
 				{
 					PluginModuleContext.bCanBuildDebugGame = true;
 					PluginModuleContext.bCanHotReload = true;
@@ -384,10 +383,8 @@ namespace UnrealBuildTool
 
 				// Compile the assembly
 				FileReference AssemblyFileName = FileReference.Combine(PluginFileName.Directory, "Intermediate", "Build", "BuildRules", Path.GetFileNameWithoutExtension(PluginFileName.FullName) + "ModuleRules" + FrameworkAssemblyExtension);
-				PluginRulesAssembly = new RulesAssembly(Scope, new List<DirectoryReference> { PluginFileName.Directory }, ForeignPlugins, ModuleFiles, TargetFiles, AssemblyFileName, bContainsEngineModules, DefaultBuildSettings: null, bReadOnly: false, bSkipCompile: bSkipCompile, bForceCompile: bForceCompile, Parent: Parent, Logger: Logger);
-				LoadedAssemblyMap.Add(PluginFileName, PluginRulesAssembly);
-			}
-			return PluginRulesAssembly;
+				return new RulesAssembly(Scope, new List<DirectoryReference> { PluginFileName.Directory }, ForeignPlugins, ModuleFiles, TargetFiles, AssemblyFileName, bContainsEngineModules, DefaultBuildSettings: null, bReadOnly: false, bSkipCompile: bSkipCompile, bForceCompile: bForceCompile, Parent: Parent, Logger: Logger);
+			});
 		}
 
 		/// <summary>
