Fix Xcode build on Development config, VALID_ARCHS was not using the correct apple name

diff --git a/Engine/Source/Programs/UnrealBuildTool/ProjectFiles/Xcode/XcodeProjectLegacy.cs b/Engine/Source/Programs/UnrealBuildTool/ProjectFiles/Xcode/XcodeProjectLegacy.cs
--- a/Engine/Source/Programs/UnrealBuildTool/ProjectFiles/Xcode/XcodeProjectLegacy.cs
+++ b/Engine/Source/Programs/UnrealBuildTool/ProjectFiles/Xcode/XcodeProjectLegacy.cs
@@ -1343,7 +1343,7 @@ namespace UnrealBuildTool.XcodeProjectLegacy
 					ValidArchs = ValidArchs.Union(this.SupportedIOSArchitectures);
 					SupportedPlatforms += " appletvos";
 				}
-				Content.Append("\t\t\t\tVALID_ARCHS = \"" + string.Join(" ", ValidArchs) + "\";" + ProjectFileGenerator.NewLine);
+				Content.Append("\t\t\t\tVALID_ARCHS = \"" + string.Join(" ", ValidArchs.Select(x => x.AppleName)) + "\";" + ProjectFileGenerator.NewLine);
 				Content.Append("\t\t\t\tSUPPORTED_PLATFORMS = \"" + SupportedPlatforms + "\";" + ProjectFileGenerator.NewLine);
 			}
 			Content.Append("\t\t\t\tGCC_PREPROCESSOR_DEFINITIONS = ();" + ProjectFileGenerator.NewLine);
