Updated to write out the PrimaryProjectName.txt only once even if the contents need to change. This fixes an issue where a certain platform was writing it twice and UBT would throw an warning about the file being written multiple times.

diff --git a/Engine/Source/Programs/UnrealBuildTool/ProjectFiles/ProjectFileGenerator.cs b/Engine/Source/Programs/UnrealBuildTool/ProjectFiles/ProjectFileGenerator.cs
--- a/Engine/Source/Programs/UnrealBuildTool/ProjectFiles/ProjectFileGenerator.cs
+++ b/Engine/Source/Programs/UnrealBuildTool/ProjectFiles/ProjectFileGenerator.cs
@@ -836,6 +836,8 @@ namespace UnrealBuildTool
 		{
 			bool bSuccess = true;
 
+			bool WritePrimaryProjectNameTxt = false;
+
 			// Parse project generator options
 			bool IncludeAllPlatforms = true;
 			ConfigureProjectFileGeneration(Arguments, ref IncludeAllPlatforms, Logger);
@@ -883,8 +885,7 @@ namespace UnrealBuildTool
 				}
 
 				// Write out the name of the primary project file, so the runtime knows to use it
-				FileReference PrimaryProjectNameLocation = FileReference.Combine(Unreal.EngineDirectory, "Intermediate", "ProjectFiles", "PrimaryProjectName.txt");
-				Utils.WriteFileIfChanged(PrimaryProjectNameLocation, PrimaryProjectName, Logger);
+				WritePrimaryProjectNameTxt = true;
 			}
 
 			// Modify the name if specific platforms were given
@@ -906,6 +907,11 @@ namespace UnrealBuildTool
 				}
 
 				// the primary project name is always read from our intermediate directory and not the overriden one for this set of platforms
+				WritePrimaryProjectNameTxt = true;
+			}
+
+			if (WritePrimaryProjectNameTxt)
+			{
 				FileReference PrimaryProjectNameLocation = FileReference.Combine(Unreal.EngineDirectory, "Intermediate", "ProjectFiles", "PrimaryProjectName.txt");
 				Utils.WriteFileIfChanged(PrimaryProjectNameLocation, PrimaryProjectName, Logger);
 			}
