Updated the search file code to use the directory cache to save time.

diff --git a/Engine/Source/Programs/UnrealBuildTool/System/SourceFileSearch.cs b/Engine/Source/Programs/UnrealBuildTool/System/SourceFileSearch.cs
--- a/Engine/Source/Programs/UnrealBuildTool/System/SourceFileSearch.cs
+++ b/Engine/Source/Programs/UnrealBuildTool/System/SourceFileSearch.cs
@@ -7,6 +7,7 @@ using System.IO;
 using System.Diagnostics;
 using System.Linq;
 using EpicGames.Core;
+using UnrealBuildBase;
 
 namespace UnrealBuildTool
 {
@@ -87,7 +88,7 @@ namespace UnrealBuildTool
 				SearchedDirectories.Add(Directory);
 			}
 
-			foreach (FileReference File in DirectoryReference.EnumerateFiles(Directory))
+			foreach (FileReference File in DirectoryLookupCache.EnumerateFiles(Directory))
 			{
 				if (ShouldInclude(File, DefaultExcludedFileSuffixes))
 				{
@@ -100,7 +101,7 @@ namespace UnrealBuildTool
 		{
 			FindFilesInternal(Directory, ExcludedDirectorySuffixes, FoundFiles, SearchedDirectories);
 
-			foreach (DirectoryReference SubDirectory in DirectoryReference.EnumerateDirectories(Directory))
+			foreach (DirectoryReference SubDirectory in DirectoryLookupCache.EnumerateDirectories(Directory))
 			{
 				if (ShouldInclude(SubDirectory, ExcludedDirectorySuffixes))
 				{
