UnrealBuildTool: Update support for minimum cpu arch for x64 platforms. Replaces ModuleRules.bUseAVX

* Add MinimumCpuArchitectureX64 enum, valid values are None, AVX, AVX2, & AVX512. Default value is None
* Obsolete boolean bUseAVX in favor of MinCpuArchX64
* Add MinCpuArchX64 setting to TargetRules
* Add Optional MinCpuArchX64 to ModuleRules, will override any target setting if not null
* Only pass this /arch for x64 platforms
* Can be controlled on the command line with -MinCpuArchX64=

diff --git a/Engine/Source/Programs/UnrealBuildTool/Configuration/ModuleRules.cs b/Engine/Source/Programs/UnrealBuildTool/Configuration/ModuleRules.cs
--- a/Engine/Source/Programs/UnrealBuildTool/Configuration/ModuleRules.cs
+++ b/Engine/Source/Programs/UnrealBuildTool/Configuration/ModuleRules.cs
@@ -764,10 +764,27 @@ namespace UnrealBuildTool
 		public bool bUseRTTI = false;
 
 		/// <summary>
-		/// Direct the compiler to generate AVX instructions wherever SSE or AVX intrinsics are used, on the platforms that support it.
+		/// Obsolete: Direct the compiler to generate AVX instructions wherever SSE or AVX intrinsics are used, on the platforms that support it.
 		/// Note that by enabling this you are changing the minspec for the PC platform, and the resultant executable will crash on machines without AVX support.
 		/// </summary>
-		public bool bUseAVX = false;
+		[Obsolete("bUseAVX is obsolete and will be removed in UE5.4, please replace with MinCpuArchX64")]
+		public bool bUseAVX
+		{
+			get
+			{
+				return MinCpuArchX64 >= MinimumCpuArchitectureX64.AVX;
+			}
+			set
+			{
+				MinCpuArchX64 = value ? MinimumCpuArchitectureX64.AVX : MinimumCpuArchitectureX64.None;
+			}
+		}
+
+		/// <summary>
+		/// Direct the compiler to generate AVX instructions wherever SSE or AVX intrinsics are used, on the x64 platforms that support it.
+		/// Note that by enabling this you are changing the minspec for the PC platform, and the resultant executable will crash on machines without AVX support.
+		/// </summary>
+		public MinimumCpuArchitectureX64? MinCpuArchX64 = null;
 
 		/// <summary>
 		/// Enable buffer security checks.  This should usually be enabled as it prevents severe security risks.
diff --git a/Engine/Source/Programs/UnrealBuildTool/Configuration/TargetRules.cs b/Engine/Source/Programs/UnrealBuildTool/Configuration/TargetRules.cs
--- a/Engine/Source/Programs/UnrealBuildTool/Configuration/TargetRules.cs
+++ b/Engine/Source/Programs/UnrealBuildTool/Configuration/TargetRules.cs
@@ -2067,6 +2067,15 @@ namespace UnrealBuildTool
 		[XmlConfigFile(Category = "BuildConfiguration")]
 		public CStandardVersion CStandard = CStandardVersion.Default;
 
+		/// <summary>
+		/// Direct the compiler to generate AVX instructions wherever SSE or AVX intrinsics are used, on the x64 platforms that support it.
+		/// Note that by enabling this you are changing the minspec for the PC platform, and the resultant executable will crash on machines without AVX support.
+		/// </summary>
+		[RequiresUniqueBuildEnvironment]
+		[CommandLine("-MinCpuArchX64")]
+		[XmlConfigFile(Category = "BuildConfiguration")]
+		public MinimumCpuArchitectureX64 MinCpuArchX64 = MinimumCpuArchitectureX64.Default;
+
 		/// <summary>
 		/// Do not allow manifest changes when building this target. Used to cause earlier errors when building multiple targets with a shared build environment.
 		/// </summary>
@@ -3780,6 +3789,11 @@ namespace UnrealBuildTool
 			get { return Inner.CStandard; }
 		}
 
+		public MinimumCpuArchitectureX64 MinCpuArchX64
+		{
+			get { return Inner.MinCpuArchX64; }
+		}
+
 		internal bool bNoManifestChanges
 		{
 			get { return Inner.bNoManifestChanges; }
diff --git a/Engine/Source/Programs/UnrealBuildTool/Configuration/UEBuildModuleCPP.cs b/Engine/Source/Programs/UnrealBuildTool/Configuration/UEBuildModuleCPP.cs
--- a/Engine/Source/Programs/UnrealBuildTool/Configuration/UEBuildModuleCPP.cs
+++ b/Engine/Source/Programs/UnrealBuildTool/Configuration/UEBuildModuleCPP.cs
@@ -1673,7 +1673,6 @@ namespace UnrealBuildTool
 			Result.bUseUnity = Rules.bUseUnity;
 			Result.bOptimizeCode = ShouldEnableOptimization(Rules.OptimizeCode, Target.Configuration, Rules.bTreatAsEngineModule);
 			Result.bUseRTTI |= Rules.bUseRTTI;
-			Result.bUseAVX = Rules.bUseAVX;
 			Result.bEnableBufferSecurityChecks = Rules.bEnableBufferSecurityChecks;
 			Result.MinSourceFilesForUnityBuildOverride = Rules.MinSourceFilesForUnityBuildOverride;
 			Result.MinFilesUsingPrecompiledHeaderOverride = Rules.MinFilesUsingPrecompiledHeaderOverride;
@@ -1729,6 +1728,12 @@ namespace UnrealBuildTool
 				Result.CStandard = (CStandardVersion)Rules.CStandard;
 			}
 
+			// If the module overrides the x64 minimum arch, override it on the compile environment
+			if (Rules.MinCpuArchX64 != null)
+			{
+				Result.MinCpuArchX64 = (MinimumCpuArchitectureX64)Rules.MinCpuArchX64;
+			}
+
 			// Set the macro used to check whether monolithic headers can be used
 			if ((Rules.bTreatAsEngineModule || Target.bWarnAboutMonolithicHeadersIncluded) && (Rules.IWYUSupport == IWYUSupport.None || !Target.bEnforceIWYU))
 			{
diff --git a/Engine/Source/Programs/UnrealBuildTool/Configuration/UEBuildTarget.cs b/Engine/Source/Programs/UnrealBuildTool/Configuration/UEBuildTarget.cs
--- a/Engine/Source/Programs/UnrealBuildTool/Configuration/UEBuildTarget.cs
+++ b/Engine/Source/Programs/UnrealBuildTool/Configuration/UEBuildTarget.cs
@@ -4445,6 +4445,7 @@ namespace UnrealBuildTool
 			GlobalCompileEnvironment.bHideSymbolsByDefault = !Rules.bPublicSymbolsByDefault;
 			GlobalCompileEnvironment.CppStandard = Rules.CppStandard;
 			GlobalCompileEnvironment.CStandard = Rules.CStandard;
+			GlobalCompileEnvironment.MinCpuArchX64 = Rules.MinCpuArchX64;
 			GlobalCompileEnvironment.AdditionalArguments = Rules.AdditionalCompilerArguments ?? String.Empty;
 			GlobalCompileEnvironment.bDeterministic = Rules.bDeterministic;
 			GlobalCompileEnvironment.CrashDiagnosticDirectory = Rules.CrashDiagnosticDirectory;
diff --git a/Engine/Source/Programs/UnrealBuildTool/Platform/Windows/VCToolChain.cs b/Engine/Source/Programs/UnrealBuildTool/Platform/Windows/VCToolChain.cs
--- a/Engine/Source/Programs/UnrealBuildTool/Platform/Windows/VCToolChain.cs
+++ b/Engine/Source/Programs/UnrealBuildTool/Platform/Windows/VCToolChain.cs
@@ -661,10 +661,10 @@ namespace UnrealBuildTool
 			//
 			//	PC
 			//
-			if (CompileEnvironment.bUseAVX)
+			if (CompileEnvironment.Architecture == UnrealArch.X64 && CompileEnvironment.MinCpuArchX64 != MinimumCpuArchitectureX64.None)
 			{
-				// Define /arch:AVX for the current compilation unit.  Machines without AVX support will crash on any SSE/AVX instructions if they run this compilation unit.
-				Arguments.Add("/arch:AVX");
+				// Define /arch:AVX[2,512] for the current compilation unit.  Machines without AVX support will crash on any SSE/AVX instructions if they run this compilation unit.
+				Arguments.Add($"/arch:{CompileEnvironment.MinCpuArchX64}");
 
 				// AVX available implies sse4 and sse2 available.
 				// Inform Unreal code that we have sse2, sse4, and AVX, both available to compile and available to run
diff --git a/Engine/Source/Programs/UnrealBuildTool/ProjectFiles/Rider/RiderProjectFile.cs b/Engine/Source/Programs/UnrealBuildTool/ProjectFiles/Rider/RiderProjectFile.cs
--- a/Engine/Source/Programs/UnrealBuildTool/ProjectFiles/Rider/RiderProjectFile.cs
+++ b/Engine/Source/Programs/UnrealBuildTool/ProjectFiles/Rider/RiderProjectFile.cs
@@ -754,7 +754,7 @@ namespace UnrealBuildTool
 				bUseUnity = CompileEnvironment.bUseUnity,
 				bCreateDebugInfo = CompileEnvironment.bCreateDebugInfo,
 				bIsBuildingLibrary = CompileEnvironment.bIsBuildingLibrary,
-				bUseAVX = CompileEnvironment.bUseAVX,
+				MinCpuArchX64 = CompileEnvironment.MinCpuArchX64,
 				bIsBuildingDLL = CompileEnvironment.bIsBuildingDLL,
 				bUseDebugCRT = CompileEnvironment.bUseDebugCRT,
 				bUseRTTI = CompileEnvironment.bUseRTTI,
diff --git a/Engine/Source/Programs/UnrealBuildTool/ProjectFiles/Rider/ToolchainInfo.cs b/Engine/Source/Programs/UnrealBuildTool/ProjectFiles/Rider/ToolchainInfo.cs
--- a/Engine/Source/Programs/UnrealBuildTool/ProjectFiles/Rider/ToolchainInfo.cs
+++ b/Engine/Source/Programs/UnrealBuildTool/ProjectFiles/Rider/ToolchainInfo.cs
@@ -20,7 +20,7 @@ namespace UnrealBuildTool
 		public bool bUseInlining;
 		public bool bUseUnity;
 		public bool bCreateDebugInfo;
-		public bool bUseAVX;
+		public MinimumCpuArchitectureX64 MinCpuArchX64;
 		public bool bUseDebugCRT;
 		public bool bUseStaticCRT;
 		public string? PrecompiledHeaderAction;
@@ -64,7 +64,7 @@ namespace UnrealBuildTool
 				bIsBuildingDLL == Other.bIsBuildingDLL && Architecture == Other.Architecture &&
 				Configuration == Other.Configuration && bOptimizeCode == Other.bOptimizeCode &&
 				bUseInlining == Other.bUseInlining && bUseUnity == Other.bUseUnity &&
-				bCreateDebugInfo == Other.bCreateDebugInfo && bUseAVX == Other.bUseAVX &&
+				bCreateDebugInfo == Other.bCreateDebugInfo && MinCpuArchX64 == Other.MinCpuArchX64 &&
 				bUseDebugCRT == Other.bUseDebugCRT && bUseStaticCRT == Other.bUseStaticCRT &&
 				PrecompiledHeaderAction == Other.PrecompiledHeaderAction && PrecompiledHeaderFile == Other.PrecompiledHeaderFile &&
 				Equals(ForceIncludeFiles, Other.ForceIncludeFiles) && Compiler == Other.Compiler &&
@@ -93,7 +93,7 @@ namespace UnrealBuildTool
 				HashCode = (HashCode * 397) ^ bUseInlining.GetHashCode();
 				HashCode = (HashCode * 397) ^ bUseUnity.GetHashCode();
 				HashCode = (HashCode * 397) ^ bCreateDebugInfo.GetHashCode();
-				HashCode = (HashCode * 397) ^ bUseAVX.GetHashCode();
+				HashCode = (HashCode * 397) ^ MinCpuArchX64.GetHashCode();
 				HashCode = (HashCode * 397) ^ bUseDebugCRT.GetHashCode();
 				HashCode = (HashCode * 397) ^ bUseStaticCRT.GetHashCode();
 				HashCode = (HashCode * 397) ^ (PrecompiledHeaderAction != null ? PrecompiledHeaderAction.GetHashCode() : 0);
diff --git a/Engine/Source/Programs/UnrealBuildTool/System/CppCompileEnvironment.cs b/Engine/Source/Programs/UnrealBuildTool/System/CppCompileEnvironment.cs
--- a/Engine/Source/Programs/UnrealBuildTool/System/CppCompileEnvironment.cs
+++ b/Engine/Source/Programs/UnrealBuildTool/System/CppCompileEnvironment.cs
@@ -96,6 +96,39 @@ namespace UnrealBuildTool
 		Default = None,
 	}
 
+	/// <summary>
+	/// Specifies the architecture for code generation on x64 for windows platforms.
+	/// Note that by enabling this you are changing the minspec for the PC platform, and the resultant executable will crash on machines without AVX support.
+	/// For more details please see https://learn.microsoft.com/en-us/cpp/build/reference/arch-x64
+	/// </summary>
+	public enum MinimumCpuArchitectureX64
+	{
+		/// <summary>
+		/// No minimum architecure
+		/// </summary>
+		None,
+
+		/// <summary>
+		/// Enables the use of Intel Advanced Vector Extensions instructions
+		/// </summary>
+		AVX,
+
+		/// <summary>
+		/// Enables the use of Intel Advanced Vector Extensions 2 instructions
+		/// </summary>
+		AVX2,
+
+		/// <summary>
+		/// Enables the use of Intel Advanced Vector Extensions 512 instructions
+		/// </summary>
+		AVX512,
+
+		/// <summary>
+		/// Use the default minimum architecure
+		/// </summary>
+		Default = None,
+	}
+
 	/// <summary>
 	/// The optimization level that may be compilation targets for C# files.
 	/// </summary>
@@ -227,12 +260,6 @@ namespace UnrealBuildTool
 		/// </summary>
 		public bool bCompileISPC = false;
 
-		/// <summary>
-		/// Direct the compiler to generate AVX instructions wherever SSE or AVX intrinsics are used.
-		/// Note that by enabling this you are changing the minspec for the PC platform, and the resultant executable will crash on machines without AVX support.
-		/// </summary>
-		public bool bUseAVX = false;
-
 		/// <summary>
 		/// Enable buffer security checks.   This should usually be enabled as it prevents severe security risks.
 		/// </summary>
@@ -595,6 +622,12 @@ namespace UnrealBuildTool
 		/// </summary>
 		public CStandardVersion CStandard = CStandardVersion.Default;
 
+		/// <summary>
+		/// Direct the compiler to generate AVX instructions wherever SSE or AVX intrinsics are used.
+		/// Note that by enabling this you are changing the minspec for the PC platform, and the resultant executable will crash on machines without AVX support.
+		/// </summary>
+		public MinimumCpuArchitectureX64 MinCpuArchX64 = MinimumCpuArchitectureX64.Default;
+
 		/// <summary>
 		/// The amount of the stack usage to report static analysis warnings.
 		/// </summary>
@@ -663,7 +696,6 @@ namespace UnrealBuildTool
 			bUseStackProtection = Other.bUseStackProtection;
 			bUseInlining = Other.bUseInlining;
 			bCompileISPC = Other.bCompileISPC;
-			bUseAVX = Other.bUseAVX;
 			bUseUnity = Other.bUseUnity;
 			MinSourceFilesForUnityBuildOverride = Other.MinSourceFilesForUnityBuildOverride;
 			MinFilesUsingPrecompiledHeaderOverride = Other.MinFilesUsingPrecompiledHeaderOverride;
@@ -724,6 +756,7 @@ namespace UnrealBuildTool
 			bHideSymbolsByDefault = Other.bHideSymbolsByDefault;
 			CppStandard = Other.CppStandard;
 			CStandard = Other.CStandard;
+			MinCpuArchX64 = Other.MinCpuArchX64;
 			bEnableCoroutines = Other.bEnableCoroutines;
 			IncludeOrderVersion = Other.IncludeOrderVersion;
 			bDeterministic = Other.bDeterministic;
