Moved some of the VS project properties to the common section.

diff --git a/Engine/Source/Programs/UnrealBuildTool/ProjectFiles/PlatformProjectGenerator.cs b/Engine/Source/Programs/UnrealBuildTool/ProjectFiles/PlatformProjectGenerator.cs
--- a/Engine/Source/Programs/UnrealBuildTool/ProjectFiles/PlatformProjectGenerator.cs
+++ b/Engine/Source/Programs/UnrealBuildTool/ProjectFiles/PlatformProjectGenerator.cs
@@ -171,17 +171,6 @@ namespace UnrealBuildTool
 		/// <returns>The custom path lines for the project file; Empty string if it doesn't require one</returns>
 		public virtual void GetVisualStudioPathsEntries(UnrealTargetPlatform InPlatform, UnrealTargetConfiguration InConfiguration, TargetType TargetType, FileReference TargetRulesPath, FileReference ProjectFilePath, FileReference NMakeOutputPath, VCProjectFileFormat InProjectFileFormat, StringBuilder ProjectFileBuilder)
 		{
-			// NOTE: We are intentionally overriding defaults for these paths with empty strings.  We never want Visual Studio's
-			//       defaults for these fields to be propagated, since they are version-sensitive paths that may not reflect
-			//       the environment that UBT is building in.  We'll set these environment variables ourselves!
-			// NOTE: We don't touch 'ExecutablePath' because that would result in Visual Studio clobbering the system "Path"
-			//       environment variable
-			ProjectFileBuilder.AppendLine("    <IncludePath />");
-			ProjectFileBuilder.AppendLine("    <ReferencePath />");
-			ProjectFileBuilder.AppendLine("    <LibraryPath />");
-			ProjectFileBuilder.AppendLine("    <LibraryWPath />");
-			ProjectFileBuilder.AppendLine("    <SourcePath />");
-			ProjectFileBuilder.AppendLine("    <ExcludePath />");
 		}
 
 		/// <summary>
diff --git a/Engine/Source/Programs/UnrealBuildTool/ProjectFiles/VisualStudio/VCProject.cs b/Engine/Source/Programs/UnrealBuildTool/ProjectFiles/VisualStudio/VCProject.cs
--- a/Engine/Source/Programs/UnrealBuildTool/ProjectFiles/VisualStudio/VCProject.cs
+++ b/Engine/Source/Programs/UnrealBuildTool/ProjectFiles/VisualStudio/VCProject.cs
@@ -1181,19 +1181,30 @@ namespace UnrealBuildTool
 			VCProjectFileContent.AppendLine("  <ImportGroup Label=\"ExtensionSettings\" />");
 			VCProjectFileContent.AppendLine("  <PropertyGroup Label=\"UserMacros\" />");
 
-			// Write the invalid configuration
+			// Write the common and invalid configuration values
 			{
 				const string InvalidMessage = "echo The selected platform/configuration is not valid for this target.";
 
 				string ProjectRelativeUnusedDirectory = NormalizeProjectPath(DirectoryReference.Combine(Unreal.EngineDirectory, "Intermediate", "Build", "Unused"));
 
-				VCProjectFileContent.AppendLine("  <PropertyGroup Condition=\"'$(Configuration)'=='Invalid'\">");
+				VCProjectFileContent.AppendLine("  <PropertyGroup>");
 				VCProjectFileContent.AppendLine("    <NMakeBuildCommandLine>{0}</NMakeBuildCommandLine>", InvalidMessage);
 				VCProjectFileContent.AppendLine("    <NMakeReBuildCommandLine>{0}</NMakeReBuildCommandLine>", InvalidMessage);
 				VCProjectFileContent.AppendLine("    <NMakeCleanCommandLine>{0}</NMakeCleanCommandLine>", InvalidMessage);
 				VCProjectFileContent.AppendLine("    <NMakeOutput>Invalid Output</NMakeOutput>", InvalidMessage);
 				VCProjectFileContent.AppendLine("    <OutDir>{0}{1}</OutDir>", ProjectRelativeUnusedDirectory, Path.DirectorySeparatorChar);
 				VCProjectFileContent.AppendLine("    <IntDir>{0}{1}</IntDir>", ProjectRelativeUnusedDirectory, Path.DirectorySeparatorChar);
+				// NOTE: We are intentionally overriding defaults for these paths with empty strings.  We never want Visual Studio's
+				//       defaults for these fields to be propagated, since they are version-sensitive paths that may not reflect
+				//       the environment that UBT is building in.  We'll set these environment variables ourselves!
+				// NOTE: We don't touch 'ExecutablePath' because that would result in Visual Studio clobbering the system "Path"
+				//       environment variable
+				VCProjectFileContent.AppendLine("    <IncludePath />");
+				VCProjectFileContent.AppendLine("    <ReferencePath />");
+				VCProjectFileContent.AppendLine("    <LibraryPath />");
+				VCProjectFileContent.AppendLine("    <LibraryWPath />");
+				VCProjectFileContent.AppendLine("    <SourcePath />");
+				VCProjectFileContent.AppendLine("    <ExcludePath />");
 				VCProjectFileContent.AppendLine("  </PropertyGroup>");
 			}
 
@@ -1845,11 +1856,7 @@ namespace UnrealBuildTool
 
 				if (IsStubProject)
 				{
-					string ProjectRelativeUnusedDirectory = NormalizeProjectPath(DirectoryReference.Combine(Unreal.EngineDirectory, "Intermediate", "Build", "Unused"));
-
 					VCProjectFileContent.AppendLine("  <PropertyGroup {0}>", ConditionString);
-					VCProjectFileContent.AppendLine("    <OutDir>{0}{1}</OutDir>", ProjectRelativeUnusedDirectory, Path.DirectorySeparatorChar);
-					VCProjectFileContent.AppendLine("    <IntDir>{0}{1}</IntDir>", ProjectRelativeUnusedDirectory, Path.DirectorySeparatorChar);
 					VCProjectFileContent.AppendLine("    <NMakeBuildCommandLine>@rem Nothing to do.</NMakeBuildCommandLine>");
 					VCProjectFileContent.AppendLine("    <NMakeReBuildCommandLine>@rem Nothing to do.</NMakeReBuildCommandLine>");
 					VCProjectFileContent.AppendLine("    <NMakeCleanCommandLine>@rem Nothing to do.</NMakeCleanCommandLine>");
@@ -1859,14 +1866,10 @@ namespace UnrealBuildTool
 				else if (Unreal.IsEngineInstalled() && Combination.ProjectTarget != null && Combination.ProjectTarget.TargetRules != null &&
 					(Combination.Platform == null || !Combination.ProjectTarget.SupportedPlatforms.Contains(Combination.Platform.Value)))
 				{
-					string ProjectRelativeUnusedDirectory = NormalizeProjectPath(DirectoryReference.Combine(Unreal.EngineDirectory, "Intermediate", "Build", "Unused"));
-
 					string TargetName = Combination.ProjectTarget.TargetFilePath.GetFileNameWithoutAnyExtensions();
 					string ValidPlatforms = String.Join(", ", Combination.ProjectTarget.SupportedPlatforms.Select(x => x.ToString()));
 
 					VCProjectFileContent.AppendLine("  <PropertyGroup {0}>", ConditionString);
-					VCProjectFileContent.AppendLine("    <OutDir>{0}{1}</OutDir>", ProjectRelativeUnusedDirectory, Path.DirectorySeparatorChar);
-					VCProjectFileContent.AppendLine("    <IntDir>{0}{1}</IntDir>", ProjectRelativeUnusedDirectory, Path.DirectorySeparatorChar);
 					VCProjectFileContent.AppendLine("    <NMakeBuildCommandLine>@echo {0} is not a supported platform for {1}. Valid platforms are {2}.</NMakeBuildCommandLine>", Combination.Platform!, TargetName, ValidPlatforms);
 					VCProjectFileContent.AppendLine("    <NMakeReBuildCommandLine>@echo {0} is not a supported platform for {1}. Valid platforms are {2}.</NMakeReBuildCommandLine>", Combination.Platform!, TargetName, ValidPlatforms);
 					VCProjectFileContent.AppendLine("    <NMakeCleanCommandLine>@echo {0} is not a supported platform for {1}. Valid platforms are {2}.</NMakeCleanCommandLine>", Combination.Platform!, TargetName, ValidPlatforms);
@@ -1935,25 +1938,13 @@ namespace UnrealBuildTool
 
 					VCProjectFileContent.AppendLine("  <PropertyGroup {0}>", ConditionString);
 
-					StringBuilder PathsStringBuilder = new StringBuilder();
 					if (ProjGenerator != null)
 					{
+						StringBuilder PathsStringBuilder = new StringBuilder();
 						ProjGenerator.GetVisualStudioPathsEntries(Platform, Configuration, TargetRulesObject.Type, TargetFilePath, ProjectFilePath, NMakePath, ProjectFileFormat, PathsStringBuilder);
+						VCProjectFileContent.Append(PathsStringBuilder.ToString());
 					}
 
-					string PathStrings = PathsStringBuilder.ToString();
-					if (string.IsNullOrEmpty(PathStrings) || (PathStrings.Contains("<IntDir>") == false))
-					{
-						string ProjectRelativeUnusedDirectory = "$(ProjectDir)..\\Build\\Unused";
-						VCProjectFileContent.Append(PathStrings);
-
-						VCProjectFileContent.AppendLine("    <OutDir>{0}{1}</OutDir>", ProjectRelativeUnusedDirectory, Path.DirectorySeparatorChar);
-						VCProjectFileContent.AppendLine("    <IntDir>{0}{1}</IntDir>", ProjectRelativeUnusedDirectory, Path.DirectorySeparatorChar);
-					}
-					else
-					{
-						VCProjectFileContent.Append(PathStrings);
-					}
 
 					// This is the standard UE based project NMake build line:
 					//	..\..\Build\BatchFiles\Build.bat <TARGETNAME> <PLATFORM> <CONFIGURATION>
