UnrealBuildTool: Actually add the modules listed in the .uproject to the build list

diff --git a/Engine/Source/Programs/UnrealBuildTool/Configuration/UEBuildTarget.cs b/Engine/Source/Programs/UnrealBuildTool/Configuration/UEBuildTarget.cs
--- a/Engine/Source/Programs/UnrealBuildTool/Configuration/UEBuildTarget.cs
+++ b/Engine/Source/Programs/UnrealBuildTool/Configuration/UEBuildTarget.cs
@@ -3107,6 +3107,9 @@ namespace UnrealBuildTool
 			// Setup the target's binaries.
 			SetupBinaries(Logger);
 
+			// Setup the target's modules referenced in the .uproject
+			SetupProjectModules(Logger);
+
 			// Setup the target's plugins
 			SetupPlugins(Logger);
 
@@ -3925,6 +3928,26 @@ namespace UnrealBuildTool
 			return UEBuildPlatform.GetBuildPlatform(Platform).FinalizeBinaryPaths(BinaryFile, ProjectFile, Rules);
 		}
 
+		/// <summary>
+		/// Sets up the uproject modules for this target
+		/// </summary>
+		public void SetupProjectModules(ILogger Logger)
+		{
+			if (ProjectDescriptor?.Modules == null)
+			{
+				return;
+			}
+
+			foreach (ModuleDescriptor Descriptor in ProjectDescriptor.Modules.Where(x => x.IsCompiledInConfiguration(Platform, Configuration, TargetName, TargetType, Rules.bBuildDeveloperTools, Rules.bBuildRequiresCookedData)))
+			{
+				UEBuildModuleCPP Module = FindOrCreateCppModuleByName(Descriptor.Name, TargetRulesFile.GetFileName(), Logger);
+				if (Module.Binary == null)
+				{
+					AddModuleToBinary(Module);
+				}
+			}
+		}
+
 		/// <summary>
 		/// Sets up the plugins for this target
 		/// </summary>
