UnrealBuildTool: Fix Action properties that weren't being serialized

#rnx

diff --git a/Engine/Source/Programs/UnrealBuildTool/Executors/ParallelExecutor.cs b/Engine/Source/Programs/UnrealBuildTool/Executors/ParallelExecutor.cs
--- a/Engine/Source/Programs/UnrealBuildTool/Executors/ParallelExecutor.cs
+++ b/Engine/Source/Programs/UnrealBuildTool/Executors/ParallelExecutor.cs
@@ -218,8 +218,8 @@ namespace UnrealBuildTool
 
 			var ActiveTasks = new Task<ExecuteResults>[ActualNumParallelProcesses];
 			var ActiveActionIndices = new int[ActualNumParallelProcesses];
-			float MaxParallelProcessWeight = (float)ActualNumParallelProcesses;
-			float ActiveTaskWeight = 0.0f;
+			double MaxParallelProcessWeight = ActualNumParallelProcesses;
+			double ActiveTaskWeight = 0.0;
 			int ActiveTaskCount = 0;
 			bool Success = true;
 			List<float> LoggedCpuUtilization = new List<float>();
diff --git a/Engine/Source/Programs/UnrealBuildTool/Platform/Windows/VCCompileAction.cs b/Engine/Source/Programs/UnrealBuildTool/Platform/Windows/VCCompileAction.cs
--- a/Engine/Source/Programs/UnrealBuildTool/Platform/Windows/VCCompileAction.cs
+++ b/Engine/Source/Programs/UnrealBuildTool/Platform/Windows/VCCompileAction.cs
@@ -182,7 +182,7 @@ namespace UnrealBuildTool
 		public bool bIsHighPriority { get => CreatePchFile != null; }
 
 		/// <inheritdoc/>
-		public float Weight { get; set; } = 1.0f;
+		public double Weight { get; set; } = 1.0;
 		#endregion
 
 		#region Implementation of IAction
diff --git a/Engine/Source/Programs/UnrealBuildTool/System/Action.cs b/Engine/Source/Programs/UnrealBuildTool/System/Action.cs
--- a/Engine/Source/Programs/UnrealBuildTool/System/Action.cs
+++ b/Engine/Source/Programs/UnrealBuildTool/System/Action.cs
@@ -143,7 +143,7 @@ namespace UnrealBuildTool
 		/// <summary>
 		/// Used to determine how much weight(CPU and Memory work) this action is.
 		/// </summary>
-		float Weight { get; }
+		double Weight { get; }
 	}
 
 	/// <summary>
@@ -248,7 +248,7 @@ namespace UnrealBuildTool
 		public bool bIsHighPriority { get; set; } = false;
 
 		/// <inheritdoc/>
-		public float Weight { get; set; } = 1.0f;
+		public double Weight { get; set; } = 1.0;
 
 		IEnumerable<FileItem> IExternalAction.PrerequisiteItems => PrerequisiteItems;
 		IEnumerable<FileItem> IExternalAction.ProducedItems => ProducedItems;
@@ -280,11 +280,13 @@ namespace UnrealBuildTool
 			StatusDescription = InOther.StatusDescription;
 			bCanExecuteRemotely = InOther.bCanExecuteRemotely;
 			bCanExecuteRemotelyWithSNDBS = InOther.bCanExecuteRemotelyWithSNDBS;
+			bCanCache = InOther.bCanCache;
 			bIsGCCCompiler = InOther.bIsGCCCompiler;
 			bShouldOutputStatusDescription = InOther.bShouldOutputStatusDescription;
 			bProducesImportLibrary = InOther.bProducesImportLibrary;
 			bUseActionHistory = InOther.bUseActionHistory;
 			bIsHighPriority = InOther.bIsHighPriority;
+			Weight = InOther.Weight;
 		}
 
 		public Action(BinaryArchiveReader Reader)
@@ -298,6 +300,7 @@ namespace UnrealBuildTool
 			StatusDescription = Reader.ReadString()!;
 			bCanExecuteRemotely = Reader.ReadBool();
 			bCanExecuteRemotelyWithSNDBS = Reader.ReadBool();
+			bCanCache = Reader.ReadBool();
 			bIsGCCCompiler = Reader.ReadBool();
 			bShouldOutputStatusDescription = Reader.ReadBool();
 			bProducesImportLibrary = Reader.ReadBool();
@@ -307,6 +310,7 @@ namespace UnrealBuildTool
 			DependencyListFile = Reader.ReadFileItem();
 			bUseActionHistory = Reader.ReadBool();
 			bIsHighPriority = Reader.ReadBool();
+			Weight = Reader.ReadDouble();
 
 		}
 
@@ -324,6 +328,7 @@ namespace UnrealBuildTool
 			Writer.WriteString(StatusDescription);
 			Writer.WriteBool(bCanExecuteRemotely);
 			Writer.WriteBool(bCanExecuteRemotelyWithSNDBS);
+			Writer.WriteBool(bCanCache);
 			Writer.WriteBool(bIsGCCCompiler);
 			Writer.WriteBool(bShouldOutputStatusDescription);
 			Writer.WriteBool(bProducesImportLibrary);
@@ -333,6 +338,7 @@ namespace UnrealBuildTool
 			Writer.WriteFileItem(DependencyListFile);
 			Writer.WriteBool(bUseActionHistory);
 			Writer.WriteBool(bIsHighPriority);
+			Writer.WriteDouble(Weight);
 		}
 
 		/// <summary>
@@ -391,6 +397,12 @@ namespace UnrealBuildTool
 				Action.bCanExecuteRemotelyWithSNDBS = bCanExecuteRemotelyWithSNDBS;
 			}
 
+			bool bCanCache;
+			if (Object.TryGetBoolField("bCanCache", out bCanCache))
+			{
+				Action.bCanCache = bCanCache;
+			}
+
 			bool bIsGCCCompiler;
 			if(Object.TryGetBoolField("bIsGCCCompiler", out bIsGCCCompiler))
 			{
@@ -433,6 +445,12 @@ namespace UnrealBuildTool
 				Action.DependencyListFile = FileItem.GetItemByPath(DependencyListFile);
 			}
 
+			double Weight;
+			if (Object.TryGetDoubleField("Weight", out Weight))
+			{
+				Action.Weight = Weight;
+			}
+
 			return Action;
 		}
 
@@ -491,9 +509,11 @@ namespace UnrealBuildTool
 			Writer.WriteValue("StatusDescription", Action.StatusDescription);
 			Writer.WriteValue("bCanExecuteRemotely", Action.bCanExecuteRemotely);
 			Writer.WriteValue("bCanExecuteRemotelyWithSNDBS", Action.bCanExecuteRemotelyWithSNDBS);
+			Writer.WriteValue("bCanCache", Action.bCanCache);
 			Writer.WriteValue("bIsGCCCompiler", Action.bIsGCCCompiler);
 			Writer.WriteValue("bShouldOutputStatusDescription", Action.bShouldOutputStatusDescription);
 			Writer.WriteValue("bProducesImportLibrary", Action.bProducesImportLibrary);
+			Writer.WriteValue("Weight", Action.Weight);
 
 			Writer.WriteArrayStart("PrerequisiteActions");
 			foreach (LinkedAction PrerequisiteAction in Action.PrerequisiteActions)
@@ -615,7 +635,7 @@ namespace UnrealBuildTool
 		public bool bProducesImportLibrary => Inner.bProducesImportLibrary;
 		public bool bUseActionHistory => Inner.bUseActionHistory;
 		public bool bIsHighPriority => IsHighPriority != 0;
-		public float Weight => Inner.Weight;
+		public double Weight => Inner.Weight;
 
 		#endregion
 
