Fix monolithic editor linux builds that are extremely large because debug symbols are not split properly

#rnx

diff --git a/Engine/Source/Programs/UnrealBuildTool/Platform/Linux/LinuxToolChain.cs b/Engine/Source/Programs/UnrealBuildTool/Platform/Linux/LinuxToolChain.cs
--- a/Engine/Source/Programs/UnrealBuildTool/Platform/Linux/LinuxToolChain.cs
+++ b/Engine/Source/Programs/UnrealBuildTool/Platform/Linux/LinuxToolChain.cs
@@ -195,7 +195,7 @@ namespace UnrealBuildTool
 				Out.WriteLine("echo DummySyms>> \"{0}\"", EncodedBinarySymbolsFile.AbsolutePath);
 			}
 
-			if (!Options.HasFlag(ClangToolChainOptions.DisableSplitDebugInfoWithObjCopy) && LinkEnvironment.bCreateDebugInfo)
+			if (LinkEnvironment.bCreateDebugInfo)
 			{
 				if (MaxBinarySizeOverrideForObjcopy > 0 && bUseCmdExe)
 				{
diff --git a/Engine/Source/Programs/UnrealBuildTool/Platform/Linux/UEBuildLinux.cs b/Engine/Source/Programs/UnrealBuildTool/Platform/Linux/UEBuildLinux.cs
--- a/Engine/Source/Programs/UnrealBuildTool/Platform/Linux/UEBuildLinux.cs
+++ b/Engine/Source/Programs/UnrealBuildTool/Platform/Linux/UEBuildLinux.cs
@@ -610,16 +610,6 @@ namespace UnrealBuildTool
 				Options |= ClangToolChainOptions.UseAutoRTFMCompiler;
 			}
 
-			// When building a monolithic editor we have to avoid using objcopy.exe as it cannot handle files
-			// larger then 4GB. This is only an issue with our binutils objcopy.exe.
-			// llvm-objcopy.exe does not have this issue and once we switch over to using that in clang 10.0.1 we can remove this!
-			if ((BuildHostPlatform.Current.Platform == UnrealTargetPlatform.Win64) &&
-				(Target.LinkType == TargetLinkType.Monolithic) &&
-				(Target.Type == TargetType.Editor))
-			{
-				Options |= ClangToolChainOptions.DisableSplitDebugInfoWithObjCopy;
-			}
-
 			if (Target.LinuxPlatform.bTuneDebugInfoForLLDB)
 			{
 				Options |= ClangToolChainOptions.TuneDebugInfoForLLDB;
diff --git a/Engine/Source/Programs/UnrealBuildTool/ToolChain/ClangToolChain.cs b/Engine/Source/Programs/UnrealBuildTool/ToolChain/ClangToolChain.cs
--- a/Engine/Source/Programs/UnrealBuildTool/ToolChain/ClangToolChain.cs
+++ b/Engine/Source/Programs/UnrealBuildTool/ToolChain/ClangToolChain.cs
@@ -68,12 +68,6 @@ namespace UnrealBuildTool
 		/// </summary>
 		EnableThinLTO = 1 << 8,
 
-		/// <summary>
-		/// If should disable using objcopy to split the debug info into its own file or now
-		/// When we support larger the 4GB files with objcopy.exe this can be removed!
-		/// </summary>
-		DisableSplitDebugInfoWithObjCopy = 1 << 9,
-
 		/// <summary>
 		/// Enable tuning of debug info for LLDB
 		/// </summary>
