Moved UserDir logic to Unreal.cs and removed duplicate code.

diff --git a/Engine/Source/Programs/Shared/EpicGames.Build/Unreal.cs b/Engine/Source/Programs/Shared/EpicGames.Build/Unreal.cs
--- a/Engine/Source/Programs/Shared/EpicGames.Build/Unreal.cs
+++ b/Engine/Source/Programs/Shared/EpicGames.Build/Unreal.cs
@@ -149,6 +149,11 @@ namespace UnrealBuildBase
 		/// </summary>
 		static public readonly DirectoryReference UserSettingDirectory = GetUserSettingDirectory();
 
+		/// <summary>
+		/// Returns the User Directory path. This matches FPlatformProcess::UserDir().
+		/// </summary>
+		static public readonly DirectoryReference? UserDirectory = GetUserDirectory();
+
 		/// <summary>
 		/// Writable engine directory. Uses the user's settings folder for installed builds.
 		/// </summary>
@@ -369,5 +374,26 @@ namespace UnrealBuildBase
 
 			return DirectoryReference.Combine(EngineDirectory, "Saved");
 		}
+
+		/// <summary>
+		/// Returns the User Directory path. This matches FPlatformProcess::UserDir().
+		/// </summary>
+		private static DirectoryReference? GetUserDirectory()
+		{
+			// Some user accounts (eg. SYSTEM on Windows) don't have a home directory. Ignore them if Environment.GetFolderPath() returns an empty string.
+			string PersonalFolder = Environment.GetFolderPath(Environment.SpecialFolder.Personal);
+			if (!String.IsNullOrEmpty(PersonalFolder))
+			{
+				if (RuntimePlatform.IsMac || RuntimePlatform.IsLinux)
+				{
+					return new DirectoryReference(System.IO.Path.Combine(PersonalFolder, "Documents"));
+				}
+				else
+				{
+					return new DirectoryReference(PersonalFolder);
+				}
+			}
+			return null;
+		}
 	}
 }
diff --git a/Engine/Source/Programs/UnrealBuildTool/System/ConfigHierarchy.cs b/Engine/Source/Programs/UnrealBuildTool/System/ConfigHierarchy.cs
--- a/Engine/Source/Programs/UnrealBuildTool/System/ConfigHierarchy.cs
+++ b/Engine/Source/Programs/UnrealBuildTool/System/ConfigHierarchy.cs
@@ -195,22 +195,6 @@ namespace UnrealBuildTool
         /// </summary>
         System.Threading.ReaderWriterLockSlim NameToSectionLock = new System.Threading.ReaderWriterLockSlim();
 
-		static readonly string? PersonalConfigFolder = null;
-
-		static ConfigHierarchy()
-		{
-			// Some user accounts (eg. SYSTEM on Windows) don't have a home directory. Ignore them if Environment.GetFolderPath() returns an empty string.
-			string PersonalFolder = Environment.GetFolderPath(Environment.SpecialFolder.Personal);
-			if (!String.IsNullOrEmpty(PersonalFolder))
-			{
-				PersonalConfigFolder = PersonalFolder;
-				if (RuntimePlatform.IsMac || RuntimePlatform.IsLinux)
-				{
-					PersonalConfigFolder = System.IO.Path.Combine(PersonalConfigFolder, "Documents");
-				}
-			}
-		}
-
 		/// <summary>
 		/// Construct a config hierarchy from the given files
 		/// </summary>
@@ -981,12 +965,6 @@ namespace UnrealBuildTool
 			new ConfigLayerExpansion { Before1 = "{ENGINE}/Config/{PLATFORM}/", After1 = "{ENGINE}/Restricted/NoRedist/Platforms/{PLATFORM}/Config/",          Before2 = "{PROJECT}/Config/{PLATFORM}/", After2 = "{RESTRICTEDPROJECT_NR}/Platforms/{PLATFORM}/{OPT_SUBDIR}Config/" },
 		};
 
-		// Match FPlatformProcess::UserDir()
-		private static string? GetUserDir()
-		{
-			return PersonalConfigFolder;
-		}
-
 		private static string PerformBasicReplacements(string InString, string BaseIniName, string CustomConfig)
 		{
 			string OutString = InString.Replace("{TYPE}", BaseIniName);
@@ -997,7 +975,12 @@ namespace UnrealBuildTool
 				OutString = OutString.Replace("{USERSETTINGS}", UserSettingsDir.FullName);
 			}
 
-			OutString = OutString.Replace("{USER}", GetUserDir());
+			DirectoryReference? UserDir = Unreal.UserDirectory;
+			if (UserDir != null)
+			{
+				OutString = OutString.Replace("{USER}", UserDir.FullName);
+			}
+			
 			OutString = OutString.Replace("{CUSTOMCONFIG}", CustomConfig);
 
 			return OutString;
@@ -1082,7 +1065,7 @@ namespace UnrealBuildTool
 				// skip certain layers if we are platform-less, project-less, or userdir-less
 				if ((bHasPlatformTag && PlatformName == "None") ||
 					(bHasProjectTag && ProjectDir == null) ||
-					(bHasUserTag && GetUserDir() == null) ||
+					(bHasUserTag && Unreal.UserSettingDirectory == null) ||
 					(bHasCustomConfigTag && String.IsNullOrEmpty(CustomConfig)))
 				{
 					continue;
diff --git a/Engine/Source/Programs/UnrealBuildTool/System/UnrealPluginLanguage.cs b/Engine/Source/Programs/UnrealBuildTool/System/UnrealPluginLanguage.cs
--- a/Engine/Source/Programs/UnrealBuildTool/System/UnrealPluginLanguage.cs
+++ b/Engine/Source/Programs/UnrealBuildTool/System/UnrealPluginLanguage.cs
@@ -3499,20 +3499,7 @@ namespace UnrealBuildTool
 			}
 
 			DirectoryReference? UserSettingsFolder = Unreal.UserSettingDirectory; // Match FPlatformProcess::UserSettingsDir()
-			DirectoryReference? PersonalFolder = null; // Match FPlatformProcess::UserDir()
-			if (RuntimePlatform.IsMac || RuntimePlatform.IsLinux)
-			{
-				PersonalFolder = new DirectoryReference(Path.Combine(Environment.GetFolderPath(Environment.SpecialFolder.Personal), "Documents"));
-			}
-			else
-			{
-				// Not all user accounts have a local application data directory (eg. SYSTEM, used by Jenkins for builds).
-				string PersonalFolderSetting = Environment.GetFolderPath(Environment.SpecialFolder.Personal);
-				if (!String.IsNullOrEmpty(PersonalFolderSetting))
-				{
-					PersonalFolder = new DirectoryReference(PersonalFolderSetting);
-				}
-			}
+			DirectoryReference? PersonalFolder = Unreal.UserDirectory; // Match FPlatformProcess::UserDir()
 
 			if(UserSettingsFolder != null)
 			{
