UnrealBuildTool: Add notice for VS2019 end of support version

#rnx

diff --git a/Engine/Source/Programs/UnrealBuildTool/Platform/Windows/VCToolChain.cs b/Engine/Source/Programs/UnrealBuildTool/Platform/Windows/VCToolChain.cs
--- a/Engine/Source/Programs/UnrealBuildTool/Platform/Windows/VCToolChain.cs
+++ b/Engine/Source/Programs/UnrealBuildTool/Platform/Windows/VCToolChain.cs
@@ -92,6 +92,11 @@ namespace UnrealBuildTool
 			{
 				Lines.Add($"Using {WindowsPlatform.GetCompilerName(EnvVars.Compiler)} {EnvVars.CompilerVersion} compiler ({EnvVars.CompilerDir}) with {WindowsPlatform.GetCompilerName(EnvVars.ToolChain)} {EnvVars.ToolChainVersion} runtime ({EnvVars.ToolChainDir}) and Windows {EnvVars.WindowsSdkVersion} SDK ({EnvVars.WindowsSdkDir}).");
 			}
+
+			if (EnvVars.ToolChain == WindowsCompiler.VisualStudio2019)
+			{
+				Lines.Add($"Notice: {WindowsPlatform.GetCompilerName(WindowsCompiler.VisualStudio2019)} will no longer be supported for the installed engine in UE 5.3. Compiling code for Microsoft platforms will require {WindowsPlatform.GetCompilerName(WindowsCompiler.VisualStudio2022)} in UE 5.4 and beyond.");
+			}
 		}
 
 		public override void GetExternalDependencies(HashSet<FileItem> ExternalDependencies)
