- Fixed an issue with Xcode having stale bundle identifiers when codesigning BP projects (something old was being cached in Xcode's build cache, this change will force the bundle ID to be there so the cache won't cause problems)

diff --git a/Engine/Source/Programs/UnrealBuildTool/ProjectFiles/Xcode/XcodeProjectFileGenerator.cs b/Engine/Source/Programs/UnrealBuildTool/ProjectFiles/Xcode/XcodeProjectFileGenerator.cs
--- a/Engine/Source/Programs/UnrealBuildTool/ProjectFiles/Xcode/XcodeProjectFileGenerator.cs
+++ b/Engine/Source/Programs/UnrealBuildTool/ProjectFiles/Xcode/XcodeProjectFileGenerator.cs
@@ -49,6 +49,11 @@ namespace UnrealBuildTool
 		/// </summary>
 		string AppName = "";
 
+		/// <summary>
+		/// Store the single game project (when using -game -project=...) to a place that XcodeProjectLegacy can easily retrieve it
+		/// </summary>
+		public static FileReference? SingleGameProject = null;
+		
 		/// <summary>
 		/// If set, only write out one target/project 
 		/// </summary>
@@ -58,6 +63,8 @@ namespace UnrealBuildTool
 		public XcodeProjectFileGenerator(FileReference? InOnlyGameProject, CommandLineArguments CommandLine)
 			: base(InOnlyGameProject)
 		{
+			SingleGameProject = InOnlyGameProject;
+			
 			if (CommandLine.HasOption("-distribution"))
 			{
 				bForDistribution = true;
diff --git a/Engine/Source/Programs/UnrealBuildTool/ProjectFiles/Xcode/XcodeProjectLegacy.cs b/Engine/Source/Programs/UnrealBuildTool/ProjectFiles/Xcode/XcodeProjectLegacy.cs
--- a/Engine/Source/Programs/UnrealBuildTool/ProjectFiles/Xcode/XcodeProjectLegacy.cs
+++ b/Engine/Source/Programs/UnrealBuildTool/ProjectFiles/Xcode/XcodeProjectLegacy.cs
@@ -1831,7 +1831,10 @@ namespace UnrealBuildTool.XcodeProjectLegacy
 				return true;
 			}
 
-			FileReference? GameProjectPath = null;
+			// if a single project was specified then start with that, this will allow content-only projects to have
+			// their BundleID set in the stub xcode project when codesigning, which Xcode needs to have now)
+			// note that usually SingleGameProject is null
+			FileReference? GameProjectPath = XcodeProjectFileGenerator.SingleGameProject;
 			foreach (Project Target in ProjectTargets)
 			{
 				if (Target.UnrealProjectFilePath != null)
