- After staging, Modern xcode projects will build with stub xcode project to generate a complete .app in the root of the staging directory (tested Mac and IOS)
- Fixed up scripts for copying files into mac .app with modern xcode
- Changed bIncludeEngineSource being false to not remove game/samples/mods/etc from projects (only stub xcode projects set it to false, afaict)

diff --git a/Engine/Source/Programs/AutomationTool/Apple/ApplePlatform.Automation.cs b/Engine/Source/Programs/AutomationTool/Apple/ApplePlatform.Automation.cs
--- a/Engine/Source/Programs/AutomationTool/Apple/ApplePlatform.Automation.cs
+++ b/Engine/Source/Programs/AutomationTool/Apple/ApplePlatform.Automation.cs
@@ -17,6 +17,8 @@ public abstract class ApplePlatform : Platform
 	{
 	}
 
+	#region SDK
+
 	public override bool GetSDKInstallCommand(out string Command, out string Params, ref bool bRequiresPrivilegeElevation, ref bool bCreateWindow, ITurnkeyContext TurnkeyContext)
 	{
 		Command = "";
@@ -63,4 +65,34 @@ public abstract class ApplePlatform : Platform
 
 		return ExitCode == 0;
 	}
+
+	#endregion
+
+
+	public override void PostStagingFileCopy(ProjectParams Params, DeploymentContext SC)
+	{
+		ConfigHierarchy Ini = ConfigCache.ReadHierarchy(ConfigHierarchyType.Engine, Params.RawProjectPath.Directory!, SC.StageTargetPlatform.PlatformType);
+		bool bUseModernXcode;
+		Ini.TryGetValue("XcodeConfiguration", "bUseModernXcode", out bUseModernXcode);
+
+		if (bUseModernXcode)
+		{
+			// modern does all the distribution via .xcarchive
+			bool bForDistribution = false;
+
+			foreach (TargetReceipt Target in SC.StageTargets.Select(x => x.Receipt))
+			{
+				Console.WriteLine($"GENERATING STUB PROJECT");
+				DirectoryReference GeneratedProjectFile;
+				IOSExports.GenerateRunOnlyXcodeProject(SC.RawProjectPath, Target.Platform, Target.TargetName, bForDistribution, Logger, out GeneratedProjectFile);
+
+
+				Console.WriteLine("FINALIZEING APP with {0}, target {1}", GeneratedProjectFile, Target.TargetName);
+				string ExtraOptions = $"SYMROOT=\"{SC.StageDirectory.ParentDirectory}\" EFFECTIVE_PLATFORM_NAME={SC.StageDirectory.GetDirectoryName()}";
+				// run xcodebuild on the generated project to make the .app
+				IOSExports.FinalizeAppWithModernXcode(GeneratedProjectFile!, Target.Platform, Target.TargetName, Target.Configuration.ToString(), ExtraOptions, bForDistribution, Logger);
+				Console.WriteLine("DOME FINALIZEING APP with {0}, target {1}", GeneratedProjectFile, Target.TargetName);
+			}
+		}
+	}
 }
diff --git a/Engine/Source/Programs/AutomationTool/IOS/IOSPlatform.Automation.cs b/Engine/Source/Programs/AutomationTool/IOS/IOSPlatform.Automation.cs
--- a/Engine/Source/Programs/AutomationTool/IOS/IOSPlatform.Automation.cs
+++ b/Engine/Source/Programs/AutomationTool/IOS/IOSPlatform.Automation.cs
@@ -1132,7 +1132,7 @@ public class IOSPlatform : ApplePlatform
 				throw new AutomationException("iOS couldn't find the appropriate Xcode Project for " + RawProjectPath);
 			}
 
-			int ReturnCode = IOSExports.FinalizeAppWithXcode(GeneratedProject, TargetPlatformType, SchemeName ?? GameName, SchemeConfiguration ?? TargetConfig.ToString(), Provision, Certificate, Team, bAutomaticSigning, Distribution, bUseModernXcode, Logger);
+			int ReturnCode = IOSExports.FinalizeAppWithXcode(GeneratedProject, TargetPlatformType, SchemeName ?? GameName, SchemeConfiguration ?? TargetConfig.ToString(), Provision, Certificate, Team, "", bAutomaticSigning, Distribution, bUseModernXcode, Logger);
 			if (ReturnCode != 0)
 			{
 				throw new AutomationException(ExitCode.Error_FailedToCodeSign, "CodeSign Failed");
@@ -2273,6 +2273,8 @@ public class IOSPlatform : ApplePlatform
 					// generate the OnDemandResources.plist
 					GenerateOnDemandResourcesPlist (ChunkData, SC.StageDirectory.FullName);
 				}*/
+
+		base.PostStagingFileCopy(Params, SC);
 	}
 
 	public override bool RequiresPackageToDeploy
diff --git a/Engine/Source/Programs/UnrealBuildTool/Platform/IOS/IOSExports.cs b/Engine/Source/Programs/UnrealBuildTool/Platform/IOS/IOSExports.cs
--- a/Engine/Source/Programs/UnrealBuildTool/Platform/IOS/IOSExports.cs
+++ b/Engine/Source/Programs/UnrealBuildTool/Platform/IOS/IOSExports.cs
@@ -254,16 +254,19 @@ namespace UnrealBuildTool
 				$"-platforms={Platform}",
 				"-DeployOnly",
 				"-NoIntellisense",
-				"-NoCPP",
 				"-NoDotNet",
 				"-IgnoreJunk",
 				bForDistribution ? "-distribution" : "-development",
 				"-IncludeTempTargets",
 				"-projectfileformat=XCode",
 				"-automated",
-				$"-singletarget={TargetName}"
 			};
 
+			if (!string.IsNullOrEmpty(TargetName))
+			{
+				Options.Add($"-singletarget={TargetName}");
+			}
+
 			if (UProjectFile == null || UProjectFile.IsUnderDirectory(Unreal.EngineDirectory))
 			{
 				// @todo do we need these? where would the bundleid come from if there's no project?
@@ -287,11 +290,12 @@ namespace UnrealBuildTool
 		/// <param name="Platform">THe platform to make the .app for</param>
 		/// <param name="SchemeName">The name of the scheme (basically the target on the .xcworkspace)</param>
 		/// <param name="Configuration">Which configuration to make (Debug, etc)</param>
+		/// <param name="ExtraOptions">Extra options to pass to xcodebuild</param>
 		/// <param name="bForDistribution">True if this is making a bild for uploading to app store</param>
 		/// <param name="Logger">Logging object</param>
-		public static void FinalizeAppWithModernXcode(DirectoryReference XcodeProject, UnrealTargetPlatform Platform, string SchemeName, string Configuration, bool bForDistribution, ILogger Logger)
+		public static void FinalizeAppWithModernXcode(DirectoryReference XcodeProject, UnrealTargetPlatform Platform, string SchemeName, string Configuration, string ExtraOptions, bool bForDistribution, ILogger Logger)
 		{
-			FinalizeAppWithXcode(XcodeProject, Platform, SchemeName, Configuration, null, null, null, false, bForDistribution, bUseModernXcode: true, Logger);
+			FinalizeAppWithXcode(XcodeProject, Platform, SchemeName, Configuration, null, null, null, ExtraOptions, false, bForDistribution, bUseModernXcode: true, Logger);
 		}
 
 		/// <summary>
@@ -304,13 +308,14 @@ namespace UnrealBuildTool
 		/// <param name="Provision">An optional provision to codesign with (ignored in modern mod)</param>
 		/// <param name="Certificate">An optional certificate to codesign with (ignored in modern mode)</param>
 		/// <param name="Team">Optional Team to use when codesigning (ignored in modern mode)</param>
+		/// <param name="ExtraOptions">Extra options to pass to xcodebuild</param>
 		/// <param name="bAutomaticSigning">True if using automatic codesigning (where provision and certificate are not used) (ignored in modern mode)</param>
 		/// <param name="bForDistribution">True if this is making a bild for uploading to app store</param>
 		/// <param name="bUseModernXcode">True if the project was made with modern xcode mode</param>
 		/// <param name="Logger">Logging object</param>
 		/// <returns></returns>
 		public static int FinalizeAppWithXcode(DirectoryReference XcodeProject, UnrealTargetPlatform Platform, string SchemeName, string Configuration,
-			string? Provision, string? Certificate, string? Team, bool bAutomaticSigning, bool bForDistribution, bool bUseModernXcode, ILogger Logger)
+			string? Provision, string? Certificate, string? Team, string ExtraOptions, bool bAutomaticSigning, bool bForDistribution, bool bUseModernXcode, ILogger Logger)
 		{
 			List<string> Arguments = new()
 			{
@@ -321,6 +326,7 @@ namespace UnrealBuildTool
 				$"-scheme \"{SchemeName}\"",
 				$"-configuration \"{Configuration}\"",
 				$"-destination generic/platform=" + (Platform == UnrealTargetPlatform.TVOS ? "tvOS" : Platform == UnrealTargetPlatform.Mac ? "macOS" : "iOS"),
+				ExtraOptions,
 				//$"-sdk {SDKName}",
 			};
 
diff --git a/Engine/Source/Programs/UnrealBuildTool/Platform/IOS/IOSToolChain.cs b/Engine/Source/Programs/UnrealBuildTool/Platform/IOS/IOSToolChain.cs
--- a/Engine/Source/Programs/UnrealBuildTool/Platform/IOS/IOSToolChain.cs
+++ b/Engine/Source/Programs/UnrealBuildTool/Platform/IOS/IOSToolChain.cs
@@ -1613,7 +1613,7 @@ namespace UnrealBuildTool
 				}
 
 				// run xcodebuild on the generated project to make the .app
-				IOSExports.FinalizeAppWithModernXcode(GeneratedProjectFile!, Target.Platform, SchemeName, ConfigName, Target.bForDistribution, Logger);
+				IOSExports.FinalizeAppWithModernXcode(GeneratedProjectFile!, Target.Platform, SchemeName, ConfigName, "", Target.bForDistribution, Logger);
 
 				return;
 			}
diff --git a/Engine/Source/Programs/UnrealBuildTool/Platform/Mac/MacToolChain.cs b/Engine/Source/Programs/UnrealBuildTool/Platform/Mac/MacToolChain.cs
--- a/Engine/Source/Programs/UnrealBuildTool/Platform/Mac/MacToolChain.cs
+++ b/Engine/Source/Programs/UnrealBuildTool/Platform/Mac/MacToolChain.cs
@@ -267,10 +267,10 @@ namespace UnrealBuildTool
 		string AddFrameworkToLinkCommand(string FrameworkName, string Arg = "-framework")
 		{
 			string Result = "";
-			if (FrameworkName.EndsWith(".framework"))
+			if (FrameworkName.EndsWith(".framework") || FrameworkName.EndsWith(".framework.zip"))
 			{
 				Result += " -F \"" + Path.GetDirectoryName(Path.GetFullPath(FrameworkName)) + "\"";
-				FrameworkName = Path.GetFileNameWithoutExtension(FrameworkName);
+				FrameworkName = Path.GetFileNameWithoutExtension(Path.GetFileNameWithoutExtension(FrameworkName));
 			}
 			Result += " " + Arg + " \"" + FrameworkName + "\"";
 			return Result;
@@ -717,7 +717,7 @@ namespace UnrealBuildTool
 			List<FileItem> InputFiles = LinkEnvironment.InputFiles;
 
 			// Create an action that invokes the linker.
-			Action LinkAction = Graph.CreateAction(ActionType.Link);
+			Action LinkAction = Graph.	CreateAction(ActionType.Link);
 
 			FileReference LinkerPath = LinkEnvironment.bIsBuildingLibrary ? Info.Archiver : Info.Clang;
 
@@ -771,6 +771,8 @@ namespace UnrealBuildTool
 			if (bIsBuildingAppBundle)
 			{
 				LinkCommand += " -rpath @executable_path/../../../";
+				string ProjName = ProjectFile == null ? "Engine" : ProjectFile.GetFileNameWithoutAnyExtensions();
+				LinkCommand += $" -rpath @executable_path/../UE/{ProjName}/Binaries/Mac";
 			}
 
 			List<string> RPaths = new List<string>();
@@ -840,7 +842,7 @@ namespace UnrealBuildTool
 				{
 					if (!AllFrameworks.ContainsKey(Framework.Name))
 					{
-						AllFrameworks.Add(Framework.Name, false);
+						AllFrameworks.Add(Framework.GetFrameworkDirectory(null, null, Logger)!.FullName, false);
 					}
 				}
 				foreach (string Framework in LinkEnvironment.WeakFrameworks)
diff --git a/Engine/Source/Programs/UnrealBuildTool/ProjectFiles/ProjectFileGenerator.cs b/Engine/Source/Programs/UnrealBuildTool/ProjectFiles/ProjectFileGenerator.cs
--- a/Engine/Source/Programs/UnrealBuildTool/ProjectFiles/ProjectFileGenerator.cs
+++ b/Engine/Source/Programs/UnrealBuildTool/ProjectFiles/ProjectFileGenerator.cs
@@ -1071,74 +1071,74 @@ namespace UnrealBuildTool
 							DebugProjectFiles.AddRange(NewProjectFiles);
 						}
 					}
+				}
+
+				foreach (ProjectFile CurModProject in ModProjects)
+				{
+					RootFolder.AddSubFolder("Mods").ChildProjects.Add(CurModProject);
+				}
+
+				DirectoryReference TemplatesDirectory = DirectoryReference.Combine(Unreal.RootDirectory, "Templates");
+				DirectoryReference SamplesDirectory = DirectoryReference.Combine(Unreal.RootDirectory, "Samples");
 
-					foreach (ProjectFile CurModProject in ModProjects)
+				foreach (ProjectFile CurGameProject in GameProjects)
+				{
+					// Templates go under a different solution folder than games
+					FileReference UnrealProjectFile = CurGameProject.ProjectTargets.First().UnrealProjectFilePath!;
+					if (UnrealProjectFile.IsUnderDirectory(TemplatesDirectory))
+					{
+						DirectoryReference TemplateGameDirectory = CurGameProject.BaseDir;
+						RootFolder.AddSubFolder("Templates").ChildProjects.Add(CurGameProject);
+					}
+					else if (UnrealProjectFile.IsUnderDirectory(SamplesDirectory))
+					{
+						DirectoryReference SampleGameDirectory = CurGameProject.BaseDir;
+						RootFolder.AddSubFolder("Samples").AddSubFolder(UnrealProjectFile.GetFileNameWithoutExtension()).ChildProjects.Add(CurGameProject);
+					}
+					else
 					{
-						RootFolder.AddSubFolder("Mods").ChildProjects.Add(CurModProject);
+						RootFolder.AddSubFolder("Games").AddSubFolder(UnrealProjectFile.GetFileNameWithoutExtension()).ChildProjects.Add(CurGameProject);
 					}
 
-					DirectoryReference TemplatesDirectory = DirectoryReference.Combine(Unreal.RootDirectory, "Templates");
-					DirectoryReference SamplesDirectory = DirectoryReference.Combine(Unreal.RootDirectory, "Samples");
+					List<Tuple<ProjectFile, string>>? NewProjectFiles = CurGameProject.WriteDebugProjectFiles(SupportedPlatforms, SupportedConfigurations, PlatformProjectGenerators, Logger);
 
-					foreach (ProjectFile CurGameProject in GameProjects)
+					if (NewProjectFiles != null)
 					{
-						// Templates go under a different solution folder than games
-						FileReference UnrealProjectFile = CurGameProject.ProjectTargets.First().UnrealProjectFilePath!;
-						if (UnrealProjectFile.IsUnderDirectory(TemplatesDirectory))
-						{
-							DirectoryReference TemplateGameDirectory = CurGameProject.BaseDir;
-							RootFolder.AddSubFolder("Templates").ChildProjects.Add(CurGameProject);
-						}
-						else if (UnrealProjectFile.IsUnderDirectory(SamplesDirectory))
-						{
-							DirectoryReference SampleGameDirectory = CurGameProject.BaseDir;
-							RootFolder.AddSubFolder("Samples").AddSubFolder(UnrealProjectFile.GetFileNameWithoutExtension()).ChildProjects.Add(CurGameProject);
-						}
-						else
-						{
-							RootFolder.AddSubFolder("Games").AddSubFolder(UnrealProjectFile.GetFileNameWithoutExtension()).ChildProjects.Add(CurGameProject);
-						}
+						DebugProjectFiles.AddRange(NewProjectFiles);
+					}
 
-						List<Tuple<ProjectFile, string>>? NewProjectFiles = CurGameProject.WriteDebugProjectFiles(SupportedPlatforms, SupportedConfigurations, PlatformProjectGenerators, Logger);
+				}
 
-						if (NewProjectFiles != null)
-						{
-							DebugProjectFiles.AddRange(NewProjectFiles);
-						}
+				//Related Debug Project Files - Tuple has the related Debug Project, SolutionFolder
+				foreach (Tuple<ProjectFile, string> DebugProjectFile in DebugProjectFiles)
+				{
+					AddExistingProjectFile(DebugProjectFile.Item1, bForceDevelopmentConfiguration: false);
 
-					}
+					//add it to the Android Debug Projects folder in the solution
+					RootFolder.AddSubFolder(DebugProjectFile.Item2).ChildProjects.Add(DebugProjectFile.Item1);
+				}
 
-					//Related Debug Project Files - Tuple has the related Debug Project, SolutionFolder
-					foreach (Tuple<ProjectFile, string> DebugProjectFile in DebugProjectFiles)
+				foreach (KeyValuePair<FileReference, ProjectFile> CurProgramProject in ProgramProjects)
+				{
+					FileReference UnrealProjectFile = CurProgramProject.Value.ProjectTargets.First().UnrealProjectFilePath!;
+					Project? Target = CurProgramProject.Value.ProjectTargets.FirstOrDefault(t => !String.IsNullOrEmpty(t.TargetRules!.SolutionDirectory));
+					if (!bIncludeEnginePrograms && UnrealProjectFile.IsUnderDirectory(Unreal.EngineDirectory))
 					{
-						AddExistingProjectFile(DebugProjectFile.Item1, bForceDevelopmentConfiguration: false);
-
-						//add it to the Android Debug Projects folder in the solution
-						RootFolder.AddSubFolder(DebugProjectFile.Item2).ChildProjects.Add(DebugProjectFile.Item1);
+						continue;
 					}
 
-					foreach (KeyValuePair<FileReference, ProjectFile> CurProgramProject in ProgramProjects)
+					if (Target != null)
 					{
-						FileReference UnrealProjectFile = CurProgramProject.Value.ProjectTargets.First().UnrealProjectFilePath!;
-						Project? Target = CurProgramProject.Value.ProjectTargets.FirstOrDefault(t => !String.IsNullOrEmpty(t.TargetRules!.SolutionDirectory));
-						if (!bIncludeEnginePrograms && UnrealProjectFile.IsUnderDirectory(Unreal.EngineDirectory))
-						{
-							continue;
-						}
-
-						if (Target != null)
-						{
-							RootFolder.AddSubFolder(Target.TargetRules!.SolutionDirectory).ChildProjects.Add(CurProgramProject.Value);
-						}
-						else
-						{
-							RootFolder.AddSubFolder("Programs").ChildProjects.Add(CurProgramProject.Value);
-						}
+						RootFolder.AddSubFolder(Target.TargetRules!.SolutionDirectory).ChildProjects.Add(CurProgramProject.Value);
+					}
+					else
+					{
+						RootFolder.AddSubFolder("Programs").ChildProjects.Add(CurProgramProject.Value);
 					}
-
-					// Add all of the config files for generated program targets
-					AddEngineProgramConfigFiles(ProgramProjects);
 				}
+
+				// Add all of the config files for generated program targets
+				AddEngineProgramConfigFiles(ProgramProjects);
 			}
 
 			// Setup "stub" projects for all modules
diff --git a/Engine/Source/Programs/UnrealBuildTool/ProjectFiles/Xcode/XcodeProject.cs b/Engine/Source/Programs/UnrealBuildTool/ProjectFiles/Xcode/XcodeProject.cs
--- a/Engine/Source/Programs/UnrealBuildTool/ProjectFiles/Xcode/XcodeProject.cs
+++ b/Engine/Source/Programs/UnrealBuildTool/ProjectFiles/Xcode/XcodeProject.cs
@@ -1211,19 +1211,26 @@ namespace UnrealBuildTool.XcodeProjectXcconfig
 
 			List<string> CopyScript = new();
 
-			if (Project.Platform == UnrealTargetPlatform.IOS || Project.Platform == UnrealTargetPlatform.TVOS)
-			{
-				// UBT no longer copies the executable into the .app directory in PostBuild, so we do it here
-				// EXECUTABLE_NAME is Foo, EXECUTABLE_PATH is Foo.app/Foo
-				// NOTE: We read from hardcoded location where UBT writes to, but we write to CONFIGURATION_BUILD_DIR because
-				// when Archiving, the .app is somewhere else
-				CopyScript.AddRange(new string[]
-				{
-					"# Copy the IOS/TVOS executable into .app",
-					$"cp \\\"${{UE_BINARIES_DIR}}/${{EXECUTABLE_NAME}}\\\" \\\"${{CONFIGURATION_BUILD_DIR}}/${{EXECUTABLE_PATH}}\\\"",
-					"",
-				});
-			}
+			// UBT no longer copies the executable into the .app directory in PostBuild, so we do it here
+			// EXECUTABLE_NAME is Foo, EXECUTABLE_PATH is Foo.app/Foo
+			// NOTE: We read from hardcoded location where UBT writes to, but we write to CONFIGURATION_BUILD_DIR because
+			// when Archiving, the .app is somewhere else
+			string SourceEnvVar = (Project.Platform == UnrealTargetPlatform.Mac) ? "EXECUTABLE_PATH" : "EXECUTABLE_NAME";
+			string EngineOrProject = Project.UnrealData.UProjectFileLocation == null ? "Engine" : Project.UnrealData.ProductName;
+
+			CopyScript.AddRange(new string[]
+			{
+				"# Copy the executable into .app if needed (-ef checks if two files/paths are equivalent)",
+				$"SRC_EXE=\\\"${{UE_BINARIES_DIR}}/${{{SourceEnvVar}}}\\\"",
+				$"DEST_EXE=\\\"${{CONFIGURATION_BUILD_DIR}}/${{EXECUTABLE_PATH}}\\\"",
+				"echo ${SRC_EXE} /// ${DEST_EXE}",
+				"if [[ ! ${SRC_EXE} -ef ${DEST_EXE} ]]; then",
+				"  echo Copying executable...",
+				"  mkdir -p `dirname ${DEST_EXE}`",
+				"  cp ${SRC_EXE} ${DEST_EXE}",
+				"fi",
+				""
+			});
 
 			// rsync the Staged build into the .app, unless the UE_SKIP_STAGEDDATA_SYNC var is set to 1
 			string TargetPlatformName = Platform.ToString()!;
@@ -1250,8 +1257,18 @@ namespace UnrealBuildTool.XcodeProjectXcconfig
 					"if [[ ! -e ${STAGED_DIR} ]]; then exit 0; fi",
 					"",
 					"echo \\\"Syncing ${STAGED_DIR} to ${CONFIGURATION_BUILD_DIR}/${CONTENTS_FOLDER_PATH}\\\"",
-					$"rsync -av --exclude=Info.plist --exclude=/Manifest_* --exclude=*.app \\\"${{STAGED_DIR}}/\\\" \\\"${{CONFIGURATION_BUILD_DIR}}/${{CONTENTS_FOLDER_PATH}}{SyncSubdir}\\\"",
+					$"rsync -a --exclude=Info.plist --exclude=/Manifest_* --exclude=*.app \\\"${{STAGED_DIR}}/\\\" \\\"${{CONFIGURATION_BUILD_DIR}}/${{CONTENTS_FOLDER_PATH}}{SyncSubdir}\\\"",
 				});
+
+				// copy any dylibs from the staged stub .app into the real one
+				if (Platform == UnrealTargetPlatform.Mac)
+				{
+					CopyScript.AddRange(new string[]
+					{
+						$"echo \\\"Syncing ${{STAGED_DIR}}/{EngineOrProject}/Binaries/Mac/${{CONTENTS_FOLDER_PATH}}/ to ${{CONFIGURATION_BUILD_DIR}}/${{CONTENTS_FOLDER_PATH}}/\\\"",
+						$"rsync -av --include='*/' --include='*.dylib' --exclude='*' \\\"${{STAGED_DIR}}/{EngineOrProject}/Binaries/Mac/${{CONTENTS_FOLDER_PATH}}/MacOS/\\\" \\\"${{CONFIGURATION_BUILD_DIR}}/${{CONTENTS_FOLDER_PATH}}/MacOS\\\"",
+					});
+				}
 			}
 
 			// run this script every time, but xcode will show a warning if there isn't _some_ output
@@ -1263,15 +1280,6 @@ namespace UnrealBuildTool.XcodeProjectXcconfig
 			// always generate a dsym file when we archive, and by having Xcode do it, it will be put into the archive properly
 			// (note bInstallOnly which will make this onle run when archiving)
 			List<string> DsymScript = new();
-			if (Platform == UnrealTargetPlatform.Mac)
-			{
-				DsymScript.AddRange(new string[]
-				{
-					"# Copy the Mac executable from where it write to into the archiving working directory",
-					$"cp \\\"${{UE_PROJECT_DIR}}/Binaries/{Platform}/${{EXECUTABLE_PATH}}\\\" \\\"${{CONFIGURATION_BUILD_DIR}}/${{EXECUTABLE_PATH}}\\\"",
-					"",
-				});
-			}
 
 			DsymScript.AddRange(new string[]
 			{
@@ -1775,7 +1783,7 @@ namespace UnrealBuildTool.XcodeProjectXcconfig
 			}
 
 			// create one run target for each platform if our platform is null (ie XcodeProjectGenerator.PerPlatformMode is RunTargetPerPlatform)
-			List<UnrealTargetPlatform> TargetPlatforms = Platform == null ? XcodeProjectFileGenerator.XcodePlatforms : new() { Platform.Value };
+			List<UnrealTargetPlatform>					 TargetPlatforms = Platform == null ? XcodeProjectFileGenerator.XcodePlatforms : new() { Platform.Value };
 			foreach (UnrealTargetPlatform TargetPlatform in TargetPlatforms)
 			{
 				XcodeRunTarget RunTarget = new XcodeRunTarget(this, UnrealData.ProductName, UnrealData.AllConfigs[0].ProjectTarget!.TargetRules!.Type, TargetPlatform, BuildTarget, ProjectFile, Logger);
diff --git a/Engine/Source/Programs/UnrealBuildTool/ProjectFiles/Xcode/XcodeProjectFileGenerator.cs b/Engine/Source/Programs/UnrealBuildTool/ProjectFiles/Xcode/XcodeProjectFileGenerator.cs
--- a/Engine/Source/Programs/UnrealBuildTool/ProjectFiles/Xcode/XcodeProjectFileGenerator.cs
+++ b/Engine/Source/Programs/UnrealBuildTool/ProjectFiles/Xcode/XcodeProjectFileGenerator.cs
@@ -264,6 +264,12 @@ namespace UnrealBuildTool
 
 							foreach (XcodeProjectXcconfig.XcodeProjectFile XcodeProject in SupportedProjects)
 							{
+								// if we are only generating a single target project, skip any others now
+								if (!string.IsNullOrEmpty(SingleTargetName) && !XcodeProject.ProjectTargets.Any(x => x.TargetRules?.Name == SingleTargetName))
+								{
+									continue;
+								}
+
 								// we have to re-check for each project - if it's a legacy project, even if we wanted it split, it won't be, so always point to 
 								// the shared legacy project
 								FileReference PathToProject = XcodeProject.ProjectFilePath;
@@ -411,7 +417,7 @@ namespace UnrealBuildTool
 					bIncludeEnginePrograms = false;
 					bIncludeTemplateFiles = false;
 				}
-				//bIncludeEngineSource = false;
+				bIncludeEngineSource = false;
 				bIncludeConfigFiles = false;
 				bIncludeDocumentation = false;
 				bIncludeShaderSource = false;
diff --git a/Engine/Source/Programs/UnrealBuildTool/ToolChain/AppleToolChain.cs b/Engine/Source/Programs/UnrealBuildTool/ToolChain/AppleToolChain.cs
--- a/Engine/Source/Programs/UnrealBuildTool/ToolChain/AppleToolChain.cs
+++ b/Engine/Source/Programs/UnrealBuildTool/ToolChain/AppleToolChain.cs
@@ -531,7 +531,7 @@ namespace UnrealBuildTool
 			// run xcodebuild on the generated project to make the .app
 			bool bForDistribution = false;
 			Logger.LogInformation($"Finalizing .app via Xcode... [generation took {(DateTime.Now - Start).TotalMilliseconds / 1000.0}s]");
-			IOSExports.FinalizeAppWithModernXcode(GeneratedProjectFile!, Target.Platform, Target.TargetName, Target.Configuration.ToString(), bForDistribution, Logger);
+			IOSExports.FinalizeAppWithModernXcode(GeneratedProjectFile!, Target.Platform, Target.TargetName, Target.Configuration.ToString(), "", bForDistribution, Logger);
 		}
 
 
