- Fixed an issue when building multiple editor targets at once (some unrelated code got caught up in a changelist)

diff --git a/Engine/Source/Programs/UnrealBuildTool/Platform/Mac/MacToolChain.cs b/Engine/Source/Programs/UnrealBuildTool/Platform/Mac/MacToolChain.cs
--- a/Engine/Source/Programs/UnrealBuildTool/Platform/Mac/MacToolChain.cs
+++ b/Engine/Source/Programs/UnrealBuildTool/Platform/Mac/MacToolChain.cs
@@ -267,10 +267,10 @@ namespace UnrealBuildTool
 		string AddFrameworkToLinkCommand(string FrameworkName, string Arg = "-framework")
 		{
 			string Result = "";
-			if (FrameworkName.EndsWith(".framework") || FrameworkName.EndsWith(".framework.zip"))
+			if (FrameworkName.EndsWith(".framework"))
 			{
 				Result += " -F \"" + Path.GetDirectoryName(Path.GetFullPath(FrameworkName)) + "\"";
-				FrameworkName = Path.GetFileNameWithoutExtension(Path.GetFileNameWithoutExtension(FrameworkName));
+				FrameworkName = Path.GetFileNameWithoutExtension(FrameworkName);
 			}
 			Result += " " + Arg + " \"" + FrameworkName + "\"";
 			return Result;
@@ -717,7 +717,7 @@ namespace UnrealBuildTool
 			List<FileItem> InputFiles = LinkEnvironment.InputFiles;
 
 			// Create an action that invokes the linker.
-			Action LinkAction = Graph.	CreateAction(ActionType.Link);
+			Action LinkAction = Graph.CreateAction(ActionType.Link);
 
 			FileReference LinkerPath = LinkEnvironment.bIsBuildingLibrary ? Info.Archiver : Info.Clang;
 
@@ -771,8 +771,6 @@ namespace UnrealBuildTool
 			if (bIsBuildingAppBundle)
 			{
 				LinkCommand += " -rpath @executable_path/../../../";
-				string ProjName = ProjectFile == null ? "Engine" : ProjectFile.GetFileNameWithoutAnyExtensions();
-				LinkCommand += $" -rpath @executable_path/../UE/{ProjName}/Binaries/Mac";
 			}
 
 			List<string> RPaths = new List<string>();
@@ -842,7 +840,7 @@ namespace UnrealBuildTool
 				{
 					if (!AllFrameworks.ContainsKey(Framework.Name))
 					{
-						AllFrameworks.Add(Framework.GetFrameworkDirectory(null, null, Logger)!.FullName, false);
+						AllFrameworks.Add(Framework.Name, false);
 					}
 				}
 				foreach (string Framework in LinkEnvironment.WeakFrameworks)
