- Fixed a crash in IOS toolchain code caused by project file writing parallelization

diff --git a/Engine/Source/Programs/UnrealBuildTool/Platform/IOS/UEBuildIOS.cs b/Engine/Source/Programs/UnrealBuildTool/Platform/IOS/UEBuildIOS.cs
--- a/Engine/Source/Programs/UnrealBuildTool/Platform/IOS/UEBuildIOS.cs
+++ b/Engine/Source/Programs/UnrealBuildTool/Platform/IOS/UEBuildIOS.cs
@@ -946,25 +946,29 @@ namespace UnrealBuildTool
 			bool bCacheByBundle = !string.IsNullOrEmpty(Bundle);
 			if (bCacheByBundle)
 			{
-				ProjectSettings = CachedProjectSettingsByBundle.FirstOrDefault(x => x.ProjectFile == ProjectFile && x.BundleIdentifier == Bundle);
+				lock (CachedProjectSettingsByBundle)
+				{
+					ProjectSettings = CachedProjectSettingsByBundle.FirstOrDefault(x => x.ProjectFile == ProjectFile && x.BundleIdentifier == Bundle);
+					if (ProjectSettings == null)
+					{
+						ProjectSettings = CreateProjectSettings(ProjectFile, Bundle);
+						CachedProjectSettingsByBundle.Add(ProjectSettings);
+					}
+				}
 			}
 			else
 			{
-				ProjectSettings = CachedProjectSettings.FirstOrDefault(x => x.ProjectFile == ProjectFile);
-			}
-
-			if(ProjectSettings == null)
-			{
-				ProjectSettings = CreateProjectSettings(ProjectFile, Bundle);
-				if (bCacheByBundle)
+				lock (CachedProjectSettings)
 				{
-					CachedProjectSettingsByBundle.Add(ProjectSettings);
-				}
-				else
-				{
-					CachedProjectSettings.Add(ProjectSettings);
+					ProjectSettings = CachedProjectSettings.FirstOrDefault(x => x.ProjectFile == ProjectFile);
+					if (ProjectSettings == null)
+					{
+						ProjectSettings = CreateProjectSettings(ProjectFile, Bundle);
+						CachedProjectSettings.Add(ProjectSettings);
+					}
 				}
 			}
+
 			return ProjectSettings;
 		}
 
