- Another parallel target threading crash fix (Mac this time)

diff --git a/Engine/Source/Programs/UnrealBuildTool/Platform/Mac/UEBuildMac.cs b/Engine/Source/Programs/UnrealBuildTool/Platform/Mac/UEBuildMac.cs
--- a/Engine/Source/Programs/UnrealBuildTool/Platform/Mac/UEBuildMac.cs
+++ b/Engine/Source/Programs/UnrealBuildTool/Platform/Mac/UEBuildMac.cs
@@ -361,13 +361,16 @@ namespace UnrealBuildTool
 			{
 				ApplePlatformSDK SDK = (ApplePlatformSDK?)GetSDK() ?? new ApplePlatformSDK(Logger);
 				// Validate the added public libraries
-				foreach (FileReference LibLoc in Module.PublicLibraries)
+				lock (ValidatedLibs)
 				{
-					if (ValidatedLibs.Contains(LibLoc))
+					foreach (FileReference LibLoc in Module.PublicLibraries)
 					{
-						continue;
+						if (ValidatedLibs.Contains(LibLoc))
+						{
+							continue;
+						}
+						ValidatedLibs.Add(LibLoc);
 					}
-					ValidatedLibs.Add(LibLoc);
 				}
 			}
 		}
