Land public UBT portions of SOL-4383

This patch makes it possible to generate header files in a way that the Verse VM will take advantage of.

diff --git a/Engine/Source/Programs/UnrealBuildTool/Configuration/ModuleRules.cs b/Engine/Source/Programs/UnrealBuildTool/Configuration/ModuleRules.cs
--- a/Engine/Source/Programs/UnrealBuildTool/Configuration/ModuleRules.cs
+++ b/Engine/Source/Programs/UnrealBuildTool/Configuration/ModuleRules.cs
@@ -11,6 +11,7 @@ using System.Threading.Tasks;
 using EpicGames.Core;
 using Microsoft.Extensions.Logging;
 using UnrealBuildBase;
+using Microsoft.Extensions.Logging;
 
 namespace UnrealBuildTool
 {
@@ -1299,6 +1300,12 @@ namespace UnrealBuildTool
 		/// </summary>
 		public CStandardVersion? CStandard;
 
+		/// <summary>
+        /// A list of subdirectory names and functions that are invoked to generate header files.
+        /// The subdirectory name is appended to the generated code directory to form a new directory
+        /// that headers are generated inside.
+		/// </summary>
+        public List<Tuple<string, Action<ILogger, DirectoryReference>>> GenerateHeaderFuncs = new List<Tuple<string, Action<ILogger, DirectoryReference>>>();
 
 		/// <summary>
 		///  Control visibility of symbols
diff --git a/Engine/Source/Programs/UnrealBuildTool/Configuration/UEBuildModuleCPP.cs b/Engine/Source/Programs/UnrealBuildTool/Configuration/UEBuildModuleCPP.cs
--- a/Engine/Source/Programs/UnrealBuildTool/Configuration/UEBuildModuleCPP.cs
+++ b/Engine/Source/Programs/UnrealBuildTool/Configuration/UEBuildModuleCPP.cs
@@ -334,6 +334,11 @@ namespace UnrealBuildTool
 				{
 					IncludePaths.Add(GeneratedCodeDirectoryVNI!);
 				}
+
+                foreach (var (Subdirectory, _) in Rules.GenerateHeaderFuncs)
+                {
+					IncludePaths.Add(DirectoryReference.Combine(GeneratedCodeDirectory, Subdirectory));
+                }
 			}
 
 			ModuleInterfacePaths.Add(UEToolChain.GetModuleInterfaceDir(IntermediateDirectory));
diff --git a/Engine/Source/Programs/UnrealBuildTool/Configuration/UEBuildTarget.cs b/Engine/Source/Programs/UnrealBuildTool/Configuration/UEBuildTarget.cs
--- a/Engine/Source/Programs/UnrealBuildTool/Configuration/UEBuildTarget.cs
+++ b/Engine/Source/Programs/UnrealBuildTool/Configuration/UEBuildTarget.cs
@@ -2391,6 +2391,16 @@ namespace UnrealBuildTool
 			VNITask.Wait();
 #endif
 
+            foreach (var Module in Modules.Values.OfType<UEBuildModuleCPP>())
+            {
+                foreach (var (Subdirectory, Func) in Module.Rules.GenerateHeaderFuncs)
+                {
+                    DirectoryReference GenDirectory = DirectoryReference.Combine(Module.GeneratedCodeDirectory!, Subdirectory);
+                    DirectoryReference.CreateDirectory(GenDirectory);
+                    Func(Logger, GenDirectory);
+                }
+            }
+
 			// Compile the resource files common to all DLLs on Windows
 			if (!ShouldCompileMonolithic())
 			{
