- Fixed codesigning of Content only projects

diff --git a/Engine/Source/Programs/UnrealBuildTool/ProjectFiles/Xcode/XcodeProjectFileGenerator.cs b/Engine/Source/Programs/UnrealBuildTool/ProjectFiles/Xcode/XcodeProjectFileGenerator.cs
--- a/Engine/Source/Programs/UnrealBuildTool/ProjectFiles/Xcode/XcodeProjectFileGenerator.cs
+++ b/Engine/Source/Programs/UnrealBuildTool/ProjectFiles/Xcode/XcodeProjectFileGenerator.cs
@@ -417,7 +417,6 @@ namespace UnrealBuildTool
 					bIncludeEnginePrograms = false;
 					bIncludeTemplateFiles = false;
 				}
-				bIncludeEngineSource = false;
 				bIncludeConfigFiles = false;
 				bIncludeDocumentation = false;
 				bIncludeShaderSource = false;
@@ -430,6 +429,8 @@ namespace UnrealBuildTool
 				// generate just the game project
 				else
 				{
+					// content only, project-only projects need engine to get the UnrealGame, etc targets (checking for Root and engine catches UnrealGame and content only projects, respectively)
+					bIncludeEngineSource = PrimaryProjectPath == UnrealBuildBase.Unreal.RootDirectory || PrimaryProjectPath == UnrealBuildBase.Unreal.EngineDirectory;
 					bGeneratingGameProjectFiles = true;
 				}
 			}
