During a RemoteMac tvOS build, the entitlement file is written during the stub ipa build which can confuse xcode, so just tell it to ignore the fact the file changed.

#rnx

diff --git a/Engine/Source/Programs/UnrealBuildTool/Platform/IOS/IOSToolChain.cs b/Engine/Source/Programs/UnrealBuildTool/Platform/IOS/IOSToolChain.cs
--- a/Engine/Source/Programs/UnrealBuildTool/Platform/IOS/IOSToolChain.cs
+++ b/Engine/Source/Programs/UnrealBuildTool/Platform/IOS/IOSToolChain.cs
@@ -1723,6 +1723,10 @@ namespace UnrealBuildTool
 
 					CmdLine += String.Format(" CODE_SIGN_IDENTITY='{0}'", SigningCertificate);
 
+					// <projectname>.entitlements file can get copied over during a RemoteMac build which confuses Xcode, so just tell it to ignore any 
+					// entitlement file changes.  See UE-181206.
+					CmdLine += String.Format(" CODE_SIGN_ALLOW_ENTITLEMENTS_MODIFICATION=YES");
+
 					if (!ProjectSettings.bAutomaticSigning)
 					{
 						CmdLine += (!string.IsNullOrEmpty(MobileProvisionUUID) ? (" PROVISIONING_PROFILE_SPECIFIER=" + MobileProvisionUUID) : "");
