UAT: Avoid throwing an exception (and breaking in the debugger) while checking for WINE during normal execution.

diff --git a/Engine/Source/Programs/UnrealBuildTool/System/BuildHostPlatform.cs b/Engine/Source/Programs/UnrealBuildTool/System/BuildHostPlatform.cs
--- a/Engine/Source/Programs/UnrealBuildTool/System/BuildHostPlatform.cs
+++ b/Engine/Source/Programs/UnrealBuildTool/System/BuildHostPlatform.cs
@@ -250,22 +250,16 @@ namespace UnrealBuildTool
 
 		public override string BinarySuffix => ".exe";
 
-		[DllImport("ntdll.dll", EntryPoint="wine_get_version", CallingConvention=CallingConvention.Cdecl, CharSet=CharSet.Ansi)]
-		private static extern string GetWineVersion();
+		[DllImport("kernel32.dll", CharSet = CharSet.Auto)]
+		private static extern IntPtr GetModuleHandle(string lpModuleName);
+
+		[DllImport("kernel32.dll", CharSet = CharSet.Auto)]
+		private static extern IntPtr GetProcAddress(IntPtr hModule, string procName);
 
 		public override bool IsRunningOnWine()
 		{
-			try
-			{
-				// if we dont throw an exception we are running in Wine
-				GetWineVersion();
-			}
-			catch
-			{
-				return false;
-			};
-
-			return true;
+			IntPtr NtdllHandle = GetModuleHandle("ntdll.dll");
+			return NtdllHandle.ToInt64() != 0 && GetProcAddress(NtdllHandle, "wine_get_version").ToInt64() != 0;
 		}
 
 		internal override IEnumerable<ProjectFileFormat> GetDefaultProjectFileFormats()
