Fixed issue where the check for first header include was not being reported as an actual warning/error.

diff --git a/Engine/Source/Programs/UnrealBuildTool/Configuration/UEBuildModuleCPP.cs b/Engine/Source/Programs/UnrealBuildTool/Configuration/UEBuildModuleCPP.cs
--- a/Engine/Source/Programs/UnrealBuildTool/Configuration/UEBuildModuleCPP.cs
+++ b/Engine/Source/Programs/UnrealBuildTool/Configuration/UEBuildModuleCPP.cs
@@ -34,6 +34,21 @@ namespace UnrealBuildTool
 			public readonly List<FileItem> ISPCFiles = new List<FileItem>();
 		}
 
+		/// <summary>
+		/// Instance where a CPP file doesn't include the expected header as the first include.
+		/// </summary>
+		public struct InvalidIncludeDirective
+		{
+			public readonly FileReference CppFile;
+			public readonly FileReference HeaderFile;
+
+			public InvalidIncludeDirective(FileReference cppFile, FileReference headerFile)
+			{
+				CppFile = cppFile;
+				HeaderFile = headerFile;
+			}
+		}
+
 		/// <summary>
 		/// If UHT found any associated UObjects in this module's source files
 		/// </summary>
@@ -73,7 +88,7 @@ namespace UnrealBuildTool
 		/// <summary>
 		/// List of invalid include directives. These are buffered up and output before we start compiling.
 		/// </summary>
-		public List<string>? InvalidIncludeDirectiveMessages;
+		public List<InvalidIncludeDirective>? InvalidIncludeDirectives;
 
 		/// <summary>
 		/// Set of source directories referenced during a build
@@ -1585,7 +1600,7 @@ namespace UnrealBuildTool
 		{
 			if(Rules.PCHUsage == ModuleRules.PCHUsageMode.UseExplicitOrSharedPCHs)
 			{
-				if(InvalidIncludeDirectiveMessages == null)
+				if(InvalidIncludeDirectives == null)
 				{
 					// Find headers used by the source file.
 					Dictionary<string, FileReference> NameToHeaderFile = new Dictionary<string, FileReference>();
@@ -1595,7 +1610,7 @@ namespace UnrealBuildTool
 					}
 
 					// Find the directly included files for each source file, and make sure it includes the matching header if possible
-					InvalidIncludeDirectiveMessages = new List<string>();
+					InvalidIncludeDirectives = new();
 					if (Rules != null && Rules.IWYUSupport != IWYUSupport.None && Target.bEnforceIWYU)
 					{
 						foreach (FileItem CppFile in CppFiles)
@@ -1610,7 +1625,7 @@ namespace UnrealBuildTool
 									FileReference? HeaderFile;
 									if (NameToHeaderFile.TryGetValue(ExpectedName, out HeaderFile) && !IgnoreMismatchedHeader(ExpectedName))
 									{
-										InvalidIncludeDirectiveMessages.Add(String.Format("{0}(1): error: Expected {1} to be first header included.", CppFile.Location, HeaderFile.GetFileName()));
+										InvalidIncludeDirectives.Add(new InvalidIncludeDirective(CppFile.Location, HeaderFile));
 									}
 								}
 							}
diff --git a/Engine/Source/Programs/UnrealBuildTool/Configuration/UEBuildTarget.cs b/Engine/Source/Programs/UnrealBuildTool/Configuration/UEBuildTarget.cs
--- a/Engine/Source/Programs/UnrealBuildTool/Configuration/UEBuildTarget.cs
+++ b/Engine/Source/Programs/UnrealBuildTool/Configuration/UEBuildTarget.cs
@@ -2510,12 +2510,12 @@ namespace UnrealBuildTool
 			}
 
 			// Make sure all the checked headers were valid
-			List<string> InvalidIncludeDirectiveMessages = Modules.Values.OfType<UEBuildModuleCPP>().Where(x => x.InvalidIncludeDirectiveMessages != null).SelectMany(x => x.InvalidIncludeDirectiveMessages!).ToList();
-			if (InvalidIncludeDirectiveMessages.Count > 0)
+			List<UEBuildModuleCPP.InvalidIncludeDirective> InvalidIncludeDirectives = Modules.Values.OfType<UEBuildModuleCPP>().Where(x => x.InvalidIncludeDirectives != null).SelectMany(x => x.InvalidIncludeDirectives!).ToList();
+			if (InvalidIncludeDirectives.Count > 0)
 			{
-				foreach (string InvalidIncludeDirectiveMessage in InvalidIncludeDirectiveMessages)
+				foreach (UEBuildModuleCPP.InvalidIncludeDirective InvalidIncludeDirective in InvalidIncludeDirectives)
 				{
-					Log.WriteLine(LogEventType.Warning, LogFormatOptions.NoSeverityPrefix, "{0}", InvalidIncludeDirectiveMessage);
+					Logger.LogWarning("{CppFile}(1): error: Expected {HeaderFile} to be first header included.", InvalidIncludeDirective.CppFile, InvalidIncludeDirective.HeaderFile.GetFileName());
 				}
 			}
 
