Added circular dependency info

diff --git a/Engine/Source/Programs/UnrealBuildTool/Modes/PrintBuildGraphInfo.cs b/Engine/Source/Programs/UnrealBuildTool/Modes/PrintBuildGraphInfo.cs
--- a/Engine/Source/Programs/UnrealBuildTool/Modes/PrintBuildGraphInfo.cs
+++ b/Engine/Source/Programs/UnrealBuildTool/Modes/PrintBuildGraphInfo.cs
@@ -151,6 +151,7 @@ namespace UnrealBuildTool
 						}
 					}
 
+					Dictionary<UEBuildModule, PrecompiledHeaderTemplate> ModulesUsingSharedPCHs = new();
 					if (GlobalCompileEnvironment.SharedPCHs.Any())
 					{
 						Logger.LogInformation(" Shared PCHs:");
@@ -174,14 +175,90 @@ namespace UnrealBuildTool
 								foreach (UEBuildModuleCPP Module in SortedModules)
 								{
 									Logger.LogInformation("    {ModuleName}", Module.Name);
+									ModulesUsingSharedPCHs.Add(Module, Template);
 								}
 							}
 						}
 					}
+
+					Logger.LogInformation("Circular Dependencies that are not declared.");
+					Logger.LogInformation("NOTE: Anything prefixed with a '!!!' means that the circular dependency could possibly create linker errors.");
+					foreach (UEBuildModule Module in Modules)
+					{
+						List<Stack<UEBuildModule>> ModuleDepStacks = new List<Stack<UEBuildModule>>();
+						FindCircularDependencyModules(Module, Module, new HashSet<UEBuildModule>(), true, new Stack<UEBuildModule>(), ModuleDepStacks);
+						if (ModuleDepStacks.Any())
+						{
+							bool UsesSharedPCH = ModulesUsingSharedPCHs.ContainsKey(Module);
+							Logger.LogInformation(" {ModuleName} - {DepCount}", Module.Name, ModuleDepStacks.Count);
+							foreach (Stack<UEBuildModule> DepModuleStack in ModuleDepStacks)
+							{
+								StringBuilder Dependencies = new StringBuilder();
+
+								if (UsesSharedPCH && DepModuleStack.Contains(ModulesUsingSharedPCHs[Module].Module))
+								{
+									Dependencies.Append("!!! ");
+								}
+
+								var DepModuleArray = DepModuleStack.Reverse().ToArray();
+								for (int DepModuleIndex = 0; DepModuleIndex < DepModuleArray.Length; DepModuleIndex++) 
+								{
+									Dependencies.Append(DepModuleArray[DepModuleIndex].Name);
+									if (DepModuleIndex != DepModuleArray.Length - 1)
+									{
+										Dependencies.Append(" -> ");
+									}
+								}
+								Logger.LogInformation($"   {Dependencies}");
+							}
+						}
+					}
 				}
 			}
 
 			return 0;
 		}
+
+		private void FindCircularDependencyModules(UEBuildModule SearchForBuildModule, UEBuildModule CurrentBuildModule, HashSet<UEBuildModule> IgnoreReferencedModules, bool bIncludePrivateDependencyModules, Stack<UEBuildModule> CurrentStack, List<Stack<UEBuildModule>> ModuleDepStacks)
+		{
+			CurrentStack.Push(CurrentBuildModule);
+
+			if (SearchForBuildModule == CurrentBuildModule && CurrentStack.Count > 1)
+			{
+				ModuleDepStacks.Add(new Stack<UEBuildModule>(CurrentStack.Reverse()));
+			}
+
+			List<UEBuildModule> AllDependencyModules = new List<UEBuildModule>(
+				((bIncludePrivateDependencyModules && CurrentBuildModule.PrivateDependencyModules != null) ? CurrentBuildModule.PrivateDependencyModules.Count : 0) +
+				(CurrentBuildModule.PublicDependencyModules != null ? CurrentBuildModule.PublicDependencyModules.Count : 0) +
+				(CurrentBuildModule.PublicIncludePathModules != null ? CurrentBuildModule.PublicIncludePathModules!.Count : 0)
+				);
+			if (bIncludePrivateDependencyModules && CurrentBuildModule.PrivateDependencyModules != null)
+			{
+				AllDependencyModules.AddRange(CurrentBuildModule.PrivateDependencyModules);
+			}
+			if (CurrentBuildModule.PublicDependencyModules != null)
+			{
+				AllDependencyModules.AddRange(CurrentBuildModule.PublicDependencyModules!);
+			}
+			if (CurrentBuildModule.PublicIncludePathModules != null)
+			{
+				AllDependencyModules.AddRange(CurrentBuildModule.PublicIncludePathModules);
+			}
+
+			foreach (UEBuildModule DependencyModule in AllDependencyModules)
+			{
+				// Don't follow circular back-references!
+				if (!CurrentBuildModule.HasCircularDependencyOn(DependencyModule.Name))
+				{
+					if (IgnoreReferencedModules.Add(DependencyModule))
+					{
+						// Recurse into dependent modules
+						FindCircularDependencyModules(SearchForBuildModule, DependencyModule, new HashSet<UEBuildModule>(IgnoreReferencedModules), false, CurrentStack, ModuleDepStacks);
+					}
+				}
+			}
+			CurrentStack.Pop();
+		}
 	}
 }
