Moving Code Signing out of iPhone Packager and in UBT to remove Mono and reduce IPP dependancy

diff --git a/Engine/Source/Developer/IOS/IOSPlatformEditor/Private/IOSTargetSettingsCustomization.cpp b/Engine/Source/Developer/IOS/IOSPlatformEditor/Private/IOSTargetSettingsCustomization.cpp
--- a/Engine/Source/Developer/IOS/IOSPlatformEditor/Private/IOSTargetSettingsCustomization.cpp
+++ b/Engine/Source/Developer/IOS/IOSPlatformEditor/Private/IOSTargetSettingsCustomization.cpp
@@ -1492,8 +1492,6 @@ FReply FIOSTargetSettingsCustomization::OnInstallProvisionClicked()
 
 FReply FIOSTargetSettingsCustomization::OnGenerateImageClicked(const FString SourceImagePath, const FString TargetImagePath, FIntPoint IconRequiredSize)
 {
-    UE_LOG(LogIOSTargetSettings, Warning, TEXT("HEY SALUT=[%s]"), *SourceImagePath);
-    
 	if (FPaths::FileExists(*TargetImagePath))
 	{
 		const EAppReturnType::Type Answer = FMessageDialog::Open(EAppMsgType::YesNo, FText::FromString(TEXT("File already exists. Do you want to overwrite it ?")));
diff --git a/Engine/Source/Programs/AutomationTool/Scripts/IPhonePackager.Automation.cs b/Engine/Source/Programs/AutomationTool/Scripts/IPhonePackager.Automation.cs
--- a/Engine/Source/Programs/AutomationTool/Scripts/IPhonePackager.Automation.cs
+++ b/Engine/Source/Programs/AutomationTool/Scripts/IPhonePackager.Automation.cs
@@ -5,6 +5,8 @@ using System.Text;
 using System.IO;
 using AutomationTool;
 using UnrealBuildTool;
+using EpicGames.Core;
+using System.Runtime.ConstrainedExecution;
 using Microsoft.Extensions.Logging;
 
 [Help("UAT command to call into the integrated IPhonePackager code")]
@@ -12,7 +14,35 @@ class IPhonePackager : BuildCommand
 {
 	public override void ExecuteBuild()
 	{
-		Logger.LogInformation("************************* Calling IPP");
+
+		string ProjectFile = ParseParamValue("project");
+		string BundleIdentifier = ParseParamValue("bundleid");
+		bool bIsTVOS = ParseParam("tvos");
+
+		CodeSigningConfig.Initialize(new FileReference(ProjectFile), bIsTVOS:bIsTVOS);
+		List<string> Certs = AppleCodeSign.FindCertificates();
+		List<FileReference> Provisions = AppleCodeSign.FindProvisions(BundleIdentifier, bForDistribution:false, out _);
+
+
+		Logger.LogInformation("Certs:");
+		foreach (string Cert in Certs)
+		{
+			Logger.LogInformation($"  {Cert}");
+		}
+		Logger.LogInformation("Provisions:");
+		foreach (FileReference Provision in Provisions)
+		{
+			Logger.LogInformation($"  {Provision}");
+		}
+
+
+		FileReference MobileProvisionFile;
+		string SigningCertificate;
+		if (AppleCodeSign.FindCertAndProvision(BundleIdentifier, out MobileProvisionFile, out SigningCertificate))
+		{
+			Logger.LogInformation($"Matched: {MobileProvisionFile}- {SigningCertificate}");
+		}
+
 
 		Platform IOS = Platform.GetPlatform(UnrealTargetPlatform.IOS);
 		string Command = ParseParamValue("cmd", "");
diff --git a/Engine/Source/Programs/UnrealBuildTool/Platform/IOS/IOSCodeSigning.cs b/Engine/Source/Programs/UnrealBuildTool/Platform/IOS/IOSCodeSigning.cs
new file mode 100644
--- /dev/null
+++ b/Engine/Source/Programs/UnrealBuildTool/Platform/IOS/IOSCodeSigning.cs
@@ -0,0 +1,2436 @@
+﻿using System;
+using System.Collections.Generic;
+using System.Text;
+using System.Security.Cryptography.X509Certificates;
+using System.Security.Cryptography;
+using System.IO;
+using System.Diagnostics;
+using System.Xml;
+using System.Linq;
+using System.Globalization;
+using EpicGames.Core;
+using UnrealBuildBase;
+using System.Collections;
+using System.Text.RegularExpressions;
+using System.Runtime.ConstrainedExecution;
+using Microsoft.Extensions.Logging;
+
+namespace UnrealBuildTool
+{
+	/// <summary>
+	/// Utilities related to code sign for Apple Platforms
+	/// </summary>
+	public class AppleCodeSign
+	{
+		/// <summary>
+		/// A logger
+		/// </summary>
+		static ILogger? _logger;
+
+		/// <summary>
+		/// Constructor
+		/// </summary>
+		public AppleCodeSign()
+		{
+			ILoggerFactory loggerFactory = LoggerFactory.Create(builder =>
+		    {
+        		builder.AddConsole();
+		    });
+			_logger = loggerFactory.CreateLogger("AppleCodeSign");
+		}
+
+		/// <summary>
+		/// The shared provision library directory (on PC)
+		/// </summary>
+		public static string ProvisionDirectory
+		{
+			get
+			{
+				if (OperatingSystem.IsMacOS())
+				{
+					return Environment.GetEnvironmentVariable("HOME") + "/Library/MobileDevice/Provisioning Profiles/";
+				}
+				else
+				{
+					return Environment.GetFolderPath(Environment.SpecialFolder.LocalApplicationData) + "/Apple Computer/MobileDevice/Provisioning Profiles/";
+				}
+			}
+		}
+
+		/// <summary>
+		/// Get the name from the Certificate
+		/// </summary>
+		public static string GetCommonNameFromCert(X509Certificate2 Cert)
+		{
+			// Make sure we have a useful friendly name
+			string CommonName = "(no common name present)";
+
+			string FullName = Cert.SubjectName.Name;
+			char[] SplitChars = { ',' };
+			string[] NameParts = FullName.Split(SplitChars);
+
+			foreach (string Part in NameParts)
+			{
+				string CleanPart = Part.Trim();
+				if (CleanPart.StartsWith("CN="))
+				{
+					CommonName = CleanPart.Substring(3);
+				}
+			}
+
+			return CommonName;
+		}
+
+		/// <summary>
+		/// Get the friendly name from the Certificate
+		/// </summary>
+		public static string GetFriendlyNameFromCert(X509Certificate2 Cert)
+		{
+			if (OperatingSystem.IsMacOS())
+			{
+				return GetCommonNameFromCert(Cert);
+			}
+			else
+			{
+				// Make sure we have a useful friendly name
+				string FriendlyName = Cert.FriendlyName;
+				if ((FriendlyName == "") || (FriendlyName == null))
+				{
+					FriendlyName = GetCommonNameFromCert(Cert);
+				}
+
+				return FriendlyName;
+			}
+		}
+
+
+		static private string CertToolData = "";
+
+		/// <summary>
+		/// When receiving the cert tool process call
+		/// </summary>
+		static public void OutputReceivedCertToolProcessCall(Object Sender, DataReceivedEventArgs Line)
+		{
+			if ((Line != null) && !String.IsNullOrEmpty(Line.Data))
+			{
+				CertToolData += Line.Data + "\n";
+			}
+		}
+
+		/// <summary>
+		/// Finds all valid installed certificates
+		/// </summary>
+		public static List<string> FindCertificates()
+		{
+			string[] ValidCertificatePrefixes = { "iPhone Developer", "iPhone Distribution", "Apple Development", "Apple Distribution" };
+
+			X509Certificate2Collection FoundCerts = new X509Certificate2Collection();
+
+			if (OperatingSystem.IsMacOS())
+			{
+				foreach (string SearchPrefix in ValidCertificatePrefixes)
+				{
+					// run certtool y to get the currently installed certificates
+					CertToolData = "";
+					Process CertTool = new Process();
+					CertTool.StartInfo.FileName = "/usr/bin/security";
+					CertTool.StartInfo.UseShellExecute = false;
+					CertTool.StartInfo.Arguments = string.Format("find-certificate -a -c \"{0}\" -p", SearchPrefix);
+					CertTool.StartInfo.RedirectStandardOutput = true;
+					CertTool.StartInfo.StandardOutputEncoding = Encoding.UTF8;
+					CertTool.OutputDataReceived += new DataReceivedEventHandler(OutputReceivedCertToolProcessCall);
+					CertTool.Start();
+					CertTool.BeginOutputReadLine();
+					CertTool.WaitForExit();
+					if (CertTool.ExitCode == 0)
+					{
+						string header = "-----BEGIN CERTIFICATE-----\n";
+						string footer = "-----END CERTIFICATE-----";
+						int start = CertToolData.IndexOf(header);
+						while (start != -1)
+						{
+							start += header.Length;
+							int end = CertToolData.IndexOf(footer, start);
+							string base64 = CertToolData.Substring(start, (end - start));
+							byte[] certData = Convert.FromBase64String(base64);
+							X509Certificate2 cert = new X509Certificate2(certData);
+							FoundCerts.Add(cert);
+							start = CertToolData.IndexOf(header, start);
+						}
+					}
+				}
+			}
+			else
+			{
+				// Open the personal certificate store on this machine
+				X509Store Store = new X509Store();
+				Store.Open(OpenFlags.ReadOnly);
+
+				foreach (string SearchPrefix in ValidCertificatePrefixes)
+				{
+					FoundCerts.AddRange(Store.Certificates.Find(X509FindType.FindBySubjectName, SearchPrefix, false));
+				}
+
+				Store.Close();
+			}
+
+			List<string> Certs = new();
+			foreach (X509Certificate2 TestCert in FoundCerts)
+			{
+				DateTime EffectiveDate = TestCert.NotBefore.ToUniversalTime();
+				DateTime ExpirationDate = TestCert.NotAfter.ToUniversalTime();
+				DateTime Now = DateTime.UtcNow;
+
+				bool bCertTimeIsValid = (EffectiveDate < Now) && (ExpirationDate > Now);
+				string CertLine = string.Format("CERTIFICATE-Name:{0},Validity:{1},StartDate:{2},EndDate:{3}", GetFriendlyNameFromCert(TestCert), bCertTimeIsValid ? "VALID" : "EXPIRED", EffectiveDate.ToString("o"), ExpirationDate.ToString("o"));
+				_logger?.LogInformation(CertLine);
+				Certs.Add(CertLine);
+			}
+			return Certs;
+		}
+
+		/// <summary>
+		/// Finds all valid installed provisions
+		/// </summary>
+		public static List<FileReference> FindProvisions(string CFBundleIdentifier, bool bForDistribution, out FileReference? MatchedProvision)
+		{
+			CodeSigningConfig.bForDistribution = bForDistribution;
+
+			List<FileReference> Provisions = new();
+
+			if (!Directory.Exists(CodeSigningConfig.ProvisionDirectory))
+			{
+				_logger?.LogError("Could not find provision directory '{0}'.", CodeSigningConfig.ProvisionDirectory);
+				MatchedProvision = null;
+				return Provisions;
+			}
+
+			// cache the provision library
+			string SelectedProvision = "";
+			string SelectedCert = "";
+			string SelectedFile = "";
+			int FoundName = -1;
+
+			Dictionary<string, MobileProvision> ProvisionLibrary = new Dictionary<string, MobileProvision>();
+
+			foreach (string ProvisionFile in Directory.EnumerateFiles(CodeSigningConfig.ProvisionDirectory, "*.mobileprovision"))
+			{
+				MobileProvision Provision = MobileProvisionParser.ParseFile(ProvisionFile);
+				ProvisionLibrary.Add(ProvisionFile, Provision);
+			}
+
+			// first sort all profiles so we look at newer ones first.
+			IEnumerable<string> ProfileKeys = ProvisionLibrary.Select(KV => KV.Key)
+				.OrderByDescending(K => ProvisionLibrary[K].CreationDate)
+				.ToArray();
+
+			// note - all of this is a near duplicate of code in MobileProvisionUtilities, which other functions
+			// in this class call to do similar work! 
+			// @todo - unify all of this.
+
+			foreach (string ProvisionFile in ProfileKeys)
+			{
+				MobileProvision p = ProvisionLibrary[ProvisionFile];
+
+				DateTime EffectiveDate = p.CreationDate;
+				DateTime ExpirationDate = p.ExpirationDate;
+				DateTime Now = DateTime.UtcNow;
+
+				bool bCertTimeIsValid = (EffectiveDate < Now) && (ExpirationDate > Now);
+				bool bValid = false;
+				X509Certificate2? Cert = FindCertificate(p);
+				if (Cert != null)
+				{
+					bValid = (Cert.NotBefore.ToUniversalTime() < Now) && (Cert.NotAfter.ToUniversalTime() > Now);
+				}
+				bool bPassesNameCheck = p.ApplicationIdentifier.Substring(p.ApplicationIdentifierPrefix.Length + 1) == CFBundleIdentifier;
+				bool bPassesCompanyCheck = false;
+				bool bPassesWildCardCheck = false;
+				if (p.ApplicationIdentifier.Contains("*"))
+				{
+					string CompanyName = p.ApplicationIdentifier.Substring(p.ApplicationIdentifierPrefix.Length + 1);
+				
+					if (CompanyName != "*")
+					{
+						CompanyName = CompanyName.Substring(0, CompanyName.LastIndexOf("."));
+						bPassesCompanyCheck = CFBundleIdentifier.StartsWith(CompanyName);
+					}
+					else
+					{
+						bPassesWildCardCheck = true;
+					}
+				}
+				bool bIsManaged = false;
+				if (p.ProvisionName == "iOS Team Provisioning Profile: " + CFBundleIdentifier)
+				{
+					bIsManaged = true;
+				}
+				bool bDistribution = ((p?.ProvisionedDeviceIDs?.Count == 0) && !(p?.bDebug ?? false));
+				string Validity = "VALID";
+				if (!bCertTimeIsValid)
+				{
+					Validity = "EXPIRED";
+				}
+				else if (!bValid)
+				{
+					Validity = "NO_CERT";
+				}
+				else if (!bPassesNameCheck && !bPassesWildCardCheck && !bPassesCompanyCheck)
+				{
+					Validity = "NO_MATCH";
+				}
+				if (bIsManaged)
+				{
+					Validity = "MANAGED";
+				}
+				if ((string.IsNullOrWhiteSpace(SelectedProvision) || FoundName < 2) && Validity == "VALID" && !bDistribution)
+				{
+					int Prev = FoundName;
+					if (bPassesNameCheck)
+					{
+						FoundName = 2;
+					}
+					else if (bPassesCompanyCheck && FoundName < 1)
+					{
+						FoundName = 1;
+					}
+					else if (bPassesWildCardCheck && FoundName == -1)
+					{
+						FoundName = 0;
+					}
+					if (Cert is not null && FoundName != Prev)
+					{
+					    SelectedProvision = p?.ProvisionName ?? "";
+    					SelectedFile = Path.GetFileName(ProvisionFile) ?? "";
+    					SelectedCert = GetFriendlyNameFromCert(Cert);
+					}	
+				}
+				if (p != null)
+				{
+				    _logger?.LogInformation("PROVISION-File:{0},Name:{1},Validity:{2},StartDate:{3},EndDate:{4},Type:{5}", Path.GetFileName(ProvisionFile), p.ProvisionName, Validity, EffectiveDate.ToString(), ExpirationDate.ToString(), bDistribution ? "DISTRIBUTION" : "DEVELOPMENT");
+				}
+				if (ProvisionFile != null)
+				{
+				    Provisions.Add(new FileReference(ProvisionFile));
+				}
+
+			}
+		    _logger?.LogInformation(" MATCHED-Provision:{0},File:{1},Cert:{2} ", SelectedProvision, SelectedFile, SelectedCert);
+			MatchedProvision = new FileReference(SelectedProvision);
+
+			return Provisions;
+		}
+
+		/// <summary>
+		/// Tries to find a matching certificate on this machine from the the serial number of one of the
+		/// certificates in the mobile provision (the one in the mobileprovision is missing the public/private key pair)
+		/// </summary>
+		public static X509Certificate2? FindCertificate(MobileProvision ProvisionToWorkFrom)
+		{
+			_logger?.LogInformation("  Looking for a certificate that matches the application identifier '{0}'", ProvisionToWorkFrom.ApplicationIdentifier);
+
+			X509Certificate2? Result = null;
+
+			if (Environment.OSVersion.Platform == PlatformID.Unix || Environment.OSVersion.Platform == PlatformID.MacOSX)
+			{
+				// run certtool y to get the currently installed certificates
+				CertToolData = "";
+				Process CertTool = new Process();
+				CertTool.StartInfo.FileName = "/usr/bin/security";
+				CertTool.StartInfo.UseShellExecute = false;
+				CertTool.StartInfo.Arguments = "find-identity -p codesigning -v";
+				CertTool.StartInfo.RedirectStandardOutput = true;
+				CertTool.StartInfo.StandardOutputEncoding = Encoding.UTF8;
+				CertTool.OutputDataReceived += new DataReceivedEventHandler(OutputReceivedCertToolProcessCall);
+				CertTool.Start();
+				CertTool.BeginOutputReadLine();
+				CertTool.WaitForExit();
+
+				_logger?.LogInformation("  Running {0} {1}\n{2}", CertTool.StartInfo.FileName, CertTool.StartInfo.Arguments, CertToolData);
+
+				if (CertTool.ExitCode == 0)
+				{
+					_logger?.LogInformation("  Provisioning profile contains the following certificate hashes:");
+					foreach (X509Certificate2 SourceCert in ProvisionToWorkFrom.DeveloperCertificates)
+					{
+						X509Certificate2? ValidInTimeCert = null;
+						// see if certificate can be found by serial number
+						string CertHash = SourceCert.GetCertHashString();
+
+						_logger?.LogInformation("    {0}", CertHash);
+
+						if (CertToolData.Contains(CertHash))
+						{
+							ValidInTimeCert = SourceCert;
+						}
+
+						if (ValidInTimeCert != null)
+						{
+							// Found a cert in the valid time range, quit now!
+							Result = ValidInTimeCert;
+							_logger?.LogInformation("  Matches!");
+							break;
+						}
+					}
+				}
+			}
+			else
+			{
+				// Open the personal certificate store on this machine
+				X509Store Store = new X509Store();
+				Store.Open(OpenFlags.ReadOnly);
+
+				// Try finding a matching certificate from the serial number (the one in the mobileprovision is missing the public/private key pair)
+				foreach (X509Certificate2 SourceCert in ProvisionToWorkFrom.DeveloperCertificates)
+				{
+					X509Certificate2Collection FoundCerts = Store.Certificates.Find(X509FindType.FindBySerialNumber, SourceCert.SerialNumber, false);
+
+					_logger?.LogInformation("  .. Provision entry SN '{0}' matched {1} installed certificate(s)", SourceCert.SerialNumber, FoundCerts.Count);
+
+					X509Certificate2? ValidInTimeCert = null;
+					foreach (X509Certificate2 TestCert in FoundCerts)
+					{
+						DateTime EffectiveDate = TestCert.NotBefore.ToUniversalTime();
+						DateTime ExpirationDate = TestCert.NotAfter.ToUniversalTime();
+						DateTime Now = DateTime.UtcNow;
+
+						bool bCertTimeIsValid = (EffectiveDate < Now) && (ExpirationDate > Now);
+
+						_logger?.LogInformation("  .. .. Installed certificate '{0}' is {1} (range '{2}' to '{3}')", GetFriendlyNameFromCert(TestCert), bCertTimeIsValid ? "valid (choosing it)" : "EXPIRED", TestCert.GetEffectiveDateString(), TestCert.GetExpirationDateString());
+						if (bCertTimeIsValid)
+						{
+							ValidInTimeCert = TestCert;
+							break;
+						}
+					}
+
+					if (ValidInTimeCert != null)
+					{
+						// Found a cert in the valid time range, quit now!
+						Result = ValidInTimeCert;
+						break;
+					}
+				}
+
+				Store.Close();
+			}
+
+			if (Result == null)
+			{
+				_logger?.LogInformation("  .. Failed to find a valid certificate that was in date");
+			}
+
+			return Result;
+		}
+
+		/// <summary>
+		/// Makes sure the required files for code signing exist and can be found
+		/// </summary>
+		public static bool FindRequiredFiles(out MobileProvision? Provision, out X509Certificate2? Cert, out bool bHasOverridesFile, out bool bNameMatch, bool bNameCheck = true)
+		{
+			Provision = null;
+			Cert = null;
+			bHasOverridesFile = File.Exists(CodeSigningConfig.GetPlistOverrideFilename());
+			bNameMatch = false;
+
+			string? CFBundleIdentifier = CodeSigningConfig.OverrideBundleName;
+
+			if (string.IsNullOrEmpty(CFBundleIdentifier))
+			{
+				// Load Info.plist, which guides nearly everything else
+				string plistFile = CodeSigningConfig.EngineBuildDirectory + "/UnrealGame-Info.plist";
+				if (!string.IsNullOrEmpty(CodeSigningConfig.ProjectFile))
+				{
+					plistFile = Path.GetDirectoryName(CodeSigningConfig.ProjectFile) + "/Intermediate/" + CodeSigningConfig.OSString + "/" + Path.GetFileNameWithoutExtension(CodeSigningConfig.ProjectFile) + "-Info.plist";
+
+					if (!File.Exists(plistFile))
+					{
+						plistFile = CodeSigningConfig.IntermediateDirectory + "/UnrealGame-Info.plist";
+						if (!File.Exists(plistFile))
+						{
+							plistFile = CodeSigningConfig.EngineBuildDirectory + "/UnrealGame-Info.plist";
+						}
+					}
+				}
+				PListHelper Info = new PListHelper();
+				try
+				{
+					string RawInfoPList = File.ReadAllText(plistFile, Encoding.UTF8);
+					Info = new PListHelper(RawInfoPList); ;
+				}
+				catch (Exception ex)
+				{
+					Console.WriteLine(ex.Message);
+				}
+
+				if (Info == null)
+				{
+					return false;
+				}
+
+				// Get the name of the bundle
+				Info.GetString("CFBundleIdentifier", out CFBundleIdentifier);
+				if (CFBundleIdentifier == null)
+				{
+					return false;
+				}
+				else
+				{
+					CFBundleIdentifier = CFBundleIdentifier.Replace("${BUNDLE_IDENTIFIER}", Path.GetFileNameWithoutExtension(CodeSigningConfig.ProjectFile));
+				}
+			}
+
+			// Check for a mobile provision
+			try
+			{
+				string MobileProvisionFilename = MobileProvision.FindCompatibleProvision(CFBundleIdentifier, out bNameMatch);
+				Provision = MobileProvisionParser.ParseFile(MobileProvisionFilename);
+			}
+			catch (Exception)
+			{
+			}
+
+			// if we have a null provision see if we can find a compatible provision without checking for a certificate
+			if (Provision == null)
+			{
+				try
+				{
+					string MobileProvisionFilename = MobileProvision.FindCompatibleProvision(CFBundleIdentifier, out bNameMatch, false);
+					Provision = MobileProvisionParser.ParseFile(MobileProvisionFilename);
+				}
+				catch (Exception)
+				{
+				}
+
+				// if we have a null provision see if we can find a valid provision without checking for name match
+				if (Provision == null && !bNameCheck)
+				{
+					try
+					{
+						string MobileProvisionFilename = MobileProvision.FindCompatibleProvision(CFBundleIdentifier, out bNameMatch, false, false);
+						Provision = MobileProvisionParser.ParseFile(MobileProvisionFilename);
+					}
+					catch (Exception)
+					{
+					}
+				}
+			}
+
+			// Check for a suitable signature to match the mobile provision
+			if (Provision != null)
+			{
+				Cert = AppleCodeSign.FindCertificate(Provision);
+			}
+
+			return true;
+		}
+
+		/// <summary>
+		/// Find the certs and provision
+		/// </summary>
+		public static bool FindCertAndProvision(string BundleID, out FileReference? Provision, out string CertName)
+		{
+			X509Certificate2? Cert;
+			MobileProvision? MobileProvision;
+			MobileProvision.CacheMobileProvisions();
+
+			if (!string.IsNullOrEmpty(BundleID))
+			{
+				CodeSigningConfig.OverrideBundleName = BundleID;
+			}
+
+			if (FindRequiredFiles(out MobileProvision, out Cert, out _, out _))
+			{
+				// print out the provision and cert name
+				Provision = new FileReference(MobileProvision?.FileName ?? "");
+				CertName = GetFriendlyNameFromCert(Cert!);
+				_logger?.LogInformation("CERTIFICATE-{0},PROVISION-{1}", CertName, Provision.FullName);
+				return true;
+			}
+
+			_logger?.LogInformation("No matching Signing Data found!");
+			Provision = null;
+			CertName = "";
+			return false;
+		}
+	}
+
+	/// <summary>
+	/// The XML doc
+	/// </summary>
+	public class PListHelper
+	{
+		/// <summary>
+		/// Helper for info plist
+		/// </summary>
+		public XmlDocument Doc;
+
+		bool bReadOnly = false;
+		
+		/// <summary>
+		/// Set read only
+		/// </summary>
+		public void SetReadOnly(bool bNowReadOnly)
+		{
+			bReadOnly = bNowReadOnly;
+		}
+		/// <summary>
+		/// constructor
+		/// </summary>
+		public PListHelper(string Source)
+		{
+			Doc = new XmlDocument();
+			Doc.XmlResolver = null;
+			Doc.LoadXml(Source);
+		}
+
+		/// <summary>
+		/// Create plist from file
+		/// </summary>
+		public static PListHelper CreateFromFile(string Filename)
+		{
+			byte[] RawPList = File.ReadAllBytes(Filename);
+			return new PListHelper(Encoding.UTF8.GetString(RawPList));
+		}
+		
+		/// <summary>
+		/// save the plist to file
+		/// </summary>
+		public void SaveToFile(string Filename)
+		{
+			File.WriteAllText(Filename, SaveToString(), Encoding.UTF8);
+		}
+
+		/// <summary>
+		/// Constructor
+		/// </summary>
+		public PListHelper()
+		{
+			string EmptyFileText =
+				"<?xml version=\"1.0\" encoding=\"UTF-8\"?>\n" +
+				"<!DOCTYPE plist PUBLIC \"-//Apple//DTD PLIST 1.0//EN\" \"http://www.apple.com/DTDs/PropertyList-1.0.dtd\">\n" +
+				"<plist version=\"1.0\">\n" +
+				"<dict>\n" +
+				"</dict>\n" +
+				"</plist>\n";
+
+			Doc = new XmlDocument();
+			Doc.XmlResolver = null;
+			Doc.LoadXml(EmptyFileText);
+		}
+
+		/// <summary>
+		/// Convert value to plist format
+		/// </summary>
+		public XmlElement? ConvertValueToPListFormat(object Value)
+		{
+			XmlElement ValueElement;
+
+			if (Value is string str)
+			{
+				ValueElement = Doc.CreateElement("string");
+				ValueElement.InnerText = str;
+			}
+			else if (Value is Dictionary<string, object> dict)
+			{
+				ValueElement = Doc.CreateElement("dict");
+				foreach (var KVP in dict)
+				{
+					AddKeyValuePair(ValueElement, KVP.Key, KVP.Value);
+				}
+			}
+			else if (Value is PListHelper)
+			{
+    			PListHelper? PList = Value as PListHelper;
+
+    			ValueElement = Doc.CreateElement("dict");
+
+    			XmlNode? docElement = PList?.Doc?.DocumentElement;
+    			XmlNode? SourceDictionaryNode = docElement?.SelectSingleNode("/plist/dict");
+   				if (SourceDictionaryNode == null)
+    			{
+			        throw new InvalidDataException("The PListHelper object contains a document without a dictionary.");
+			    }
+
+			    foreach (XmlNode TheirChild in SourceDictionaryNode)
+			    {
+			        ValueElement.AppendChild(Doc.ImportNode(TheirChild, true));
+			    }
+			}
+			else if (Value is Array)
+			{
+				if (Value is byte[])
+				{
+					ValueElement = Doc.CreateElement("data");
+					ValueElement.InnerText = (Value is byte[] byteArray) ? Convert.ToBase64String(byteArray) : string.Empty;
+				}
+				else
+				{
+					ValueElement = Doc.CreateElement("array");
+					if (Value is Array array)
+					{
+					    foreach (var A in array)
+					    {
+				       		var childNode = ConvertValueToPListFormat(A);
+				        	if (childNode != null)
+	        				{
+					            ValueElement.AppendChild(childNode);
+ 							}
+					    }
+					}
+				}
+			}
+			else if (Value is IList)
+			{
+				ValueElement = Doc.CreateElement("array");
+				if (Value is IList list  && ValueElement != null)
+				{
+					foreach (var A in list)
+			        {
+						var child = ConvertValueToPListFormat(A);
+          				if (child != null)
+		            	{
+        		        	ValueElement.AppendChild(child);
+	            		}
+					}
+				}
+			}
+			else if (Value is bool)
+			{
+				ValueElement = Doc.CreateElement(((bool)Value) ? "true" : "false");
+			}
+			else if (Value is double)
+			{
+				ValueElement = Doc.CreateElement("real");
+				ValueElement.InnerText = ((double)Value).ToString();
+			}
+			else if (Value is int)
+			{
+				ValueElement = Doc.CreateElement("integer");
+				ValueElement.InnerText = ((int)Value).ToString();
+			}
+			else
+			{
+				throw new InvalidDataException(String.Format("Object '{0}' is in an unknown type that cannot be converted to PList format", Value));
+			}
+			return ValueElement;
+		}
+
+		/// <summary>
+		/// Add key value pair
+		/// </summary>
+		public void AddKeyValuePair(XmlNode DictRoot, string KeyName, object Value)
+		{
+			if (bReadOnly)
+			{
+				throw new AccessViolationException("PList has been set to read only and may not be modified");
+			}
+
+			XmlElement KeyElement = Doc.CreateElement("key");
+			KeyElement.InnerText = KeyName;
+
+			DictRoot.AppendChild(KeyElement);
+			var newChild = ConvertValueToPListFormat(Value);
+			if (newChild != null)
+			{
+			    DictRoot.AppendChild(newChild);
+			}
+		}
+
+		/// <summary>
+		/// Add key value pair
+		/// </summary>
+		public void AddKeyValuePair(string KeyName, object? Value)
+		{
+			XmlNode? dictRoot = null;
+
+			if (Doc.DocumentElement != null)
+			{
+			    dictRoot = Doc.DocumentElement.SelectSingleNode("/plist/dict");
+			}
+
+			if (dictRoot != null)
+   			{
+				XmlNode? childNode = Doc.CreateElement("null");
+		        if (Value != null)
+        		{
+		            childNode = ConvertValueToPListFormat(Value);
+        		}
+		        else
+        		{
+        		    childNode = Doc.CreateElement("null");
+		        }
+				if (childNode != null)
+				{
+        			dictRoot.AppendChild(childNode);
+				}
+		        if (Value != null)
+        		{
+            		AddKeyValuePair(dictRoot, KeyName, Value);
+		        }
+    }
+		}
+
+		/// <summary>
+		/// Clones a dictionary from an existing .plist into a new one.  Root should point to the dict key in the source plist.
+		/// </summary>
+		public static PListHelper CloneDictionaryRootedAt(XmlNode Root)
+		{
+			// Create a new empty dictionary
+			PListHelper Result = new PListHelper();
+
+			// Copy all of the entries in the source dictionary into the new one
+			XmlNode? NewDictRoot = Result.Doc.DocumentElement?.SelectSingleNode("/plist/dict");
+			if (NewDictRoot != null)
+		    {
+				foreach (XmlNode TheirChild in Root)
+				{
+        			XmlNode? importedNode = Result.Doc.ImportNode(TheirChild, true);
+			        if (importedNode != null)
+        			{
+			            NewDictRoot.AppendChild(importedNode);
+			        }
+    			}
+			}
+			return Result;
+		}
+
+		/// <summary>
+		/// Get String
+		/// </summary>
+		public bool GetString(string Key, out string Value)
+		{
+    		string PathToValue = String.Format("/plist/dict/key[.='{0}']/following-sibling::string[1]", Key);
+
+		    XmlNode? ValueNode = Doc.DocumentElement?.SelectSingleNode(PathToValue);
+    		if (ValueNode == null)
+		    {
+        		Value = "";
+		        return false;
+		    }
+
+		    Value = ValueNode.InnerText ?? "";
+    		return true;
+		}
+
+		/// <summary>
+		/// Get date
+		/// </summary>
+		public bool GetDate(string Key, out string Value)
+		{
+			string PathToValue = String.Format("/plist/dict/key[.='{0}']/following-sibling::date[1]", Key);
+
+			XmlNode? ValueNode = Doc.DocumentElement?.SelectSingleNode(PathToValue);
+			if (ValueNode == null)
+			{
+				Value = "";
+				return false;
+			}
+
+			Value = ValueNode.InnerText;
+			return true;
+		}
+
+		/// <summary>
+		/// Get bool
+		/// </summary>
+		public bool GetBool(string Key)
+		{
+			string PathToValue = String.Format("/plist/dict/key[.='{0}']/following-sibling::node()", Key);
+
+			XmlNode? ValueNode = Doc.DocumentElement?.SelectSingleNode(PathToValue);
+			if (ValueNode == null)
+			{
+				return false;
+			}
+
+			return ValueNode.Name == "true";
+		}
+
+		/// <summary>
+		/// Process one node event
+		/// </summary>
+		public delegate void ProcessOneNodeEvent(XmlNode ValueNode);
+
+		/// <summary>
+		/// Process the value for the key
+		/// </summary>
+		public void ProcessValueForKey(string Key, string ExpectedValueType, ProcessOneNodeEvent ValueHandler)
+		{
+			string PathToValue = String.Format("/plist/dict/key[.='{0}']/following-sibling::{1}[1]", Key, ExpectedValueType);
+
+			XmlNode? ValueNode = Doc.DocumentElement?.SelectSingleNode(PathToValue);
+			if (ValueNode != null)
+			{
+				ValueHandler(ValueNode);
+			}
+		}
+
+		/// <summary>
+		/// Merge two plists together.  Whenever both have the same key, the value in the dominant source list wins.
+		/// </summary>
+		public void MergePlistIn(string dominantPlist)
+		{
+    		if (bReadOnly)
+    		{
+     	  		throw new AccessViolationException("PList has been set to read only and may not be modified");
+	    	}
+
+    		XmlDocument dominant = new XmlDocument();
+    		dominant.XmlResolver = null;
+    		dominant.LoadXml(dominantPlist);
+
+    		XmlNode? dictionaryNode = Doc.DocumentElement?.SelectSingleNode("/plist/dict");
+    		if (dictionaryNode == null)
+    		{
+    		    throw new InvalidOperationException("Invalid PList format: missing dictionary node");
+    		}
+
+   			// Merge any key-value pairs in the strong .plist into the weak .plist
+    		XmlNodeList? strongKeys = dominant.DocumentElement?.SelectNodes("/plist/dict/key");
+    		if (strongKeys != null)
+    		{
+        		foreach (XmlNode strongKeyNode in strongKeys)
+        		{
+            		string strongKey = strongKeyNode.InnerText;
+
+            		XmlNode? weakNode = Doc.DocumentElement?.SelectSingleNode($"/plist/dict/key[.='{strongKey}']");
+            		if (weakNode == null)
+            		{
+             		   // Doesn't exist in dominant plist, inject key-value pair
+					   XmlNode? valueNode = strongKeyNode.NextSibling;
+					   if (valueNode != null)
+					   {
+							dictionaryNode.AppendChild(Doc.ImportNode(strongKeyNode, true));
+                    		dictionaryNode.AppendChild(Doc.ImportNode(valueNode, true));
+					   }
+		            }
+        	    else
+            	{
+            	    // Remove the existing value node from the weak file
+            	    XmlNode? existingValueNode = weakNode.NextSibling;
+					if (existingValueNode != null && existingValueNode.Name == "string")
+					{
+						weakNode.ParentNode?.RemoveChild(existingValueNode);
+
+                    	// Insert a clone of the dominant value node
+                    	XmlNode? dominantValueNode = strongKeyNode.NextSibling;
+                    	if (dominantValueNode != null && dominantValueNode.Name == "string")
+                    	{
+                        	weakNode.ParentNode?.InsertAfter(Doc.ImportNode(dominantValueNode, true), weakNode);
+                    	}
+                	}
+            	}
+        	}
+    	}
+	}
+
+
+		/// <summary>
+		/// Returns each of the entries in the value tag of type array for a given key
+		/// If the key is missing, an empty array is returned.
+		/// Only entries of a given type within the array are returned.
+		/// </summary>
+		public List<string> GetArray(string Key, string EntryType)
+		{
+			List<string> Result = new List<string>();
+
+			ProcessValueForKey(Key, "array",
+				delegate (XmlNode ValueNode)
+				{
+					foreach (XmlNode ChildNode in ValueNode.ChildNodes)
+					{
+						if (EntryType == ChildNode.Name)
+						{
+							string Value = ChildNode.InnerText;
+							Result.Add(Value);
+						}
+					}
+				});
+
+			return Result;
+		}
+
+		/// <summary>
+		/// Returns true if the key exists (and has a value) and false otherwise
+		/// </summary>
+		public bool HasKey(string KeyName)
+		{
+			string PathToKey = String.Format("/plist/dict/key[.='{0}']", KeyName);
+
+			XmlNode? KeyNode = Doc.DocumentElement?.SelectSingleNode(PathToKey);
+			return (KeyNode != null);
+		}
+
+		/// <summary>
+		/// Remove the key value
+		/// </summary>
+		public void RemoveKeyValue(string KeyName)
+		{
+			if (bReadOnly)
+			{
+				throw new AccessViolationException("PList has been set to read only and may not be modified");
+			}
+
+			XmlNode? DictionaryNode = Doc.DocumentElement?.SelectSingleNode("/plist/dict");
+
+			string PathToKey = String.Format("/plist/dict/key[.='{0}']", KeyName);
+			XmlNode? KeyNode = Doc.DocumentElement?.SelectSingleNode(PathToKey);
+			if (KeyNode != null && KeyNode.ParentNode != null)
+			{
+				XmlNode? ValueNode = KeyNode.NextSibling;
+				//remove value
+				if (ValueNode != null)
+				{
+					ValueNode.RemoveAll();
+					ValueNode?.ParentNode?.RemoveChild(ValueNode);
+				}
+				//remove key
+				KeyNode.RemoveAll();
+				KeyNode.ParentNode.RemoveChild(KeyNode);
+			}
+		}
+
+		/// <summary>
+		/// Set the value for the key
+		/// </summary>
+		public void SetValueForKey(string KeyName, object Value)
+		{
+			if (bReadOnly)
+			{
+				throw new AccessViolationException("PList has been set to read only and may not be modified");
+			}
+
+			XmlNode? DictionaryNode = Doc.DocumentElement?.SelectSingleNode("/plist/dict");
+
+			string PathToKey = String.Format("/plist/dict/key[.='{0}']", KeyName);
+			XmlNode? KeyNode = Doc.DocumentElement?.SelectSingleNode(PathToKey);
+
+			XmlNode? ValueNode = null;
+			if (KeyNode != null)
+			{
+				ValueNode = KeyNode.NextSibling;
+			}
+
+			if (ValueNode == null)
+			{
+			    KeyNode = Doc.CreateNode(XmlNodeType.Element, "key", null);
+			    KeyNode.InnerText = KeyName;
+
+			    ValueNode = ConvertValueToPListFormat(Value);
+
+			    if (DictionaryNode != null && KeyNode != null && ValueNode != null)
+				{
+				    DictionaryNode.AppendChild(KeyNode);
+				    DictionaryNode.AppendChild(ValueNode);
+				}
+			}
+			else
+			{
+			    // Remove the existing value and create a new one
+			    if (ValueNode.ParentNode != null)
+			    {
+        			ValueNode.ParentNode.RemoveChild(ValueNode);
+			    }
+			    ValueNode = ConvertValueToPListFormat(Value);
+
+			    // Insert the value after the key
+				KeyNode?.ParentNode?.InsertAfter(ValueNode!, KeyNode);
+
+			}
+
+		}
+
+		/// <summary>
+		/// Set the string
+		/// </summary>
+		public void SetString(string Key, string Value)
+		{
+			SetValueForKey(Key, Value);
+		}
+
+		/// <summary>
+		/// Save the string
+		/// </summary>
+		public string SaveToString()
+		{
+			// Convert the XML back to text in the same style as the original .plist
+			StringBuilder TextOut = new StringBuilder();
+
+			// Work around the fact it outputs the wrong encoding by default (and set some other settings to get something similar to the input file)
+			TextOut.Append("<?xml version=\"1.0\" encoding=\"UTF-8\"?>\n");
+			XmlWriterSettings Settings = new XmlWriterSettings();
+			Settings.Indent = true;
+			Settings.IndentChars = "\t";
+			Settings.NewLineChars = "\n";
+			Settings.NewLineHandling = NewLineHandling.Replace;
+			Settings.OmitXmlDeclaration = true;
+			Settings.Encoding = new UTF8Encoding(false);
+
+			// Work around the fact that it embeds an empty declaration list to the document type which codesign dislikes...
+			// Replacing InternalSubset with null if it's empty.  The property is readonly, so we have to reconstruct it entirely
+			Doc.ReplaceChild(Doc.CreateDocumentType(
+				Doc.DocumentType!.Name,
+				Doc.DocumentType!.PublicId,
+				Doc.DocumentType!.SystemId,
+				String.IsNullOrEmpty(Doc.DocumentType!.InternalSubset) ? null : Doc.DocumentType!.InternalSubset),
+				Doc.DocumentType!);
+
+			XmlWriter Writer = XmlWriter.Create(TextOut, Settings);
+
+			Doc.Save(Writer);
+
+			// Remove the space from any standalone XML elements because the iOS parser does not handle them
+			return Regex.Replace(TextOut.ToString(), @"<(?<tag>\S+) />", "<${tag}/>");
+		}
+	}
+
+	internal class FileOperations
+	{
+		/// <summary>
+		/// Find a prefixable file (checks first with the prefix and then without it)
+		/// </summary>
+		/// <param name="BasePath">The path that the file should be in</param>
+		/// <param name="FileName">The filename to check for (with and without SigningPrefix prepended to it)</param>
+		/// <returns>The path to the most desirable file (may still not exist)</returns>
+		static public string FindPrefixedFile(string BasePath, string FileName)
+		{
+			return FindPrefixedFileOrDirectory(BasePath, FileName, false);
+		}
+
+		/// <summary>
+		/// Finds the first file in the directory with a given extension (if it exists).  This method is to provide some
+		/// idiot-proofing (for example, .mobileprovision files, which will be expected to be saved from the website as
+		/// GameName.mobileprovision or Distro_GameName.mobileprovision, but falling back to using the first one we find
+		/// if the correctly named one is not found helps the user if they mess up during the process.
+		/// </summary>
+		static public string FindAnyFileWithExtension(string BasePath, string FileExt)
+		{
+			string[] FileList = Directory.GetFiles(BasePath);
+			foreach (string Filename in FileList)
+			{
+				if (Path.GetExtension(Filename).Equals(FileExt, StringComparison.InvariantCultureIgnoreCase))
+				{
+					return Filename;
+				}
+			}
+			return "";
+		}
+
+		/// <summary>
+		/// Find a prefixable file or directory (checks first with the prefix and then without it)
+		/// </summary>
+		/// <param name="BasePath">The path that the file should be in</param>
+		/// <param name="Name">The name to check for (with and without SigningPrefix prepended to it)</param>
+		/// <param name="bIsDirectory">Is the desired name a directory?</param>
+		/// 
+		/// <returns>The path to the most desirable file (may still not exist)</returns>
+		static public string FindPrefixedFileOrDirectory(string BasePath, string Name, bool bIsDirectory)
+		{
+			string PrefixedPath = Path.Combine(BasePath, CodeSigningConfig.SigningPrefix + Name);
+
+			if ((bIsDirectory && Directory.Exists(PrefixedPath)) || (!bIsDirectory && File.Exists(PrefixedPath)))
+			{
+				return PrefixedPath;
+			}
+			else
+			{
+				return Path.Combine(BasePath, Name);
+			}
+		}
+	}
+
+	internal class Utilities
+	{
+		static ILogger? _logger;
+
+		/// <summary>
+		/// The constructor
+		/// </summary>
+		public Utilities()
+		{
+			ILoggerFactory loggerFactory = LoggerFactory.Create(builder =>
+		    {
+        		builder.AddConsole();
+		    });
+			_logger = loggerFactory.CreateLogger("CodeSigningUtilities");
+		}
+
+		/// <summary>
+		/// Get the string from the plist
+		/// </summary>
+		static public string GetStringFromPList(string KeyName)
+		{
+			// Open the .plist and read out the specified key
+			string PListAsString;
+			if (!Utilities.GetSourcePList(out PListAsString))
+			{
+				return "(unknown)";
+			}
+
+			PListHelper Helper = new PListHelper(PListAsString);
+
+			string Result;
+			if (Helper.GetString(KeyName, out Result))
+			{
+				return Result;
+			}
+			else
+			{
+				return "(unknown)";
+			}
+		}
+
+		/// <summary>
+		/// Get the precompile source plist filename
+		/// </summary>
+		static public string GetPrecompileSourcePListFilename()
+		{
+			// check for one in the project directory
+			string SourceName = FileOperations.FindPrefixedFile(CodeSigningConfig.IntermediateDirectory, CodeSigningConfig.Program_GameName + "-Info.plist");
+
+			if (!File.Exists(SourceName))
+			{
+				// Check for a premade one
+				SourceName = FileOperations.FindPrefixedFile(CodeSigningConfig.BuildDirectory, CodeSigningConfig.Program_GameName + "-Info.plist");
+
+				if (!File.Exists(SourceName))
+				{
+					// fallback to the shared one
+					SourceName = FileOperations.FindPrefixedFile(CodeSigningConfig.EngineBuildDirectory, "UnrealGame-Info.plist");
+
+					if (!File.Exists(SourceName))
+					{
+						 if (_logger != null)
+						 {
+							_logger.LogInformation("Failed to find " + CodeSigningConfig.Program_GameName + "-Info.plist. Please create new .plist or copy a base .plist from a provided game sample.");
+						 }
+					}
+				}
+			}
+
+			return SourceName;
+		}
+
+		/**
+		 * Handle grabbing the initial plist
+		 */
+		static public bool GetSourcePList(out string PListSource)
+		{
+			// Check for a premade one
+			string SourceName = GetPrecompileSourcePListFilename();
+
+			if (File.Exists(SourceName))
+			{
+				_logger?.LogInformation(" ... reading source .plist: " + SourceName);
+				PListSource = File.ReadAllText(SourceName);
+				return true;
+			}
+			else
+			{
+				_logger?.LogError(" ... failed to locate the source .plist file");
+				//Program.ReturnCode = (int)ErrorCodes.Error_KeyNotFoundInPList;
+				PListSource = "";
+				return false;
+			}
+		}
+	}
+
+	/// <summary>
+	/// The mobile provision
+	/// </summary>
+	public class MobileProvision
+	{
+		/// <summary>
+		/// Tag
+		/// </summary>
+		public object? Tag;
+
+		/// <summary>
+		/// ApplicationIdentifierPrefix
+		/// </summary>
+		public string ApplicationIdentifierPrefix = "";
+
+		/// <summary>
+		/// ApplicationIdentifier
+		/// </summary>
+		public string ApplicationIdentifier = "";
+
+		/// <summary>
+		/// DeveloperCertificates
+		/// </summary>
+		public List<X509Certificate2> DeveloperCertificates = new List<X509Certificate2>();
+
+		/// <summary>
+		/// ProvisionedDeviceIDs
+		/// </summary>
+		public List<string>? ProvisionedDeviceIDs;
+
+		/// <summary>
+		/// ProvisionName
+		/// </summary>
+		public string? ProvisionName;
+
+		/// <summary>
+		/// bDebugn
+		/// </summary>
+		public bool bDebug;
+
+		/// <summary>
+		/// Data
+		/// </summary>
+		public PListHelper? Data;
+
+		/// <summary>
+		/// CreationDate
+		/// </summary>
+		public DateTime CreationDate;
+
+		/// <summary>
+		/// ExpirationDate
+		/// </summary>
+		public DateTime ExpirationDate;
+
+		/// <summary>
+		/// FileName
+		/// </summary>
+		public string? FileName;
+
+		/// <summary>
+		/// UUID
+		/// </summary>
+		public string? UUID;
+
+		/// <summary>
+		/// The Platform
+		/// </summary>
+		public string? Platform;
+
+		/// <summary>
+		/// The logger
+		/// </summary>
+		static ILogger? _logger;
+
+		/// <summary>
+		/// The constructor
+		/// </summary>
+		public MobileProvision()
+		{
+			ILoggerFactory loggerFactory = LoggerFactory.Create(builder =>
+		    {
+        		builder.AddConsole();
+		    });
+			_logger = loggerFactory.CreateLogger("MobileProvision");
+		}
+
+		/// <summary>
+		/// Cache the mobile provisions
+		/// </summary>
+		public static void CacheMobileProvisions()
+		{
+			_logger?.LogInformation("Caching provisions");
+			if (!Directory.Exists(CodeSigningConfig.ProvisionDirectory))
+			{
+				_logger?.LogInformation("Provision Folder {0} doesn't exist, creating..", CodeSigningConfig.ProvisionDirectory);
+				Directory.CreateDirectory(CodeSigningConfig.ProvisionDirectory);
+			}
+
+			List<string> ProvisionCopySrcDirectories = new List<string>();
+
+			// Paths for provisions under game directory 
+			if (!String.IsNullOrEmpty(CodeSigningConfig.ProjectFile))
+			{
+				ProvisionCopySrcDirectories.Add(Path.GetDirectoryName(CodeSigningConfig.ProjectFile) + "/Build/" + CodeSigningConfig.OSString + "/");
+				ProvisionCopySrcDirectories.Add(Path.GetDirectoryName(CodeSigningConfig.ProjectFile) + "/Restricted/NoRedist/Build/" + CodeSigningConfig.OSString + "/");
+				ProvisionCopySrcDirectories.Add(Path.GetDirectoryName(CodeSigningConfig.ProjectFile) + "/Restricted/NotForLicensees/Build/" + CodeSigningConfig.OSString + "/");
+			}
+
+			// Paths for provisions under the Engine directory
+			string? OverrideProvisionDirectory = Environment.GetEnvironmentVariable("ProvisionDirectory");
+			if (!String.IsNullOrEmpty(OverrideProvisionDirectory))
+			{
+				ProvisionCopySrcDirectories.Add(OverrideProvisionDirectory);
+			}
+			else
+			{
+				ProvisionCopySrcDirectories.Add(CodeSigningConfig.EngineBuildDirectory);
+				ProvisionCopySrcDirectories.Add(CodeSigningConfig.EngineDirectory + "/Restricted/NoRedist/Build/" + CodeSigningConfig.OSString + "/");
+				ProvisionCopySrcDirectories.Add(CodeSigningConfig.EngineDirectory + "/Restricted/NotForLicensees/Build/" + CodeSigningConfig.OSString + "/");
+			}
+
+			// Copy all provisions from the above paths to the library
+			foreach (string ProvisionCopySrcDirectory in ProvisionCopySrcDirectories)
+			{
+				if (Directory.Exists(ProvisionCopySrcDirectory))
+				{
+					_logger?.LogInformation("Finding provisions in {0}", ProvisionCopySrcDirectory);
+					foreach (string Provision in Directory.EnumerateFiles(ProvisionCopySrcDirectory, "*.mobileprovision", SearchOption.AllDirectories))
+					{
+						string TargetFile = CodeSigningConfig.ProvisionDirectory + Path.GetFileName(Provision);
+						if (!File.Exists(TargetFile) || File.GetLastWriteTime(TargetFile) < File.GetLastWriteTime(Provision))
+						{
+							FileInfo DestFileInfo;
+							if (File.Exists(TargetFile))
+							{
+								DestFileInfo = new FileInfo(TargetFile);
+								DestFileInfo.Attributes = DestFileInfo.Attributes & ~FileAttributes.ReadOnly;
+							}
+							_logger?.LogInformation("  Copying {0} -> {1}", Provision, TargetFile);
+							File.Copy(Provision, TargetFile, true);
+							DestFileInfo = new FileInfo(TargetFile);
+							DestFileInfo.Attributes = DestFileInfo.Attributes & ~FileAttributes.ReadOnly;
+							if (!File.Exists(TargetFile))
+							{
+								_logger?.LogInformation("ERROR: Failed to copy {0} -> {1}", Provision, TargetFile);
+							}
+						}
+						else
+						{
+							_logger?.LogInformation("  Not copying {0} as {1} already exists and is not older", Provision, TargetFile);
+						}
+					}
+				}
+			}
+		}
+
+		/// <summary>
+		/// Clean the mobile provisions
+		/// </summary>
+		public static void CleanMobileProvisions()
+		{
+			if (!Directory.Exists(CodeSigningConfig.ProvisionDirectory))
+			{
+				_logger?.LogInformation("Provision Folder {0} doesn't exist, nothing to do.", CodeSigningConfig.ProvisionDirectory);
+			}
+			else
+			{
+				_logger?.LogInformation("Cleaning out contents of  Provision Folder {0}", CodeSigningConfig.ProvisionDirectory);
+				foreach (string Provision in Directory.GetFiles(CodeSigningConfig.ProvisionDirectory))
+				{
+					File.Delete(Provision);
+				}
+			}
+		}
+
+		/// <summary>
+		/// Fond compatible provision
+		/// </summary>
+		public static string FindCompatibleProvision(string CFBundleIdentifier, out bool bNameMatch, bool bCheckCert = true, bool bCheckIdentifier = true, bool bCheckDistro = true)
+		{
+			bNameMatch = false;
+
+			// remap the gamename if necessary
+			string GameName = CodeSigningConfig.Program_GameName;
+			if (GameName == "UnrealGame")
+			{
+				if (CodeSigningConfig.ProjectFile.Length > 0)
+				{
+					GameName = Path.GetFileNameWithoutExtension(CodeSigningConfig.ProjectFile);
+				}
+			}
+
+			// ensure the provision directory exists
+			if (!Directory.Exists(CodeSigningConfig.ProvisionDirectory))
+			{
+				Directory.CreateDirectory(CodeSigningConfig.ProvisionDirectory);
+			}
+
+			if (CodeSigningConfig.bProvision)
+			{
+				if (File.Exists(CodeSigningConfig.ProvisionDirectory + "/" + CodeSigningConfig.Provision))
+				{
+					return CodeSigningConfig.ProvisionDirectory + "/" + CodeSigningConfig.Provision;
+				}
+			}
+
+			#region remove after we provide an install mechanism
+			MobileProvision.CacheMobileProvisions();
+			#endregion
+
+			// cache the provision library
+			Dictionary<string, MobileProvision> ProvisionLibrary = new Dictionary<string, MobileProvision>();
+			foreach (string Provision in Directory.EnumerateFiles(CodeSigningConfig.ProvisionDirectory, "*.mobileprovision"))
+			{
+   				MobileProvision? p = MobileProvisionParser.ParseFile(Provision);
+    			if (p != null)
+   				{
+        			ProvisionLibrary.Add(Provision, p);
+        			if (p.FileName != null && p.UUID != null && p.FileName.Contains(p.UUID) && !File.Exists(Path.Combine(CodeSigningConfig.ProvisionDirectory, "UE4_" + p.UUID + ".mobileprovision")))
+        			{
+            			File.Copy(Provision, Path.Combine(CodeSigningConfig.ProvisionDirectory, "UE4_" + p.UUID + ".mobileprovision"));
+            			p = MobileProvisionParser.ParseFile(Path.Combine(CodeSigningConfig.ProvisionDirectory, "UE4_" + p.UUID + ".mobileprovision"));
+            			if (p != null)
+            			{
+                			ProvisionLibrary.Add(Path.Combine(CodeSigningConfig.ProvisionDirectory, "UE4_" + p.UUID + ".mobileprovision"), p);
+            			}
+        			}
+    			}
+			}
+
+			_logger?.LogInformation("Searching for mobile provisions that match the game '{0}' (distribution: {3}) with CFBundleIdentifier='{1}' in '{2}'", GameName, CFBundleIdentifier, CodeSigningConfig.ProvisionDirectory, CodeSigningConfig.bForDistribution);
+
+			// first sort all profiles so we look at newer ones first.
+			IEnumerable<string> ProfileKeys = ProvisionLibrary.Select(KV => KV.Key)
+				.OrderByDescending(K => ProvisionLibrary[K].CreationDate)
+				.ToArray();
+
+			// check the cache for a provision matching the app id (com.company.Game)
+			// First checking for a contains match and then for a wildcard match
+			for (int Phase = -1; Phase < 3; ++Phase)
+			{
+				if (Phase == -1 && string.IsNullOrEmpty(CodeSigningConfig.ProvisionUUID))
+				{
+					continue;
+				}
+				foreach (string Key in ProfileKeys)
+				{
+					string DebugName = Path.GetFileName(Key);
+					MobileProvision TestProvision = ProvisionLibrary[Key];
+
+					// make sure the file is not managed by Xcode
+					if (TestProvision.FileName != null && TestProvision.UUID != null &&
+					Path.GetFileName(TestProvision.FileName).ToLower().Equals(TestProvision.UUID.ToLower() + ".mobileprovision"))
+					{
+						continue;
+					}
+
+					_logger?.LogInformation(" Phase {0} considering provision '{1}' named '{2}'", Phase, DebugName, TestProvision.ProvisionName);
+
+					if (TestProvision.ProvisionName == "iOS Team Provisioning Profile: " + CFBundleIdentifier)
+					{
+						_logger?.LogInformation("  Failing as provisioning is automatic");
+						continue;
+					}
+
+					// check to see if the platform is the same as what we are looking for
+					if (!string.IsNullOrEmpty(TestProvision.Platform) && TestProvision.Platform != CodeSigningConfig.OSString && !string.IsNullOrEmpty(CodeSigningConfig.OSString))
+					{
+						_logger?.LogInformation("  Failing platform {0} Config: {1}", TestProvision.Platform, CodeSigningConfig.OSString);
+						continue;
+					}
+
+					// Validate the name
+					bool bPassesNameCheck = false;
+					if (Phase == -1)
+					{
+						bPassesNameCheck = TestProvision.UUID == CodeSigningConfig.ProvisionUUID;
+						bNameMatch = bPassesNameCheck;
+					}
+					else if (Phase == 0)
+					{
+						bPassesNameCheck = TestProvision.ApplicationIdentifier.Substring(TestProvision.ApplicationIdentifierPrefix.Length + 1) == CFBundleIdentifier;
+						bNameMatch = bPassesNameCheck;
+					}
+					else if (Phase == 1)
+					{
+						if (TestProvision.ApplicationIdentifier.Contains("*"))
+						{
+							string CompanyName = TestProvision.ApplicationIdentifier.Substring(TestProvision.ApplicationIdentifierPrefix.Length + 1);
+							if (CompanyName != "*")
+							{
+								CompanyName = CompanyName.Substring(0, CompanyName.LastIndexOf("."));
+								bPassesNameCheck = CFBundleIdentifier.StartsWith(CompanyName);
+							}
+						}
+					}
+					else
+					{
+						if (TestProvision.ApplicationIdentifier.Contains("*"))
+						{
+							string CompanyName = TestProvision.ApplicationIdentifier.Substring(TestProvision.ApplicationIdentifierPrefix.Length + 1);
+							bPassesNameCheck = CompanyName == "*";
+						}
+					}
+					if (!bPassesNameCheck && bCheckIdentifier)
+					{
+						_logger?.LogInformation("  .. Failed phase {0} name check (provision app ID was {1})", Phase, TestProvision.ApplicationIdentifier);
+						continue;
+					}
+
+					if (CodeSigningConfig.bForDistribution)
+					{
+						// Check to see if this is a distribution provision. get-task-allow must be false for distro profiles.
+						// TestProvision.ProvisionedDeviceIDs.Count==0 is not a valid check as ad-hoc distro profiles do list devices.
+						bool bDistroProv = !TestProvision.bDebug;
+						if (TestProvision != null && !bDistroProv)
+						{
+							_logger?.LogInformation("  .. Failed distribution check (mode={0}, get-task-allow={1}, #devices={2})", CodeSigningConfig.bForDistribution, TestProvision.bDebug, TestProvision.ProvisionedDeviceIDs?.Count);
+							continue;
+						}
+					}
+					else
+					{
+					    if (bCheckDistro)
+    					{
+      						bool bPassesDebugCheck = TestProvision?.bDebug ?? false;
+        					if (!bPassesDebugCheck)
+        					{
+            					_logger?.LogInformation(" .. Failed debugging check (mode={0}, get-task-allow={1}, #devices={2})", CodeSigningConfig.bForDistribution, TestProvision?.bDebug ?? false, TestProvision?.ProvisionedDeviceIDs?.Count ?? 0);
+								continue;
+					        }
+    					}
+   						else if (!TestProvision?.bDebug ?? true)
+					    {
+					        CodeSigningConfig.bForceStripSymbols = true;
+					    }
+					}
+
+
+					DateTime CurrentUTCTime = DateTime.UtcNow;
+					bool bPassesDateCheck = (CurrentUTCTime >= TestProvision?.CreationDate) && (CurrentUTCTime < TestProvision?.ExpirationDate);
+					if (!bPassesDateCheck)
+					{
+					    _logger?.LogInformation("  .. Failed time period check (valid from {0} to {1}, but UTC time is now {2})",
+				        TestProvision?.CreationDate, TestProvision?.ExpirationDate, CurrentUTCTime);
+					    continue;
+					}
+
+					// check to see if we have a certificate for this provision
+					bool bPassesHasMatchingCertCheck = false;
+					if (bCheckCert)
+					{
+
+						X509Certificate2? Cert;
+						try
+						{
+						    if (TestProvision == null)
+							{
+       							throw new ArgumentNullException(nameof(TestProvision));
+						    }
+    						Cert = AppleCodeSign.FindCertificate(TestProvision);
+						}
+						catch (Exception ex)
+						{
+    						// handle the exception
+    						Console.WriteLine($"Error finding certificate: {ex.Message}");
+    						return "";
+						}
+						bPassesHasMatchingCertCheck = (Cert != null);
+						if (bPassesHasMatchingCertCheck && CodeSigningConfig.bCert)
+						{
+							bPassesHasMatchingCertCheck &= (AppleCodeSign.GetFriendlyNameFromCert(Cert!) == CodeSigningConfig.Certificate);
+						}
+					}
+					else
+					{
+						bPassesHasMatchingCertCheck = true;
+					}
+
+					if (!bPassesHasMatchingCertCheck)
+					{
+						_logger?.LogInformation("  .. Failed to find a matching certificate that was in date");
+						continue;
+					}
+
+					// Made it past all the tests
+					_logger?.LogInformation("  Picked '{0}' with AppID '{1}' and Name '{2}' as a matching provision for the game '{3}'", DebugName, TestProvision?.ApplicationIdentifier, TestProvision?.ProvisionName, GameName);
+					return Key;
+				}
+			}
+
+			// check to see if there is already an embedded provision
+			string EmbeddedMobileProvisionFilename = Path.Combine(CodeSigningConfig.RepackageStagingDirectory, "embedded.mobileprovision");
+
+			_logger?.LogWarning("Failed to find a valid matching mobile provision, will attempt to use the embedded mobile provision instead if present");
+			return EmbeddedMobileProvisionFilename;
+		}
+
+		/// <summary>
+		/// Extracts the dict values for the Entitlements key and creates a new full .plist file
+		/// from them (with outer plist and dict keys as well as doctype, etc...)
+		/// </summary>
+		public string GetEntitlementsString(string CFBundleIdentifier, out string TeamIdentifier)
+		{
+			PListHelper XCentPList = new PListHelper();
+
+			Data?.ProcessValueForKey("Entitlements", "dict", (XmlNode ValueNode) =>
+			{
+		    	if (ValueNode != null)
+				{
+        			XCentPList = PListHelper.CloneDictionaryRootedAt(ValueNode)!;
+			    }
+			});
+
+			// Modify the application-identifier to be fully qualified if needed
+			string CurrentApplicationIdentifier;
+			XCentPList.GetString("application-identifier", out CurrentApplicationIdentifier);
+			XCentPList.GetString("com.apple.developer.team-identifier", out TeamIdentifier);
+
+			//			if (CurrentApplicationIdentifier.Contains("*"))
+			{
+				// Replace the application identifier
+				string NewApplicationIdentifier = String.Format("{0}.{1}", ApplicationIdentifierPrefix, CFBundleIdentifier);
+				XCentPList.SetString("application-identifier", NewApplicationIdentifier);
+
+
+				// Replace the keychain access groups
+				// Note: This isn't robust, it ignores the existing value in the wildcard and uses the same value for
+				// each entry.  If there is a legitimate need for more than one entry in the access group list, then
+				// don't use a wildcard!
+				List<string> KeyGroups = XCentPList.GetArray("keychain-access-groups", "string");
+
+				for (int i = 0; i < KeyGroups.Count; ++i)
+				{
+					string Entry = KeyGroups[i];
+					if (Entry.Contains("*"))
+					{
+						Entry = NewApplicationIdentifier;
+					}
+					KeyGroups[i] = Entry;
+				}
+
+				XCentPList.SetValueForKey("keychain-access-groups", KeyGroups);
+			}
+
+			// must have CloudKit and CloudDocuments for com.apple.developer.icloud-services
+			// otherwise the game will not be listed in the Settings->iCloud apps menu on the device
+			{
+				// iOS only
+				if (Platform == "IOS" && XCentPList.HasKey("com.apple.developer.icloud-services"))
+				{
+					List<string> ServicesGroups = XCentPList.GetArray("com.apple.developer.icloud-services", "string");
+					ServicesGroups.Clear();
+
+					ServicesGroups.Add("CloudKit");
+					ServicesGroups.Add("CloudDocuments");
+					XCentPList.SetValueForKey("com.apple.developer.icloud-services", ServicesGroups);
+				}
+
+				// For distribution builds, the entitlements from mobileprovisioning have a modified syntax
+				if (CodeSigningConfig.bForDistribution)
+				{
+					// remove the wildcards from the ubiquity-kvstore-identifier string
+					if (XCentPList.HasKey("com.apple.developer.ubiquity-kvstore-identifier"))
+					{
+						string UbiquityKvstoreString;
+						XCentPList.GetString("com.apple.developer.ubiquity-kvstore-identifier", out UbiquityKvstoreString);
+
+						int DotPosition = UbiquityKvstoreString.LastIndexOf("*");
+						if (DotPosition >= 0)
+						{
+							string TeamPrefix = DotPosition > 1 ? UbiquityKvstoreString.Substring(0, DotPosition - 1) : TeamIdentifier;
+							string NewUbiquityKvstoreIdentifier = String.Format("{0}.{1}", TeamPrefix, CFBundleIdentifier);
+							XCentPList.SetValueForKey("com.apple.developer.ubiquity-kvstore-identifier", NewUbiquityKvstoreIdentifier);
+						}
+					}
+
+					// remove the wildcards from the ubiquity-container-identifiers array
+					if (XCentPList.HasKey("com.apple.developer.ubiquity-container-identifiers"))
+					{
+						List<string> UbiquityContainerIdentifiersGroups = XCentPList.GetArray("com.apple.developer.ubiquity-container-identifiers", "string");
+
+						for (int i = 0; i < UbiquityContainerIdentifiersGroups.Count; i++)
+						{
+							int DotPosition = UbiquityContainerIdentifiersGroups[i].LastIndexOf("*");
+							if (DotPosition >= 0)
+							{
+								string TeamPrefix = DotPosition > 1 ? UbiquityContainerIdentifiersGroups[i].Substring(0, DotPosition - 1) : TeamIdentifier;
+								string NewUbiquityContainerIdentifier = String.Format("{0}.{1}", TeamPrefix, CFBundleIdentifier);
+								UbiquityContainerIdentifiersGroups[i] = NewUbiquityContainerIdentifier;
+							}
+						}
+
+						if (UbiquityContainerIdentifiersGroups.Count == 0)
+						{
+							string NewUbiquityKvstoreIdentifier = String.Format("{0}.{1}", TeamIdentifier, CFBundleIdentifier);
+							UbiquityContainerIdentifiersGroups.Add(NewUbiquityKvstoreIdentifier);
+						}
+
+						XCentPList.SetValueForKey("com.apple.developer.ubiquity-container-identifiers", UbiquityContainerIdentifiersGroups);
+					}
+
+					// remove the wildcards from the developer.associated-domains array or string
+					if (XCentPList.HasKey("com.apple.developer.associated-domains"))
+					{
+						string AssociatedDomainsString;
+						XCentPList.GetString("com.apple.developer.associated-domains", out AssociatedDomainsString);
+
+						//check if the value is string
+						if (AssociatedDomainsString != null && AssociatedDomainsString.Contains("*"))
+						{
+							XCentPList.RemoveKeyValue("com.apple.developer.associated-domains");
+						}
+						else
+						{
+							//check if the value is an array
+							List<string> AssociatedDomainsGroup = XCentPList.GetArray("com.apple.developer.associated-domains", "string");
+
+							if (AssociatedDomainsGroup.Count == 1 && AssociatedDomainsGroup[0].Contains("*"))
+							{
+								XCentPList.RemoveKeyValue("com.apple.developer.associated-domains");
+							}
+						}
+					}
+
+					// remove development keys - generated when the cloudkit container is in development mode
+					XCentPList.RemoveKeyValue("com.apple.developer.icloud-container-development-container-identifiers");
+				}
+
+				// set the icloud-container-environment according to the project settings
+				if (XCentPList.HasKey("com.apple.developer.icloud-container-environment"))
+				{
+					List<string> ContainerEnvironmentGroup = XCentPList.GetArray("com.apple.developer.icloud-container-environment", "string");
+
+					if (ContainerEnvironmentGroup.Count != 0)
+					{
+						ContainerEnvironmentGroup.Clear();
+
+						// The new value is a string, not an array
+						string NewContainerEnvironment = CodeSigningConfig.bForDistribution ? "Production" : "Development";
+						XCentPList.SetValueForKey("com.apple.developer.icloud-container-environment", NewContainerEnvironment);
+					}
+				}
+			}
+
+			return XCentPList.SaveToString();
+		}
+
+		/// <summary>
+		/// Constructs a MobileProvision from an xml blob extracted from the real ASN.1 file
+		/// </summary>
+		public MobileProvision(string EmbeddedPListText)
+		{
+			Data = new PListHelper(EmbeddedPListText);
+
+			// Now extract things
+
+			// Key: ApplicationIdentifierPrefix, Array<String>
+			List<string> PrefixList = Data.GetArray("ApplicationIdentifierPrefix", "string");
+			if (PrefixList.Count > 1)
+			{
+				_logger?.LogWarning("Found more than one entry for ApplicationIdentifierPrefix in the .mobileprovision, using the first one found");
+			}
+
+			if (PrefixList.Count > 0)
+			{
+				ApplicationIdentifierPrefix = PrefixList[0];
+			}
+
+			// Example date string from the XML: "2014-06-30T20:45:55Z";
+			DateTimeStyles AppleDateStyle = DateTimeStyles.AllowWhiteSpaces | DateTimeStyles.AssumeUniversal | DateTimeStyles.AdjustToUniversal;
+
+			string CreationDateString;
+			if (Data.GetDate("CreationDate", out CreationDateString))
+			{
+				CreationDate = DateTime.Parse(CreationDateString, CultureInfo.InvariantCulture, AppleDateStyle);
+			}
+
+			string ExpirationDateString;
+			if (Data.GetDate("ExpirationDate", out ExpirationDateString))
+			{
+				ExpirationDate = DateTime.Parse(ExpirationDateString, CultureInfo.InvariantCulture, AppleDateStyle);
+			}
+
+			// Key: DeveloperCertificates, Array<Data> (uuencoded)
+			string CertificatePassword = "";
+			List<string> CertificateList = Data.GetArray("DeveloperCertificates", "data");
+			foreach (string EncodedCert in CertificateList)
+			{
+				byte[] RawCert = Convert.FromBase64String(EncodedCert);
+				DeveloperCertificates.Add(new X509Certificate2(RawCert, CertificatePassword));
+			}
+
+			// Key: Name, String
+			if (!Data.GetString("Name", out ProvisionName))
+			{
+				ProvisionName = "(unknown)";
+			}
+
+			// Key: ProvisionedDevices, Array<String>
+			ProvisionedDeviceIDs = Data.GetArray("ProvisionedDevices", "string");
+
+			// Key: application-identifier, Array<String>
+			PListHelper? XCentPList = null;
+			Data.ProcessValueForKey("Entitlements", "dict", delegate (XmlNode ValueNode)
+			{
+				XCentPList = PListHelper.CloneDictionaryRootedAt(ValueNode);
+			});
+
+			// Modify the application-identifier to be fully qualified if needed
+			#pragma warning disable CS8602
+			if (!XCentPList.GetString("application-identifier", out ApplicationIdentifier))
+			#pragma warning restore CS8602 
+			{
+				ApplicationIdentifier = "(unknown)";
+			}
+
+
+			// check for get-task-allow
+			if (XCentPList != null)
+			{
+			    bDebug = XCentPList.GetBool("get-task-allow");
+			}
+			
+			if (!Data.GetString("UUID", out UUID))
+			{
+				UUID = "(unkown)";
+			}
+
+			List<string> Platforms = Data.GetArray("Platform", "string");
+			if (Platforms.Contains("iOS"))
+			{
+				Platform = "IOS";
+			}
+			else if (Platforms.Contains("tvOS"))
+			{
+				Platform = "TVOS";
+			}
+			else
+			{
+				Platform = "";
+			}
+		}
+
+		/// <summary>
+		/// Does this provision contain the specified UDID?
+		/// </summary>
+		public bool ContainsUDID(string UDID)
+		{
+		    bool bFound = false;
+    		if (ProvisionedDeviceIDs != null)
+		    {
+        		foreach (string TestUDID in ProvisionedDeviceIDs)
+		        {
+        		    if (TestUDID.Equals(UDID, StringComparison.InvariantCultureIgnoreCase))
+            		{
+		                bFound = true;
+        		        break;
+            		}
+		        }
+		    }
+		return bFound;
+		}
+
+	}
+
+	/// <summary>
+	/// This class understands how to get the embedded plist in a .mobileprovision file.  It doesn't
+	/// understand the full format and is not capable of writing a new one out or anything similar.
+	/// </summary>
+	public class MobileProvisionParser
+	{
+		static ILogger? _logger;
+
+		/// <summary>
+		/// The mobile provision parser
+		/// </summary>
+		public MobileProvisionParser()
+		{
+			ILoggerFactory loggerFactory = LoggerFactory.Create(builder =>
+		    {
+        		builder.AddConsole();
+		    });
+			_logger = loggerFactory.CreateLogger("MobileProvisionParser");
+		}
+
+		private static int StrStrByteArray(byte[] Haystack, int Offset, string Needle)
+		{
+			byte[] NeedleBytes = Encoding.UTF8.GetBytes(Needle);
+
+			//@TODO: Is there anything better in .NET? That's going to be pretty slow on large files
+			for (int i = Offset; i < Haystack.Length - NeedleBytes.Length; ++i)
+			{
+				bool bMatch = true;
+
+				for (int j = 0; j < NeedleBytes.Length; ++j)
+				{
+					if (Haystack[i + j] != NeedleBytes[j])
+					{
+						bMatch = false;
+						break;
+					}
+				}
+
+				if (bMatch)
+				{
+					return i;
+				}
+			}
+
+			return -1;
+		}
+
+		/// <summary>
+		/// Parse the file
+		/// </summary>
+		public static MobileProvision? ParseFile(byte[] RawData)
+		{
+			//@TODO: This file is just an ASN.1 stream, should find or make a raw ASN1 parser and use
+			// that instead of this (theoretically fragile) code (particularly the length extraction)
+
+			string StartPattern = "<?xml";
+			string EndPattern = "</plist>";
+
+			// Search the start pattern
+			int StartPos = StrStrByteArray(RawData, 0, StartPattern);
+			if (StartPos != -1)
+			{
+				// Search the end pattern
+				int EndPos = StrStrByteArray(RawData, StartPos, EndPattern);
+				if (EndPos != -1)
+				{
+					// Offset the end position to take in account the end pattern
+					EndPos += EndPattern.Length;
+
+					// Convert the data to a string
+					string PlistText = Encoding.UTF8.GetString(RawData, StartPos, EndPos - StartPos);
+
+					// Return the constructed 'mobile provision'
+					return new MobileProvision(PlistText);
+				}
+			}
+
+			// Unable to find the start of the plist data
+			_logger?.LogError("Failed to find embedded plist in .mobileprovision file");
+			return null;
+		}
+
+		/// <summary>
+		/// Parse the file
+		/// </summary>
+		public static MobileProvision? ParseFile(Stream InputStream)
+		{
+			// Read in the entire file
+			int NumBytes = (int)InputStream.Length;
+			byte[] RawData = new byte[NumBytes];
+			InputStream.Read(RawData, 0, NumBytes);
+
+			return ParseFile(RawData);
+		}
+
+		/// <summary>
+		/// Parse the file
+		/// </summary>
+		public static MobileProvision ParseFile(string Filename)
+		{
+			FileStream InputStream = File.OpenRead(Filename);
+			MobileProvision? Result = ParseFile(InputStream);
+			InputStream.Close();
+			if (Result == null)
+		    {
+        		throw new Exception("Failed to parse mobile provision file");
+		    }
+		    Result.FileName ??= Filename;
+
+		    return Result;
+		}
+	}
+
+	/// <summary>
+	/// The code signing config class
+	/// </summary>
+	public class CodeSigningConfig
+	{
+		/// <summary>
+		/// The display name in title bars, popups, etc ...
+		/// </summary>
+		public static string AppDisplayName = "Unreal iOS Configuration";
+
+
+		/// <summary>
+		/// The game directory
+		/// </summary>
+		public static string GameDirectory = "";        // "..\\..\\..\\..\\Engine\\Source\\UE4";
+
+		/// <summary>
+		/// The binaries directory
+		/// </summary>
+		public static string BinariesDirectory = "";    // "..\\..\\..\\..\\Engine\\Binaries\\IOS\\"
+
+		/// <summary>
+		/// Optional Prefix to append to .xcent and .mobileprovision files, for handling multiple certificates on same source game
+		/// </summary>
+		public static string SigningPrefix = "";
+
+		/// <summary>
+		/// PCStagingRootDir
+		/// </summary>
+		public static string PCStagingRootDir = "";
+
+		/// <summary>
+		/// The local staging directory for files needed by Xcode on the Mac (on PC)
+		/// </summary>
+		public static string PCXcodeStagingDir = "";
+
+		/// <summary>
+		/// The staging directory from UAT that is passed into IPP (for repackaging, etc...)
+		/// </summary>
+		public static string RepackageStagingDirectory = "";
+
+		/// <summary>
+		/// The delta manifest for deploying files for iterative deploy
+		/// </summary>
+		public static string DeltaManifest = "";
+
+		/// <summary>
+		/// The files to be retrieved from the device
+		/// </summary>
+		public static List<string> FilesForBackup = new List<string>();
+
+		/// <summary>
+		/// The project file that is passed into IPP from UAT
+		/// </summary>
+		public static string ProjectFile = "";
+
+		/// <summary>
+		/// The device to deploy or launch on
+		/// </summary>
+		public static string DeviceId = "";
+
+		/// <summary>
+		/// Determine whether to use RPC Util or not
+		/// </summary>
+		public static bool bUseRPCUtil = true;
+
+		/// <summary>
+		/// Optional override for the dev root on the target mac for remote builds.
+		/// </summary>
+		public static string? OverrideDevRoot = null;
+
+		/// <summary>
+		/// Optional value specifying that we are working with tvOS
+		/// </summary>
+		public static string OSString = "IOS";
+
+		/// <summary>
+		/// The local build directory (on PC)
+		/// </summary>
+		public static string BuildDirectory
+		{
+			get
+			{
+				string StandardGameBuildDir = GameDirectory + @"\Build\IOS";
+				// if the normal Build dir exists, return it, otherwise, use the program Resources directory
+				return Path.GetFullPath(Directory.Exists(StandardGameBuildDir) ?
+					StandardGameBuildDir :
+					GameDirectory + @"\Resources\IOS");
+			}
+		}
+
+		/// <summary>
+		/// The shared provision library directory (on PC)
+		/// </summary>
+		public static string ProvisionDirectory
+		{
+			get
+			{
+				if (Environment.OSVersion.Platform == PlatformID.MacOSX || Environment.OSVersion.Platform == PlatformID.Unix)
+				{
+					return Environment.GetEnvironmentVariable("HOME") + "/Library/MobileDevice/Provisioning Profiles/";
+				}
+				else
+				{
+					return Environment.GetFolderPath(Environment.SpecialFolder.LocalApplicationData) + "/Apple Computer/MobileDevice/Provisioning Profiles/";
+				}
+			}
+		}
+
+		/// <summary>
+		/// The shared (Engine) build directory (on PC)
+		/// </summary>
+		public static string EngineBuildDirectory
+		{
+			get
+			{
+				return Path.Combine(Unreal.EngineDirectory.FullName, "Build", CodeSigningConfig.OSString);
+			}
+		}
+
+		/// <summary>
+		/// The local build intermediate directory (on PC)
+		/// </summary>
+		public static string IntermediateDirectory
+		{
+			get
+			{
+				string IntermediateGameBuildDir = GameDirectory + @"\Intermediate\" + CodeSigningConfig.OSString;
+				// if the normal Build dir exists, return it, otherwise, use the program Resources directory
+				return Path.GetFullPath(Directory.Exists(IntermediateGameBuildDir) ?
+					IntermediateGameBuildDir :
+					BuildDirectory);
+			}
+		}
+
+		/// <summary>
+		/// The local directory cooked files are placed (on PC)
+		/// </summary>
+		public static string CookedDirectory
+		{
+			get { return Path.GetFullPath(GameDirectory + @"\Saved\Cooked\" + CodeSigningConfig.OSString); }
+		}
+
+		/// <summary>
+		/// The local directory config files are placed (on PC)
+		/// </summary>
+		public static string ConfigDirectory
+		{
+			get { return Path.GetFullPath(GameDirectory + @"\Saved\Config"); }
+		}
+
+		/// <summary>
+		/// The engine config files are placed (on PC)
+		/// </summary>
+		public static string DefaultConfigDirectory
+		{
+			get { return Path.Combine(Unreal.EngineDirectory.FullName, "Engine/Config"); }
+		}
+
+		/// <summary>
+		/// The engine directory
+		/// </summary>
+		public static string EngineDirectory
+		{
+			get { return Unreal.EngineDirectory.FullName; }
+		}
+
+		/// <summary>
+		/// The local directory that a payload (GameName.app) is assembled into before being copied to the Mac (on PC)
+		/// </summary>
+		public static string PayloadRootDirectory
+		{
+			get { return Path.GetFullPath(PCStagingRootDir + @"\Payload"); }
+		}
+		
+		/// <summary>
+		/// Program_Architecture
+		/// </summary>
+		public static string Program_Architecture = "";
+
+		/// <summary>
+		/// Program_GameName
+		/// </summary>
+		public static string Program_GameName = "";
+
+		/// <summary>
+		/// Returns the name of the file containing user overrides that will be applied when packaging on PC
+		/// </summary>
+		public static string GetPlistOverrideFilename()
+		{
+			return GetPlistOverrideFilename(false);
+		}
+		
+		/// <summary>
+		/// GetPlistOverrideFilename
+		/// </summary>
+		public static string GetPlistOverrideFilename(bool bWantDistributionOverlay)
+		{
+			string Prefix = "";
+			if (bWantDistributionOverlay)
+			{
+				Prefix = "Distro_";
+			}
+
+			return Path.Combine(BuildDirectory, Prefix + Program_GameName + "Overrides.plist");
+		}
+
+		/// <summary>
+		/// Whether or not to output extra information (like every file copy and date/time stamp)
+		/// </summary>
+		public static bool bVerbose = true;
+
+		/// <summary>
+		/// Whether or not to output extra information in code signing
+		/// </summary>
+		public static bool bCodeSignVerbose = false;
+
+		/// <summary>
+		/// Whether or not non-critical files will be packaged (critical == required for signing or .app validation, the app may still fail to
+		/// run with only 'critical' files present).  Also affects the name and location of the IPA back on the PC
+		/// </summary>
+		public static bool bCreateStubSet = false;
+
+		/// <summary>
+		/// Is this a distribution packaging build?  Controls a number of aspects of packing (which signing prefix and provisioning profile to use, etc...)
+		/// </summary>
+		public static bool bForDistribution = false;
+
+		/// <summary>
+		/// Is this a c++ code based project?  So far used on repackagefromstage to properly choose the .ipa name
+		/// </summary>
+		public static bool bIsCodeBasedProject = false;
+
+		/// <summary>
+		/// Is this going to try automatic signing?
+		/// </summary>
+		public static bool bAutomaticSigning = false;
+
+		/// <summary>
+		/// Whether or not to strip symbols (they will always be stripped when packaging for distribution)
+		/// </summary>
+		public static bool bForceStripSymbols = false;
+
+		/// <summary>
+		/// Do a code signing update when repackaging?
+		/// </summary>
+		public static bool bPerformResignWhenRepackaging = false;
+
+		/// <summary>
+		/// Whether the cooked data will be cooked on the fly or was already cooked by the books.
+		/// </summary>
+		public static bool bCookOnTheFly = false;
+
+		/// <summary>
+		/// Whether the install should be performed on a provision
+		/// </summary>
+		public static bool bProvision = false;
+
+		/// <summary>
+		/// provision to be installed
+		/// </summary>
+		public static string Provision = "";
+
+		/// <summary>
+		/// provision to be installed
+		/// </summary>
+		public static string ProvisionUUID = "";
+
+		/// <summary>
+		/// provision to extract certificate for
+		/// </summary>
+		public static string? ProvisionFile;
+
+		/// <summary>
+		/// target name specificied by the command line, if not set GameName + optional Client will be used
+		/// </summary>
+		public static string TargetName = "";
+
+		/// <summary>
+		/// IOS Team ID to be used for automatic signing
+		/// </summary>
+		public static string TeamID = "";
+
+		/// <summary>
+		/// Whether the install should be performed on a certificate
+		/// </summary>
+		public static bool bCert = false;
+
+		/// <summary>
+		/// Certificate to be installed
+		/// </summary>
+		public static string Certificate = "";
+
+		/// <summary>
+		/// Certificate to be output
+		/// </summary>
+		public static string? OutputCertificate = null;
+
+		/// <summary>
+		/// An override server Mac name
+		/// </summary>
+		public static string? OverrideMacName = null;
+
+		/// <summary>
+		/// A name to use for the bundle indentifier and display name when resigning
+		/// </summary>
+		public static string? OverrideBundleName = null;
+
+		/// <summary>
+		/// Whether to use None or Best for the compression setting when repackaging IPAs (int version of Ionic.Zlib.CompressionLevel)
+		/// By making this an int, we are able to delay load the Ionic assembly from a different path (which is required)
+		/// </summary>
+		public static int RecompressionSetting = 0; //Ionic.Zlib.CompressionLevel.None;
+
+		/// <summary>
+		/// If the requirements blob is present in the existing executable when code signing, should it be carried forward
+		/// (true) or should a dummy requirements blob be created with no actual requirements (false)
+		/// </summary>
+		public static bool bMaintainExistingRequirementsWhenCodeSigning = false;
+
+		/// <summary>
+		/// The code signing identity to use when signing via RPC
+		/// </summary>
+		public static string? CodeSigningIdentity;
+
+		/// <summary>
+		/// The minimum OS version
+		/// </summary>
+		public static string? MinOSVersion;
+
+		/// <summary>
+		/// Create an iterative IPA (i.e. a stub only with Icons and Splash images)
+		/// </summary>
+		public static bool bIterate = false;
+
+		/// <summary>
+		/// Returns a path to the place to back up documents from a device
+		/// </summary>
+		/// <returns></returns>
+		public static string GetRootBackedUpDocumentsDirectory()
+		{
+			return Path.GetFullPath(Path.Combine(GameDirectory + @"\" + CodeSigningConfig.OSString + "_Backups"));
+		}
+
+		/// <summary>
+		/// The init
+		/// </summary>
+		public static bool Initialize(FileReference? ProjectFile, bool bIsTVOS)
+		{
+			bool bIsEpicInternal = File.Exists(@"..\..\EpicInternal.txt");
+
+			// if the path is a directory (relative to where the game was launched from), then get the absolute directory
+			string FullGamePath;
+			string OrigGamePath;
+
+			if (ProjectFile == null)
+			{
+				FullGamePath = Unreal.EngineDirectory.FullName;
+				OrigGamePath = "Engine";
+				Program_GameName = "UnrealGame";
+			}
+			else
+			{
+				FullGamePath = OrigGamePath = CodeSigningConfig.ProjectFile = ProjectFile.FullName;
+				Program_GameName = ProjectFile.GetFileNameWithoutAnyExtensions();
+				CodeSigningConfig.ProjectFile = ProjectFile.FullName;
+			}
+
+			OSString = bIsTVOS ? "TVOS" : "IOS";
+
+			if (Directory.Exists(FullGamePath))
+			{
+				GameDirectory = FullGamePath;
+			}
+			// is it a file? if so, just use the file's directory
+			else if (File.Exists(FullGamePath))
+			{
+				GameDirectory = Path.GetDirectoryName(FullGamePath)!;
+			}
+
+			if (!Directory.Exists(GameDirectory))
+			{
+				Console.ForegroundColor = ConsoleColor.Red;
+				Console.WriteLine();
+				Console.WriteLine("Unable to find a game or program {0}. You may need to specify a path to the program", OrigGamePath);
+				Console.ResetColor();
+				return false;
+			}
+
+			// special case handling for anything inside Engine/Source, it will go to the Engine/Binaries directory, not the Game's binaries directory
+			if (OrigGamePath.Replace("\\", "/").Contains("Engine/Source"))
+			{
+				BinariesDirectory = Path.Combine(Unreal.EngineDirectory.FullName, "Binaries", CodeSigningConfig.OSString);
+			}
+			else if (!OrigGamePath.Replace("\\", "/").Contains(@"Binaries/" + CodeSigningConfig.OSString))
+			{
+				// no sense in adding Binaries\IOS when it's already there. This is a special case to handle packaging UnrealLaunchDaemon from the command line.
+				BinariesDirectory = Path.Combine(GameDirectory, "Binaries", CodeSigningConfig.OSString) + "/";
+			}
+			else
+			{
+				BinariesDirectory = GameDirectory;
+			}
+
+			CodeSigningConfig.CodeSigningIdentity = CodeSigningConfig.bForDistribution ? "iPhone Distribution" : "iPhone Developer";
+
+			// Windows doesn't allow environment vars to be set to blank so detect "none" and treat it as such
+			if (CodeSigningConfig.SigningPrefix == "none")
+			{
+				CodeSigningConfig.SigningPrefix = CodeSigningConfig.bForDistribution ? "Distro_" : "";
+			}
+
+			return true;
+		}
+	}
+
+}
+
diff --git a/Engine/Source/Programs/UnrealBuildTool/Platform/IOS/UEBuildIOS.cs b/Engine/Source/Programs/UnrealBuildTool/Platform/IOS/UEBuildIOS.cs
--- a/Engine/Source/Programs/UnrealBuildTool/Platform/IOS/UEBuildIOS.cs
+++ b/Engine/Source/Programs/UnrealBuildTool/Platform/IOS/UEBuildIOS.cs
@@ -10,6 +10,7 @@ using System.Text.RegularExpressions;
 using EpicGames.Core;
 using UnrealBuildBase;
 using Microsoft.Extensions.Logging;
+using System.Security.Cryptography.X509Certificates;
 
 namespace UnrealBuildTool
 {
@@ -505,33 +506,12 @@ namespace UnrealBuildTool
             string? MobileProvision = ProjectSettings.MobileProvision;
 
 			FileReference? ProjectFile = ProjectSettings.ProjectFile;
-			FileReference IPhonePackager = FileReference.Combine(Unreal.EngineDirectory, "Binaries/DotNET/IOS/IPhonePackager.exe");
+			CodeSigningConfig.Initialize(ProjectFile, bIsTVOS);
 
 			if (!string.IsNullOrEmpty(SigningCertificate))
             {
-                // verify the certificate
-                Process IPPProcess = new Process();
-
-				string IPPCmd = "certificates " + ((ProjectFile != null) ? ("\"" + ProjectFile.ToString() + "\"") : "Engine") + " -bundlename " + ProjectSettings.BundleIdentifier + (bForDistribution ? " -distribution" : "");
-				IPPProcess.StartInfo.StandardOutputEncoding = Encoding.UTF8;
-				IPPProcess.StartInfo.StandardErrorEncoding = Encoding.UTF8;
-				IPPProcess.StartInfo.WorkingDirectory = Unreal.EngineDirectory.ToString();
-				IPPProcess.OutputDataReceived += new DataReceivedEventHandler((s, e) => IPPDataReceivedHandler(s, e, Logger));
-				IPPProcess.ErrorDataReceived += new DataReceivedEventHandler((s, e) => IPPDataReceivedHandler(s, e, Logger));
-
-				if (BuildHostPlatform.Current.Platform == UnrealTargetPlatform.Mac)
-                {
-                    IPPProcess.StartInfo.FileName = FileReference.Combine(Unreal.EngineDirectory, "Build/BatchFiles/Mac/RunMono.sh").FullName;
-					IPPProcess.StartInfo.Arguments = string.Format("\"{0}\" {1}", IPhonePackager ,IPPCmd);
-				}
-                else
-                {
-					IPPProcess.StartInfo.FileName = IPhonePackager.FullName;
-					IPPProcess.StartInfo.Arguments = IPPCmd;
-                }
-
-				Logger.LogInformation("Getting certifcate information via {IPPProcessStartInfoFileName} {IPPProcessStartInfoArguments}", IPPProcess.StartInfo.FileName, IPPProcess.StartInfo.Arguments);
-                Utils.RunLocalProcess(IPPProcess);
+				List<string> Certs = AppleCodeSign.FindCertificates();
+				List<FileReference> Provisions = AppleCodeSign.FindProvisions(ProjectSettings.BundleIdentifier, bForDistribution, out _);
             }
             else
             {
@@ -565,31 +545,13 @@ namespace UnrealBuildTool
                 MobileProvision = "";
 				MobileProvisionFile = null;
                 Logger.LogInformation("Provision not specified or not found for {Project}, searching for compatible match...", ((ProjectFile != null) ? ProjectFile.GetFileNameWithoutAnyExtensions() : "UnrealGame"));
-                Process IPPProcess = new Process();
 
-				IPPProcess.StartInfo.StandardOutputEncoding = Encoding.UTF8;
-				IPPProcess.StartInfo.StandardErrorEncoding = Encoding.UTF8;
-				IPPProcess.OutputDataReceived += new DataReceivedEventHandler((s, e) => IPPDataReceivedHandler(s, e, Logger));
-				IPPProcess.ErrorDataReceived += new DataReceivedEventHandler((s, e) => IPPDataReceivedHandler(s, e, Logger));
-				IPPProcess.StartInfo.WorkingDirectory = Unreal.EngineDirectory.ToString();
-
-				string IPPCmd = "signing_match " + ((ProjectFile != null) ? ("\"" + ProjectFile.ToString() + "\"") : "Engine") + " -bundlename " + ProjectSettings.BundleIdentifier + (bIsTVOS ? " -tvos" : "") + (bForDistribution ? " -distribution" : "");
-
-				if (BuildHostPlatform.Current.Platform == UnrealTargetPlatform.Mac)
-				{
-					IPPProcess.StartInfo.FileName = FileReference.Combine(Unreal.EngineDirectory, "Build/BatchFiles/Mac/RunMono.sh").FullName;
-					IPPProcess.StartInfo.Arguments = string.Format("\"{0}\" {1}", IPhonePackager, IPPCmd);
-				}
-				else
+				if (AppleCodeSign.FindCertAndProvision(ProjectSettings.BundleIdentifier, out MobileProvisionFile, out SigningCertificate))
 				{
-					IPPProcess.StartInfo.FileName = IPhonePackager.FullName;
-					IPPProcess.StartInfo.Arguments = IPPCmd;
+					MobileProvision = MobileProvisionFile!.FullName;
 				}
-
-				Logger.LogInformation("Getting signing information via {IPPProcessStartInfoFileName} {IPPProcessStartInfoArguments}", IPPProcess.StartInfo.FileName, IPPProcess.StartInfo.Arguments);
-
-				Utils.RunLocalProcess(IPPProcess);
-				if(MobileProvisionFile != null)
+				
+				if (MobileProvisionFile != null)
 				{
 					Logger.LogInformation("Provision found for {Project}, Provision: {Provision}, Certificate: {Certificate}", ((ProjectFile != null) ? ProjectFile.GetFileNameWithoutAnyExtensions() : "UnrealGame"), MobileProvisionFile, SigningCertificate);
 				}
diff --git a/Engine/Source/Programs/UnrealBuildTool/UnrealBuildTool.csproj b/Engine/Source/Programs/UnrealBuildTool/UnrealBuildTool.csproj
--- a/Engine/Source/Programs/UnrealBuildTool/UnrealBuildTool.csproj
+++ b/Engine/Source/Programs/UnrealBuildTool/UnrealBuildTool.csproj
@@ -1,4 +1,4 @@
-﻿<Project Sdk="Microsoft.NET.Sdk">
+<Project Sdk="Microsoft.NET.Sdk">
   <Import Project="..\Shared\UnrealEngine.csproj.props" />
   
   <PropertyGroup>
@@ -96,6 +96,9 @@
       <Link>Properties\MetaData.cs</Link>
     </Compile>
   </ItemGroup>
+  <ItemGroup>
+     <PackageReference Include="Microsoft.Extensions.Logging.Console" Version="5.0.0" />
+  </ItemGroup>
 
   <!-- The following will glob all platform extensions' UnrealBuildTool source
   files and include them in the build. They will also appear in VisualStudio.
