Fix warnings at proj generation

diff --git a/Engine/Source/Programs/UnrealBuildTool/Platform/IOS/IOSCodeSigning.cs b/Engine/Source/Programs/UnrealBuildTool/Platform/IOS/IOSCodeSigning.cs
--- a/Engine/Source/Programs/UnrealBuildTool/Platform/IOS/IOSCodeSigning.cs
+++ b/Engine/Source/Programs/UnrealBuildTool/Platform/IOS/IOSCodeSigning.cs
@@ -316,7 +316,14 @@ namespace UnrealBuildTool
 
 			}
 		    _logger?.LogInformation(" MATCHED-Provision:{0},File:{1},Cert:{2} ", SelectedProvision, SelectedFile, SelectedCert);
-			MatchedProvision = new FileReference(SelectedProvision);
+			if (SelectedProvision != "")
+			{
+				MatchedProvision = new FileReference(SelectedProvision);
+			}
+			else
+			{
+				MatchedProvision = null;
+			}
 
 			return Provisions;
 		}
@@ -541,11 +548,14 @@ namespace UnrealBuildTool
 
 			if (FindRequiredFiles(out MobileProvision, out Cert, out _, out _))
 			{
-				// print out the provision and cert name
-				Provision = new FileReference(MobileProvision?.FileName ?? "");
-				CertName = GetFriendlyNameFromCert(Cert!);
-				_logger?.LogInformation("CERTIFICATE-{0},PROVISION-{1}", CertName, Provision.FullName);
-				return true;
+				if (MobileProvision?.FileName != null)
+				{
+					// print out the provision and cert name
+					Provision = new FileReference(MobileProvision?.FileName);
+					CertName = GetFriendlyNameFromCert(Cert!);
+					_logger?.LogInformation("CERTIFICATE-{0},PROVISION-{1}", CertName, Provision.FullName);
+					return true;
+				}
 			}
 
 			_logger?.LogInformation("No matching Signing Data found!");
