fix boilerplate

diff --git a/Engine/Source/Programs/UnrealBuildTool/Platform/IOS/IOSCodeSigning.cs b/Engine/Source/Programs/UnrealBuildTool/Platform/IOS/IOSCodeSigning.cs
--- a/Engine/Source/Programs/UnrealBuildTool/Platform/IOS/IOSCodeSigning.cs
+++ b/Engine/Source/Programs/UnrealBuildTool/Platform/IOS/IOSCodeSigning.cs
@@ -1,4 +1,6 @@
-﻿using System;
+﻿// Copyright Epic Games, Inc. All Rights Reserved.
+
+using System;
 using System.Collections.Generic;
 using System.Text;
 using System.Security.Cryptography.X509Certificates;
