UBT: Ignore ConfigRedirects file if the stream does not have it.

diff --git a/Engine/Source/Programs/UnrealBuildTool/System/ConfigFile.cs b/Engine/Source/Programs/UnrealBuildTool/System/ConfigFile.cs
--- a/Engine/Source/Programs/UnrealBuildTool/System/ConfigFile.cs
+++ b/Engine/Source/Programs/UnrealBuildTool/System/ConfigFile.cs
@@ -174,7 +174,10 @@ namespace UnrealBuildTool
 			{
 				// read the special ConfigRemap.ini file into sections
 				FileReference ConfigRemapFile = FileReference.Combine(UnrealBuildBase.Unreal.EngineDirectory, "Config", "ConfigRedirects.ini");
-				ReadIntoSections(ConfigRemapFile, Sections, ConfigLineAction.Set);
+				if (FileReference.Exists(ConfigRemapFile))
+				{
+					ReadIntoSections(ConfigRemapFile, Sections, ConfigLineAction.Set);
+				}
 			}
 			catch (Exception)
 			{
