UBT: Fix warning.

diff --git a/Engine/Source/Programs/UnrealBuildTool/Platform/IOS/IOSCodeSigning.cs b/Engine/Source/Programs/UnrealBuildTool/Platform/IOS/IOSCodeSigning.cs
--- a/Engine/Source/Programs/UnrealBuildTool/Platform/IOS/IOSCodeSigning.cs
+++ b/Engine/Source/Programs/UnrealBuildTool/Platform/IOS/IOSCodeSigning.cs
@@ -550,10 +550,11 @@ namespace UnrealBuildTool
 
 			if (FindRequiredFiles(out MobileProvision, out Cert, out _, out _))
 			{
-				if (MobileProvision?.FileName != null)
+				string? FileName = MobileProvision?.FileName;
+				if (FileName != null)
 				{
 					// print out the provision and cert name
-					Provision = new FileReference(MobileProvision?.FileName);
+					Provision = new FileReference(FileName);
 					CertName = GetFriendlyNameFromCert(Cert!);
 					_logger?.LogInformation("CERTIFICATE-{0},PROVISION-{1}", CertName, Provision.FullName);
 					return true;
