Optimized GatherDependencyModules and CreateSharedPCHInstances

diff --git a/Engine/Source/Programs/UnrealBuildTool/Configuration/UEBuildTarget.cs b/Engine/Source/Programs/UnrealBuildTool/Configuration/UEBuildTarget.cs
--- a/Engine/Source/Programs/UnrealBuildTool/Configuration/UEBuildTarget.cs
+++ b/Engine/Source/Programs/UnrealBuildTool/Configuration/UEBuildTarget.cs
@@ -1,6 +1,7 @@
 // Copyright Epic Games, Inc. All Rights Reserved.
 
 using System;
+using System.Collections.Concurrent;
 using System.Collections.Generic;
 using System.ComponentModel;
 using System.Linq;
@@ -2147,21 +2148,19 @@ namespace UnrealBuildTool
 		/// <returns>Dependency modules set.</returns>
 		static HashSet<UEBuildModuleCPP> GatherDependencyModules(List<UEBuildBinary> Binaries)
 		{
-			HashSet<UEBuildModuleCPP> Output = new HashSet<UEBuildModuleCPP>();
-
-			foreach (UEBuildBinary Binary in Binaries)
+			ConcurrentBag<UEBuildModuleCPP> Bag = new();
+			Parallel.ForEach(Binaries, Binary =>
 			{
 				List<UEBuildModule> DependencyModules = Binary.GetAllDependencyModules(bIncludeDynamicallyLoaded: false, bForceCircular: false);
 				foreach (UEBuildModuleCPP Module in DependencyModules.OfType<UEBuildModuleCPP>())
 				{
 					if (Module.Binary != null)
 					{
-						Output.Add(Module);
+						Bag.Add(Module);
 					}
 				}
-			}
-
-			return Output;
+			});
+			return Bag.ToHashSet();
 		}
 
 		/// <summary>
@@ -3449,55 +3448,51 @@ namespace UnrealBuildTool
 				return;
 			}
 
-			List<UEBuildModuleCPP> Modules = new List<UEBuildModuleCPP>();
-			Dictionary<UEBuildBinary, CppCompileEnvironment> BinaryCompileEnvironments = new Dictionary<UEBuildBinary, CppCompileEnvironment>();
-			foreach (UEBuildBinary Binary in OriginalBinaries)
-			{
-				BinaryCompileEnvironments.Add(Binary, Binary.CreateBinaryCompileEnvironment(GlobalCompileEnvironment));
-				Modules.AddRange(Binary.Modules.OfType<UEBuildModuleCPP>());
-			}
-
-			// Find which instances are needed
-			Dictionary<UEBuildModuleCPP, CppCompileEnvironment> ModuleCompileEnvironments = new Dictionary<UEBuildModuleCPP, CppCompileEnvironment>();
-			Dictionary<PrecompiledHeaderTemplate, List<UEBuildModuleCPP>> PCHsAndModulesDict = new Dictionary<PrecompiledHeaderTemplate, List<UEBuildModuleCPP>>();
-			foreach (UEBuildModuleCPP Module in Modules)
+			ConcurrentDictionary<UEBuildModuleCPP, CppCompileEnvironment> ModuleCompileEnvironments = new();
+			ConcurrentDictionary<PrecompiledHeaderTemplate, List<UEBuildModuleCPP>> PCHsAndModulesDict = new();
+			Parallel.ForEach(OriginalBinaries, Binary =>
 			{
-				CppCompileEnvironment BinaryCompileEnvironment = BinaryCompileEnvironments[Module.Binary];
-				CppCompileEnvironment ModuleCompileEnvironment = Module.CreateModuleCompileEnvironment(Target, BinaryCompileEnvironment, Logger);
+				CppCompileEnvironment BinaryCompileEnvironment = Binary.CreateBinaryCompileEnvironment(GlobalCompileEnvironment);
 
-				// Is it using a private PCH?
-				if (Module.Rules.PrivatePCHHeaderFile != null && (Module.Rules.PCHUsage == ModuleRules.PCHUsageMode.NoSharedPCHs || Module.Rules.PCHUsage == ModuleRules.PCHUsageMode.UseExplicitOrSharedPCHs))
+				foreach (UEBuildModuleCPP Module in Binary.Modules.OfType<UEBuildModuleCPP>())
 				{
-					continue;
-				}
+					CppCompileEnvironment ModuleCompileEnvironment = Module.CreateModuleCompileEnvironment(Target, BinaryCompileEnvironment, Logger);
 
-				// Is it using a shared PCH?
-				if (Module.Rules.PCHUsage != ModuleRules.PCHUsageMode.NoPCHs)
-				{
-					if (!ModuleCompileEnvironment.bIsBuildingLibrary && Module.Rules.PCHUsage != ModuleRules.PCHUsageMode.NoSharedPCHs)
+					// Is it using a private PCH?
+					if (Module.Rules.PrivatePCHHeaderFile != null && (Module.Rules.PCHUsage == ModuleRules.PCHUsageMode.NoSharedPCHs || Module.Rules.PCHUsage == ModuleRules.PCHUsageMode.UseExplicitOrSharedPCHs))
 					{
-						// Find all the dependencies of this module
-						HashSet<UEBuildModule> ReferencedModules = new HashSet<UEBuildModule>();
-						Module.GetAllDependencyModules(new List<UEBuildModule>(), ReferencedModules, bIncludeDynamicallyLoaded: false, bForceCircular: false, bOnlyDirectDependencies: true);
+						continue;
+					}
 
-						// Find the first shared PCH module we can use
-						PrecompiledHeaderTemplate? Template = ModuleCompileEnvironment.SharedPCHs.FirstOrDefault(x => ReferencedModules.Contains(x.Module));
-						if (Template != null && Template.IsValidFor(ModuleCompileEnvironment))
+					// Is it using a shared PCH?
+					if (Module.Rules.PCHUsage != ModuleRules.PCHUsageMode.NoPCHs)
+					{
+						if (!ModuleCompileEnvironment.bIsBuildingLibrary && Module.Rules.PCHUsage != ModuleRules.PCHUsageMode.NoSharedPCHs)
 						{
-							ModuleCompileEnvironments[Module] = ModuleCompileEnvironment;
+							// Find all the dependencies of this module
+							HashSet<UEBuildModule> ReferencedModules = new HashSet<UEBuildModule>();
+							Module.GetAllDependencyModules(new List<UEBuildModule>(), ReferencedModules, bIncludeDynamicallyLoaded: false, bForceCircular: false, bOnlyDirectDependencies: true);
 
-							if (PCHsAndModulesDict.TryGetValue(Template, out List<UEBuildModuleCPP>? ModulesUsingSharedPCHs))
-							{
-								ModulesUsingSharedPCHs.Add(Module);
-							}
-							else
+							// Find the first shared PCH module we can use
+							PrecompiledHeaderTemplate? Template = ModuleCompileEnvironment.SharedPCHs.FirstOrDefault(x => ReferencedModules.Contains(x.Module));
+							if (Template != null && Template.IsValidFor(ModuleCompileEnvironment))
 							{
-								PCHsAndModulesDict[Template] = new List<UEBuildModuleCPP>() { Module };
+								ModuleCompileEnvironments.TryAdd(Module, ModuleCompileEnvironment);
+
+								PCHsAndModulesDict.AddOrUpdate(Template, k =>
+								{
+									return new List<UEBuildModuleCPP>() { Module };
+								},
+								(k,v) =>
+								{
+									v.Add(Module);
+									return v;
+								});
 							}
 						}
 					}
 				}
-			}
+			});
 
 			// Create the PCH instances
 			List<PrecompiledHeaderTemplate> SharedPCHs = new List<PrecompiledHeaderTemplate>(GlobalCompileEnvironment.SharedPCHs);
