- FIxed a crash that was happening for some poeple with the new IPhonePackager-in-UBT stuff

diff --git a/Engine/Source/Programs/UnrealBuildTool/Platform/IOS/IOSCodeSigning.cs b/Engine/Source/Programs/UnrealBuildTool/Platform/IOS/IOSCodeSigning.cs
--- a/Engine/Source/Programs/UnrealBuildTool/Platform/IOS/IOSCodeSigning.cs
+++ b/Engine/Source/Programs/UnrealBuildTool/Platform/IOS/IOSCodeSigning.cs
@@ -548,7 +548,7 @@ namespace UnrealBuildTool
 				CodeSigningConfig.OverrideBundleName = BundleID;
 			}
 
-			if (FindRequiredFiles(out MobileProvision, out Cert, out _, out _))
+			if (FindRequiredFiles(out MobileProvision, out Cert, out _, out _) && Cert != null)
 			{
 				string? FileName = MobileProvision?.FileName;
 				if (FileName != null)
