Optimized the ItemGroup values in VCProject by moving values to properties so can be shared across ClCompile items.

This change drastically cut down on the generated project folder size:
Before: 641 MB
After: 77.1 MB

diff --git a/Engine/Source/Programs/UnrealBuildTool/ProjectFiles/VisualStudio/VCProject.cs b/Engine/Source/Programs/UnrealBuildTool/ProjectFiles/VisualStudio/VCProject.cs
--- a/Engine/Source/Programs/UnrealBuildTool/ProjectFiles/VisualStudio/VCProject.cs
+++ b/Engine/Source/Programs/UnrealBuildTool/ProjectFiles/VisualStudio/VCProject.cs
@@ -13,6 +13,8 @@ using UnrealBuildBase;
 using System.Diagnostics.CodeAnalysis;
 using Microsoft.Extensions.Logging;
 using System.Reflection.Metadata;
+using System.Threading.Tasks;
+using System.Collections.Concurrent;
 
 namespace UnrealBuildTool
 {
@@ -1274,17 +1276,184 @@ namespace UnrealBuildTool
 					LocalAliasedFiles.Add(new AliasedFile(CurFile.Reference, ProjectRelativeSourceFile, FilterRelativeSourceDirectory));
 				}
 
+				ConcurrentDictionary<DirectoryReference, string?> DirectoryToPchFile = new();
+				ConcurrentDictionary<DirectoryReference, string> DirectoryToForceIncludePaths = new();
+				ConcurrentDictionary<DirectoryReference, string> DirectoryToIncludeSearchPaths = new();
+				Parallel.ForEach(LocalAliasedFiles, LocalAliasedFile =>
+				{
+					// get the filetype as represented to Visual Studio
+					string VCFileType = GetVCFileType(LocalAliasedFile.FileSystemPath);
+					DirectoryReference FileSystemPathDir = new DirectoryReference(LocalAliasedFile.FileSystemPath);
+
+					// if the filetype is an include and its path is filtered out, skip it entirely (should we do this for any type of
+					// file? Possibly, but not today due to potential fallout)
+					if (VCFileType == "ClInclude" && IncludePathIsFilteredOut(FileSystemPathDir))
+					{
+						return;
+					}
+
+					// Allow filtering of any type of file
+					if (FilePathIsFilteredOut(FileSystemPathDir))
+					{
+						return;
+					}
+
+					if (VCFileType != "ClCompile")
+					{
+						return;
+					}
+
+					DirectoryReference Directory = LocalAliasedFile.Location.Directory;
+
+					// Find the PCH file
+					DirectoryToPchFile.GetOrAdd(Directory, _ =>
+					{
+						string? PchHeaderFile = null;
+						for (DirectoryReference? ParentDir = Directory; ParentDir != null; ParentDir = ParentDir.ParentDirectory)
+						{
+							if (ModuleDirToPchHeaderFile.TryGetValue(ParentDir, out PchHeaderFile))
+							{
+								break;
+							}
+						}
+						return PchHeaderFile;
+					});
+
+					// Find the force-included headers
+					DirectoryToForceIncludePaths.GetOrAdd(Directory, _ =>
+					{
+						string? ForceIncludePaths = null;
+						for (DirectoryReference? ParentDir = Directory; ParentDir != null; ParentDir = ParentDir.ParentDirectory)
+						{
+							if (ModuleDirToForceIncludePaths.TryGetValue(ParentDir, out ForceIncludePaths))
+							{
+								break;
+							}
+						}
+
+						// filter here. It's a little more graceful to do it where this info is built but easier to follow if we filter 
+						// things our right before they're written.
+						if (!string.IsNullOrEmpty(ForceIncludePaths))
+						{
+							IEnumerable<string> PathList = ForceIncludePaths.Split(new[] { ';' }, StringSplitOptions.RemoveEmptyEntries);
+							ForceIncludePaths = string.Join(";", PathList.Where(P => !IncludePathIsFilteredOut(new DirectoryReference(P))));
+						}
+
+						ForceIncludePaths ??= String.Empty;
+
+						return ForceIncludePaths;
+					});
+
+					if (TryGetBuildEnvironment(Directory, out BuildEnvironment? BuildEnvironment))
+					{
+						DirectoryToIncludeSearchPaths.GetOrAdd(Directory, _ =>
+						{
+							StringBuilder Builder = new StringBuilder();
+							AppendIncludePaths(Builder, BuildEnvironment.UserIncludePaths, SharedIncludeSearchPathsSet);
+							AppendIncludePaths(Builder, BuildEnvironment.SystemIncludePaths, SharedIncludeSearchPathsSet);
+							return Builder.ToString();
+						});
+					}
+				});
+
+				// Collapse common values
+				{
+					StringBuilder CommonProjectFileContent = new StringBuilder();
+
+					//AdditionalIncludeDirectories
+					{
+						int ValueIndex = 0;
+						Dictionary<string, string> UpdatedValues = new();
+						var KVPList = DirectoryToIncludeSearchPaths.ToList();
+						KVPList.SortBy(kvp => kvp.Key);
+						foreach (var DirectoryKVP in KVPList)
+						{
+							string? Value = DirectoryKVP.Value;
+							if (!String.IsNullOrEmpty(Value))
+							{
+								if (!UpdatedValues.ContainsKey(Value))
+								{
+									string PropertyName = "ClCompile_AdditionalIncludeDirectories";
+									if (ValueIndex > 0)
+									{
+										PropertyName += "_" + ValueIndex;
+									}
+									ValueIndex++;
+									UpdatedValues.Add(Value, PropertyName);
+									CommonProjectFileContent.AppendLine($"    <{PropertyName}>$(NMakeIncludeSearchPath);{Value}</{PropertyName}>");
+								}
+								DirectoryToIncludeSearchPaths[DirectoryKVP.Key] = UpdatedValues[Value];
+							}
+						}
+					}
+					//ForcedIncludeFiles
+					{
+						int ValueIndex = 0;
+						Dictionary<string, string> UpdatedValues = new();
+						var KVPList = DirectoryToForceIncludePaths.ToList();
+						KVPList.SortBy(kvp => kvp.Key);
+						foreach (var DirectoryKVP in KVPList)
+						{
+							string? Value = DirectoryKVP.Value;
+							if (!String.IsNullOrEmpty(Value))
+							{
+								if (!UpdatedValues.ContainsKey(Value))
+								{
+									string PropertyName = "ClCompile_ForcedIncludeFiles";
+									if (ValueIndex > 0)
+									{
+										PropertyName += "_" + ValueIndex;
+									}
+									ValueIndex++;
+									UpdatedValues.Add(Value, PropertyName);
+									CommonProjectFileContent.AppendLine($"    <{PropertyName}>{Value}</{PropertyName}>");
+								}
+								DirectoryToForceIncludePaths[DirectoryKVP.Key] = UpdatedValues[Value];
+							}
+						}
+					}
+					// AdditionalOptions
+					{
+						int ValueIndex = 0;
+						Dictionary<string, string> UpdatedValues = new();
+						var KVPList = DirectoryToPchFile.ToList();
+						KVPList.SortBy(kvp => kvp.Key);
+						foreach (var DirectoryKVP in KVPList)
+						{
+							string? Value = DirectoryKVP.Value;
+							if (!String.IsNullOrEmpty(Value))
+							{
+								if (!UpdatedValues.ContainsKey(Value))
+								{
+									string PropertyName = "ClCompile_AdditionalOptions";
+									if (ValueIndex > 0)
+									{
+										PropertyName += "_" + ValueIndex;
+									}
+									ValueIndex++;
+									UpdatedValues.Add(Value, PropertyName);
+									CommonProjectFileContent.AppendLine($"    <{PropertyName}>$(AdditionalOptions) /Yu\"{Value}\"</{PropertyName}>");
+								}
+								DirectoryToPchFile[DirectoryKVP.Key] = UpdatedValues[Value];
+							}
+						}
+					}
+
+					if (CommonProjectFileContent.Length > 0)
+					{
+						VCProjectFileContent.AppendLine("  <PropertyGroup>");
+						VCProjectFileContent.Append(CommonProjectFileContent);
+						VCProjectFileContent.AppendLine("  </PropertyGroup>");
+					}
+				}
+
 				VCFiltersFileContent.AppendLine("  <ItemGroup>");
 
 				VCProjectFileContent.AppendLine("  <ItemGroup>");
 
 				// Add all file directories to the filters file as solution filters
 				HashSet<string> FilterDirectories = new HashSet<string>();
-				//UEBuildPlatform BuildPlatform = UEBuildPlatform.GetBuildPlatform(BuildHostPlatform.Current.Platform);
 
-				Dictionary<DirectoryReference, string> DirectoryToIncludeSearchPaths = new Dictionary<DirectoryReference, string>();
-				Dictionary<DirectoryReference, string> DirectoryToForceIncludePaths = new Dictionary<DirectoryReference, string>();
-				Dictionary<DirectoryReference, string?> DirectoryToPchFile = new Dictionary<DirectoryReference, string?>();
 				foreach (AliasedFile AliasedFile in LocalAliasedFiles)
 				{
 					// No need to add the root directory relative to the project (it would just be an empty string!)
@@ -1295,16 +1464,17 @@ namespace UnrealBuildTool
 
 					// get the filetype as represented to Visual Studio
 					string VCFileType = GetVCFileType(AliasedFile.FileSystemPath);
+					DirectoryReference FileSystemPathDir = new DirectoryReference(AliasedFile.FileSystemPath);
 
 					// if the filetype is an include and its path is filtered out, skip it entirely (should we do this for any type of
 					// file? Possibly, but not today due to potential fallout)
-					if (VCFileType == "ClInclude" && IncludePathIsFilteredOut(new DirectoryReference(AliasedFile.FileSystemPath)))
+					if (VCFileType == "ClInclude" && IncludePathIsFilteredOut(FileSystemPathDir))
 					{
 						continue;
 					}
 
 					// Allow filtering of any type of file
-					if (FilePathIsFilteredOut(new DirectoryReference(AliasedFile.FileSystemPath)))
+					if (FilePathIsFilteredOut(FileSystemPathDir))
 					{
 						continue;
 					}
@@ -1317,68 +1487,26 @@ namespace UnrealBuildTool
 					{
 						DirectoryReference Directory = AliasedFile.Location.Directory;
 
-						// Find the force-included headers
-						string? ForceIncludePaths;
-						if (!DirectoryToForceIncludePaths.TryGetValue(Directory, out ForceIncludePaths))
-						{
-							for (DirectoryReference? ParentDir = Directory; ParentDir != null; ParentDir = ParentDir.ParentDirectory)
-							{
-								if (ModuleDirToForceIncludePaths.TryGetValue(ParentDir, out ForceIncludePaths))
-								{
-									break;
-								}
-							}
-
-							// filter here. It's a little more graceful to do it where this info is built but easier to follow if we filter 
-							// things our right before they're written.
-							if (!string.IsNullOrEmpty(ForceIncludePaths))
-							{
-								IEnumerable<string> PathList = ForceIncludePaths.Split(new[] { ';' }, StringSplitOptions.RemoveEmptyEntries);
-								ForceIncludePaths = string.Join(";", PathList.Where(P => !IncludePathIsFilteredOut(new DirectoryReference(P))));
-							}
-
-							ForceIncludePaths ??= String.Empty;
-
-							DirectoryToForceIncludePaths[Directory] = ForceIncludePaths;
-						}
-
-						// Find the PCH file
-						string? PchHeaderFile;
-						if (!DirectoryToPchFile.TryGetValue(Directory, out PchHeaderFile))
-						{
-							for (DirectoryReference? ParentDir = Directory; ParentDir != null; ParentDir = ParentDir.ParentDirectory)
-							{
-								if (ModuleDirToPchHeaderFile.TryGetValue(ParentDir, out PchHeaderFile))
-								{
-									break;
-								}
-							}
-							DirectoryToPchFile[Directory] = PchHeaderFile;
-						}
-
 						// Find the include search paths
-						VCProjectFileContent.AppendLine("    <{0} Include=\"{1}\">", VCFileType, EscapeFileName(AliasedFile.FileSystemPath));
 						if (TryGetBuildEnvironment(Directory, out BuildEnvironment? BuildEnvironment))
 						{
-							string? IncludeSearchPaths = String.Empty;
-							if (!DirectoryToIncludeSearchPaths.TryGetValue(Directory, out IncludeSearchPaths))
-							{
-								StringBuilder Builder = new StringBuilder();
-								AppendIncludePaths(Builder, BuildEnvironment.UserIncludePaths, SharedIncludeSearchPathsSet);
-								AppendIncludePaths(Builder, BuildEnvironment.SystemIncludePaths, SharedIncludeSearchPathsSet);
-								IncludeSearchPaths = Builder.ToString();
+							VCProjectFileContent.AppendLine("    <{0} Include=\"{1}\">", VCFileType, EscapeFileName(AliasedFile.FileSystemPath));
 
-								DirectoryToIncludeSearchPaths.Add(Directory, IncludeSearchPaths);
-							}
-							VCProjectFileContent.AppendLine("      <AdditionalIncludeDirectories>$(NMakeIncludeSearchPath);{0}</AdditionalIncludeDirectories>", IncludeSearchPaths);
-							VCProjectFileContent.AppendLine("      <ForcedIncludeFiles>{0}</ForcedIncludeFiles>", ForceIncludePaths);
+							VCProjectFileContent.AppendLine($"      <AdditionalIncludeDirectories>$({DirectoryToIncludeSearchPaths[Directory]})</AdditionalIncludeDirectories>");
+							VCProjectFileContent.AppendLine($"      <ForcedIncludeFiles>$({DirectoryToForceIncludePaths[Directory]})</ForcedIncludeFiles>");
+
+							string? PchHeaderFile = DirectoryToPchFile[Directory];
 							if (PchHeaderFile != null && ProjectFileFormat >= VCProjectFileFormat.VisualStudio2022)
 							{
-								VCProjectFileContent.AppendLine("      <AdditionalOptions>$(AdditionalOptions) /Yu\"{0}\"</AdditionalOptions>",
-									PchHeaderFile);
+								VCProjectFileContent.AppendLine($"      <AdditionalOptions>$({DirectoryToPchFile[Directory]})</AdditionalOptions>");
 							}
+
+							VCProjectFileContent.AppendLine("    </{0}>", VCFileType);
+						}
+						else
+						{
+							VCProjectFileContent.AppendLine("    <{0} Include=\"{1}\"/>", VCFileType, EscapeFileName(AliasedFile.FileSystemPath));
 						}
-						VCProjectFileContent.AppendLine("    </{0}>", VCFileType);
 					}
 
 					if (!String.IsNullOrWhiteSpace(AliasedFile.ProjectPath))
