Optimized the reading of the text files in the SourceFileMetadataCache. The same file was being redundantly read for different information.

diff --git a/Engine/Source/Programs/UnrealBuildTool/System/SourceFileMetadataCache.cs b/Engine/Source/Programs/UnrealBuildTool/System/SourceFileMetadataCache.cs
--- a/Engine/Source/Programs/UnrealBuildTool/System/SourceFileMetadataCache.cs
+++ b/Engine/Source/Programs/UnrealBuildTool/System/SourceFileMetadataCache.cs
@@ -37,9 +37,9 @@ namespace UnrealBuildTool
 	class SourceFileMetadataCache
 	{
 		/// <summary>
-		/// Information about the first file included from a source file
+		/// Source(cpp/c) file info
 		/// </summary>
-		class IncludeInfo
+		class SourceFileInfo
 		{
 			/// <summary>
 			/// Last write time of the file when the data was cached
@@ -50,12 +50,17 @@ namespace UnrealBuildTool
 			/// Contents of the include directive
 			/// </summary>
 			public string? IncludeText;
+
+			/// <summary>
+			/// List of files this particular file is inlining
+			/// </summary>
+			public List<string> InlinedFileNames = new List<string>();
 		}
 
 		/// <summary>
-		/// Information about whether a header file contains reflection markup and what includes it has (the superset, ignoring #if/#endif)
+		/// Header file info
 		/// </summary>
-		class HeaderInfo
+		class HeaderFileInfo
 		{
 			/// <summary>
 			/// Last write time of the file when the data was cached
@@ -75,26 +80,10 @@ namespace UnrealBuildTool
 			public HeaderUnitType UnitType;
 		}
 
-		/// <summary>
-		/// Information about whether a file contains the inline gen.cpp macro
-		/// </summary>
-		class InlineReflectionInfo
-		{
-			/// <summary>
-			/// Last write time of the file when the data was cached
-			/// </summary>
-			public long LastWriteTimeUtc;
-
-			/// <summary>
-			/// List of files this particular file is inlining
-			/// </summary>
-			public List<string> InlinedFileNames = new List<string>();
-		}
-
 		/// <summary>
 		/// The current file version
 		/// </summary>
-		public const int CurrentVersion = 5;
+		public const int CurrentVersion = 6;
 
 		/// <summary>
 		/// Location of this dependency cache
@@ -114,20 +103,15 @@ namespace UnrealBuildTool
 		/// <summary>
 		/// Map from file item to source file info
 		/// </summary>
-		ConcurrentDictionary<FileItem, IncludeInfo> FileToIncludeInfo = new ConcurrentDictionary<FileItem, IncludeInfo>();
+		ConcurrentDictionary<FileItem, SourceFileInfo> FileToSourceFileInfo = new ConcurrentDictionary<FileItem, SourceFileInfo>();
 
 		/// <summary>
 		/// Map from file item to header file info
 		/// </summary>
-		ConcurrentDictionary<FileItem, HeaderInfo> FileToHeaderInfo = new ConcurrentDictionary<FileItem, HeaderInfo>();
-
-		/// <summary>
-		/// Map from file item to inline info
-		/// </summary>
-		ConcurrentDictionary<FileItem, InlineReflectionInfo> FileToInlineReflectionInfo = new ConcurrentDictionary<FileItem, InlineReflectionInfo>();
+		ConcurrentDictionary<FileItem, HeaderFileInfo> FileToHeaderFileInfo = new ConcurrentDictionary<FileItem, HeaderFileInfo>();
 
 		/// <summary>
-		/// Map from file item to source file info
+		/// Map from file item to source file
 		/// </summary>
 		ConcurrentDictionary<FileItem, SourceFile> FileToSourceFile = new ConcurrentDictionary<FileItem, SourceFile>();
 
@@ -164,7 +148,7 @@ namespace UnrealBuildTool
 		/// <summary>
 		/// Static cache of all constructed dependency caches
 		/// </summary>
-		static Dictionary<FileReference, SourceFileMetadataCache> Caches = new Dictionary<FileReference, SourceFileMetadataCache>();
+		static ConcurrentDictionary<FileReference, SourceFileMetadataCache> Caches = new ConcurrentDictionary<FileReference, SourceFileMetadataCache>();
 
 		/// <summary>
 		/// Constructs a dependency cache. This method is private; call CppDependencyCache.Create() to create a cache hierarchy for a given project.
@@ -190,112 +174,136 @@ namespace UnrealBuildTool
 		}
 
 		/// <summary>
-		/// Gets the first included file from a source file
+		/// Returns a SourceFileInfo struct for a file (and parse the file if not already cached)
 		/// </summary>
-		/// <param name="SourceFile">The source file to parse</param>
-		/// <returns>Text from the first include directive. Null if the file did not contain any include directives.</returns>
-		public string? GetFirstInclude(FileItem SourceFile)
+		/// <param name="SourceFile">The file to parse</param>
+		/// <returns>SourceFileInfo for file</returns>
+		SourceFileInfo GetSourceFileInfo(FileItem SourceFile)
 		{
-			if(Parent != null && !SourceFile.Location.IsUnderDirectory(BaseDirectory))
+			if (Parent != null && !SourceFile.Location.IsUnderDirectory(BaseDirectory))
 			{
-				return Parent.GetFirstInclude(SourceFile);
+				return Parent.GetSourceFileInfo(SourceFile);
 			}
 			else
 			{
-				IncludeInfo? IncludeInfo;
-				if(!FileToIncludeInfo.TryGetValue(SourceFile, out IncludeInfo) || SourceFile.LastWriteTimeUtc.Ticks > IncludeInfo.LastWriteTimeUtc)
+				var UpdateSourceFileInfo = (FileItem SourceFile) =>
 				{
-					IncludeInfo = new IncludeInfo();
-					IncludeInfo.LastWriteTimeUtc = SourceFile.LastWriteTimeUtc.Ticks;
-					IncludeInfo.IncludeText = ParseFirstInclude(SourceFile.Location);
-					FileToIncludeInfo[SourceFile] = IncludeInfo;
+					var SourceFileInfo = new SourceFileInfo();
+					var FileText = FileReference.ReadAllText(SourceFile.Location);
+					var FileTextLines = FileText.Split('\n');
+
+					SourceFileInfo.LastWriteTimeUtc = SourceFile.LastWriteTimeUtc.Ticks;
+
+					// Inline reflection data
+					MatchCollection FileMatches = InlineReflectionMarkupRegex.Matches(FileText);
+					foreach (Match Match in FileMatches)
+					{
+						SourceFileInfo.InlinedFileNames.Add(Match.Groups[1].Value);
+					}
+
+					SourceFileInfo.IncludeText = ParseFirstInclude(SourceFile, FileTextLines);
+
 					bModified = true;
+					return SourceFileInfo;
+				};
+
+				return FileToSourceFileInfo.AddOrUpdate(SourceFile, _ =>
+				{
+					return UpdateSourceFileInfo(SourceFile);
+				},
+				(k, v) =>
+				{
+					if (SourceFile.LastWriteTimeUtc.Ticks > v.LastWriteTimeUtc)
+					{
+						return UpdateSourceFileInfo(SourceFile);
+					}
+					return v;
 				}
-				return IncludeInfo.IncludeText;
+				);
 			}
 		}
 
 		/// <summary>
-		/// Finds or adds a SourceFile class for the given file
+		/// Parse the first include directive from a source file
 		/// </summary>
-		/// <param name="File">File to fetch the source file data for</param>
-		/// <returns>SourceFile instance corresponding to the given source file</returns>
-		public SourceFile GetSourceFile(FileItem File)
+		/// <param name="SourceFile">The source file to parse</param>
+		/// <param name="FileToSourceFileFileText">The source file contents</param>
+		/// <returns>The first include directive</returns>
+		static string? ParseFirstInclude(FileItem SourceFile, string[] FileToSourceFileFileText)
 		{
-			if (Parent != null && !File.Location.IsUnderDirectory(BaseDirectory))
-			{
-				return Parent.GetSourceFile(File);
-			}
-			else
+			bool bMatchImport = SourceFile.HasExtension(".m") || SourceFile.HasExtension(".mm");
+			foreach (string Line in FileToSourceFileFileText)
 			{
-				SourceFile? Result;
-				if (!FileToSourceFile.TryGetValue(File, out Result) || File.LastWriteTimeUtc.Ticks > Result.LastWriteTimeUtc)
+				if (Line == null)
 				{
-					SourceFile NewSourceFile = new SourceFile(File);
-					if (Result == null)
-					{
-						if (FileToSourceFile.TryAdd(File, NewSourceFile))
-						{
-							Result = NewSourceFile;
-						}
-						else
-						{
-							Result = FileToSourceFile[File];
-						}
-					}
-					else
+					return null;
+				}
+
+				Match IncludeMatch = IncludeRegex.Match(Line);
+				if (IncludeMatch.Success)
+				{
+					return IncludeMatch.Groups[1].Value;
+				}
+
+				if (bMatchImport)
+				{
+					Match ImportMatch = ImportRegex.Match(Line);
+					if (ImportMatch.Success)
 					{
-						if (FileToSourceFile.TryUpdate(File, NewSourceFile, Result))
-						{
-							Result = NewSourceFile;
-						}
-						else
-						{
-							Result = FileToSourceFile[File];
-						}
+						return IncludeMatch.Groups[1].Value;
 					}
 				}
-				return Result;
 			}
+			return null;
 		}
 
-		/// <summary>
-		/// Determines whether the given file contains reflection markup
-		/// </summary>
-		/// <param name="HeaderFile">The source file to parse</param>
-		/// <returns>True if the file contains reflection markup</returns>
-		public bool ContainsReflectionMarkup(FileItem HeaderFile)
+		HeaderFileInfo GetHeaderFileInfo(FileItem HeaderFile)
 		{
-			if(Parent != null && !HeaderFile.Location.IsUnderDirectory(BaseDirectory))
+			if (Parent != null && !HeaderFile.Location.IsUnderDirectory(BaseDirectory))
 			{
-				return Parent.ContainsReflectionMarkup(HeaderFile);
+				return Parent.GetHeaderFileInfo(HeaderFile);
 			}
 			else
 			{
-				HeaderInfo? HeaderInfo;
-				if(!FileToHeaderInfo.TryGetValue(HeaderFile, out HeaderInfo) || HeaderFile.LastWriteTimeUtc.Ticks > HeaderInfo.LastWriteTimeUtc)
+				var UpdateHeaderFileInfo = (FileItem HeaderFile) =>
 				{
-					HeaderInfo = ParseHeader(HeaderFile);
-					FileToHeaderInfo[HeaderFile] = HeaderInfo;
+					var HeaderFileInfo = new HeaderFileInfo();
+					var FileText = FileReference.ReadAllText(HeaderFile.Location);
+					var FileTextLines = FileText.Split('\n');
+
+					HeaderFileInfo.LastWriteTimeUtc = HeaderFile.LastWriteTimeUtc.Ticks;
+
+					ParseHeader(HeaderFileInfo, FileTextLines);
+
 					bModified = true;
+					return HeaderFileInfo;
+				};
+
+				return FileToHeaderFileInfo.AddOrUpdate(HeaderFile, _ =>
+				{
+					return UpdateHeaderFileInfo(HeaderFile);
+				},
+				(k, v) =>
+				{
+					if (HeaderFile.LastWriteTimeUtc.Ticks > v.LastWriteTimeUtc)
+					{
+						return UpdateHeaderFileInfo(HeaderFile);
+					}
+					return v;
 				}
-				return HeaderInfo.bContainsMarkup;
+				);
 			}
 		}
 
 		/// <summary>
 		/// Read entire header file to find markup and includes
 		/// </summary>
-		/// <param name="HeaderFile">The header file to parse</param>
 		/// <returns>A HeaderInfo struct containing information about header</returns>
-		private HeaderInfo ParseHeader(FileItem HeaderFile)
+		private void ParseHeader(HeaderFileInfo HeaderFileInfo, string[] FileText)
 		{
-			HeaderInfo HeaderInfo = new();
-			HeaderInfo.LastWriteTimeUtc = HeaderFile.LastWriteTimeUtc.Ticks;
-
 			bool bContainsMarkup = false;
 			SortedSet<string> Includes = new();
-			foreach (string Line in FileReference.ReadAllLines(HeaderFile.Location))
+			foreach (string Line in FileText)
 			{
 				if (!bContainsMarkup)
 				{
@@ -349,71 +357,72 @@ namespace UnrealBuildTool
 				{
 					ReadOnlySpan<char> Span = Line.AsSpan(HeaderUnitIndex + "HEADER_UNIT_".Length);
 					if (Span.StartsWith("UNSUPPORTED"))
-						HeaderInfo.UnitType = HeaderUnitType.Unsupported;
+						HeaderFileInfo.UnitType = HeaderUnitType.Unsupported;
 					else if (Span.StartsWith("SKIP"))
-						HeaderInfo.UnitType = HeaderUnitType.Skip;
+						HeaderFileInfo.UnitType = HeaderUnitType.Skip;
 				}
 			}
 
-			HeaderInfo.bContainsMarkup = bContainsMarkup;
-			HeaderInfo.Includes = Includes.ToList();
-			return HeaderInfo;
+			HeaderFileInfo.bContainsMarkup = bContainsMarkup;
+			HeaderFileInfo.Includes = Includes.ToList();
 		}
 
-
-
 		/// <summary>
-		/// Returns a list of inlined generated cpps that this source file contains.
+		/// Gets the first included file from a source file
 		/// </summary>
 		/// <param name="SourceFile">The source file to parse</param>
-		/// <returns>List of marked files this source file contains</returns>
-		public IList<string> GetListOfInlinedGeneratedCppFiles(FileItem SourceFile)
+		/// <returns>Text from the first include directive. Null if the file did not contain any include directives.</returns>
+		public string? GetFirstInclude(FileItem SourceFile)
 		{
-			if (Parent != null && !SourceFile.Location.IsUnderDirectory(BaseDirectory))
+			return GetSourceFileInfo(SourceFile).IncludeText;
+		}
+
+		/// <summary>
+		/// Finds or adds a SourceFile class for the given file
+		/// </summary>
+		/// <param name="File">File to fetch the source file data for</param>
+		/// <returns>SourceFile instance corresponding to the given source file</returns>
+		public SourceFile GetSourceFile(FileItem File)
+		{
+			if (Parent != null && !File.Location.IsUnderDirectory(BaseDirectory))
 			{
-				return Parent.GetListOfInlinedGeneratedCppFiles(SourceFile);
+				return Parent.GetSourceFile(File);
 			}
 			else
 			{
-				InlineReflectionInfo? InlineReflectionInfo;
-				if (!FileToInlineReflectionInfo.TryGetValue(SourceFile, out InlineReflectionInfo) || SourceFile.LastWriteTimeUtc.Ticks > InlineReflectionInfo.LastWriteTimeUtc)
+				return FileToSourceFile.AddOrUpdate(File, _ =>
 				{
-					InlineReflectionInfo = new InlineReflectionInfo();
-					InlineReflectionInfo.LastWriteTimeUtc = SourceFile.LastWriteTimeUtc.Ticks;
-					MatchCollection FileMatches = InlineReflectionMarkupRegex.Matches(FileReference.ReadAllText(SourceFile.Location));
-					foreach (Match Match in FileMatches)
+					return new SourceFile(File);
+				},
+				(k,v) =>
+				{
+					if (File.LastWriteTimeUtc.Ticks > v.LastWriteTimeUtc)
 					{
-						InlineReflectionInfo.InlinedFileNames.Add(Match.Groups[1].Value);
+						return new SourceFile(File);
 					}
-					FileToInlineReflectionInfo[SourceFile] = InlineReflectionInfo;
-					bModified = true;
-				}
-				return InlineReflectionInfo.InlinedFileNames;
+					return v;
+				});
 			}
 		}
 
 		/// <summary>
-		/// Returns a HeaderInfo struct for a header file (and parse the file if not already cached)
+		/// Returns a list of inlined generated cpps that this source file contains.
 		/// </summary>
-		/// <param name="HeaderFile">The header file to parse</param>
-		/// <returns>HeaderInfo for header file</returns>
-		HeaderInfo GetHeaderInfo(FileItem HeaderFile)
+		/// <param name="SourceFile">The source file to parse</param>
+		/// <returns>List of marked files this source file contains</returns>
+		public IList<string> GetListOfInlinedGeneratedCppFiles(FileItem SourceFile)
 		{
-			if (Parent != null && !HeaderFile.Location.IsUnderDirectory(BaseDirectory))
-			{
-				return Parent.GetHeaderInfo(HeaderFile);
-			}
-			else
-			{
-				HeaderInfo? HeaderInfo;
-				if (!FileToHeaderInfo.TryGetValue(HeaderFile, out HeaderInfo) || HeaderFile.LastWriteTimeUtc.Ticks > HeaderInfo.LastWriteTimeUtc)
-				{
-					HeaderInfo = ParseHeader(HeaderFile);
-					FileToHeaderInfo[HeaderFile] = HeaderInfo;
-					bModified = true;
-				}
-				return HeaderInfo;
-			}
+			return GetSourceFileInfo(SourceFile).InlinedFileNames;
+		}
+
+		/// <summary>
+		/// Determines whether the given file contains reflection markup
+		/// </summary>
+		/// <param name="HeaderFile">The source file to parse</param>
+		/// <returns>True if the file contains reflection markup</returns>
+		public bool ContainsReflectionMarkup(FileItem HeaderFile)
+		{
+			return GetHeaderFileInfo(HeaderFile).bContainsMarkup;
 		}
 
 		/// <summary>
@@ -423,7 +432,7 @@ namespace UnrealBuildTool
 		/// <returns>Header unit type</returns>
 		public HeaderUnitType GetHeaderUnitType(FileItem HeaderFile)
 		{
-			return GetHeaderInfo(HeaderFile).UnitType;
+			return GetHeaderFileInfo(HeaderFile).UnitType;
 		}
 
 		/// <summary>
@@ -433,43 +442,7 @@ namespace UnrealBuildTool
 		/// <returns>List of includes</returns>
 		public List<string> GetHeaderIncludes(FileItem HeaderFile)
 		{
-			return GetHeaderInfo(HeaderFile).Includes!;
-		}
-
-		/// <summary>
-		/// Parse the first include directive from a source file
-		/// </summary>
-		/// <param name="SourceFile">The source file to parse</param>
-		/// <returns>The first include directive</returns>
-		static string? ParseFirstInclude(FileReference SourceFile)
-		{
-			bool bMatchImport = SourceFile.HasExtension(".m") || SourceFile.HasExtension(".mm");
-			using(StreamReader Reader = new StreamReader(SourceFile.FullName, true))
-			{
-				for(;;)
-				{
-					string? Line = Reader.ReadLine();
-					if(Line == null)
-					{
-						return null;
-					}
-
-					Match IncludeMatch = IncludeRegex.Match(Line);
-					if(IncludeMatch.Success)
-					{
-						return IncludeMatch.Groups[1].Value;
-					}
-
-					if(bMatchImport)
-					{
-						Match ImportMatch = ImportRegex.Match(Line);
-						if(ImportMatch.Success)
-						{
-							return IncludeMatch.Groups[1].Value;
-						}
-					}
-				}
-			}
+			return GetHeaderFileInfo(HeaderFile).Includes!;
 		}
 
 		/// <summary>
@@ -524,21 +497,15 @@ namespace UnrealBuildTool
 		/// <returns>Reference to a dependency cache with the given settings</returns>
 		static SourceFileMetadataCache FindOrAddCache(FileReference Location, DirectoryReference BaseDirectory, SourceFileMetadataCache? Parent, ILogger Logger)
 		{
-			lock(Caches)
+			SourceFileMetadataCache Cache = Caches.GetOrAdd(Location, _ =>
 			{
-				SourceFileMetadataCache? Cache;
-				if(Caches.TryGetValue(Location, out Cache))
-				{
-					Debug.Assert(Cache.BaseDirectory == BaseDirectory);
-					Debug.Assert(Cache.Parent == Parent);
-				}
-				else
-				{
-					Cache = new SourceFileMetadataCache(Location, BaseDirectory, Parent, Logger);
-					Caches.Add(Location, Cache);
-				}
-				return Cache;
-			}
+				return new SourceFileMetadataCache(Location, BaseDirectory, Parent, Logger); ;
+			});
+
+			Debug.Assert(Cache.BaseDirectory == BaseDirectory);
+			Debug.Assert(Cache.Parent == Parent);
+
+			return Cache;
 		}
 
 		/// <summary>
@@ -565,30 +532,18 @@ namespace UnrealBuildTool
 						return;
 					}
 
-					int FileToFirstIncludeCount = Reader.ReadInt();
-					for(int Idx = 0; Idx < FileToFirstIncludeCount; Idx++)
-					{
-						FileItem File = Reader.ReadCompactFileItem();
-						
-						IncludeInfo IncludeInfo = new IncludeInfo();
-						IncludeInfo.LastWriteTimeUtc = Reader.ReadLong();
-						IncludeInfo.IncludeText = Reader.ReadString();
-
-						FileToIncludeInfo[File] = IncludeInfo;
-					}
-
 					int FileToMarkupFlagCount = Reader.ReadInt();
 					for(int Idx = 0; Idx < FileToMarkupFlagCount; Idx++)
 					{
 						FileItem File = Reader.ReadCompactFileItem();
 
-						HeaderInfo HeaderInfo = new HeaderInfo();
-						HeaderInfo.LastWriteTimeUtc = Reader.ReadLong();
-						HeaderInfo.bContainsMarkup = Reader.ReadBool();
-						HeaderInfo.UnitType = (HeaderUnitType)Reader.ReadByte();
-						HeaderInfo.Includes = Reader.ReadList(() => Reader.ReadString())!;
+						HeaderFileInfo HeaderFileInfo = new HeaderFileInfo();
+						HeaderFileInfo.LastWriteTimeUtc = Reader.ReadLong();
+						HeaderFileInfo.bContainsMarkup = Reader.ReadBool();
+						HeaderFileInfo.UnitType = (HeaderUnitType)Reader.ReadByte();
+						HeaderFileInfo.Includes = Reader.ReadList(() => Reader.ReadString())!;
 
-						FileToHeaderInfo[File] = HeaderInfo;
+						FileToHeaderFileInfo[File] = HeaderFileInfo;
 					}
 
 					int FileToInlineMarkupFlagCount = Reader.ReadInt();
@@ -596,11 +551,12 @@ namespace UnrealBuildTool
 					{
 						FileItem File = Reader.ReadCompactFileItem();
 
-						InlineReflectionInfo InlineReflectionInfo = new InlineReflectionInfo();
-						InlineReflectionInfo.LastWriteTimeUtc = Reader.ReadLong();
-						InlineReflectionInfo.InlinedFileNames = Reader.ReadList(() => Reader.ReadString())!;
+						SourceFileInfo SourceFileInfo = new SourceFileInfo();
+						SourceFileInfo.LastWriteTimeUtc = Reader.ReadLong();
+						SourceFileInfo.IncludeText = Reader.ReadString();
+						SourceFileInfo.InlinedFileNames = Reader.ReadList(() => Reader.ReadString())!;
 
-						FileToInlineReflectionInfo[File] = InlineReflectionInfo;
+						FileToSourceFileInfo[File] = SourceFileInfo;
 					}
 				}
 			}
@@ -623,16 +579,8 @@ namespace UnrealBuildTool
 				{
 					Writer.WriteInt(CurrentVersion);
 
-					Writer.WriteInt(FileToIncludeInfo.Count);
-					foreach(KeyValuePair<FileItem, IncludeInfo> Pair in FileToIncludeInfo)
-					{
-						Writer.WriteCompactFileItem(Pair.Key);
-						Writer.WriteLong(Pair.Value.LastWriteTimeUtc);
-						Writer.WriteString(Pair.Value.IncludeText);
-					}
-
-					Writer.WriteInt(FileToHeaderInfo.Count);
-					foreach(KeyValuePair<FileItem, HeaderInfo> Pair in FileToHeaderInfo)
+					Writer.WriteInt(FileToHeaderFileInfo.Count);
+					foreach(KeyValuePair<FileItem, HeaderFileInfo> Pair in FileToHeaderFileInfo)
 					{
 						Writer.WriteCompactFileItem(Pair.Key);
 						Writer.WriteLong(Pair.Value.LastWriteTimeUtc);
@@ -641,11 +589,12 @@ namespace UnrealBuildTool
 						Writer.WriteList(Pair.Value.Includes, Item => Writer.WriteString(Item));
 					}
 
-					Writer.WriteInt(FileToInlineReflectionInfo.Count);
-					foreach (KeyValuePair<FileItem, InlineReflectionInfo> Pair in FileToInlineReflectionInfo)
+					Writer.WriteInt(FileToSourceFileInfo.Count);
+					foreach (KeyValuePair<FileItem, SourceFileInfo> Pair in FileToSourceFileInfo)
 					{
 						Writer.WriteCompactFileItem(Pair.Key);
 						Writer.WriteLong(Pair.Value.LastWriteTimeUtc);
+						Writer.WriteString(Pair.Value.IncludeText);
 						Writer.WriteList(Pair.Value.InlinedFileNames, Item => Writer.WriteString(Item));
 					}
 				}
