- Hiding warning/error about missing identity (which is okay to be missing - depending Mac configuration it could be an error not a warning)
- Only use new dSYM utility when doing monolothic binaries, otherwise the normal path works fine (reduces risk)

diff --git a/Engine/Source/Programs/UnrealBuildTool/Configuration/UEBuildBinary.cs b/Engine/Source/Programs/UnrealBuildTool/Configuration/UEBuildBinary.cs
--- a/Engine/Source/Programs/UnrealBuildTool/Configuration/UEBuildBinary.cs
+++ b/Engine/Source/Programs/UnrealBuildTool/Configuration/UEBuildBinary.cs
@@ -767,6 +767,9 @@ namespace UnrealBuildTool
 			// Set the link output file.
 			BinaryLinkEnvironment.OutputFilePaths = OutputFilePaths.ToList();
 
+			// Rembmer the link type
+			BinaryLinkEnvironment.LinkType = Target.LinkType;
+
 			// Set whether the link is allowed to have exports.
 			BinaryLinkEnvironment.bHasExports = bAllowExports;
 
diff --git a/Engine/Source/Programs/UnrealBuildTool/Platform/Mac/MacToolChain.cs b/Engine/Source/Programs/UnrealBuildTool/Platform/Mac/MacToolChain.cs
--- a/Engine/Source/Programs/UnrealBuildTool/Platform/Mac/MacToolChain.cs
+++ b/Engine/Source/Programs/UnrealBuildTool/Platform/Mac/MacToolChain.cs
@@ -696,7 +696,7 @@ namespace UnrealBuildTool
 				// codesign with ad-hoc signature for when building outside of Xcode the there will be at least some signature
 				// (and it can be a BuildProduct down in ModifyBuildProducts)
 				AppendMacLine(FinalizeAppBundleScript, "codesign -f -s - \"{0}.app\"", ExeName);
-				AppendMacLine(FinalizeAppBundleScript, "codesign -f -s \"Developer ID Application\" \"{0}.app\"", ExeName);
+				AppendMacLine(FinalizeAppBundleScript, "codesign -f -s \"Developer ID Application\" \"{0}.app\" 2> /dev/null", ExeName);
 				AppendMacLine(FinalizeAppBundleScript, "echo done > /dev/null");
 
 				FinalizeAppBundleScript.Close();
@@ -970,7 +970,8 @@ namespace UnrealBuildTool
 			// a problem where dsymutil would exit with an error saying the input file did not exist.
 			// Note that the source and dest are switched from a copy command
 			string ExtraOptions = "";
-			string DsymutilPath = FileReference.Combine(Unreal.EngineDirectory, "Build/BatchFiles/Mac/GenerateUniversalDSYM.sh").FullName;
+			string DsymutilPath = GetDsymutilPath(Logger, out ExtraOptions, bIsForLTOBuild: false);
+			string UniversalDsymutilScriptPath = FileReference.Combine(Unreal.EngineDirectory, "Build/BatchFiles/Mac/GenerateUniversalDSYM.sh").FullName;
 
 			string ArgumentString = "-c \"";
 			ArgumentString += string.Format("for i in {{1..30}}; ");
@@ -985,11 +986,22 @@ namespace UnrealBuildTool
 			ArgumentString += string.Format("if [ ! -f \\\"{1}\\\" ] || [ \\\"{0}\\\" -nt \\\"{1}\\\" ] ; ", MachOBinary.AbsolutePath, OutputFile.AbsolutePath);
 			ArgumentString += string.Format("then ");
 				ArgumentString += string.Format("rm -rf \\\"{0}\\\"; ", OutputFile.AbsolutePath);
-				ArgumentString += string.Format(" \\\"{0}\\\" \\\"{1}\\\" \\\"{2}\\\"; ",
-					DsymutilPath,
-					MachOBinary.AbsolutePath,
-					OutputFile.AbsolutePath,
-					ExtraOptions);
+				// use the new script for monolthic (ie large) targets
+				if (LinkEnvironment.LinkType == TargetLinkType.Monolithic)
+				{
+					ArgumentString += string.Format(" \\\"{0}\\\" \\\"{1}\\\" \\\"{2}\\\"; ",
+						UniversalDsymutilScriptPath,
+						MachOBinary.AbsolutePath,
+						OutputFile.AbsolutePath);
+				}
+				else
+				{
+					ArgumentString += string.Format(" \\\"{0}\\\" {3} -f \\\"{1}\\\" -o \\\"{2}\\\"; ",
+						DsymutilPath,
+						MachOBinary.AbsolutePath,
+						OutputFile.AbsolutePath,
+						ExtraOptions);
+				}
 			ArgumentString += string.Format("fi; ");
 			ArgumentString += "\"";
 
diff --git a/Engine/Source/Programs/UnrealBuildTool/System/LinkEnvironment.cs b/Engine/Source/Programs/UnrealBuildTool/System/LinkEnvironment.cs
--- a/Engine/Source/Programs/UnrealBuildTool/System/LinkEnvironment.cs
+++ b/Engine/Source/Programs/UnrealBuildTool/System/LinkEnvironment.cs
@@ -158,6 +158,11 @@ namespace UnrealBuildTool
 		/// </summary>
 		public bool bIsBuildingDLL = false;
 
+		/// <summary>
+		/// The method of linking the target uses
+		/// </summary>
+		public TargetLinkType LinkType;
+		
 		/// <summary>
 		/// Whether we should compile using the statically-linked CRT. This is not widely supported for the whole engine, but is required for programs that need to run without dependencies.
 		/// </summary>
