Simplified the logic a little to hopefully fix a crash seen by others.

diff --git a/Engine/Source/Programs/UnrealBuildTool/Configuration/UEBuildTarget.cs b/Engine/Source/Programs/UnrealBuildTool/Configuration/UEBuildTarget.cs
--- a/Engine/Source/Programs/UnrealBuildTool/Configuration/UEBuildTarget.cs
+++ b/Engine/Source/Programs/UnrealBuildTool/Configuration/UEBuildTarget.cs
@@ -3448,8 +3448,7 @@ namespace UnrealBuildTool
 				return;
 			}
 
-			ConcurrentDictionary<UEBuildModuleCPP, CppCompileEnvironment> ModuleCompileEnvironments = new();
-			ConcurrentDictionary<PrecompiledHeaderTemplate, List<UEBuildModuleCPP>> PCHsAndModulesDict = new();
+			ConcurrentDictionary<PrecompiledHeaderTemplate, List<Tuple<UEBuildModuleCPP, CppCompileEnvironment>>> PCHsAndModulesDict = new();
 			Parallel.ForEach(OriginalBinaries, Binary =>
 			{
 				CppCompileEnvironment BinaryCompileEnvironment = Binary.CreateBinaryCompileEnvironment(GlobalCompileEnvironment);
@@ -3477,15 +3476,15 @@ namespace UnrealBuildTool
 							PrecompiledHeaderTemplate? Template = ModuleCompileEnvironment.SharedPCHs.FirstOrDefault(x => ReferencedModules.Contains(x.Module));
 							if (Template != null && Template.IsValidFor(ModuleCompileEnvironment))
 							{
-								ModuleCompileEnvironments.TryAdd(Module, ModuleCompileEnvironment);
+								Tuple<UEBuildModuleCPP, CppCompileEnvironment> NewModuleInfo = new(Module, ModuleCompileEnvironment);
 
 								PCHsAndModulesDict.AddOrUpdate(Template, k =>
 								{
-									return new List<UEBuildModuleCPP>() { Module };
+									return new List<Tuple<UEBuildModuleCPP, CppCompileEnvironment>>() { NewModuleInfo };
 								},
 								(k,v) =>
 								{
-									v.Add(Module);
+									v.Add(NewModuleInfo);
 									return v;
 								});
 							}
@@ -3499,11 +3498,11 @@ namespace UnrealBuildTool
 			SharedPCHs.Reverse();
 			foreach (var SharedPCH in SharedPCHs)
 			{
-				if (PCHsAndModulesDict.TryGetValue(SharedPCH, out List<UEBuildModuleCPP>? ModulesUsingSharedPCHs))
+				if (PCHsAndModulesDict.TryGetValue(SharedPCH, out List<Tuple<UEBuildModuleCPP, CppCompileEnvironment>>? ModulesUsingSharedPCHs))
 				{
-					foreach (var Module in ModulesUsingSharedPCHs)
+					foreach (var ModuleInfo in ModulesUsingSharedPCHs)
 					{
-						Module.FindOrCreateSharedPCH(ToolChain, SharedPCH, ModuleCompileEnvironments[Module], Graph);
+						ModuleInfo.Item1.FindOrCreateSharedPCH(ToolChain, SharedPCH, ModuleInfo.Item2, Graph);
 					}
 				}
 			}
