C++20 fixes

#rnx

diff --git a/Engine/Plugins/Experimental/PCG/Source/PCG/Private/Tests/PCGGraphParametersTest.cpp b/Engine/Plugins/Experimental/PCG/Source/PCG/Private/Tests/PCGGraphParametersTest.cpp
--- a/Engine/Plugins/Experimental/PCG/Source/PCG/Private/Tests/PCGGraphParametersTest.cpp
+++ b/Engine/Plugins/Experimental/PCG/Source/PCG/Private/Tests/PCGGraphParametersTest.cpp
@@ -452,7 +452,7 @@ bool FPCGGraphSetValueUserParametersPropagates::RunTest(const FString& Parameter
 	});
 
 	// Setting MyProperty to 3.0
-	constexpr double NewValue = 3.0;
+	const double NewValue = 3.0;
 	PCGTests::EmulateModifyingUserParametersValue(Graph, MyPropertyName, [Graph, MyPropertyName, NewValue]()
 	{
 		const FPropertyBagPropertyDesc* PropertyDesc = Graph->GetUserParametersStruct()->FindPropertyDescByName(MyPropertyName);
diff --git a/Engine/Source/Programs/UnrealBuildTool/Platform/Windows/VCToolChain.cs b/Engine/Source/Programs/UnrealBuildTool/Platform/Windows/VCToolChain.cs
--- a/Engine/Source/Programs/UnrealBuildTool/Platform/Windows/VCToolChain.cs
+++ b/Engine/Source/Programs/UnrealBuildTool/Platform/Windows/VCToolChain.cs
@@ -926,32 +926,34 @@ namespace UnrealBuildTool
 				case CppStandardVersion.Cpp20:
 					Arguments.Add("/std:c++20");
 					Arguments.Add("/Zc:preprocessor");
-
-					// warning C5054: operator ___: deprecated between enumerations of different types
-					// re: http://www.open-std.org/jtc1/sc22/wg21/docs/papers/2018/p1120r0.html
-						
-					// It seems unclear whether the deprecation will be enacted in C++23 or not
-					// e.g. http://www.open-std.org/jtc1/sc22/wg21/docs/papers/2020/p2139r2.html
-					// Until the path forward is clearer, it seems reasonable to leave things as they are.
-					Arguments.Add("/wd5054");
-
 					break;
 				case CppStandardVersion.Latest:
 					Arguments.Add("/std:c++latest");
-						
-					// warning C5054: operator ___: deprecated between enumerations of different types
-					// re: http://www.open-std.org/jtc1/sc22/wg21/docs/papers/2018/p1120r0.html
-						
-					// It seems unclear whether the deprecation will be enacted in C++23 or not
-					// e.g. http://www.open-std.org/jtc1/sc22/wg21/docs/papers/2020/p2139r2.html
-					// Until the path forward is clearer, it seems reasonable to leave things as they are.
-					Arguments.Add("/wd5054");
-						
+					Arguments.Add("/Zc:preprocessor");
 					break;
 				default:
 					throw new BuildException($"Unsupported C++ standard type set: {CompileEnvironment.CppStandard}");
 			}
 
+			if (CompileEnvironment.CppStandard >= CppStandardVersion.Cpp20)
+			{
+				// warning C5054: operator ___: deprecated between enumerations of different types
+				// re: http://www.open-std.org/jtc1/sc22/wg21/docs/papers/2018/p1120r0.html
+
+				// It seems unclear whether the deprecation will be enacted in C++23 or not
+				// e.g. http://www.open-std.org/jtc1/sc22/wg21/docs/papers/2020/p2139r2.html
+				// Until the path forward is clearer, it seems reasonable to leave things as they are.
+				Arguments.Add("/wd5054");
+
+				// Disable VS2019 C++20 warnings fixed with VS2022
+				if (Target.WindowsPlatform.Compiler == WindowsCompiler.VisualStudio2019)
+				{
+					Arguments.Add("/wd4005");
+					Arguments.Add("/wd4668");
+					Arguments.Add("/wd5105");
+				}
+			}
+
 			if (CompileEnvironment.bEnableCoroutines)
 			{
 				if (Target.WindowsPlatform.Compiler.IsMSVC())
diff --git a/Engine/Source/Runtime/AudioCaptureImplementations/AudioCaptureRtAudio/AudioCaptureRtAudio.Build.cs b/Engine/Source/Runtime/AudioCaptureImplementations/AudioCaptureRtAudio/AudioCaptureRtAudio.Build.cs
--- a/Engine/Source/Runtime/AudioCaptureImplementations/AudioCaptureRtAudio/AudioCaptureRtAudio.Build.cs
+++ b/Engine/Source/Runtime/AudioCaptureImplementations/AudioCaptureRtAudio/AudioCaptureRtAudio.Build.cs
@@ -8,6 +8,9 @@ public class AudioCaptureRtAudio : ModuleRules
 
 	public AudioCaptureRtAudio(ReadOnlyTargetRules Target) : base(Target)
 	{
+		// warning C5056: operator '!=': deprecated for array types
+		CppStandard = CppStandardVersion.Cpp17;
+
 		// RtAudio throws exceptions for errors.
 		bEnableExceptions = true;
 
