Fixed thread unsafe code in CreateSharedPCHInstances.

diff --git a/Engine/Source/Programs/UnrealBuildTool/Configuration/UEBuildTarget.cs b/Engine/Source/Programs/UnrealBuildTool/Configuration/UEBuildTarget.cs
--- a/Engine/Source/Programs/UnrealBuildTool/Configuration/UEBuildTarget.cs
+++ b/Engine/Source/Programs/UnrealBuildTool/Configuration/UEBuildTarget.cs
@@ -3443,12 +3443,18 @@ namespace UnrealBuildTool
 		/// <param name="Graph">List of build actions</param>
 		void CreateSharedPCHInstances(ReadOnlyTargetRules Target, UEToolChain ToolChain, List<UEBuildBinary> OriginalBinaries, CppCompileEnvironment GlobalCompileEnvironment, IActionGraphBuilder Graph, ILogger Logger)
 		{
-			if (!Target.bUsePCHFiles || GlobalCompileEnvironment.SharedPCHs.Count == 0)
+			int NumSharedPCHs = GlobalCompileEnvironment.SharedPCHs.Count;
+			if (!Target.bUsePCHFiles || NumSharedPCHs == 0)
 			{
 				return;
 			}
 
-			ConcurrentDictionary<PrecompiledHeaderTemplate, List<Tuple<UEBuildModuleCPP, CppCompileEnvironment>>> PCHsAndModulesDict = new();
+			ConcurrentBag<Tuple<UEBuildModuleCPP, CppCompileEnvironment>>[] ModuleInfos = new ConcurrentBag<Tuple<UEBuildModuleCPP, CppCompileEnvironment>>[NumSharedPCHs];
+			for (int SharedPCHIndex = 0; SharedPCHIndex < NumSharedPCHs; SharedPCHIndex++)
+			{
+				ModuleInfos[SharedPCHIndex] = new ConcurrentBag<Tuple<UEBuildModuleCPP, CppCompileEnvironment>>();
+			}
+
 			Parallel.ForEach(OriginalBinaries, Binary =>
 			{
 				CppCompileEnvironment BinaryCompileEnvironment = Binary.CreateBinaryCompileEnvironment(GlobalCompileEnvironment);
@@ -3473,20 +3479,10 @@ namespace UnrealBuildTool
 							Module.GetAllDependencyModules(new List<UEBuildModule>(), ReferencedModules, bIncludeDynamicallyLoaded: false, bForceCircular: false, bOnlyDirectDependencies: true);
 
 							// Find the first shared PCH module we can use
-							PrecompiledHeaderTemplate? Template = ModuleCompileEnvironment.SharedPCHs.FirstOrDefault(x => ReferencedModules.Contains(x.Module));
-							if (Template != null && Template.IsValidFor(ModuleCompileEnvironment))
+							int SharedPCHIndex = ModuleCompileEnvironment.SharedPCHs.FindIndex(x => ReferencedModules.Contains(x.Module));
+							if (SharedPCHIndex != -1 && GlobalCompileEnvironment.SharedPCHs[SharedPCHIndex].IsValidFor(ModuleCompileEnvironment))
 							{
-								Tuple<UEBuildModuleCPP, CppCompileEnvironment> NewModuleInfo = new(Module, ModuleCompileEnvironment);
-
-								PCHsAndModulesDict.AddOrUpdate(Template, k =>
-								{
-									return new List<Tuple<UEBuildModuleCPP, CppCompileEnvironment>>() { NewModuleInfo };
-								},
-								(k,v) =>
-								{
-									v.Add(NewModuleInfo);
-									return v;
-								});
+								ModuleInfos[SharedPCHIndex].Add(new Tuple<UEBuildModuleCPP, CppCompileEnvironment>(Module, ModuleCompileEnvironment));
 							}
 						}
 					}
@@ -3494,16 +3490,11 @@ namespace UnrealBuildTool
 			});
 
 			// Create the PCH instances
-			List<PrecompiledHeaderTemplate> SharedPCHs = new List<PrecompiledHeaderTemplate>(GlobalCompileEnvironment.SharedPCHs);
-			SharedPCHs.Reverse();
-			foreach (var SharedPCH in SharedPCHs)
+			for (int SharedPCHIndex = NumSharedPCHs - 1; SharedPCHIndex >= 0; SharedPCHIndex--)
 			{
-				if (PCHsAndModulesDict.TryGetValue(SharedPCH, out List<Tuple<UEBuildModuleCPP, CppCompileEnvironment>>? ModulesUsingSharedPCHs))
+				foreach (var ModuleInfo in ModuleInfos[SharedPCHIndex])
 				{
-					foreach (var ModuleInfo in ModulesUsingSharedPCHs)
-					{
-						ModuleInfo.Item1.FindOrCreateSharedPCH(ToolChain, SharedPCH, ModuleInfo.Item2, Graph);
-					}
+					ModuleInfo.Item1.FindOrCreateSharedPCH(ToolChain, GlobalCompileEnvironment.SharedPCHs[SharedPCHIndex], ModuleInfo.Item2, Graph);
 				}
 			}
 		}
