Small opts in the project generation code.
#p4v-preflight-copy 25094409

diff --git a/Engine/Source/Programs/UnrealBuildTool/Modes/GenerateProjectFilesMode.cs b/Engine/Source/Programs/UnrealBuildTool/Modes/GenerateProjectFilesMode.cs
--- a/Engine/Source/Programs/UnrealBuildTool/Modes/GenerateProjectFilesMode.cs
+++ b/Engine/Source/Programs/UnrealBuildTool/Modes/GenerateProjectFilesMode.cs
@@ -18,7 +18,7 @@ namespace UnrealBuildTool
 	/// <summary>
 	/// Generates project files for one or more projects
 	/// </summary>
-	[ToolMode("GenerateProjectFiles", ToolModeOptions.XmlConfig | ToolModeOptions.BuildPlatforms | ToolModeOptions.SingleInstance | ToolModeOptions.UseStartupTraceListener)]
+	[ToolMode("GenerateProjectFiles", ToolModeOptions.XmlConfig | ToolModeOptions.BuildPlatforms | ToolModeOptions.SingleInstance | ToolModeOptions.UseStartupTraceListener | ToolModeOptions.StartPrefetchingEngine)]
 	class GenerateProjectFilesMode : ToolMode
 	{
 		/// <summary>
diff --git a/Engine/Source/Programs/UnrealBuildTool/ProjectFiles/VisualStudio/VCProject.cs b/Engine/Source/Programs/UnrealBuildTool/ProjectFiles/VisualStudio/VCProject.cs
--- a/Engine/Source/Programs/UnrealBuildTool/ProjectFiles/VisualStudio/VCProject.cs
+++ b/Engine/Source/Programs/UnrealBuildTool/ProjectFiles/VisualStudio/VCProject.cs
@@ -865,43 +865,30 @@ namespace UnrealBuildTool
 			StringBuilder SharedIncludeSearchPaths = new StringBuilder();
 			{
 				// Find out how many source files there are in each directory
-				Dictionary<DirectoryReference, int> SourceDirToCount = new Dictionary<DirectoryReference, int>();
-				foreach (SourceFile SourceFile in SourceFiles)
+				ConcurrentDictionary<DirectoryReference, int> SourceDirToCount = new();
+				Parallel.ForEach(SourceFiles.Where(sf => sf.Reference.HasExtension(".cpp")), SourceFile =>
 				{
-					if (SourceFile.Reference.HasExtension(".cpp"))
-					{
-						DirectoryReference SourceDir = SourceFile.Reference.Directory;
-
-						int Count;
-						SourceDirToCount.TryGetValue(SourceDir, out Count);
-						SourceDirToCount[SourceDir] = Count + 1;
-					}
-				}
-
+					SourceDirToCount.AddOrUpdate(SourceFile.Reference.Directory, _ => 1, (k, v) => v + 1);
+				});
+				
 				// Figure out the most common include paths
-				Dictionary<DirectoryReference, int> IncludePathToCount = new Dictionary<DirectoryReference, int>();
-				foreach (KeyValuePair<DirectoryReference, int> Pair in SourceDirToCount)
+				ConcurrentDictionary<DirectoryReference, int> IncludePathToCount = new();
+				Parallel.ForEach(SourceDirToCount, Pair =>
 				{
-					for (DirectoryReference? CurrentDir = Pair.Key; CurrentDir != null; CurrentDir = CurrentDir.ParentDirectory)
+					if (TryGetBuildEnvironment(Pair.Key, out BuildEnvironment? OutBuildEnvironment))
 					{
-						BuildEnvironment? BuildEnvironment;
-						if (BaseDirToBuildEnvironment.TryGetValue(CurrentDir, out BuildEnvironment))
+						foreach (DirectoryReference IncludePath in OutBuildEnvironment.UserIncludePaths.AbsolutePaths)
 						{
-							foreach (DirectoryReference IncludePath in BuildEnvironment.UserIncludePaths.AbsolutePaths)
-							{
-								int Count;
-								IncludePathToCount.TryGetValue(IncludePath, out Count);
-								IncludePathToCount[IncludePath] = Count + Pair.Value;
-							}
-							break;
+							IncludePathToCount.AddOrUpdate(IncludePath, _ => Pair.Value, (k, v) => v + Pair.Value);
 						}
+						return;
 					}
-				}
+				});
 
 				// Append the most common include paths to the search list.
 				if (Settings.MaxSharedIncludePaths > 0)
 				{
-					foreach (DirectoryReference IncludePath in IncludePathToCount.OrderByDescending(x => x.Value).Select(x => x.Key))
+					foreach (DirectoryReference IncludePath in IncludePathToCount.OrderByDescending(x => x.Value).ThenBy(x => x.Key).Select(x => x.Key))
 					{
 						string RelativePath = NormalizeProjectPath(IncludePath);
 						if (SharedIncludeSearchPaths.Length + RelativePath.Length >= Settings.MaxSharedIncludePaths)
