[UBT]
* New Box binaries with tons of fixes.

diff --git a/Engine/Source/Programs/AutomationTool/AutomationTool.csproj b/Engine/Source/Programs/AutomationTool/AutomationTool.csproj
--- a/Engine/Source/Programs/AutomationTool/AutomationTool.csproj
+++ b/Engine/Source/Programs/AutomationTool/AutomationTool.csproj
@@ -111,19 +111,13 @@
 
   <Choose>
     <When Condition="'$(IsWindows)'=='true'">
-      <ItemGroup Condition="Exists('..\..\..\Restricted\NotForLicensees\Binaries\Win64\Box\Box.dll')">
-        <Reference Include="..\..\..\Restricted\NotForLicensees\Binaries\Win64\Box\Box.dll" />
-        <Content Include="..\..\..\Restricted\NotForLicensees\Binaries\Win64\Box\BoxNative.dll">
-          <CopyToOutputDirectory>PreserveNewest</CopyToOutputDirectory>
-        </Content>
+      <ItemGroup Condition="Exists('..\..\..\Restricted\NotForLicensees\Binaries\Win64\Box\BoxManaged.dll')">
+        <Reference Include="..\..\..\Restricted\NotForLicensees\Binaries\Win64\Box\BoxManaged.dll" />
         <Content Include="..\..\..\Restricted\NotForLicensees\Binaries\Win64\Box\BoxDetours.dll">
           <CopyToOutputDirectory>PreserveNewest</CopyToOutputDirectory>
         </Content>
-        <Content Include="..\..\..\Restricted\NotForLicensees\Binaries\Win64\Box\zstd.dll">
-          <CopyToOutputDirectory>PreserveNewest</CopyToOutputDirectory>
-        </Content>
       </ItemGroup>
-      <PropertyGroup Condition="Exists('..\..\..\Restricted\NotForLicensees\Binaries\Win64\Box\Box.dll')">
+      <PropertyGroup Condition="Exists('..\..\..\Restricted\NotForLicensees\Binaries\Win64\Box\BoxManaged.dll')">
         <DefineConstants>$(DefineConstants);__BOXEXECUTOR_AVAILABLE__</DefineConstants>
       </PropertyGroup>
     </When>
diff --git a/Engine/Source/Programs/UnrealBuildTool/Platform/Windows/VCToolChain.cs b/Engine/Source/Programs/UnrealBuildTool/Platform/Windows/VCToolChain.cs
--- a/Engine/Source/Programs/UnrealBuildTool/Platform/Windows/VCToolChain.cs
+++ b/Engine/Source/Programs/UnrealBuildTool/Platform/Windows/VCToolChain.cs
@@ -1566,7 +1566,7 @@ namespace UnrealBuildTool
 					}
 				}
 
-				CompileAction.bCanCache = CompileEnvironment.bDeterministic;
+				CompileAction.bCanCache = CompileEnvironment.bDeterministic && !Target.WindowsPlatform.Compiler.IsClang();
 
 				// Don't farm out creation of precompiled headers as it is the critical path task.
 				CompileAction.bCanExecuteRemotely =
@@ -1709,8 +1709,6 @@ namespace UnrealBuildTool
 					CompileAction.WriteResponseFile(Graph, Logger);
 				}
 
-				CompileAction.bCanCache = CompileEnvironment.bDeterministic;
-
 				// When compiling with SN-DBS, modules that contain a #import must be built locally
 				CompileAction.bCanExecuteRemotelyWithSNDBS = CompileAction.bCanExecuteRemotely;
 				if (CompileEnvironment.bBuildLocallyWithSNDBS == true)
diff --git a/Engine/Source/Programs/UnrealBuildTool/UnrealBuildTool.csproj b/Engine/Source/Programs/UnrealBuildTool/UnrealBuildTool.csproj
--- a/Engine/Source/Programs/UnrealBuildTool/UnrealBuildTool.csproj
+++ b/Engine/Source/Programs/UnrealBuildTool/UnrealBuildTool.csproj
@@ -73,19 +73,13 @@
 
   <Choose>
     <When Condition="'$([System.Runtime.InteropServices.RuntimeInformation]::IsOSPlatform($([System.Runtime.InteropServices.OSPlatform]::Windows)))' == 'true'">
-      <ItemGroup Condition="Exists('..\..\..\Restricted\NotForLicensees\Binaries\Win64\Box\Box.dll')">
-        <Reference Include="..\..\..\Restricted\NotForLicensees\Binaries\Win64\Box\Box.dll" />
-        <Content Include="..\..\..\Restricted\NotForLicensees\Binaries\Win64\Box\BoxNative.dll">
-          <CopyToOutputDirectory>PreserveNewest</CopyToOutputDirectory>
-        </Content>
+      <ItemGroup Condition="Exists('..\..\..\Restricted\NotForLicensees\Binaries\Win64\Box\BoxManaged.dll')">
+        <Reference Include="..\..\..\Restricted\NotForLicensees\Binaries\Win64\Box\BoxManaged.dll" />
         <Content Include="..\..\..\Restricted\NotForLicensees\Binaries\Win64\Box\BoxDetours.dll">
           <CopyToOutputDirectory>PreserveNewest</CopyToOutputDirectory>
         </Content>
-        <Content Include="..\..\..\Restricted\NotForLicensees\Binaries\Win64\Box\zstd.dll">
-          <CopyToOutputDirectory>PreserveNewest</CopyToOutputDirectory>
-        </Content>
       </ItemGroup>
-      <PropertyGroup Condition="Exists('..\..\..\Restricted\NotForLicensees\Binaries\Win64\Box\Box.dll')">
+      <PropertyGroup Condition="Exists('..\..\..\Restricted\NotForLicensees\Binaries\Win64\Box\BoxManaged.dll')">
         <DefineConstants>$(DefineConstants);__BOXEXECUTOR_AVAILABLE__</DefineConstants>
       </PropertyGroup>
     </When>
