[UBT]
* Added so ljwhost.dll is copied when built

diff --git a/Engine/Source/Programs/AutomationTool/AutomationTool.csproj b/Engine/Source/Programs/AutomationTool/AutomationTool.csproj
--- a/Engine/Source/Programs/AutomationTool/AutomationTool.csproj
+++ b/Engine/Source/Programs/AutomationTool/AutomationTool.csproj
@@ -116,6 +116,9 @@
         <Content Include="..\..\..\Restricted\NotForLicensees\Binaries\Win64\Box\BoxDetours.dll">
           <CopyToOutputDirectory>PreserveNewest</CopyToOutputDirectory>
         </Content>
+        <Content Include="..\..\..\Restricted\NotForLicensees\Binaries\Win64\Box\Ijwhost.dll">
+          <CopyToOutputDirectory>PreserveNewest</CopyToOutputDirectory>
+        </Content>
       </ItemGroup>
       <PropertyGroup Condition="Exists('..\..\..\Restricted\NotForLicensees\Binaries\Win64\Box\BoxManaged.dll')">
         <DefineConstants>$(DefineConstants);__BOXEXECUTOR_AVAILABLE__</DefineConstants>
diff --git a/Engine/Source/Programs/UnrealBuildTool/UnrealBuildTool.csproj b/Engine/Source/Programs/UnrealBuildTool/UnrealBuildTool.csproj
--- a/Engine/Source/Programs/UnrealBuildTool/UnrealBuildTool.csproj
+++ b/Engine/Source/Programs/UnrealBuildTool/UnrealBuildTool.csproj
@@ -78,6 +78,9 @@
         <Content Include="..\..\..\Restricted\NotForLicensees\Binaries\Win64\Box\BoxDetours.dll">
           <CopyToOutputDirectory>PreserveNewest</CopyToOutputDirectory>
         </Content>
+        <Content Include="..\..\..\Restricted\NotForLicensees\Binaries\Win64\Box\Ijwhost.dll">
+          <CopyToOutputDirectory>PreserveNewest</CopyToOutputDirectory>
+        </Content>
       </ItemGroup>
       <PropertyGroup Condition="Exists('..\..\..\Restricted\NotForLicensees\Binaries\Win64\Box\BoxManaged.dll')">
         <DefineConstants>$(DefineConstants);__BOXEXECUTOR_AVAILABLE__</DefineConstants>
