- Fixed warnings when building C++ projects with InstalledBuilds, where it was linking the binary with both arm64 and x64 engine .o files

diff --git a/Engine/Source/Programs/UnrealBuildTool/Platform/Mac/MacToolChain.cs b/Engine/Source/Programs/UnrealBuildTool/Platform/Mac/MacToolChain.cs
--- a/Engine/Source/Programs/UnrealBuildTool/Platform/Mac/MacToolChain.cs
+++ b/Engine/Source/Programs/UnrealBuildTool/Platform/Mac/MacToolChain.cs
@@ -716,6 +716,11 @@ namespace UnrealBuildTool
 		{
 			List<FileItem> InputFiles = LinkEnvironment.InputFiles;
 
+			// with precompiled builds, we end up trying to link the arm64 and x64 versions of the engine .o files. Remove the wrong one
+			// until we can separate out the .o files by architecture
+			string ArchToRemove = LinkEnvironment.Architecture == UnrealArch.Arm64 ? "x64" : "arm64";
+			InputFiles.RemoveAll(x => x.Location.ContainsName(ArchToRemove, 0));
+
 			// Create an action that invokes the linker.
 			Action LinkAction = Graph.CreateAction(ActionType.Link);
 
