Adds a new `bEnableLibFuzzer` option for desktop platforms.

diff --git a/Engine/Source/Programs/UnrealBuildTool/Platform/Linux/LinuxToolChain.cs b/Engine/Source/Programs/UnrealBuildTool/Platform/Linux/LinuxToolChain.cs
--- a/Engine/Source/Programs/UnrealBuildTool/Platform/Linux/LinuxToolChain.cs
+++ b/Engine/Source/Programs/UnrealBuildTool/Platform/Linux/LinuxToolChain.cs
@@ -453,6 +453,12 @@ namespace UnrealBuildTool
 				Arguments.Add("-fsanitize-memory-track-origins");
 				Arguments.Add("-DFORCE_ANSI_ALLOCATOR=1");
 			}
+
+			// LibFuzzer
+			if (Options.HasFlag(ClangToolChainOptions.EnableLibFuzzer))
+			{
+				Arguments.Add("-fsanitize=fuzzer");
+			}
 		}
 
 		/// <inheritdoc/>
@@ -585,7 +591,8 @@ namespace UnrealBuildTool
 			if (Options.HasFlag(ClangToolChainOptions.EnableAddressSanitizer) ||
 				Options.HasFlag(ClangToolChainOptions.EnableThreadSanitizer) ||
 				Options.HasFlag(ClangToolChainOptions.EnableUndefinedBehaviorSanitizer) ||
-				Options.HasFlag(ClangToolChainOptions.EnableMemorySanitizer))
+				Options.HasFlag(ClangToolChainOptions.EnableMemorySanitizer) ||
+				Options.HasFlag(ClangToolChainOptions.EnableLibFuzzer))
 			{
 				Arguments.Add("-g");
 
@@ -618,6 +625,11 @@ namespace UnrealBuildTool
 					Arguments.Add("-fsanitize=memory -fsanitize-memory-track-origins");
 				}
 
+				if (Options.HasFlag(ClangToolChainOptions.EnableLibFuzzer))
+				{
+					Arguments.Add("-fsanitize=fuzzer");
+				}
+
 				if (CrossCompiling())
 				{
 					// x64 only replaced the linux folder with arch, while on arm64 its still linux
@@ -780,7 +792,8 @@ namespace UnrealBuildTool
 			if (Options.HasFlag(ClangToolChainOptions.EnableAddressSanitizer) ||
 				Options.HasFlag(ClangToolChainOptions.EnableThreadSanitizer) ||
 				Options.HasFlag(ClangToolChainOptions.EnableUndefinedBehaviorSanitizer) ||
-				Options.HasFlag(ClangToolChainOptions.EnableMemorySanitizer))
+				Options.HasFlag(ClangToolChainOptions.EnableMemorySanitizer) ||
+				Options.HasFlag(ClangToolChainOptions.EnableLibFuzzer))
 			{
 				string SanitizerInfo = "Building with:";
 				string StaticOrShared = Options.HasFlag(ClangToolChainOptions.EnableSharedSanitizer) ? " dynamically" : " statically";
@@ -789,6 +802,7 @@ namespace UnrealBuildTool
 				SanitizerInfo += Options.HasFlag(ClangToolChainOptions.EnableThreadSanitizer) ? StaticOrShared + " linked ThreadSanitizer" : "";
 				SanitizerInfo += Options.HasFlag(ClangToolChainOptions.EnableUndefinedBehaviorSanitizer) ? StaticOrShared + " linked UndefinedBehaviorSanitizer" : "";
 				SanitizerInfo += Options.HasFlag(ClangToolChainOptions.EnableMemorySanitizer) ? StaticOrShared + " linked MemorySanitizer" : "";
+				SanitizerInfo += Options.HasFlag(ClangToolChainOptions.EnableLibFuzzer) ? StaticOrShared + " linked LibFuzzer" : "";
 
 				Logger.LogInformation("{SanitizerInfo}", SanitizerInfo);
 			}
diff --git a/Engine/Source/Programs/UnrealBuildTool/Platform/Linux/UEBuildLinux.cs b/Engine/Source/Programs/UnrealBuildTool/Platform/Linux/UEBuildLinux.cs
--- a/Engine/Source/Programs/UnrealBuildTool/Platform/Linux/UEBuildLinux.cs
+++ b/Engine/Source/Programs/UnrealBuildTool/Platform/Linux/UEBuildLinux.cs
@@ -54,6 +54,13 @@ namespace UnrealBuildTool
 		[XmlConfigFile(Category = "BuildConfiguration", Name = "bEnableAddressSanitizer")]
 		public bool bEnableAddressSanitizer = false;
 
+		/// <summary>
+		/// Enables LibFuzzer
+		/// </summary>
+		[CommandLine("-EnableLibFuzzer")]
+		[XmlConfigFile(Category = "BuildConfiguration", Name = "bEnableLibFuzzer")]
+		public bool bEnableLibFuzzer = false;
+
 		/// <summary>
 		/// Enables thread sanitizer (TSan)
 		/// </summary>
@@ -134,6 +141,11 @@ namespace UnrealBuildTool
 			get { return Inner.bEnableAddressSanitizer; }
 		}
 
+		public bool bEnableLibFuzzer
+		{
+			get { return Inner.bEnableLibFuzzer; }
+		}
+
 		public bool bEnableThreadSanitizer
 		{
 			get { return Inner.bEnableThreadSanitizer; }
@@ -255,6 +267,10 @@ namespace UnrealBuildTool
 				{
 					SanitizerSuffix = "MSan";
 				}
+				if (Target.LinuxPlatform.bEnableLibFuzzer)
+				{
+					SanitizerSuffix += "LibFuzzer";
+				}
 
 				if (!String.IsNullOrEmpty(SanitizerSuffix))
 				{
@@ -600,6 +616,15 @@ namespace UnrealBuildTool
 			{
 				Options |= ClangToolChainOptions.DisableDumpSyms;
 			}
+			if (Target.LinuxPlatform.bEnableLibFuzzer)
+			{
+				Options |= ClangToolChainOptions.EnableLibFuzzer;
+
+				if (Target.LinkType != TargetLinkType.Monolithic)
+				{
+					throw new BuildException("LibFuzzer is unsupported for non-monolithic builds.");
+				}
+			}
 			if (Target.bAllowLTCG && Target.bPreferThinLTO)
 			{
 				Options |= ClangToolChainOptions.EnableThinLTO;
diff --git a/Engine/Source/Programs/UnrealBuildTool/Platform/Mac/MacToolChain.cs b/Engine/Source/Programs/UnrealBuildTool/Platform/Mac/MacToolChain.cs
--- a/Engine/Source/Programs/UnrealBuildTool/Platform/Mac/MacToolChain.cs
+++ b/Engine/Source/Programs/UnrealBuildTool/Platform/Mac/MacToolChain.cs
@@ -291,7 +291,10 @@ namespace UnrealBuildTool
 				Arguments.Add("-Wl,-fatal_warnings");
 			}
 
-			if (Options.HasFlag(ClangToolChainOptions.EnableAddressSanitizer) || Options.HasFlag(ClangToolChainOptions.EnableThreadSanitizer) || Options.HasFlag(ClangToolChainOptions.EnableUndefinedBehaviorSanitizer))
+			if (Options.HasFlag(ClangToolChainOptions.EnableAddressSanitizer) ||
+				Options.HasFlag(ClangToolChainOptions.EnableThreadSanitizer) ||
+				Options.HasFlag(ClangToolChainOptions.EnableUndefinedBehaviorSanitizer) ||
+				Options.HasFlag(ClangToolChainOptions.EnableLibFuzzer))
 			{
 				Arguments.Add("-g");
 				if (Options.HasFlag(ClangToolChainOptions.EnableAddressSanitizer))
@@ -306,6 +309,10 @@ namespace UnrealBuildTool
 				{
 					Arguments.Add("-fsanitize=undefined");
 				}
+				if (Options.HasFlag(ClangToolChainOptions.EnableLibFuzzer))
+				{
+					Arguments.Add("-fsanitize=fuzzer");
+				}
 			}
 
 			if (LinkEnvironment.bIsBuildingDLL)
diff --git a/Engine/Source/Programs/UnrealBuildTool/Platform/Mac/UEBuildMac.cs b/Engine/Source/Programs/UnrealBuildTool/Platform/Mac/UEBuildMac.cs
--- a/Engine/Source/Programs/UnrealBuildTool/Platform/Mac/UEBuildMac.cs
+++ b/Engine/Source/Programs/UnrealBuildTool/Platform/Mac/UEBuildMac.cs
@@ -26,6 +26,13 @@ namespace UnrealBuildTool
 		[XmlConfigFile(Category = "BuildConfiguration", Name = "bEnableAddressSanitizer")]
 		public bool bEnableAddressSanitizer = false;
 
+		/// <summary>
+		/// Enables LibFuzzer.
+		/// </summary>
+		[CommandLine("-EnableLibFuzzer")]
+		[XmlConfigFile(Category = "BuildConfiguration", Name = "bEnableLibFuzzer")]
+		public bool bEnableLibFuzzer = false;
+
 		/// <summary>
 		/// Enables thread sanitizer (TSan).
 		/// </summary>
@@ -97,6 +104,11 @@ namespace UnrealBuildTool
 			get { return Inner.bEnableAddressSanitizer; }
 		}
 
+		public bool bEnableLibFuzzer
+		{
+			get { return Inner.bEnableLibFuzzer; }
+		}
+
 		public bool bEnableThreadSanitizer
 		{
 			get { return Inner.bEnableThreadSanitizer; }
diff --git a/Engine/Source/Programs/UnrealBuildTool/Platform/Windows/UEBuildWindows.cs b/Engine/Source/Programs/UnrealBuildTool/Platform/Windows/UEBuildWindows.cs
--- a/Engine/Source/Programs/UnrealBuildTool/Platform/Windows/UEBuildWindows.cs
+++ b/Engine/Source/Programs/UnrealBuildTool/Platform/Windows/UEBuildWindows.cs
@@ -364,6 +364,13 @@ namespace UnrealBuildTool
 		[CommandLine("-EnableASan")]
 		public bool bEnableAddressSanitizer = false;
 
+		/// <summary>
+		/// Enables LibFuzzer. Only supported for Visual Studio 2022 17.0.0 and up.
+		/// </summary>
+		[XmlConfigFile(Category = "BuildConfiguration", Name = "bEnableLibFuzzer")]
+		[CommandLine("-EnableLibFuzzer")]
+		public bool bEnableLibFuzzer = false;
+
 		/// <summary>
 		/// Whether .sarif files containing errors and warnings are written alongside each .obj, if supported
 		/// </summary>
@@ -771,6 +778,11 @@ namespace UnrealBuildTool
 			get { return Inner.bEnableAddressSanitizer; }
 		}
 
+		public bool bEnableLibFuzzer
+		{ 
+			get { return Inner.bEnableLibFuzzer; }
+		}
+        
 		public bool bWriteSarif
 		{
 			get { return Inner.bWriteSarif; }
@@ -1184,6 +1196,12 @@ namespace UnrealBuildTool
 				throw new BuildException("Address sanitizer requires Visual Studio 2019 16.7 or later.");
 			}
 
+			// Ensure we're using recent enough version of Visual Studio to support LibFuzzer.
+			if (Target.WindowsPlatform.bEnableLibFuzzer && Target.WindowsPlatform.Environment.CompilerVersion < new VersionNumber(14, 30, 0))
+			{
+				throw new BuildException("LibFuzzer MSVC support requires Visual Studio 2022 17.0 or later (14.30.0). The current compiler version was detected as: {0}", Target.WindowsPlatform.Environment.CompilerVersion.ToString());
+			}
+
 //			@Todo: Still getting reports of frequent OOM issues with this enabled as of 15.7.
 //			// Enable fast PDB linking if we're on VS2017 15.7 or later. Previous versions have OOM issues with large projects.
 //			if(!Target.bFormalBuild && !Target.bUseFastPDBLinking.HasValue && Target.WindowsPlatform.Compiler.IsMSVC())
diff --git a/Engine/Source/Programs/UnrealBuildTool/Platform/Windows/VCToolChain.cs b/Engine/Source/Programs/UnrealBuildTool/Platform/Windows/VCToolChain.cs
--- a/Engine/Source/Programs/UnrealBuildTool/Platform/Windows/VCToolChain.cs
+++ b/Engine/Source/Programs/UnrealBuildTool/Platform/Windows/VCToolChain.cs
@@ -596,6 +596,18 @@ namespace UnrealBuildTool
 				AddSystemIncludePath(Arguments, DirectoryReference.Combine(EnvVars.CompilerDir, "crt", "src"), Target.WindowsPlatform.Compiler);
 			}
 
+			if (Target.WindowsPlatform.bEnableLibFuzzer)
+			{
+				if (Target.WindowsPlatform.Compiler.IsClang())
+				{
+					Arguments.Add("-fsanitize=fuzzer");
+				}
+				else
+				{
+					Arguments.Add("/fsanitize=fuzzer");
+				}
+			}
+
 			//
 			//	Debug
 			//
diff --git a/Engine/Source/Programs/UnrealBuildTool/ToolChain/ClangToolChain.cs b/Engine/Source/Programs/UnrealBuildTool/ToolChain/ClangToolChain.cs
--- a/Engine/Source/Programs/UnrealBuildTool/ToolChain/ClangToolChain.cs
+++ b/Engine/Source/Programs/UnrealBuildTool/ToolChain/ClangToolChain.cs
@@ -107,6 +107,11 @@ namespace UnrealBuildTool
 		/// Indicates that the AutoRTFM Clang compiler should be used instead of the standard clang compiler
 		/// </summary>
 		UseAutoRTFMCompiler = 1 << 17,
+
+		/// <summary>
+		/// Enable LibFuzzer
+		/// </summary>
+		EnableLibFuzzer = 1 << 18,
 	}
 
 	abstract class ClangToolChain : ISPCToolChain
@@ -769,6 +774,12 @@ namespace UnrealBuildTool
 			{
 				Arguments.Add("-fsanitize=memory");
 			}
+
+			// LibFuzzer
+			if (Options.HasFlag(ClangToolChainOptions.EnableLibFuzzer))
+			{
+				Arguments.Add("-fsanitize=fuzzer");
+			}
 		}
 
 		/// <summary>
