UnrealBuildTool: MSVC deterministic mode pch variant

#rnx

diff --git a/Engine/Source/Programs/UnrealBuildTool/Configuration/UEBuildModuleCPP.cs b/Engine/Source/Programs/UnrealBuildTool/Configuration/UEBuildModuleCPP.cs
--- a/Engine/Source/Programs/UnrealBuildTool/Configuration/UEBuildModuleCPP.cs
+++ b/Engine/Source/Programs/UnrealBuildTool/Configuration/UEBuildModuleCPP.cs
@@ -1151,6 +1151,18 @@ namespace UnrealBuildTool
 				}
 			}
 
+			if (CompileEnvironment.bDeterministic != BaseCompileEnvironment.bDeterministic)
+			{
+				if (CompileEnvironment.bDeterministic)
+				{
+					Variant += ".Determ";
+				}
+				else
+				{
+					Variant += ".NonDeterm";
+				}
+			}
+
 			if (CompileEnvironment.CppStandard != BaseCompileEnvironment.CppStandard)
 			{
 				Variant += String.Format(".{0}", CompileEnvironment.CppStandard);
diff --git a/Engine/Source/Programs/UnrealBuildTool/Platform/Windows/VCToolChain.cs b/Engine/Source/Programs/UnrealBuildTool/Platform/Windows/VCToolChain.cs
--- a/Engine/Source/Programs/UnrealBuildTool/Platform/Windows/VCToolChain.cs
+++ b/Engine/Source/Programs/UnrealBuildTool/Platform/Windows/VCToolChain.cs
@@ -541,7 +541,7 @@ namespace UnrealBuildTool
 				Arguments.Add("/Ob0");
 			}
 
-			// Experimental deterministic compile support
+			// Deterministic compile support
 			if (CompileEnvironment.bDeterministic)
 			{
 				if (Target.WindowsPlatform.Compiler.IsMSVC())
@@ -1271,7 +1271,7 @@ namespace UnrealBuildTool
 			// Don't embed the full PDB path; we want to be able to move binaries elsewhere. They will always be side by side.
 			Arguments.Add("/PDBALTPATH:%_PDB%");
 
-			// Experimental deterministic link support
+			// Deterministic link support
 			if (LinkEnvironment.bDeterministic)
 			{
 				if (Target.WindowsPlatform.Compiler.IsMSVC())
diff --git a/Engine/Source/Programs/UnrealBuildTool/System/CppCompileEnvironment.cs b/Engine/Source/Programs/UnrealBuildTool/System/CppCompileEnvironment.cs
--- a/Engine/Source/Programs/UnrealBuildTool/System/CppCompileEnvironment.cs
+++ b/Engine/Source/Programs/UnrealBuildTool/System/CppCompileEnvironment.cs
@@ -646,12 +646,12 @@ namespace UnrealBuildTool
 		public EngineIncludeOrderVersion IncludeOrderVersion = EngineIncludeOrderVersion.Latest;
 
 		/// <summary>
-		/// Set flags for determinstic compiles (experimental).
+		/// Set flags for determinstic compiles.
 		/// </summary>
-		public bool bDeterministic = false;
+		public bool bDeterministic;
 
 		/// <summary>
-		/// Set flags for determinstic warnings (experimental).
+		/// Set flags for determinstic compile warnings.
 		/// </summary>
 		public WarningLevel DeterministicWarningLevel = WarningLevel.Off;
 
diff --git a/Engine/Source/Programs/UnrealBuildTool/System/LinkEnvironment.cs b/Engine/Source/Programs/UnrealBuildTool/System/LinkEnvironment.cs
--- a/Engine/Source/Programs/UnrealBuildTool/System/LinkEnvironment.cs
+++ b/Engine/Source/Programs/UnrealBuildTool/System/LinkEnvironment.cs
@@ -277,9 +277,9 @@ namespace UnrealBuildTool
 		public bool bIgnoreUnresolvedSymbols;
 
 		/// <summary>
-		/// Set flags for determinstic compiles (experimental).
+		/// Set flags for determinstic compiles.
 		/// </summary>
-		public bool bDeterministic = false;
+		public bool bDeterministic;
 
 		/// <summary>
 		/// Whether to log detailed timing information
