UnrealBuildTool: Enable deterministic msvc compiles by default

diff --git a/Engine/Source/Programs/UnrealBuildTool/Configuration/TargetRules.cs b/Engine/Source/Programs/UnrealBuildTool/Configuration/TargetRules.cs
--- a/Engine/Source/Programs/UnrealBuildTool/Configuration/TargetRules.cs
+++ b/Engine/Source/Programs/UnrealBuildTool/Configuration/TargetRules.cs
@@ -1681,12 +1681,13 @@ namespace UnrealBuildTool
 		public bool bUsePCHFiles = true;
 
 		/// <summary>
-		/// Set flags require for deterministic compiling and linking (experimental).
+		/// Set flags require for deterministic compiling and linking.
 		/// </summary>
-		[CommandLine("-Deterministic")]
+		[CommandLine("-Deterministic", Value = "true")]
+		[CommandLine("-NonDeterministic", Value = "false")]
 		[XmlConfigFile(Category = "BuildConfiguration")]
 		[RequiresUniqueBuildEnvironment]
-		public bool bDeterministic = false;
+		public bool bDeterministic = true;
 
 		/// <summary>
 		/// Whether PCHs should be chained when compiling with clang.
