- Changed AppleToolchain to use -gdwarf-4 instead of -gdwarf-2. Generates smaller debug info

diff --git a/Engine/Source/Programs/UnrealBuildTool/ToolChain/AppleToolChain.cs b/Engine/Source/Programs/UnrealBuildTool/ToolChain/AppleToolChain.cs
--- a/Engine/Source/Programs/UnrealBuildTool/ToolChain/AppleToolChain.cs
+++ b/Engine/Source/Programs/UnrealBuildTool/ToolChain/AppleToolChain.cs
@@ -441,7 +441,7 @@ namespace UnrealBuildTool
 			// Create DWARF format debug info if wanted,
 			if (CompileEnvironment.bCreateDebugInfo)
 			{
-				Arguments.Add("-gdwarf-2");
+				Arguments.Add("-gdwarf-4");
 			}
 		}
 
