UnrealBuildTool: Use a force include file for preprocessor defines in .c files when compiling

diff --git a/Engine/Source/Programs/UnrealBuildTool/Configuration/UEBuildModuleCPP.cs b/Engine/Source/Programs/UnrealBuildTool/Configuration/UEBuildModuleCPP.cs
--- a/Engine/Source/Programs/UnrealBuildTool/Configuration/UEBuildModuleCPP.cs
+++ b/Engine/Source/Programs/UnrealBuildTool/Configuration/UEBuildModuleCPP.cs
@@ -695,7 +695,9 @@ namespace UnrealBuildTool
 			// Compile C files directly. Do not use a PCH here, because a C++ PCH is not compatible with C source files.
 			if(InputFiles.CFiles.Count > 0)
 			{
-				LinkInputFiles.AddRange(ToolChain.CompileAllCPPFiles(ModuleCompileEnvironment, InputFiles.CFiles, IntermediateDirectory, Name, Graph).ObjectFiles);
+				CppCompileEnvironment CCompileEnvironment = new CppCompileEnvironment(ModuleCompileEnvironment);
+				CreateHeaderForDefinitions(CCompileEnvironment, IntermediateDirectory, "c", Graph);
+				LinkInputFiles.AddRange(ToolChain.CompileAllCPPFiles(CCompileEnvironment, InputFiles.CFiles, IntermediateDirectory, Name, Graph).ObjectFiles);
 			}
 
 			// Compile CC files directly.
