UBT: Merging CL 25121561 to support async ToolMode and ActionExecutor instances.

diff --git a/Engine/Source/Programs/UnrealBuildTool/Configuration/UEBuildTarget.cs b/Engine/Source/Programs/UnrealBuildTool/Configuration/UEBuildTarget.cs
--- a/Engine/Source/Programs/UnrealBuildTool/Configuration/UEBuildTarget.cs
+++ b/Engine/Source/Programs/UnrealBuildTool/Configuration/UEBuildTarget.cs
@@ -2187,7 +2187,7 @@ namespace UnrealBuildTool
 		/// <summary>
 		/// Builds the target, appending list of output files and returns building result.
 		/// </summary>
-		public TargetMakefile Build(BuildConfiguration BuildConfiguration, ISourceFileWorkingSet WorkingSet, TargetDescriptor TargetDescriptor, ILogger Logger, bool bInitOnly = false)
+		public async Task<TargetMakefile> BuildAsync(BuildConfiguration BuildConfiguration, ISourceFileWorkingSet WorkingSet, TargetDescriptor TargetDescriptor, ILogger Logger, bool bInitOnly = false)
 		{
 			CppConfiguration CppConfiguration = GetCppConfiguration(Configuration);
 
@@ -2366,7 +2366,7 @@ namespace UnrealBuildTool
 			// files that are injected as top level prerequisites.  If UHT only emitted included header files, we wouldn't need to run it during the Gather phase at all.
 			if (Makefile.UObjectModules.Count > 0)
 			{
-				ExternalExecution.ExecuteHeaderToolIfNecessary(BuildConfiguration, ProjectFile, Makefile, TargetName, WorkingSet, Logger);
+				await ExternalExecution.ExecuteHeaderToolIfNecessaryAsync(BuildConfiguration, ProjectFile, Makefile, TargetName, WorkingSet, Logger);
 			}
 
 			// Prefetch directory items for UHT folders since they are going to be used later
diff --git a/Engine/Source/Programs/UnrealBuildTool/Executors/Experimental/FASTBuild.cs b/Engine/Source/Programs/UnrealBuildTool/Executors/Experimental/FASTBuild.cs
--- a/Engine/Source/Programs/UnrealBuildTool/Executors/Experimental/FASTBuild.cs
+++ b/Engine/Source/Programs/UnrealBuildTool/Executors/Experimental/FASTBuild.cs
@@ -24,6 +24,7 @@ using System.Text.RegularExpressions;
 using UnrealBuildBase;
 using System.Runtime.Versioning;
 using Microsoft.Extensions.Logging;
+using System.Threading.Tasks;
 
 namespace UnrealBuildTool
 {
@@ -459,7 +460,7 @@ namespace UnrealBuildTool
 		}
 
 		[SupportedOSPlatform("windows")]
-		public override bool ExecuteActions(IEnumerable<LinkedAction> Actions, ILogger Logger)
+		public override async Task<bool> ExecuteActionsAsync(IEnumerable<LinkedAction> Actions, ILogger Logger)
 		{
 			if (!Actions.Any())
 			{
@@ -478,7 +479,7 @@ namespace UnrealBuildTool
 				return ExecuteBffFile(FASTBuildFilePath, Logger);
 			}
 
-			return LocalExecutor.ExecuteActions(Actions, Logger);
+			return await LocalExecutor.ExecuteActionsAsync(Actions, Logger);
 		}
 
 		private void AddText(string StringToWrite)
diff --git a/Engine/Source/Programs/UnrealBuildTool/Executors/HybridExecutor.cs b/Engine/Source/Programs/UnrealBuildTool/Executors/HybridExecutor.cs
--- a/Engine/Source/Programs/UnrealBuildTool/Executors/HybridExecutor.cs
+++ b/Engine/Source/Programs/UnrealBuildTool/Executors/HybridExecutor.cs
@@ -4,6 +4,7 @@ using EpicGames.Core;
 using Microsoft.Extensions.Logging;
 using System.Collections.Generic;
 using System.Linq;
+using System.Threading.Tasks;
 
 namespace UnrealBuildTool
 {
@@ -77,7 +78,7 @@ namespace UnrealBuildTool
 		/// <param name="ActionsToExecute">Actions to be executed</param>
 		/// <param name="Logger">Logger for output</param>
 		/// <returns>True if the build succeeded, false otherwise</returns>
-		public override bool ExecuteActions(IEnumerable<LinkedAction> ActionsToExecute, ILogger Logger)
+		public override async Task<bool> ExecuteActionsAsync(IEnumerable<LinkedAction> ActionsToExecute, ILogger Logger)
 		{
 			if (!ActionsToExecute.Any())
 			{
@@ -143,13 +144,13 @@ namespace UnrealBuildTool
 			}
 
 			// Execute the remote actions
-			if (RemoteActionsToExecute.Count > 0 && !RemoteExecutor.ExecuteActions(RemoteActionsToExecute, Logger))
+			if (RemoteActionsToExecute.Count > 0 && !await RemoteExecutor.ExecuteActionsAsync(RemoteActionsToExecute, Logger))
 			{
 				return false;
 			}
 
 			// Pass all the local actions through to the parallel executor
-			if (LocalActionsToExecute.Count > 0 && !LocalExecutor.ExecuteActions(LocalActionsToExecute, Logger))
+			if (LocalActionsToExecute.Count > 0 && !await LocalExecutor.ExecuteActionsAsync(LocalActionsToExecute, Logger))
 			{
 				return false;
 			}
diff --git a/Engine/Source/Programs/UnrealBuildTool/Executors/ParallelExecutor.cs b/Engine/Source/Programs/UnrealBuildTool/Executors/ParallelExecutor.cs
--- a/Engine/Source/Programs/UnrealBuildTool/Executors/ParallelExecutor.cs
+++ b/Engine/Source/Programs/UnrealBuildTool/Executors/ParallelExecutor.cs
@@ -181,11 +181,16 @@ namespace UnrealBuildTool
 			Error,
 		}
 
+		public override Task<bool> ExecuteActionsAsync(IEnumerable<LinkedAction> ActionsToExecute, ILogger Logger)
+		{
+			return Task.FromResult(ExecuteActions(ActionsToExecute, Logger));
+		}
+
 		/// <summary>
 		/// Executes the specified actions locally.
 		/// </summary>
 		/// <returns>True if all the tasks successfully executed, or false if any of them failed.</returns>
-		public override bool ExecuteActions(IEnumerable<LinkedAction> InputActions, ILogger Logger)
+		bool ExecuteActions(IEnumerable<LinkedAction> InputActions, ILogger Logger)
 		{
 			if (!InputActions.Any())
 			{
diff --git a/Engine/Source/Programs/UnrealBuildTool/Executors/SNDBS.cs b/Engine/Source/Programs/UnrealBuildTool/Executors/SNDBS.cs
--- a/Engine/Source/Programs/UnrealBuildTool/Executors/SNDBS.cs
+++ b/Engine/Source/Programs/UnrealBuildTool/Executors/SNDBS.cs
@@ -14,6 +14,7 @@ using System.ServiceProcess;
 using System.Text;
 using System.Text.Json;
 using System.Text.RegularExpressions;
+using System.Threading.Tasks;
 using EpicGames.Core;
 using Microsoft.Extensions.Logging;
 using UnrealBuildBase;
@@ -214,7 +215,12 @@ namespace UnrealBuildTool
 			return true;
 		}
 
-		public override bool ExecuteActions(IEnumerable<LinkedAction> Actions, ILogger Logger)
+		public override Task<bool> ExecuteActionsAsync(IEnumerable<LinkedAction> ActionsToExecute, ILogger Logger)
+		{
+			return Task.FromResult(ExecuteActions(ActionsToExecute, Logger));
+		}
+
+		bool ExecuteActions(IEnumerable<LinkedAction> Actions, ILogger Logger)
 		{
 			if (!Actions.Any())
 			{
diff --git a/Engine/Source/Programs/UnrealBuildTool/Executors/XGE.cs b/Engine/Source/Programs/UnrealBuildTool/Executors/XGE.cs
--- a/Engine/Source/Programs/UnrealBuildTool/Executors/XGE.cs
+++ b/Engine/Source/Programs/UnrealBuildTool/Executors/XGE.cs
@@ -19,6 +19,7 @@ using UnrealBuildBase;
 using System.Diagnostics.CodeAnalysis;
 using System.Runtime.Versioning;
 using Microsoft.Extensions.Logging;
+using System.Threading.Tasks;
 
 namespace UnrealBuildTool
 {
@@ -395,7 +396,12 @@ namespace UnrealBuildTool
 			Logger.LogInformation("XGEEXPORT: Exported '{OutFile}'", OutFile);
 		}
 
-		public override bool ExecuteActions(IEnumerable<LinkedAction> Actions, ILogger Logger)
+		public override Task<bool> ExecuteActionsAsync(IEnumerable<LinkedAction> ActionsToExecute, ILogger Logger)
+		{
+			return Task.FromResult(ExecuteActions(ActionsToExecute, Logger));
+		}
+
+		bool ExecuteActions(IEnumerable<LinkedAction> Actions, ILogger Logger)
 		{
 			if (!Actions.Any())
 			{
diff --git a/Engine/Source/Programs/UnrealBuildTool/Modes/AggregateClangTimingInfo.cs b/Engine/Source/Programs/UnrealBuildTool/Modes/AggregateClangTimingInfo.cs
--- a/Engine/Source/Programs/UnrealBuildTool/Modes/AggregateClangTimingInfo.cs
+++ b/Engine/Source/Programs/UnrealBuildTool/Modes/AggregateClangTimingInfo.cs
@@ -170,7 +170,7 @@ namespace UnrealBuildTool
 
 		private ConcurrentDictionary<FileReference, ClangTrace> ClangTraceCache = new();
 
-		public override int Execute(CommandLineArguments Arguments, ILogger Logger)
+		public override Task<int> ExecuteAsync(CommandLineArguments Arguments, ILogger Logger)
 		{
 			FileReference ManifestFile = Arguments.GetFileReference("-ManifestFile=");
 			IEnumerable<FileReference> SourceFiles = FileReference.ReadAllLines(ManifestFile).Select(x => new FileReference(x));
@@ -256,7 +256,7 @@ namespace UnrealBuildTool
 				File.Move(TempFilePath, ArchiveFile.FullName, true);
 			}
 
-			return 0;
+			return Task.FromResult(0);
 		}
 
 
diff --git a/Engine/Source/Programs/UnrealBuildTool/Modes/AggregateParsedTimingInfo.cs b/Engine/Source/Programs/UnrealBuildTool/Modes/AggregateParsedTimingInfo.cs
--- a/Engine/Source/Programs/UnrealBuildTool/Modes/AggregateParsedTimingInfo.cs
+++ b/Engine/Source/Programs/UnrealBuildTool/Modes/AggregateParsedTimingInfo.cs
@@ -29,7 +29,7 @@ namespace UnrealBuildTool
 		private ConcurrentBag<TimingData> AggregateFunctions = new ConcurrentBag<TimingData>();
 		private object AddFileLock = new object();
 
-		public override int Execute(CommandLineArguments Arguments, ILogger Logger)
+		public override Task<int> ExecuteAsync(CommandLineArguments Arguments, ILogger Logger)
 		{
 			FileReference ManifestFile = Arguments.GetFileReference("-ManifestFile=");
 			string[] ParsedFileNames = FileReference.ReadAllLines(ManifestFile);
@@ -87,7 +87,7 @@ namespace UnrealBuildTool
 				FileReference.WriteAllText(CompileTimingFile, JsonCompileTimes);
 			}
 
-			return 0;
+			return Task.FromResult(0);
 		}
 
 		private void GroupTimingDataOnName(TimingData AggregateData, string TimingName, IEnumerable<TimingData> UngroupedData)
diff --git a/Engine/Source/Programs/UnrealBuildTool/Modes/AnalyzeMode.cs b/Engine/Source/Programs/UnrealBuildTool/Modes/AnalyzeMode.cs
--- a/Engine/Source/Programs/UnrealBuildTool/Modes/AnalyzeMode.cs
+++ b/Engine/Source/Programs/UnrealBuildTool/Modes/AnalyzeMode.cs
@@ -6,6 +6,7 @@ using System.Collections.Generic;
 using System.IO;
 using System.Linq;
 using System.Text;
+using System.Threading.Tasks;
 using System.Xml;
 using Microsoft.Extensions.Logging;
 using UnrealBuildBase;
@@ -24,7 +25,7 @@ namespace UnrealBuildTool.Modes
 		/// <param name="Arguments">Command line arguments</param>
 		/// <returns>Exit code</returns>
 		/// <param name="Logger"></param>
-		public override int Execute(CommandLineArguments Arguments, ILogger Logger)
+		public override Task<int> ExecuteAsync(CommandLineArguments Arguments, ILogger Logger)
 		{
 			Arguments.ApplyTo(this);
 
@@ -47,7 +48,7 @@ namespace UnrealBuildTool.Modes
 				}
 			}
 
-			return 0;
+			return Task.FromResult(0);
 		}
 
 		class ModuleInfo
diff --git a/Engine/Source/Programs/UnrealBuildTool/Modes/BuildMode.cs b/Engine/Source/Programs/UnrealBuildTool/Modes/BuildMode.cs
--- a/Engine/Source/Programs/UnrealBuildTool/Modes/BuildMode.cs
+++ b/Engine/Source/Programs/UnrealBuildTool/Modes/BuildMode.cs
@@ -102,7 +102,7 @@ namespace UnrealBuildTool
 		/// <param name="Arguments">Command-line arguments</param>
 		/// <returns>One of the values of ECompilationResult</returns>
 		/// <param name="Logger"></param>
-		public override int Execute(CommandLineArguments Arguments, ILogger Logger)
+		public override async Task<int> ExecuteAsync(CommandLineArguments Arguments, ILogger Logger)
 		{
 			Arguments.ApplyTo(this);
 
@@ -249,7 +249,7 @@ namespace UnrealBuildTool
 					// Create the working set provider per group.
 					using (ISourceFileWorkingSet WorkingSet = SourceFileWorkingSet.Create(Unreal.RootDirectory, ProjectDirs, Logger))
 					{
-						Build(TargetDescriptors, BuildConfiguration, WorkingSet, Options, WriteOutdatedActionsFile, Logger, bSkipPreBuildTargets);
+						await BuildAsync(TargetDescriptors, BuildConfiguration, WorkingSet, Options, WriteOutdatedActionsFile, Logger, bSkipPreBuildTargets);
 					}
 				}
 			}
@@ -380,7 +380,7 @@ namespace UnrealBuildTool
 		/// <param name="Logger">Logger for output</param>
 		/// <param name="bSkipPreBuildTargets">If true then only the current target descriptors will be built.</param>
 		/// <returns>Result from the compilation</returns>
-		public static void Build(List<TargetDescriptor> TargetDescriptors, BuildConfiguration BuildConfiguration, ISourceFileWorkingSet WorkingSet, BuildOptions Options, FileReference? WriteOutdatedActionsFile, ILogger Logger, bool bSkipPreBuildTargets = false)
+		public static async Task BuildAsync(List<TargetDescriptor> TargetDescriptors, BuildConfiguration BuildConfiguration, ISourceFileWorkingSet WorkingSet, BuildOptions Options, FileReference? WriteOutdatedActionsFile, ILogger Logger, bool bSkipPreBuildTargets = false)
 		{
 			List<TargetMakefile> TargetMakefiles = new List<TargetMakefile>();
 
@@ -393,7 +393,7 @@ namespace UnrealBuildTool
 					Utils.ExecuteCustomBuildSteps(InitScripts, Logger);
 				}
 
-				TargetMakefile NewMakefile = CreateMakefile(BuildConfiguration, TargetDescriptors[Idx], WorkingSet, Logger);
+				TargetMakefile NewMakefile = await CreateMakefileAsync(BuildConfiguration, TargetDescriptors[Idx], WorkingSet, Logger);
 				TargetMakefiles.Add(NewMakefile);
 				if (!bSkipPreBuildTargets)
 				{
@@ -408,7 +408,7 @@ namespace UnrealBuildTool
 				}
 			}
 
-			Build(TargetMakefiles.ToArray(), TargetDescriptors, BuildConfiguration, Options, WriteOutdatedActionsFile, Logger);
+			await BuildAsync(TargetMakefiles.ToArray(), TargetDescriptors, BuildConfiguration, Options, WriteOutdatedActionsFile, Logger);
 		}
 
 		/// <summary>
@@ -421,7 +421,7 @@ namespace UnrealBuildTool
 		/// <param name="WriteOutdatedActionsFile">Files to write the list of outdated actions to (rather than building them)</param>
 		/// <param name="Logger">Logger for output</param>
 		/// <returns>Result from the compilation</returns>
-		internal static void Build(TargetMakefile[] Makefiles, List<TargetDescriptor> TargetDescriptors, BuildConfiguration BuildConfiguration, BuildOptions Options, FileReference? WriteOutdatedActionsFile, ILogger Logger)
+		internal static async Task BuildAsync(TargetMakefile[] Makefiles, List<TargetDescriptor> TargetDescriptors, BuildConfiguration BuildConfiguration, BuildOptions Options, FileReference? WriteOutdatedActionsFile, ILogger Logger)
 		{
 			// Execute the build
 			if ((Options & BuildOptions.SkipBuild) == 0)
@@ -543,7 +543,7 @@ namespace UnrealBuildTool
 					if (PreprocessActions.Count > 0)
 					{
 						Logger.LogInformation("Updating module dependencies...");
-						ActionGraph.ExecuteActions(BuildConfiguration, PreprocessActions, TargetDescriptors, Logger);
+						await ActionGraph.ExecuteActionsAsync(BuildConfiguration, PreprocessActions, TargetDescriptors, Logger);
 
 						foreach (FileItem ProducedItem in PreprocessActions.SelectMany(x => x.ProducedItems))
 						{
@@ -767,7 +767,7 @@ namespace UnrealBuildTool
 
 						using (GlobalTracer.Instance.BuildSpan("ActionGraph.ExecuteActions()").StartActive())
 						{
-							ActionGraph.ExecuteActions(BuildConfiguration, MergedActionsToExecute, TargetDescriptors, Logger);
+							await ActionGraph.ExecuteActionsAsync(BuildConfiguration, MergedActionsToExecute, TargetDescriptors, Logger);
 						}
 					}
 
@@ -1028,7 +1028,7 @@ namespace UnrealBuildTool
 		/// <param name="WorkingSet">Set of source files which are part of the working set</param>
 		/// <param name="Logger">Logger for output</param>
 		/// <returns>Makefile for the given target</returns>
-		internal static TargetMakefile CreateMakefile(BuildConfiguration BuildConfiguration, TargetDescriptor TargetDescriptor, ISourceFileWorkingSet WorkingSet, ILogger Logger)
+		internal static async Task<TargetMakefile> CreateMakefileAsync(BuildConfiguration BuildConfiguration, TargetDescriptor TargetDescriptor, ISourceFileWorkingSet WorkingSet, ILogger Logger)
 		{
 			// Get the path to the makefile for this target
 			FileReference? MakefileLocation = null;
@@ -1094,7 +1094,7 @@ namespace UnrealBuildTool
 				// Build the target
 				using (GlobalTracer.Instance.BuildSpan("UEBuildTarget.Build()").StartActive())
 				{
-					Makefile = Target.Build(BuildConfiguration, WorkingSet, TargetDescriptor, Logger);
+					Makefile = await Target.BuildAsync(BuildConfiguration, WorkingSet, TargetDescriptor, Logger);
 				}
 
 				Makefile.MemoryPerActionGB = Target.Rules.MemoryPerActionGB;
@@ -1142,7 +1142,7 @@ namespace UnrealBuildTool
 				// If the target needs UHT to be run, we'll go ahead and do that now
 				if (Makefile.UObjectModules.Count > 0)
 				{
-					ExternalExecution.ExecuteHeaderToolIfNecessary(BuildConfiguration, TargetDescriptor.ProjectFile, Makefile, TargetDescriptor.Name, WorkingSet, Logger);
+					await ExternalExecution.ExecuteHeaderToolIfNecessaryAsync(BuildConfiguration, TargetDescriptor.ProjectFile, Makefile, TargetDescriptor.Name, WorkingSet, Logger);
 				}
 
 #if __VPROJECT_AVAILABLE__
diff --git a/Engine/Source/Programs/UnrealBuildTool/Modes/CleanMode.cs b/Engine/Source/Programs/UnrealBuildTool/Modes/CleanMode.cs
--- a/Engine/Source/Programs/UnrealBuildTool/Modes/CleanMode.cs
+++ b/Engine/Source/Programs/UnrealBuildTool/Modes/CleanMode.cs
@@ -40,7 +40,7 @@ namespace UnrealBuildTool
 		/// <param name="Arguments">Command-line arguments</param>
 		/// <returns>One of the values of ECompilationResult</returns>
 		/// <param name="Logger"></param>
-		public override int Execute(CommandLineArguments Arguments, ILogger Logger)
+		public override Task<int> ExecuteAsync(CommandLineArguments Arguments, ILogger Logger)
 		{
 			Arguments.ApplyTo(this);
 
@@ -54,7 +54,7 @@ namespace UnrealBuildTool
 
 			Clean(TargetDescriptors, BuildConfiguration, Logger);
 
-			return 0;
+			return Task.FromResult(0);
 		}
 
 		public void Clean(List<TargetDescriptor> TargetDescriptors, BuildConfiguration BuildConfiguration, ILogger Logger)
diff --git a/Engine/Source/Programs/UnrealBuildTool/Modes/DeployMode.cs b/Engine/Source/Programs/UnrealBuildTool/Modes/DeployMode.cs
--- a/Engine/Source/Programs/UnrealBuildTool/Modes/DeployMode.cs
+++ b/Engine/Source/Programs/UnrealBuildTool/Modes/DeployMode.cs
@@ -28,7 +28,7 @@ namespace UnrealBuildTool
 		/// <param name="Arguments">Command line arguments</param>
 		/// <returns>Exit code</returns>
 		/// <param name="Logger"></param>
-		public override int Execute(CommandLineArguments Arguments, ILogger Logger)
+		public override Task<int> ExecuteAsync(CommandLineArguments Arguments, ILogger Logger)
 		{
 			// Apply the arguments
 			Arguments.ApplyTo(this);
@@ -39,7 +39,7 @@ namespace UnrealBuildTool
 			Logger.LogInformation("Deploying {ReceiptTargetName} {ReceiptPlatform} {ReceiptConfiguration}...", Receipt.TargetName, Receipt.Platform, Receipt.Configuration);
 			UEBuildPlatform.GetBuildPlatform(Receipt.Platform).Deploy(Receipt);
 
-			return (int)CompilationResult.Succeeded;
+			return Task.FromResult((int)CompilationResult.Succeeded);
 		}
 	}
 }
diff --git a/Engine/Source/Programs/UnrealBuildTool/Modes/ExecuteMode.cs b/Engine/Source/Programs/UnrealBuildTool/Modes/ExecuteMode.cs
--- a/Engine/Source/Programs/UnrealBuildTool/Modes/ExecuteMode.cs
+++ b/Engine/Source/Programs/UnrealBuildTool/Modes/ExecuteMode.cs
@@ -32,7 +32,7 @@ namespace UnrealBuildTool
 		/// <param name="Arguments">Command-line arguments</param>
 		/// <returns>One of the values of ECompilationResult</returns>
 		/// <param name="Logger"></param>
-		public override int Execute(CommandLineArguments Arguments, ILogger Logger)
+		public override async Task<int> ExecuteAsync(CommandLineArguments Arguments, ILogger Logger)
 		{
 			Arguments.ApplyTo(this);
 
@@ -61,7 +61,7 @@ namespace UnrealBuildTool
 			using (GlobalTracer.Instance.BuildSpan("ActionGraph.ExecuteActions()").StartActive())
 			{
 				List<TargetDescriptor> TargetDescriptors = TargetDescriptor.ParseCommandLine(Arguments, false, false, false, Logger);
-				ActionGraph.ExecuteActions(BuildConfiguration, Actions, TargetDescriptors, Logger);
+				await ActionGraph.ExecuteActionsAsync(BuildConfiguration, Actions, TargetDescriptors, Logger);
 			}
 
 			return 0;
diff --git a/Engine/Source/Programs/UnrealBuildTool/Modes/FixIncludePathsMode.cs b/Engine/Source/Programs/UnrealBuildTool/Modes/FixIncludePathsMode.cs
--- a/Engine/Source/Programs/UnrealBuildTool/Modes/FixIncludePathsMode.cs
+++ b/Engine/Source/Programs/UnrealBuildTool/Modes/FixIncludePathsMode.cs
@@ -7,6 +7,7 @@ using System.Text.RegularExpressions;
 using EpicGames.Core;
 using UnrealBuildBase;
 using Microsoft.Extensions.Logging;
+using System.Threading.Tasks;
 using System.Diagnostics;
 
 namespace UnrealBuildTool
@@ -47,7 +48,7 @@ namespace UnrealBuildTool
 		/// <param name="Arguments">Command line arguments</param>
 		/// <returns>Exit code</returns>
 		/// <param name="Logger"></param>
-		public override int Execute(CommandLineArguments Arguments, ILogger Logger)
+		public override Task<int> ExecuteAsync(CommandLineArguments Arguments, ILogger Logger)
 		{
 			Arguments.ApplyTo(this);
 
@@ -339,7 +340,7 @@ namespace UnrealBuildTool
 				}
 			}
 
-			return 0;
+			return Task.FromResult(0);
 		}
 
 		class HeaderSortComparison : IComparer<string>
diff --git a/Engine/Source/Programs/UnrealBuildTool/Modes/GenerateClangDatabase.cs b/Engine/Source/Programs/UnrealBuildTool/Modes/GenerateClangDatabase.cs
--- a/Engine/Source/Programs/UnrealBuildTool/Modes/GenerateClangDatabase.cs
+++ b/Engine/Source/Programs/UnrealBuildTool/Modes/GenerateClangDatabase.cs
@@ -5,6 +5,7 @@ using System.Collections.Generic;
 using System.Linq;
 using System.Text;
 using System.Text.RegularExpressions;
+using System.Threading.Tasks;
 using EpicGames.Core;
 using UnrealBuildBase;
 using Microsoft.Extensions.Logging;
@@ -110,7 +111,7 @@ namespace UnrealBuildTool
 		/// <param name="Arguments">Command line arguments</param>
 		/// <returns>Exit code</returns>
 		/// <param name="Logger"></param>
-		public override int Execute(CommandLineArguments Arguments, ILogger Logger)
+		public override async Task<int> ExecuteAsync(CommandLineArguments Arguments, ILogger Logger)
 		{
 			Arguments.ApplyTo(this);
 
@@ -154,7 +155,7 @@ namespace UnrealBuildTool
 					if (bExecCodeGenActions)
 					{
 						// Create the makefile
-						TargetMakefile Makefile = Target.Build(BuildConfiguration, WorkingSet, TargetDescriptor, Logger);
+						TargetMakefile Makefile = await Target.BuildAsync(BuildConfiguration, WorkingSet, TargetDescriptor, Logger);
 						List<LinkedAction> Actions = Makefile.Actions.ConvertAll(x => new LinkedAction(x, TargetDescriptor));
 						ActionGraph.Link(Actions, Logger);
 
@@ -166,7 +167,7 @@ namespace UnrealBuildTool
 						if (PrerequisiteActions.Count > 0)
 						{
 							Logger.LogInformation("Executing actions that produce source files...");
-							ActionGraph.ExecuteActions(BuildConfiguration, PrerequisiteActions, new List<TargetDescriptor> { TargetDescriptor }, Logger);
+							await ActionGraph.ExecuteActionsAsync(BuildConfiguration, PrerequisiteActions, new List<TargetDescriptor> { TargetDescriptor }, Logger);
 						}
 					}
 
diff --git a/Engine/Source/Programs/UnrealBuildTool/Modes/GenerateProjectFilesMode.cs b/Engine/Source/Programs/UnrealBuildTool/Modes/GenerateProjectFilesMode.cs
--- a/Engine/Source/Programs/UnrealBuildTool/Modes/GenerateProjectFilesMode.cs
+++ b/Engine/Source/Programs/UnrealBuildTool/Modes/GenerateProjectFilesMode.cs
@@ -61,7 +61,7 @@ namespace UnrealBuildTool
 		/// <param name="Arguments">Command line arguments</param>
 		/// <returns>Exit code</returns>
 		/// <param name="Logger"></param>
-		public override int Execute(CommandLineArguments Arguments, ILogger Logger)
+		public override Task<int> ExecuteAsync(CommandLineArguments Arguments, ILogger Logger)
 		{
 			// Apply any command line arguments to this class
 			Arguments.ApplyTo(this);
@@ -254,13 +254,13 @@ namespace UnrealBuildTool
 
 				if (!bGenerateSuccess)
 				{
-					return (int)CompilationResult.OtherCompilationError;
+					return Task.FromResult((int)CompilationResult.OtherCompilationError);
 				}
 
 				// any further generators can skip one-time operations
 				bPerformOneTimeOperations = false;
 			}
-			return (int)CompilationResult.Succeeded;
+			return Task.FromResult((int)CompilationResult.Succeeded);
 		}
 
 		/// <summary>
diff --git a/Engine/Source/Programs/UnrealBuildTool/Modes/IWYUMode.cs b/Engine/Source/Programs/UnrealBuildTool/Modes/IWYUMode.cs
--- a/Engine/Source/Programs/UnrealBuildTool/Modes/IWYUMode.cs
+++ b/Engine/Source/Programs/UnrealBuildTool/Modes/IWYUMode.cs
@@ -346,7 +346,7 @@ namespace UnrealBuildTool
 		/// <param name="Arguments">Command line arguments</param>
 		/// <returns>Exit code</returns>
 		/// <param name="Logger"></param>
-		public override int Execute(CommandLineArguments Arguments, ILogger Logger)
+		public override async Task<int> ExecuteAsync(CommandLineArguments Arguments, ILogger Logger)
 		{
 			Logger.LogInformation($"====================================================");
 			Logger.LogInformation($"Running IWYU. {(bWrite ? "" : "(Preview mode. Add -Write to write modifications to files)")}");
@@ -598,7 +598,7 @@ namespace UnrealBuildTool
 					TargetMakefile Makefile;
 					try
 					{
-						Makefile = BuildMode.CreateMakefile(BuildConfiguration, Descriptor, WorkingSet, Logger);
+						Makefile = await BuildMode.CreateMakefileAsync(BuildConfiguration, Descriptor, WorkingSet, Logger);
 					}
 					finally
 					{
@@ -610,7 +610,7 @@ namespace UnrealBuildTool
 					{
 						try
 						{
-							BuildMode.Build(new TargetMakefile[] { Makefile }, new List<TargetDescriptor>() { Descriptor }, BuildConfiguration, BuildOptions.None, null, Logger);
+							await BuildMode.BuildAsync(new TargetMakefile[] { Makefile }, new List<TargetDescriptor>() { Descriptor }, BuildConfiguration, BuildOptions.None, null, Logger);
 						}
 						finally
 						{
diff --git a/Engine/Source/Programs/UnrealBuildTool/Modes/InlineGeneratedCpps.cs b/Engine/Source/Programs/UnrealBuildTool/Modes/InlineGeneratedCpps.cs
--- a/Engine/Source/Programs/UnrealBuildTool/Modes/InlineGeneratedCpps.cs
+++ b/Engine/Source/Programs/UnrealBuildTool/Modes/InlineGeneratedCpps.cs
@@ -5,6 +5,7 @@ using System.Collections.Generic;
 using System.Linq;
 using System.Text;
 using System.Text.RegularExpressions;
+using System.Threading.Tasks;
 using EpicGames.Core;
 using UnrealBuildBase;
 using Microsoft.Extensions.Logging;
@@ -54,7 +55,7 @@ namespace UnrealBuildTool
 		/// <param name="Arguments">Command line arguments</param>
 		/// <returns>Exit code</returns>
 		/// <param name="Logger"></param>
-		public override int Execute(CommandLineArguments Arguments, ILogger Logger)
+		public override Task<int> ExecuteAsync(CommandLineArguments Arguments, ILogger Logger)
 		{
 			Arguments.ApplyTo(this);
 
@@ -256,7 +257,7 @@ namespace UnrealBuildTool
 				}
 			}
 
-			return 0;
+			return Task.FromResult(0);
 		}
 	}
 }
diff --git a/Engine/Source/Programs/UnrealBuildTool/Modes/JsonExportMode.cs b/Engine/Source/Programs/UnrealBuildTool/Modes/JsonExportMode.cs
--- a/Engine/Source/Programs/UnrealBuildTool/Modes/JsonExportMode.cs
+++ b/Engine/Source/Programs/UnrealBuildTool/Modes/JsonExportMode.cs
@@ -29,7 +29,7 @@ namespace UnrealBuildTool
 		/// <param name="Arguments">Command line arguments</param>
 		/// <returns>Exit code (always zero)</returns>
 		/// <param name="Logger"></param>
-		public override int Execute(CommandLineArguments Arguments, ILogger Logger)
+		public override async Task<int> ExecuteAsync(CommandLineArguments Arguments, ILogger Logger)
 		{
 			Arguments.ApplyTo(this);
 
@@ -57,7 +57,7 @@ namespace UnrealBuildTool
 						Arguments.ApplyTo(BuildConfiguration);
 
 						// Create the makefile
-						TargetMakefile Makefile = Target.Build(BuildConfiguration, WorkingSet, TargetDescriptor, Logger);
+						TargetMakefile Makefile = await Target.BuildAsync(BuildConfiguration, WorkingSet, TargetDescriptor, Logger);
 						List<LinkedAction> Actions = Makefile.Actions.ConvertAll(x => new LinkedAction(x, TargetDescriptor));
 						ActionGraph.Link(Actions, Logger);
 
@@ -68,8 +68,8 @@ namespace UnrealBuildTool
 						// Execute these actions
 						if (PrerequisiteActions.Count > 0)
 						{
-							Logger.LogInformation("Exeucting actions that produce source files...");
-							ActionGraph.ExecuteActions(BuildConfiguration, PrerequisiteActions, new List<TargetDescriptor> { TargetDescriptor }, Logger);
+							Logger.LogInformation("Executing actions that produce source files...");
+							await ActionGraph.ExecuteActionsAsync(BuildConfiguration, PrerequisiteActions, new List<TargetDescriptor> { TargetDescriptor }, Logger);
 						}
 					}
 				}
diff --git a/Engine/Source/Programs/UnrealBuildTool/Modes/ParseMsvcTimingInfo.cs b/Engine/Source/Programs/UnrealBuildTool/Modes/ParseMsvcTimingInfo.cs
--- a/Engine/Source/Programs/UnrealBuildTool/Modes/ParseMsvcTimingInfo.cs
+++ b/Engine/Source/Programs/UnrealBuildTool/Modes/ParseMsvcTimingInfo.cs
@@ -5,6 +5,7 @@ using System.Collections.Generic;
 using System.IO;
 using System.Linq;
 using System.Text.RegularExpressions;
+using System.Threading.Tasks;
 using EpicGames.Core;
 using Microsoft.Extensions.Logging;
 
@@ -19,7 +20,7 @@ namespace UnrealBuildTool
 	{
 		const string TimingDataRegex = @"^\t\t(?<Indent>\t*)(?<Name>[^\t]+):\s*(?<Duration>[0-9\.]+)s$";
 
-		public override int Execute(CommandLineArguments Arguments, ILogger Logger)
+		public override Task<int> ExecuteAsync(CommandLineArguments Arguments, ILogger Logger)
 		{
 			FileReference InputFile = Arguments.GetFileReference("-TimingFile=");
 
@@ -28,7 +29,7 @@ namespace UnrealBuildTool
 			if (Arguments.HasOption("-Tracing"))
 			{
 				ParseTimingDataToTracingFiles(InputFile);
-				return 0;
+				return Task.FromResult(0);
 			}
 
 			// Break the input file into the various sections for processing.
@@ -111,7 +112,7 @@ namespace UnrealBuildTool
 				Writer.Write(Summary);
 			}
 
-			return 0;
+			return Task.FromResult(0);
 		}
 
 		TimingData SummarizeParsedTimingData(string SummaryName, TimingDataType TimingType,  IEnumerable<string> Lines)
diff --git a/Engine/Source/Programs/UnrealBuildTool/Modes/PrintBuildGraphInfo.cs b/Engine/Source/Programs/UnrealBuildTool/Modes/PrintBuildGraphInfo.cs
--- a/Engine/Source/Programs/UnrealBuildTool/Modes/PrintBuildGraphInfo.cs
+++ b/Engine/Source/Programs/UnrealBuildTool/Modes/PrintBuildGraphInfo.cs
@@ -5,6 +5,7 @@ using System.Collections.Generic;
 using System.Linq;
 using System.Text;
 using System.Text.RegularExpressions;
+using System.Threading.Tasks;
 using EpicGames.Core;
 using UnrealBuildBase;
 using Microsoft.Extensions.Logging;
@@ -23,7 +24,7 @@ namespace UnrealBuildTool
 		/// <param name="Arguments">Command line arguments</param>
 		/// <returns>Exit code</returns>
 		/// <param name="Logger"></param>
-		public override int Execute(CommandLineArguments Arguments, ILogger Logger)
+		public override Task<int> ExecuteAsync(CommandLineArguments Arguments, ILogger Logger)
 		{
 			Arguments.ApplyTo(this);
 
@@ -216,7 +217,7 @@ namespace UnrealBuildTool
 				}
 			}
 
-			return 0;
+			return Task.FromResult(0);
 		}
 
 		private void FindCircularDependencyModules(UEBuildModule SearchForBuildModule, UEBuildModule CurrentBuildModule, HashSet<UEBuildModule> IgnoreReferencedModules, bool bIncludePrivateDependencyModules, Stack<UEBuildModule> CurrentStack, List<Stack<UEBuildModule>> ModuleDepStacks)
diff --git a/Engine/Source/Programs/UnrealBuildTool/Modes/ProfileUnitySizesMode.cs b/Engine/Source/Programs/UnrealBuildTool/Modes/ProfileUnitySizesMode.cs
--- a/Engine/Source/Programs/UnrealBuildTool/Modes/ProfileUnitySizesMode.cs
+++ b/Engine/Source/Programs/UnrealBuildTool/Modes/ProfileUnitySizesMode.cs
@@ -3,6 +3,7 @@
 using System;
 using System.Collections.Generic;
 using System.Text.RegularExpressions;
+using System.Threading.Tasks;
 using EpicGames.Core;
 using UnrealBuildBase;
 using Microsoft.Extensions.Logging;
@@ -110,7 +111,7 @@ namespace UnrealBuildTool
 		/// <param name="Arguments">Command line arguments</param>
 		/// <returns>Exit code</returns>
 		/// <param name="Logger"></param>
-		public override int Execute(CommandLineArguments Arguments, ILogger Logger)
+		public override async Task<int> ExecuteAsync(CommandLineArguments Arguments, ILogger Logger)
 		{
 			Arguments.ApplyTo(this);
 
@@ -168,7 +169,7 @@ namespace UnrealBuildTool
 				ModuleList.SortBy(module => module.Name);
 				foreach (UEBuildModule Module in ModuleList)
 				{					
-					CompileModule(BuildConfiguration, TargetDescriptor, Target, Module, Logger);
+					await CompileModuleAsync(BuildConfiguration, TargetDescriptor, Target, Module, Logger);
 				}
 			}
 
@@ -178,7 +179,7 @@ namespace UnrealBuildTool
 		/// <summary>
 		/// Compile the module multiple times looking for the best unity size 
 		/// </summary>
-		private void CompileModule(BuildConfiguration BuildConfiguration, TargetDescriptor TargetDescriptor, UEBuildTarget Target, UEBuildModule Module, ILogger Logger)
+		private async Task CompileModuleAsync(BuildConfiguration BuildConfiguration, TargetDescriptor TargetDescriptor, UEBuildTarget Target, UEBuildModule Module, ILogger Logger)
 		{
 			TargetDescriptor.OnlyModuleNames.Clear();
 			TargetDescriptor.OnlyModuleNames.Add(Module.Name);
@@ -192,16 +193,16 @@ namespace UnrealBuildTool
 			int TargetUnitySize = Target.Rules.NumIncludedBytesPerUnityCPP;
 
 			int BuildNum = 1;
-			CompileModule($"  [{BuildNum++}/{TotalBuilds}] ", BuildConfiguration, TargetDescriptor, Module, Logger, CurrentModuleUnitySize, true, false);
+			await CompileModuleAsync($"  [{BuildNum++}/{TotalBuilds}] ", BuildConfiguration, TargetDescriptor, Module, Logger, CurrentModuleUnitySize, true, false);
 			
-			TimingData CurrentCompileTime = GetBestCompileModuleTime($"  [{BuildNum++}/{TotalBuilds}] ", BuildConfiguration, TargetDescriptor, Module, Logger, CurrentModuleUnitySize, false, false);
+			TimingData CurrentCompileTime = await GetBestCompileModuleTimeAsync($"  [{BuildNum++}/{TotalBuilds}] ", BuildConfiguration, TargetDescriptor, Module, Logger, CurrentModuleUnitySize, false, false);
 			if (!CurrentCompileTime.IsValid())
 			{
 				Logger.LogInformation($"Skipping module because it doesn't compile with current settings.");
 				return;
 			}
 
-			TimingData DisableUnityCompileTime = GetBestCompileModuleTime($"  [{BuildNum++}/{TotalBuilds}] ", BuildConfiguration, TargetDescriptor, Module, Logger, TargetUnitySize, false, true);
+			TimingData DisableUnityCompileTime = await GetBestCompileModuleTimeAsync($"  [{BuildNum++}/{TotalBuilds}] ", BuildConfiguration, TargetDescriptor, Module, Logger, TargetUnitySize, false, true);
 
 			int MaxUnitySize = TargetUnitySize * 2;
 			List<TimingData> Timings = new();
@@ -209,7 +210,7 @@ namespace UnrealBuildTool
 			int CurrentUnitySize = UnitySizeInc;
 			for (int UnitySizeIndex = 0; UnitySizeIndex < UnitySizeDivision; UnitySizeIndex++)
 			{
-				TimingData NewTiming = GetBestCompileModuleTime($"  [{BuildNum++}/{TotalBuilds}] ", BuildConfiguration, TargetDescriptor, Module, Logger, CurrentUnitySize, false, false);
+				TimingData NewTiming = await GetBestCompileModuleTimeAsync($"  [{BuildNum++}/{TotalBuilds}] ", BuildConfiguration, TargetDescriptor, Module, Logger, CurrentUnitySize, false, false);
 				Timings.Add(NewTiming);
 				CurrentUnitySize += UnitySizeInc;
 
@@ -270,13 +271,13 @@ namespace UnrealBuildTool
 		/// <summary>
 		/// Returns best compile timings after building the module times
 		/// </summary>
-		private TimingData GetBestCompileModuleTime(string LogPrefix, BuildConfiguration BuildConfiguration, TargetDescriptor TargetDescriptor, UEBuildModule Module, ILogger Logger, int UnitySize, bool bPriming, bool bDisableUnity)
+		private async Task<TimingData> GetBestCompileModuleTimeAsync(string LogPrefix, BuildConfiguration BuildConfiguration, TargetDescriptor TargetDescriptor, UEBuildModule Module, ILogger Logger, int UnitySize, bool bPriming, bool bDisableUnity)
 		{
 			const int CompileCount = 3;
 			List<TimingData> AllTimingData = new();
 			for (int CompileIndex = 0; CompileIndex < CompileCount; CompileIndex++)
 			{
-				TimingData NewTimingData = CompileModule(LogPrefix, BuildConfiguration, TargetDescriptor, Module, Logger, UnitySize, bPriming, bDisableUnity);
+				TimingData NewTimingData = await CompileModuleAsync(LogPrefix, BuildConfiguration, TargetDescriptor, Module, Logger, UnitySize, bPriming, bDisableUnity);
 				if (!NewTimingData.IsValid())
 				{
 					return NewTimingData;
@@ -296,7 +297,7 @@ namespace UnrealBuildTool
 		/// <summary>
 		/// Compiles the module and returns the timing information
 		/// </summary>
-		private TimingData CompileModule(string LogPrefix, BuildConfiguration BuildConfiguration, TargetDescriptor TargetDescriptor, UEBuildModule Module, ILogger Logger, int UnitySize, bool bPriming, bool bDisableUnity)
+		private async Task<TimingData> CompileModuleAsync(string LogPrefix, BuildConfiguration BuildConfiguration, TargetDescriptor TargetDescriptor, UEBuildModule Module, ILogger Logger, int UnitySize, bool bPriming, bool bDisableUnity)
 		{
 			// Store the old arguments
 			string[] OldArgs = TargetDescriptor.AdditionalArguments.GetRawArray();
@@ -345,7 +346,7 @@ namespace UnrealBuildTool
 					}
 
 					TimingLogger NewTimingLogger = new(Logger);
-					BuildMode.Build(new List<TargetDescriptor>() { TargetDescriptor }, BuildConfiguration, WorkingSet, BuildOptions.None, null, NewTimingLogger);
+					await BuildMode.BuildAsync(new List<TargetDescriptor>() { TargetDescriptor }, BuildConfiguration, WorkingSet, BuildOptions.None, null, NewTimingLogger);
 					NewTimingData = NewTimingLogger.TimingData;
 				}
 			}
diff --git a/Engine/Source/Programs/UnrealBuildTool/Modes/QueryMode.cs b/Engine/Source/Programs/UnrealBuildTool/Modes/QueryMode.cs
--- a/Engine/Source/Programs/UnrealBuildTool/Modes/QueryMode.cs
+++ b/Engine/Source/Programs/UnrealBuildTool/Modes/QueryMode.cs
@@ -6,6 +6,7 @@ using System;
 using System.Linq;
 using System.Text.Json;
 using System.Text.Json.Serialization;
+using System.Threading.Tasks;
 using System.Collections.Generic;
 using OpenTracing.Util;
 using UnrealBuildTool;
@@ -97,7 +98,7 @@ namespace UnrealBuildTool
 
 		private BuildConfiguration BuildConfiguration = new();
 
-		public override int Execute(CommandLineArguments Arguments, ILogger Logger)
+		public override Task<int> ExecuteAsync(CommandLineArguments Arguments, ILogger Logger)
 		{
 			Arguments.ApplyTo(this);
 
@@ -119,7 +120,7 @@ namespace UnrealBuildTool
 
 			if (Query == null)
 			{
-				return 0;
+				return Task.FromResult(0);
 			}
 
 			var ResponseOptions = new JsonSerializerOptions
@@ -131,16 +132,16 @@ namespace UnrealBuildTool
 			{
 				case QueryType.Capabilities:
 					Logger.LogInformation("QueryCapabilities");
-					return QueryCapabilities(Arguments, Logger, ResponseOptions);
+					return Task.FromResult(QueryCapabilities(Arguments, Logger, ResponseOptions));
 				case QueryType.AvailableTargets:
 					Logger.LogInformation("QueryAvailableTargets");
-					return QueryAvailableTargets(Arguments, Logger, ResponseOptions);
+					return Task.FromResult(QueryAvailableTargets(Arguments, Logger, ResponseOptions));
 				case QueryType.TargetDetails:
 					Logger.LogInformation("QueryTargetDetails");
-					return QueryTargetDetails(Arguments, Logger, ResponseOptions);
+					return Task.FromResult(QueryTargetDetails(Arguments, Logger, ResponseOptions));
 			}
 
-			return 0;
+			return Task.FromResult(0);
 		}
 
 		private int QueryCapabilities(CommandLineArguments Arguments, ILogger Logger, JsonSerializerOptions JsonOptions)
diff --git a/Engine/Source/Programs/UnrealBuildTool/Modes/QueryTargetsMode.cs b/Engine/Source/Programs/UnrealBuildTool/Modes/QueryTargetsMode.cs
--- a/Engine/Source/Programs/UnrealBuildTool/Modes/QueryTargetsMode.cs
+++ b/Engine/Source/Programs/UnrealBuildTool/Modes/QueryTargetsMode.cs
@@ -41,7 +41,7 @@ namespace UnrealBuildTool
 		/// <param name="Arguments">Command line arguments</param>
 		/// <returns></returns>
 		/// <param name="Logger"></param>
-		public override int Execute(CommandLineArguments Arguments, ILogger Logger)
+		public override Task<int> ExecuteAsync(CommandLineArguments Arguments, ILogger Logger)
 		{
 			Arguments.ApplyTo(this);
 
@@ -70,7 +70,7 @@ namespace UnrealBuildTool
 			// Write information about these targets
 			WriteTargetInfo(ProjectFile, Assembly, OutputFile, Arguments, Logger, bIncludeAllTargets);
 			Logger.LogInformation("Written {OutputFile}", OutputFile);
-			return 0;
+			return Task.FromResult(0);
 		}
 
 		/// <summary>
diff --git a/Engine/Source/Programs/UnrealBuildTool/Modes/ServerMode.cs b/Engine/Source/Programs/UnrealBuildTool/Modes/ServerMode.cs
--- a/Engine/Source/Programs/UnrealBuildTool/Modes/ServerMode.cs
+++ b/Engine/Source/Programs/UnrealBuildTool/Modes/ServerMode.cs
@@ -6,6 +6,7 @@ using System;
 using System.Linq;
 using System.Text.Json;
 using System.Text.Json.Serialization;
+using System.Threading.Tasks;
 using System.Collections.Generic;
 using OpenTracing.Util;
 using UnrealBuildTool;
@@ -140,7 +141,7 @@ namespace UnrealBuildTool
 		private TargetIntellisenseInfo? CurrentTargetIntellisenseInfo = null;
 		private GetBrowseConfigurationResponse CurrentBrowseConfiguration;
 
-		public override int Execute(CommandLineArguments Arguments, ILogger Logger)
+		public override Task<int> ExecuteAsync(CommandLineArguments Arguments, ILogger Logger)
 		{
 			Arguments.ApplyTo(this);
 
@@ -171,7 +172,7 @@ namespace UnrealBuildTool
 				if (Line == null) { continue; }
 				HandleCommand(Line, Logger);
 			}
-			return 0;
+			return Task.FromResult(0);
 		}
 
 		private void HandleCommand(string CommandString, ILogger Logger)
diff --git a/Engine/Source/Programs/UnrealBuildTool/Modes/SetupPlatformsMode.cs b/Engine/Source/Programs/UnrealBuildTool/Modes/SetupPlatformsMode.cs
--- a/Engine/Source/Programs/UnrealBuildTool/Modes/SetupPlatformsMode.cs
+++ b/Engine/Source/Programs/UnrealBuildTool/Modes/SetupPlatformsMode.cs
@@ -22,11 +22,11 @@ namespace UnrealBuildTool
 		/// <param name="Arguments">Command line arguments</param>
 		/// <returns>Exit code</returns>
 		/// <param name="Logger"></param>
-		public override int Execute(CommandLineArguments Arguments, ILogger Logger)
+		public override Task<int> ExecuteAsync(CommandLineArguments Arguments, ILogger Logger)
 		{
 			// Output a warning if there are any arguments that are still unused
 			Arguments.CheckAllArgumentsUsed();
-			return 0;
+			return Task.FromResult(0);
 		}
 	}
 }
diff --git a/Engine/Source/Programs/UnrealBuildTool/Modes/TestMode.cs b/Engine/Source/Programs/UnrealBuildTool/Modes/TestMode.cs
--- a/Engine/Source/Programs/UnrealBuildTool/Modes/TestMode.cs
+++ b/Engine/Source/Programs/UnrealBuildTool/Modes/TestMode.cs
@@ -1,6 +1,7 @@
 // Copyright Epic Games, Inc. All Rights Reserved.
 
 using System.Collections.Generic;
+using System.Threading.Tasks;
 using EpicGames.Core;
 using UnrealBuildBase;
 using Microsoft.Extensions.Logging;
@@ -18,7 +19,7 @@ namespace UnrealBuildTool
 		/// </summary>
 		/// <param name="Arguments">Command-line arguments</param>
 		/// <param name="Logger"></param>
-		public override int Execute(CommandLineArguments Arguments, ILogger Logger)
+		public override async Task<int> ExecuteAsync(CommandLineArguments Arguments, ILogger Logger)
 		{
 			Arguments.ApplyTo(this);
 
@@ -30,7 +31,7 @@ namespace UnrealBuildTool
 			// Parse all the targets being built
 			List<TargetDescriptor> TargetDescriptors = TargetDescriptor.ParseCommandLine(Arguments, false, false, false, Logger);
 
-			BuildTests(TargetDescriptors, BuildConfiguration, Logger);
+			await BuildTestsAsync(TargetDescriptors, BuildConfiguration, Logger);
 
 			return 0;
 		}
@@ -46,7 +47,7 @@ namespace UnrealBuildTool
 		/// <param name="TargetDescriptors">Target descriptors</param>
 		/// <param name="BuildConfiguration">Current build configuration</param>
 		/// <param name="Logger"></param>
-		public static void BuildTests(List<TargetDescriptor> TargetDescriptors, BuildConfiguration BuildConfiguration, ILogger Logger)
+		public static async Task BuildTestsAsync(List<TargetDescriptor> TargetDescriptors, BuildConfiguration BuildConfiguration, ILogger Logger)
 		{
 			List<TargetDescriptor> TestTargetDescriptors = new List<TargetDescriptor>();
 
@@ -60,7 +61,7 @@ namespace UnrealBuildTool
 
 			using (ISourceFileWorkingSet WorkingSet = SourceFileWorkingSet.Create(Unreal.RootDirectory, new HashSet<DirectoryReference>(), Logger))
 			{
-				BuildMode.Build(TestTargetDescriptors, BuildConfiguration, WorkingSet, BuildOptions.None, null, Logger);
+				await BuildMode.BuildAsync(TestTargetDescriptors, BuildConfiguration, WorkingSet, BuildOptions.None, null, Logger);
 			}
 		}
 	}
diff --git a/Engine/Source/Programs/UnrealBuildTool/Modes/UnrealHeaderToolMode.cs b/Engine/Source/Programs/UnrealBuildTool/Modes/UnrealHeaderToolMode.cs
--- a/Engine/Source/Programs/UnrealBuildTool/Modes/UnrealHeaderToolMode.cs
+++ b/Engine/Source/Programs/UnrealBuildTool/Modes/UnrealHeaderToolMode.cs
@@ -10,6 +10,7 @@ using System.Collections.Generic;
 using System.IO;
 using System.Reflection;
 using System.Text;
+using System.Threading.Tasks;
 using UnrealBuildBase;
 using Microsoft.Extensions.Logging;
 
@@ -1113,7 +1114,7 @@ namespace UnrealBuildTool.Modes
 		/// <param name="Arguments">Command line arguments</param>
 		/// <returns>Exit code</returns>
 		/// <param name="Logger"></param>
-		public override int Execute(CommandLineArguments Arguments, ILogger Logger)
+		public override async Task<int> ExecuteAsync(CommandLineArguments Arguments, ILogger Logger)
 		{
 			try
 			{
@@ -1205,7 +1206,7 @@ namespace UnrealBuildTool.Modes
 					Arguments.ApplyTo(BuildConfiguration);
 
 					// Create the makefile
-					TargetMakefile Makefile = Target.Build(BuildConfiguration, WorkingSet, TargetDesc, Logger, true);
+					TargetMakefile Makefile = await Target.BuildAsync(BuildConfiguration, WorkingSet, TargetDesc, Logger, true);
 
 					FileReference ModuleInfoFileName = ExternalExecution.GetUHTModuleInfoFileName(Makefile, Target.TargetName);
 					FileReference DepsFileName = ExternalExecution.GetUHTDepsFileName(ModuleInfoFileName);
diff --git a/Engine/Source/Programs/UnrealBuildTool/Modes/ValidatePlatformsMode.cs b/Engine/Source/Programs/UnrealBuildTool/Modes/ValidatePlatformsMode.cs
--- a/Engine/Source/Programs/UnrealBuildTool/Modes/ValidatePlatformsMode.cs
+++ b/Engine/Source/Programs/UnrealBuildTool/Modes/ValidatePlatformsMode.cs
@@ -40,7 +40,7 @@ namespace UnrealBuildTool
 		/// <param name="Arguments">Command line arguments</param>
 		/// <returns>Exit code</returns>
 		/// <param name="Logger"></param>
-		public override int Execute(CommandLineArguments Arguments, ILogger Logger)
+		public override Task<int> ExecuteAsync(CommandLineArguments Arguments, ILogger Logger)
 		{
 			// Output a message if there are any arguments that are still unused
 			Arguments.ApplyTo(this);
@@ -71,7 +71,7 @@ namespace UnrealBuildTool
 					Logger.LogInformation("##PlatformValidate: {Platform} INVALID {PlatformSdkString}", Platform.ToString(), PlatformSDKString);
 				}
 			} 
-			return 0;
+			return Task.FromResult(0);
 		}
 	}
 }
diff --git a/Engine/Source/Programs/UnrealBuildTool/Modes/WriteDocumentationMode.cs b/Engine/Source/Programs/UnrealBuildTool/Modes/WriteDocumentationMode.cs
--- a/Engine/Source/Programs/UnrealBuildTool/Modes/WriteDocumentationMode.cs
+++ b/Engine/Source/Programs/UnrealBuildTool/Modes/WriteDocumentationMode.cs
@@ -43,7 +43,7 @@ namespace UnrealBuildTool
 		/// Entry point for this command
 		/// </summary>
 		/// <returns></returns>
-		public override int Execute(CommandLineArguments Arguments, ILogger Logger)
+		public override Task<int> ExecuteAsync(CommandLineArguments Arguments, ILogger Logger)
 		{
 			Arguments.ApplyTo(this);
 			Arguments.CheckAllArgumentsUsed();
@@ -62,7 +62,7 @@ namespace UnrealBuildTool
 				default:
 					throw new BuildException("Invalid documentation type: {0}", Type);
 			}
-			return 0;
+			return Task.FromResult(0);
 		}
 	}
 }
diff --git a/Engine/Source/Programs/UnrealBuildTool/Modes/WriteMetadataMode.cs b/Engine/Source/Programs/UnrealBuildTool/Modes/WriteMetadataMode.cs
--- a/Engine/Source/Programs/UnrealBuildTool/Modes/WriteMetadataMode.cs
+++ b/Engine/Source/Programs/UnrealBuildTool/Modes/WriteMetadataMode.cs
@@ -95,7 +95,7 @@ namespace UnrealBuildTool
 		/// <param name="Arguments">Command line arguments</param>
 		/// <returns>Exit code</returns>
 		/// <param name="Logger"></param>
-		public override int Execute(CommandLineArguments Arguments, ILogger Logger)
+		public override Task<int> ExecuteAsync(CommandLineArguments Arguments, ILogger Logger)
 		{
 			// Acquire a different mutex to the regular UBT instance, since this mode will be called as part of a build. We need the mutex to ensure that building two modular configurations 
 			// in parallel don't clash over writing shared *.modules files (eg. DebugGame and Development editors).
@@ -112,7 +112,7 @@ namespace UnrealBuildTool
 		/// <param name="Arguments">Command line arguments</param>
 		/// <param name="Logger">Logger for output</param>
 		/// <returns>Exit code</returns>
-		private int ExecuteInternal(CommandLineArguments Arguments, ILogger Logger)
+		private Task<int> ExecuteInternal(CommandLineArguments Arguments, ILogger Logger)
 		{
 			// Read the target info
 			WriteMetadataTargetInfo TargetInfo = BinaryFormatterUtils.Load<WriteMetadataTargetInfo>(Arguments.GetFileReference("-Input="));
@@ -252,7 +252,7 @@ namespace UnrealBuildTool
 				TargetInfo.Receipt.Write(TargetInfo.ReceiptFile);
 			}
 
-			return 0;
+			return Task.FromResult(0);
 		}
 
 		/// <summary>
diff --git a/Engine/Source/Programs/UnrealBuildTool/Platform/IOS/IOSPostBuildSyncMode.cs b/Engine/Source/Programs/UnrealBuildTool/Platform/IOS/IOSPostBuildSyncMode.cs
--- a/Engine/Source/Programs/UnrealBuildTool/Platform/IOS/IOSPostBuildSyncMode.cs
+++ b/Engine/Source/Programs/UnrealBuildTool/Platform/IOS/IOSPostBuildSyncMode.cs
@@ -65,7 +65,7 @@ namespace UnrealBuildTool
 		[CommandLine("-XmlConfigCache=")]
 		public FileReference? XmlConfigCache = null;
 
-		public override int Execute(CommandLineArguments Arguments, ILogger Logger)
+		public override Task<int> ExecuteAsync(CommandLineArguments Arguments, ILogger Logger)
 		{
 			Arguments.ApplyTo(this);
 			Arguments.CheckAllArgumentsUsed();
@@ -74,7 +74,7 @@ namespace UnrealBuildTool
 			IOSPostBuildSyncTarget Target = BinaryFormatterUtils.Load<IOSPostBuildSyncTarget>(InputFile!);
 			IOSToolChain.PostBuildSync(Target, Logger);
 
-			return 0;
+			return Task.FromResult(0);
 		}
 	}
 }
diff --git a/Engine/Source/Programs/UnrealBuildTool/Platform/Windows/PVSToolChain.cs b/Engine/Source/Programs/UnrealBuildTool/Platform/Windows/PVSToolChain.cs
--- a/Engine/Source/Programs/UnrealBuildTool/Platform/Windows/PVSToolChain.cs
+++ b/Engine/Source/Programs/UnrealBuildTool/Platform/Windows/PVSToolChain.cs
@@ -420,7 +420,7 @@ namespace UnrealBuildTool
 		/// <param name="Arguments">List of command line arguments</param>
 		/// <returns>Always zero, or throws an exception</returns>
 		/// <param name="Logger"></param>
-		public override int Execute(CommandLineArguments Arguments, ILogger Logger)
+		public override Task<int> ExecuteAsync(CommandLineArguments Arguments, ILogger Logger)
 		{
 			Arguments.ApplyTo(this);
 			Arguments.CheckAllArgumentsUsed();
@@ -489,7 +489,7 @@ namespace UnrealBuildTool
 				}
 			}
 			Logger.LogInformation("Written {NumItems} {Noun} to {File}.", UniqueItems.Count, (UniqueItems.Count == 1)? "diagnostic" : "diagnostics", OutputFile.FullName);
-			return 0;
+			return Task.FromResult(0);
 		}
 	}
 
diff --git a/Engine/Source/Programs/UnrealBuildTool/System/ActionExecutor.cs b/Engine/Source/Programs/UnrealBuildTool/System/ActionExecutor.cs
--- a/Engine/Source/Programs/UnrealBuildTool/System/ActionExecutor.cs
+++ b/Engine/Source/Programs/UnrealBuildTool/System/ActionExecutor.cs
@@ -5,6 +5,7 @@ using Microsoft.Extensions.Logging;
 using System;
 using System.Collections.Generic;
 using System.Reflection;
+using System.Threading.Tasks;
 
 namespace UnrealBuildTool
 {
@@ -53,7 +54,7 @@ namespace UnrealBuildTool
 			MemoryPerActionBytesOverride = Math.Max(MemoryPerActionBytesOverride, MemoryPerActionOverrideGB * 1024 * 1024 * 1024);
 		}
 
-		public abstract bool ExecuteActions(IEnumerable<LinkedAction> ActionsToExecute, ILogger Logger);
+		public abstract Task<bool> ExecuteActionsAsync(IEnumerable<LinkedAction> ActionsToExecute, ILogger Logger);
 
 		protected void WriteToolOutput(string Line)
 		{
diff --git a/Engine/Source/Programs/UnrealBuildTool/System/ActionGraph.cs b/Engine/Source/Programs/UnrealBuildTool/System/ActionGraph.cs
--- a/Engine/Source/Programs/UnrealBuildTool/System/ActionGraph.cs
+++ b/Engine/Source/Programs/UnrealBuildTool/System/ActionGraph.cs
@@ -18,6 +18,7 @@ using Microsoft.Extensions.Logging;
 using System.Collections.Concurrent;
 using System.Collections.Immutable;
 using System.Reflection;
+using System.Runtime.CompilerServices;
 
 namespace UnrealBuildTool
 {
@@ -356,7 +357,7 @@ namespace UnrealBuildTool
 #if __BOXEXECUTOR_AVAILABLE__
 			if (BuildConfiguration.bAllowBoxExecutor)
 			{
-				return new BoxExecutor(BuildConfiguration.MaxParallelActions, BuildConfiguration.bAllCores, BuildConfiguration.bCompactOutput, Logger);
+//				return new BoxExecutor(BuildConfiguration.MaxParallelActions, BuildConfiguration.bAllCores, BuildConfiguration.bCompactOutput, Logger);
 			}
 #endif // #if __BOXEXECUTOR_AVAILABLE__
 
@@ -386,7 +387,7 @@ namespace UnrealBuildTool
 		/// <summary>
 		/// Executes a list of actions.
 		/// </summary>
-		public static void ExecuteActions(BuildConfiguration BuildConfiguration, List<LinkedAction> ActionsToExecute, List<TargetDescriptor> TargetDescriptors, ILogger Logger)
+		public static async Task ExecuteActionsAsync(BuildConfiguration BuildConfiguration, List<LinkedAction> ActionsToExecute, List<TargetDescriptor> TargetDescriptors, ILogger Logger)
 		{
 			if (ActionsToExecute.Count == 0)
 			{
@@ -399,7 +400,7 @@ namespace UnrealBuildTool
 
 				// Execute the build
 				Stopwatch Timer = Stopwatch.StartNew();
-				if (!Executor.ExecuteActions(ActionsToExecute, Logger))
+				if (!await Executor.ExecuteActionsAsync(ActionsToExecute, Logger))
 				{
 					throw new CompilationResultException(CompilationResult.OtherCompilationError);
 				}
diff --git a/Engine/Source/Programs/UnrealBuildTool/System/ExternalExecution.cs b/Engine/Source/Programs/UnrealBuildTool/System/ExternalExecution.cs
--- a/Engine/Source/Programs/UnrealBuildTool/System/ExternalExecution.cs
+++ b/Engine/Source/Programs/UnrealBuildTool/System/ExternalExecution.cs
@@ -913,7 +913,7 @@ namespace UnrealBuildTool
 		/// Builds and runs the header tool and touches the header directories.
 		/// Performs any early outs if headers need no changes, given the UObject modules, tool path, game name, and configuration
 		/// </summary>
-		public static void ExecuteHeaderToolIfNecessary(BuildConfiguration BuildConfiguration, FileReference? ProjectFile, TargetMakefile Makefile, string TargetName, ISourceFileWorkingSet WorkingSet, ILogger Logger)
+		public static async Task ExecuteHeaderToolIfNecessaryAsync(BuildConfiguration BuildConfiguration, FileReference? ProjectFile, TargetMakefile Makefile, string TargetName, ISourceFileWorkingSet WorkingSet, ILogger Logger)
 		{
 
 			// No need to run UHT on itself
@@ -923,7 +923,7 @@ namespace UnrealBuildTool
 				return;
 			}
 
-			ExecuteHeaderToolIfNecessaryInternal(BuildConfiguration, ProjectFile, Makefile, TargetName, WorkingSet, Logger);
+			await ExecuteHeaderToolIfNecessaryInternalAsync(BuildConfiguration, ProjectFile, Makefile, TargetName, WorkingSet, Logger);
 		}
 
 		/// <summary>
@@ -1020,7 +1020,7 @@ namespace UnrealBuildTool
 		/// Builds and runs the header tool and touches the header directories.
 		/// Performs any early outs if headers need no changes, given the UObject modules, tool path, game name, and configuration
 		/// </summary>
-		private static void ExecuteHeaderToolIfNecessaryInternal(BuildConfiguration BuildConfiguration, FileReference? ProjectFile, 
+		private static async Task ExecuteHeaderToolIfNecessaryInternalAsync(BuildConfiguration BuildConfiguration, FileReference? ProjectFile, 
 			TargetMakefile Makefile, string TargetName, ISourceFileWorkingSet WorkingSet, ILogger Logger)
 		{
 			if (ProgressWriter.bWriteMarkup)
@@ -1140,7 +1140,7 @@ namespace UnrealBuildTool
 					s.Start();
 					IScope Timer = GlobalTracer.Instance.BuildSpan("Executing UnrealHeaderTool").StartActive();
 					UnrealHeaderToolMode UHTTool = new UnrealHeaderToolMode();
-					CompilationResult UHTResult = (CompilationResult)UHTTool.Execute(Arguments, Logger);
+					CompilationResult UHTResult = (CompilationResult)await UHTTool.ExecuteAsync(Arguments, Logger);
 					Timer.Span.Finish();
 					s.Stop();
 
diff --git a/Engine/Source/Programs/UnrealBuildTool/System/ToolMode.cs b/Engine/Source/Programs/UnrealBuildTool/System/ToolMode.cs
--- a/Engine/Source/Programs/UnrealBuildTool/System/ToolMode.cs
+++ b/Engine/Source/Programs/UnrealBuildTool/System/ToolMode.cs
@@ -101,6 +101,6 @@ namespace UnrealBuildTool
 		/// <param name="Arguments">List of command line arguments</param>
 		/// <param name="Logger"></param>
 		/// <returns>Exit code for the process</returns>
-		public abstract int Execute(CommandLineArguments Arguments, ILogger Logger);
+		public abstract Task<int> ExecuteAsync(CommandLineArguments Arguments, ILogger Logger);
 	}
 }
diff --git a/Engine/Source/Programs/UnrealBuildTool/ToolChain/AppleToolChain.cs b/Engine/Source/Programs/UnrealBuildTool/ToolChain/AppleToolChain.cs
--- a/Engine/Source/Programs/UnrealBuildTool/ToolChain/AppleToolChain.cs
+++ b/Engine/Source/Programs/UnrealBuildTool/ToolChain/AppleToolChain.cs
@@ -6,6 +6,7 @@ using System.Collections.Generic;
 using System.IO;
 using System.Diagnostics;
 using System.Text;
+using System.Threading.Tasks;
 using EpicGames.Core;
 using System.Text.RegularExpressions;
 using UnrealBuildBase;
@@ -576,7 +577,7 @@ namespace UnrealBuildTool
 		[CommandLine("-XmlConfigCache=")]
 		public FileReference? XmlConfigCache = null;
 
-		public override int Execute(CommandLineArguments Arguments, ILogger Logger)
+		public override Task<int> ExecuteAsync(CommandLineArguments Arguments, ILogger Logger)
 		{
 			Arguments.ApplyTo(this);
 			Arguments.CheckAllArgumentsUsed();
@@ -585,7 +586,7 @@ namespace UnrealBuildTool
 			ApplePostBuildSyncTarget Target = BinaryFormatterUtils.Load<ApplePostBuildSyncTarget>(InputFile!);
 			PostBuildSync(Target, Logger);
 
-			return 0;
+			return Task.FromResult(0);
 		}
 
 		private void PostBuildSync(ApplePostBuildSyncTarget Target, ILogger Logger)
diff --git a/Engine/Source/Programs/UnrealBuildTool/UnrealBuildTool.cs b/Engine/Source/Programs/UnrealBuildTool/UnrealBuildTool.cs
--- a/Engine/Source/Programs/UnrealBuildTool/UnrealBuildTool.cs
+++ b/Engine/Source/Programs/UnrealBuildTool/UnrealBuildTool.cs
@@ -413,7 +413,7 @@ namespace UnrealBuildTool
 		/// </summary>
 		/// <param name="ArgumentsArray">Command line arguments</param>
 		/// <returns>Zero on success, non-zero on error</returns>
-		private static int Main(string[] ArgumentsArray)
+		private static async Task<int> Main(string[] ArgumentsArray)
 		{
 			FileReference? RunFile = null;
 			DirectoryReference? TempDirectory = null;
@@ -652,7 +652,7 @@ namespace UnrealBuildTool
 				ToolMode Mode = (ToolMode)Activator.CreateInstance(ModeType)!;
 
 				// Execute the mode
-				int Result = Mode.Execute(Arguments, Logger);
+				int Result = await Mode.ExecuteAsync(Arguments, Logger);
 				if((ModeOptions & ToolModeOptions.ShowExecutionTime) != 0)
 				{
 					Logger.LogInformation("Total execution time: {Time:0.00} seconds", Timeline.Elapsed.TotalSeconds);
