UnrealBuildTool: Reconcile common clang warnings between VCToolChain and ClangToolChain

#rnx

diff --git a/Engine/Source/Programs/UnrealBuildTool/Platform/Windows/VCToolChain.cs b/Engine/Source/Programs/UnrealBuildTool/Platform/Windows/VCToolChain.cs
--- a/Engine/Source/Programs/UnrealBuildTool/Platform/Windows/VCToolChain.cs
+++ b/Engine/Source/Programs/UnrealBuildTool/Platform/Windows/VCToolChain.cs
@@ -937,11 +937,9 @@ namespace UnrealBuildTool
 					break;
 				case CppStandardVersion.Cpp20:
 					Arguments.Add("/std:c++20");
-					Arguments.Add("/Zc:preprocessor");
 					break;
 				case CppStandardVersion.Latest:
 					Arguments.Add("/std:c++latest");
-					Arguments.Add("/Zc:preprocessor");
 					break;
 				default:
 					throw new BuildException($"Unsupported C++ standard type set: {CompileEnvironment.CppStandard}");
@@ -949,6 +947,11 @@ namespace UnrealBuildTool
 
 			if (CompileEnvironment.CppStandard >= CppStandardVersion.Cpp20)
 			{
+				if (Target.WindowsPlatform.Compiler.IsMSVC())
+				{
+					Arguments.Add("/Zc:preprocessor");
+				}
+
 				// warning C5054: operator ___: deprecated between enumerations of different types
 				// re: http://www.open-std.org/jtc1/sc22/wg21/docs/papers/2018/p1120r0.html
 
@@ -993,137 +996,17 @@ namespace UnrealBuildTool
 				//
 
 				// Treat all warnings as errors by default
-				Arguments.Add("-Werror");
-
-				// Allow Microsoft-specific syntax to slide, even though it may be non-standard.  Needed for Windows headers.
-				Arguments.Add("-Wno-microsoft");
+				Arguments.Add("-Werror");                                   // https://clang.llvm.org/docs/UsersManual.html#cmdoption-werror
 
-				// @todo clang: Hack due to how we have our 'DummyPCH' wrappers setup when using unity builds.  This warning should not be disabled!!
-				Arguments.Add("-Wno-msvc-include");
-
-				if (CompileEnvironment.ShadowVariableWarningLevel != WarningLevel.Off)
-				{
-					Arguments.Add("-Wshadow" + ((CompileEnvironment.ShadowVariableWarningLevel == WarningLevel.Error) ? "" : " -Wno-error=shadow"));
-				}
-
-				if (CompileEnvironment.bEnableUndefinedIdentifierWarnings)
-				{
-					Arguments.Add(" -Wundef" + (CompileEnvironment.bUndefinedIdentifierWarningsAsErrors ? "" : " -Wno-error=undef"));
-				}
-
-				// Note: This should be kept in sync with PRAGMA_DISABLE_UNSAFE_TYPECAST_WARNINGS in ClangPlatformCompilerPreSetup.h
-				string[] UnsafeTypeCastWarningList = {
-					"float-conversion",
-					"implicit-float-conversion",
-					"implicit-int-conversion",
-					"c++11-narrowing"
-					//"shorten-64-to-32",	<-- too many hits right now, probably want it *soon*
-					//"sign-conversion",	<-- too many hits right now, probably want it eventually
-				};
-
-				if (CompileEnvironment.UnsafeTypeCastWarningLevel == WarningLevel.Error)
-				{
-					foreach (string Warning in UnsafeTypeCastWarningList)
-					{
-						Arguments.Add("-W" + Warning);
-					}
-				}
-				else if (CompileEnvironment.UnsafeTypeCastWarningLevel == WarningLevel.Warning)
-				{
-					foreach (string Warning in UnsafeTypeCastWarningList)
-					{
-						Arguments.Add("-W" + Warning + " -Wno-error=" + Warning);
-					}
-				}
-				else
-				{
-					foreach (string Warning in UnsafeTypeCastWarningList)
-					{
-						Arguments.Add("-Wno-" + Warning);
-					}
-				}
+				ClangWarnings.GetEnabledWarnings(Arguments); 
 
-				// Warn if __DATE__ or __TIME__ are used as they prevent reproducible builds
-				if (CompileEnvironment.bDeterministic)
-				{
-					if (CompileEnvironment.DeterministicWarningLevel == WarningLevel.Error)
-					{
-						Arguments.Add("-Wdate-time");
-					}
-					else if (CompileEnvironment.DeterministicWarningLevel == WarningLevel.Warning)
-					{
-						Arguments.Add("-Wdate-time -Wno-error=date-time");
-					}
-				}
+				ClangWarnings.GetDisabledWarnings(CompileEnvironment,
+					Target.StaticAnalyzer,
+					EnvVars.CompilerVersion,
+					Arguments);
 
-				// This is disabled because clang explicitly warns about changing pack alignment in a header and not
-				// restoring it afterwards, which is something we do with the Pre/PostWindowsApi.h headers.
-				Arguments.Add("-Wno-pragma-pack");
-
-				// @todo clang: Kind of a shame to turn these off.  We'd like to catch unused variables, but it is tricky with how our assertion macros work.
-				Arguments.Add("-Wno-inconsistent-missing-override");
-				Arguments.Add("-Wno-unused-variable");
-				if (EnvVars.CompilerVersion >= new VersionNumber(13))
-				{
-					Arguments.Add("-Wno-unused-but-set-variable");
-					Arguments.Add("-Wno-unused-but-set-parameter");
-				}
-				if (EnvVars.CompilerVersion >= new VersionNumber(14))
-				{
-					Arguments.Add("-Wno-bitwise-instead-of-logical");
-				}
-				Arguments.Add("-Wno-unused-local-typedefs");
-				Arguments.Add("-Wno-unused-function");
-				Arguments.Add("-Wno-unused-private-field");
-				Arguments.Add("-Wno-unused-value");
-
-				Arguments.Add("-Wno-inline-new-delete");	// @todo clang: We declare operator new as inline.  Clang doesn't seem to like that.
-				Arguments.Add("-Wno-implicit-exception-spec-mismatch");
-
-				// Sometimes we compare 'this' pointers against nullptr, which Clang warns about by default
-				Arguments.Add("-Wno-undefined-bool-conversion");
-
-				// @todo clang: Disabled warnings were copied from MacToolChain for the most part
-				Arguments.Add("-Wno-deprecated-declarations");
-				Arguments.Add("-Wno-deprecated-writable-strings");
-				Arguments.Add("-Wno-deprecated-register");
-				Arguments.Add("-Wno-switch-enum");
-				Arguments.Add("-Wno-logical-op-parentheses");	// needed for external headers we shan't change
-				Arguments.Add("-Wno-null-arithmetic");			// needed for external headers we shan't change
-				Arguments.Add("-Wno-deprecated-declarations");	// needed for wxWidgets
-				Arguments.Add("-Wno-return-type-c-linkage");	// needed for PhysX
-				Arguments.Add("-Wno-ignored-attributes");		// needed for nvtesslib
-				Arguments.Add("-Wno-uninitialized");
-				Arguments.Add("-Wno-tautological-compare");
-				Arguments.Add("-Wno-switch");
-				Arguments.Add("-Wno-invalid-offsetof"); // needed to suppress warnings about using offsetof on non-POD types.
-				Arguments.Add("-Wno-return-type");				// needed for external headers we shan't change 
-
-				// @todo clang: Sorry for adding more of these, but I couldn't read my output log. Most should probably be looked at
-				Arguments.Add("-Wno-unused-parameter");			// Unused function parameter. A lot are named 'bUnused'...
-				Arguments.Add("-Wno-ignored-qualifiers");		// const ignored when returning by value e.g. 'const int foo() { return 4; }'
-				Arguments.Add("-Wno-expansion-to-defined");		// Usage of 'defined(X)' in a macro definition. Gives different results under MSVC
-				Arguments.Add("-Wno-gnu-string-literal-operator-template");	// String literal operator"" in template, used by delegates
-				Arguments.Add("-Wno-sign-compare");				// Signed/unsigned comparison - millions of these
-				Arguments.Add("-Wno-undefined-var-template");	// Variable template instantiation required but no definition available
-				Arguments.Add("-Wno-missing-field-initializers"); // Stupid warning, generated when you initialize with MyStruct A = {0};
-				Arguments.Add("-Wno-unused-lambda-capture");
-				Arguments.Add("-Wno-nonportable-include-path");
-				Arguments.Add("-Wno-invalid-token-paste");
-				Arguments.Add("-Wno-null-pointer-arithmetic");
-				Arguments.Add("-Wno-constant-logical-operand"); // Triggered by || of two template-derived values inside a static_assert
-				if (EnvVars.CompilerVersion >= new VersionNumber(13))
-				{
-					Arguments.Add("-Wno-ordered-compare-function-pointers");
-				}
-
-				// C++20 warnings that should be addressed
-				if (CompileEnvironment.CppStandard >= CppStandardVersion.Cpp20)
-				{
-					Arguments.Add("-Wno-ambiguous-reversed-operator");
-					Arguments.Add("-Wno-deprecated-anon-enum-enum-conversion");
-					Arguments.Add("-Wno-deprecated-volatile");
-				}
+				// Additional disabled warnings for msvc. Everything below should be checked if it is necessary
+				ClangWarnings.GetVCDisabledWarnings(Arguments);
 			}
 		}
 
diff --git a/Engine/Source/Programs/UnrealBuildTool/ToolChain/ClangToolChain.cs b/Engine/Source/Programs/UnrealBuildTool/ToolChain/ClangToolChain.cs
--- a/Engine/Source/Programs/UnrealBuildTool/ToolChain/ClangToolChain.cs
+++ b/Engine/Source/Programs/UnrealBuildTool/ToolChain/ClangToolChain.cs
@@ -540,132 +540,12 @@ namespace UnrealBuildTool
 			Arguments.Add("-Wall");                                     // https://clang.llvm.org/docs/DiagnosticsReference.html#wall
 			Arguments.Add("-Werror");                                   // https://clang.llvm.org/docs/UsersManual.html#cmdoption-werror
 
-			Arguments.Add("-Wdelete-non-virtual-dtor");                 // https://clang.llvm.org/docs/DiagnosticsReference.html#wdelete-non-virtual-dtor
-			Arguments.Add("-Wenum-conversion");                         // https://clang.llvm.org/docs/DiagnosticsReference.html#wenum-conversion
-			Arguments.Add("-Wbitfield-enum-conversion");                // https://clang.llvm.org/docs/DiagnosticsReference.html#wbitfield-enum-conversion
+			ClangWarnings.GetEnabledWarnings(Arguments);
 
-			Arguments.Add("-Wno-enum-enum-conversion");                 // https://clang.llvm.org/docs/DiagnosticsReference.html#wenum-enum-conversion					// ?? no reason given
-			Arguments.Add("-Wno-enum-float-conversion");                // https://clang.llvm.org/docs/DiagnosticsReference.html#wenum-float-conversion					// ?? no reason given
-
-			// C++20 warnings that should be addressed
-			if (CompileEnvironment.CppStandard >= CppStandardVersion.Cpp20)
-			{
-				Arguments.Add("-Wno-ambiguous-reversed-operator");          // https://clang.llvm.org/docs/DiagnosticsReference.html#wambiguous-reversed-operator
-				Arguments.Add("-Wno-deprecated-anon-enum-enum-conversion"); // https://clang.llvm.org/docs/DiagnosticsReference.html#wdeprecated-anon-enum-enum-conversion
-				Arguments.Add("-Wno-deprecated-volatile");                  // https://clang.llvm.org/docs/DiagnosticsReference.html#wdeprecated-volatile
-			}
-
-			if (CompilerVersionGreaterOrEqual(13, 0, 0))
-			{
-				Arguments.Add("-Wno-unused-but-set-variable");           // https://clang.llvm.org/docs/DiagnosticsReference.html#wunused-but-set-variable				// new warning for clang 13
-				Arguments.Add("-Wno-unused-but-set-parameter");          // https://clang.llvm.org/docs/DiagnosticsReference.html#wunused-but-set-parameter				// new warning for clang 13
-				Arguments.Add("-Wno-ordered-compare-function-pointers"); // https://clang.llvm.org/docs/DiagnosticsReference.html#wordered-compare-function-pointers	// new warning for clang 13
-			}
-			if (CompilerVersionGreaterOrEqual(14, 0, 0))
-			{
-				Arguments.Add("-Wno-bitwise-instead-of-logical");       // https://clang.llvm.org/docs/DiagnosticsReference.html#wbitwise-instead-of-logical			// new warning for clang 14
-			}
-
-			Arguments.Add("-Wno-gnu-string-literal-operator-template"); // https://clang.llvm.org/docs/DiagnosticsReference.html#wgnu-string-literal-operator-template	// We use this feature to allow static FNames.
-			Arguments.Add("-Wno-inconsistent-missing-override");        // https://clang.llvm.org/docs/DiagnosticsReference.html#winconsistent-missing-override			// ?? no reason given
-			Arguments.Add("-Wno-invalid-offsetof");                     // https://clang.llvm.org/docs/DiagnosticsReference.html#winvalid-offsetof						// needed to suppress warnings about using offsetof on non-POD types.
-			Arguments.Add("-Wno-switch");                               // https://clang.llvm.org/docs/DiagnosticsReference.html#wswitch								// this hides the "enumeration value 'XXXXX' not handled in switch [-Wswitch]" warnings - we should maybe remove this at some point and add UE_LOG(, Fatal, ) to default cases
-			Arguments.Add("-Wno-tautological-compare");                 // https://clang.llvm.org/docs/DiagnosticsReference.html#wtautological-compare					// this hides the "warning : comparison of unsigned expression < 0 is always false" type warnings due to constant comparisons, which are possible with template arguments
-			Arguments.Add("-Wno-unknown-pragmas");                      // https://clang.llvm.org/docs/DiagnosticsReference.html#wunknown-pragmas						// Slate triggers this (with its optimize on/off pragmas)
-			Arguments.Add("-Wno-unused-function");                      // https://clang.llvm.org/docs/DiagnosticsReference.html#wunused-function						// this will hide the warnings about static functions in headers that aren't used in every single .cpp file
-			Arguments.Add("-Wno-unused-lambda-capture");                // https://clang.llvm.org/docs/DiagnosticsReference.html#wunused-lambda-capture					// suppressed because capturing of compile-time constants is seemingly inconsistent. And MSVC doesn't do that.
-			Arguments.Add("-Wno-unused-local-typedef");                 // https://clang.llvm.org/docs/DiagnosticsReference.html#wunused-local-typedef					// clang is being overly strict here? PhysX headers trigger this.
-			Arguments.Add("-Wno-unused-private-field");                 // https://clang.llvm.org/docs/DiagnosticsReference.html#wunused-private-field					// this will prevent the issue of warnings for unused private variables. MultichannelTcpSocket.h triggers this, possibly more
-			Arguments.Add("-Wno-unused-variable");                      // https://clang.llvm.org/docs/DiagnosticsReference.html#wunused-variable						// ?? no reason given
-			Arguments.Add("-Wno-undefined-var-template");               // https://clang.llvm.org/docs/DiagnosticsReference.html#wundefined-var-template				// not really a good warning to disable
-
-			// Profile Guided Optimization (PGO) and Link Time Optimization (LTO)
-			if (CompileEnvironment.bPGOOptimize)
-			{
-				//
-				// Clang emits warnings for each compiled object file that doesn't have a matching entry in the profile data.
-				// This can happen when the profile data is older than the binaries we're compiling.
-				//
-				// Disable these warnings. They are far too verbose.
-				//
-				Arguments.Add("-Wno-profile-instr-out-of-date");        // https://clang.llvm.org/docs/DiagnosticsReference.html#wprofile-instr-out-of-date
-				Arguments.Add("-Wno-profile-instr-unprofiled");         // https://clang.llvm.org/docs/DiagnosticsReference.html#wprofile-instr-unprofiled
-
-				// apparently there can be hashing conflicts with PGO which can result in:
-				// 'Function control flow change detected (hash mismatch)' warnings. 
-				Arguments.Add("-Wno-backend-plugin");                   // https://clang.llvm.org/docs/DiagnosticsReference.html#wbackend-plugin
-			}
-
-			// shipping builds will cause this warning with "ensure", so disable only in those case
-			if (CompileEnvironment.Configuration == CppConfiguration.Shipping || StaticAnalyzer != StaticAnalyzer.None)
-			{
-				Arguments.Add("-Wno-unused-value");                     // https://clang.llvm.org/docs/DiagnosticsReference.html#wunused-value
-			}
-
-			// https://clang.llvm.org/docs/DiagnosticsReference.html#wdeprecated-declarations
-			if (CompileEnvironment.DeprecationWarningLevel == WarningLevel.Error)
-			{
-				// TODO: This may be unnecessary with -Werror
-				Arguments.Add("-Werror=deprecated-declarations");
-			}
-
-			// Warn if __DATE__ or __TIME__ are used as they prevent reproducible builds
-			if (CompileEnvironment.bDeterministic)
-			{
-				// https://clang.llvm.org/docs/DiagnosticsReference.html#wdate-time
-				if (CompileEnvironment.DeterministicWarningLevel == WarningLevel.Error)
-				{
-					Arguments.Add("-Wdate-time");
-				}
-				else if (CompileEnvironment.DeterministicWarningLevel == WarningLevel.Warning)
-				{
-					Arguments.Add("-Wdate-time -Wno-error=date-time");
-				}
-			}
-
-			// https://clang.llvm.org/docs/DiagnosticsReference.html#wshadow
-			if (CompileEnvironment.ShadowVariableWarningLevel != WarningLevel.Off)
-			{
-				Arguments.Add("-Wshadow" + ((CompileEnvironment.ShadowVariableWarningLevel == WarningLevel.Error) ? "" : " -Wno-error=shadow"));
-			}
-
-			// https://clang.llvm.org/docs/DiagnosticsReference.html#wundef
-			if (CompileEnvironment.bEnableUndefinedIdentifierWarnings)
-			{
-				Arguments.Add("-Wundef" + (CompileEnvironment.bUndefinedIdentifierWarningsAsErrors ? "" : " -Wno-error=undef"));
-			}
-
-			// Note: This should be kept in sync with PRAGMA_DISABLE_UNSAFE_TYPECAST_WARNINGS in ClangPlatformCompilerPreSetup.h
-			string[] UnsafeTypeCastWarningList = {
-				"float-conversion",
-				"implicit-float-conversion",
-				"implicit-int-conversion",
-				"c++11-narrowing"
-				//"shorten-64-to-32",	<-- too many hits right now, probably want it *soon*
-				//"sign-conversion",	<-- too many hits right now, probably want it eventually
-			};
-
-			if (CompileEnvironment.UnsafeTypeCastWarningLevel == WarningLevel.Error)
-			{
-				foreach (string Warning in UnsafeTypeCastWarningList)
-				{
-					Arguments.Add("-W" + Warning);
-				}
-			}
-			else if (CompileEnvironment.UnsafeTypeCastWarningLevel == WarningLevel.Warning)
-			{
-				foreach (string Warning in UnsafeTypeCastWarningList)
-				{
-					Arguments.Add("-W" + Warning + " -Wno-error=" + Warning);
-				}
-			}
-			else
-			{
-				foreach (string Warning in UnsafeTypeCastWarningList)
-				{
-					Arguments.Add("-Wno-" + Warning);
-				}
-			}
+			ClangWarnings.GetDisabledWarnings(CompileEnvironment,
+				StaticAnalyzer,
+				new VersionNumber(Info.ClangVersion.Major, Info.ClangVersion.Minor, Info.ClangVersion.Build),
+				Arguments);
 
 			// always use absolute paths for errors, this can help IDEs go to the error properly
 			Arguments.Add("-fdiagnostics-absolute-paths");  // https://clang.llvm.org/docs/ClangCommandLineReference.html#cmdoption-clang-fdiagnostics-absolute-paths
diff --git a/Engine/Source/Programs/UnrealBuildTool/ToolChain/ClangWarnings.cs b/Engine/Source/Programs/UnrealBuildTool/ToolChain/ClangWarnings.cs
new file mode 100644
--- /dev/null
+++ b/Engine/Source/Programs/UnrealBuildTool/ToolChain/ClangWarnings.cs
@@ -0,0 +1,186 @@
+﻿using System.Collections.Generic;
+
+namespace UnrealBuildTool
+{
+	internal static class ClangWarnings
+	{
+		internal static void GetEnabledWarnings(List<string> Arguments)
+		{
+			Arguments.Add("-Wdelete-non-virtual-dtor");                 // https://clang.llvm.org/docs/DiagnosticsReference.html#wdelete-non-virtual-dtor
+			Arguments.Add("-Wenum-conversion");                         // https://clang.llvm.org/docs/DiagnosticsReference.html#wenum-conversion
+			Arguments.Add("-Wbitfield-enum-conversion");                // https://clang.llvm.org/docs/DiagnosticsReference.html#wbitfield-enum-conversion
+		}
+
+		internal static void GetDisabledWarnings(CppCompileEnvironment CompileEnvironment, StaticAnalyzer Analyzer, VersionNumber ClangVersion, List<string> Arguments)
+		{
+			Arguments.Add("-Wno-enum-enum-conversion");                 // https://clang.llvm.org/docs/DiagnosticsReference.html#wenum-enum-conversion					// ?? no reason given
+			Arguments.Add("-Wno-enum-float-conversion");                // https://clang.llvm.org/docs/DiagnosticsReference.html#wenum-float-conversion					// ?? no reason given
+
+			// C++20 warnings that should be addressed
+			if (CompileEnvironment.CppStandard >= CppStandardVersion.Cpp20)
+			{
+				Arguments.Add("-Wno-ambiguous-reversed-operator");          // https://clang.llvm.org/docs/DiagnosticsReference.html#wambiguous-reversed-operator
+				Arguments.Add("-Wno-deprecated-anon-enum-enum-conversion"); // https://clang.llvm.org/docs/DiagnosticsReference.html#wdeprecated-anon-enum-enum-conversion
+				Arguments.Add("-Wno-deprecated-volatile");                  // https://clang.llvm.org/docs/DiagnosticsReference.html#wdeprecated-volatile
+			}
+
+			if (ClangVersion >= new VersionNumber(13))
+			{
+				Arguments.Add("-Wno-unused-but-set-variable");           // https://clang.llvm.org/docs/DiagnosticsReference.html#wunused-but-set-variable				// new warning for clang 13
+				Arguments.Add("-Wno-unused-but-set-parameter");          // https://clang.llvm.org/docs/DiagnosticsReference.html#wunused-but-set-parameter				// new warning for clang 13
+				Arguments.Add("-Wno-ordered-compare-function-pointers"); // https://clang.llvm.org/docs/DiagnosticsReference.html#wordered-compare-function-pointers	// new warning for clang 13
+			}
+			if (ClangVersion >= new VersionNumber(14))
+			{
+				Arguments.Add("-Wno-bitwise-instead-of-logical");       // https://clang.llvm.org/docs/DiagnosticsReference.html#wbitwise-instead-of-logical			// new warning for clang 14
+			}
+
+			Arguments.Add("-Wno-gnu-string-literal-operator-template"); // https://clang.llvm.org/docs/DiagnosticsReference.html#wgnu-string-literal-operator-template	// We use this feature to allow static FNames.
+			Arguments.Add("-Wno-inconsistent-missing-override");        // https://clang.llvm.org/docs/DiagnosticsReference.html#winconsistent-missing-override			// ?? no reason given
+			Arguments.Add("-Wno-invalid-offsetof");                     // https://clang.llvm.org/docs/DiagnosticsReference.html#winvalid-offsetof						// needed to suppress warnings about using offsetof on non-POD types.
+			Arguments.Add("-Wno-switch");                               // https://clang.llvm.org/docs/DiagnosticsReference.html#wswitch								// this hides the "enumeration value 'XXXXX' not handled in switch [-Wswitch]" warnings - we should maybe remove this at some point and add UE_LOG(, Fatal, ) to default cases
+			Arguments.Add("-Wno-tautological-compare");                 // https://clang.llvm.org/docs/DiagnosticsReference.html#wtautological-compare					// this hides the "warning : comparison of unsigned expression < 0 is always false" type warnings due to constant comparisons, which are possible with template arguments
+			Arguments.Add("-Wno-unknown-pragmas");                      // https://clang.llvm.org/docs/DiagnosticsReference.html#wunknown-pragmas						// Slate triggers this (with its optimize on/off pragmas)
+			Arguments.Add("-Wno-unused-function");                      // https://clang.llvm.org/docs/DiagnosticsReference.html#wunused-function						// this will hide the warnings about static functions in headers that aren't used in every single .cpp file
+			Arguments.Add("-Wno-unused-lambda-capture");                // https://clang.llvm.org/docs/DiagnosticsReference.html#wunused-lambda-capture					// suppressed because capturing of compile-time constants is seemingly inconsistent. And MSVC doesn't do that.
+			Arguments.Add("-Wno-unused-local-typedef");                 // https://clang.llvm.org/docs/DiagnosticsReference.html#wunused-local-typedef					// clang is being overly strict here? PhysX headers trigger this.
+			Arguments.Add("-Wno-unused-private-field");                 // https://clang.llvm.org/docs/DiagnosticsReference.html#wunused-private-field					// this will prevent the issue of warnings for unused private variables. MultichannelTcpSocket.h triggers this, possibly more
+			Arguments.Add("-Wno-unused-variable");                      // https://clang.llvm.org/docs/DiagnosticsReference.html#wunused-variable						// ?? no reason given
+			Arguments.Add("-Wno-undefined-var-template");               // https://clang.llvm.org/docs/DiagnosticsReference.html#wundefined-var-template				// not really a good warning to disable
+
+			// Profile Guided Optimization (PGO) and Link Time Optimization (LTO)
+			if (CompileEnvironment.bPGOOptimize)
+			{
+				//
+				// Clang emits warnings for each compiled object file that doesn't have a matching entry in the profile data.
+				// This can happen when the profile data is older than the binaries we're compiling.
+				//
+				// Disable these warnings. They are far too verbose.
+				//
+				Arguments.Add("-Wno-profile-instr-out-of-date");        // https://clang.llvm.org/docs/DiagnosticsReference.html#wprofile-instr-out-of-date
+				Arguments.Add("-Wno-profile-instr-unprofiled");         // https://clang.llvm.org/docs/DiagnosticsReference.html#wprofile-instr-unprofiled
+
+				// apparently there can be hashing conflicts with PGO which can result in:
+				// 'Function control flow change detected (hash mismatch)' warnings. 
+				Arguments.Add("-Wno-backend-plugin");                   // https://clang.llvm.org/docs/DiagnosticsReference.html#wbackend-plugin
+			}
+
+			// shipping builds will cause this warning with "ensure", so disable only in those case
+			if (CompileEnvironment.Configuration == CppConfiguration.Shipping || Analyzer != StaticAnalyzer.None)
+			{
+				Arguments.Add("-Wno-unused-value");                     // https://clang.llvm.org/docs/DiagnosticsReference.html#wunused-value
+			}
+
+			// https://clang.llvm.org/docs/DiagnosticsReference.html#wdeprecated-declarations
+			if (CompileEnvironment.DeprecationWarningLevel == WarningLevel.Error)
+			{
+				// TODO: This may be unnecessary with -Werror
+				Arguments.Add("-Werror=deprecated-declarations");
+			}
+
+			// Warn if __DATE__ or __TIME__ are used as they prevent reproducible builds
+			if (CompileEnvironment.bDeterministic)
+			{
+				// https://clang.llvm.org/docs/DiagnosticsReference.html#wdate-time
+				if (CompileEnvironment.DeterministicWarningLevel == WarningLevel.Error)
+				{
+					Arguments.Add("-Wdate-time");
+				}
+				else if (CompileEnvironment.DeterministicWarningLevel == WarningLevel.Warning)
+				{
+					Arguments.Add("-Wdate-time -Wno-error=date-time");
+				}
+			}
+
+			// https://clang.llvm.org/docs/DiagnosticsReference.html#wshadow
+			if (CompileEnvironment.ShadowVariableWarningLevel != WarningLevel.Off)
+			{
+				Arguments.Add("-Wshadow" + ((CompileEnvironment.ShadowVariableWarningLevel == WarningLevel.Error) ? "" : " -Wno-error=shadow"));
+			}
+
+			// https://clang.llvm.org/docs/DiagnosticsReference.html#wundef
+			if (CompileEnvironment.bEnableUndefinedIdentifierWarnings)
+			{
+				Arguments.Add("-Wundef" + (CompileEnvironment.bUndefinedIdentifierWarningsAsErrors ? "" : " -Wno-error=undef"));
+			}
+
+			// Note: This should be kept in sync with PRAGMA_DISABLE_UNSAFE_TYPECAST_WARNINGS in ClangPlatformCompilerPreSetup.h
+			string[] UnsafeTypeCastWarningList = {
+				"float-conversion",
+				"implicit-float-conversion",
+				"implicit-int-conversion",
+				"c++11-narrowing"
+				//"shorten-64-to-32",	<-- too many hits right now, probably want it *soon*
+				//"sign-conversion",	<-- too many hits right now, probably want it eventually
+			};
+
+			if (CompileEnvironment.UnsafeTypeCastWarningLevel == WarningLevel.Error)
+			{
+				foreach (string Warning in UnsafeTypeCastWarningList)
+				{
+					Arguments.Add("-W" + Warning);
+				}
+			}
+			else if (CompileEnvironment.UnsafeTypeCastWarningLevel == WarningLevel.Warning)
+			{
+				foreach (string Warning in UnsafeTypeCastWarningList)
+				{
+					Arguments.Add("-W" + Warning + " -Wno-error=" + Warning);
+				}
+			}
+			else
+			{
+				foreach (string Warning in UnsafeTypeCastWarningList)
+				{
+					Arguments.Add("-Wno-" + Warning);
+				}
+			}
+		}
+
+		// Additional disabled warnings for msvc. Everything below should be checked if it is necessary
+		internal static void GetVCDisabledWarnings(List<string> Arguments)
+		{
+			// Allow Microsoft-specific syntax to slide, even though it may be non-standard.  Needed for Windows headers.
+			Arguments.Add("-Wno-microsoft");
+
+			// @todo clang: Hack due to how we have our 'DummyPCH' wrappers setup when using unity builds.  This warning should not be disabled!!
+			Arguments.Add("-Wno-msvc-include");
+
+			// This is disabled because clang explicitly warns about changing pack alignment in a header and not
+			// restoring it afterwards, which is something we do with the Pre/PostWindowsApi.h headers.
+			// @todo clang: push/pop this in  Pre/PostWindowsApi.h headers instead?
+			Arguments.Add("-Wno-pragma-pack");
+
+			Arguments.Add("-Wno-inline-new-delete");    // @todo clang: We declare operator new as inline.  Clang doesn't seem to like that.
+			Arguments.Add("-Wno-implicit-exception-spec-mismatch");
+
+			// Sometimes we compare 'this' pointers against nullptr, which Clang warns about by default
+			Arguments.Add("-Wno-undefined-bool-conversion");
+
+			// @todo clang: Disabled warnings were copied from MacToolChain for the most part
+			Arguments.Add("-Wno-deprecated-writable-strings");
+			Arguments.Add("-Wno-deprecated-register");
+			Arguments.Add("-Wno-switch-enum");
+			Arguments.Add("-Wno-logical-op-parentheses");   // needed for external headers we shan't change
+			Arguments.Add("-Wno-null-arithmetic");          // needed for external headers we shan't change
+			Arguments.Add("-Wno-deprecated-declarations");  // needed for wxWidgets
+			Arguments.Add("-Wno-return-type-c-linkage");    // needed for PhysX
+			Arguments.Add("-Wno-ignored-attributes");       // needed for nvtesslib
+			Arguments.Add("-Wno-uninitialized");
+			Arguments.Add("-Wno-return-type");              // needed for external headers we shan't change 
+
+			// @todo clang: Sorry for adding more of these, but I couldn't read my output log. Most should probably be looked at
+			Arguments.Add("-Wno-unused-parameter");         // Unused function parameter. A lot are named 'bUnused'...
+			Arguments.Add("-Wno-ignored-qualifiers");       // const ignored when returning by value e.g. 'const int foo() { return 4; }'
+			Arguments.Add("-Wno-expansion-to-defined");     // Usage of 'defined(X)' in a macro definition. Gives different results under MSVC
+			Arguments.Add("-Wno-sign-compare");             // Signed/unsigned comparison - millions of these
+			Arguments.Add("-Wno-missing-field-initializers"); // Stupid warning, generated when you initialize with MyStruct A = {0};
+			Arguments.Add("-Wno-nonportable-include-path");
+			Arguments.Add("-Wno-invalid-token-paste");
+			Arguments.Add("-Wno-null-pointer-arithmetic");
+			Arguments.Add("-Wno-constant-logical-operand"); // Triggered by || of two template-derived values inside a static_assert
+			Arguments.Add("-Wno-unused-value");
+			Arguments.Add("-Wno-bitfield-enum-conversion");
+		}
+	}
+}
