PR #10354: Allow AdditionalPlugins use in multiple projects (Contributed by [REDACTED_PERSON])

diff --git a/Engine/Source/Programs/UnrealBuildTool/ProjectFiles/ProjectFileGenerator.cs b/Engine/Source/Programs/UnrealBuildTool/ProjectFiles/ProjectFileGenerator.cs
--- a/Engine/Source/Programs/UnrealBuildTool/ProjectFiles/ProjectFileGenerator.cs
+++ b/Engine/Source/Programs/UnrealBuildTool/ProjectFiles/ProjectFileGenerator.cs
@@ -18,6 +18,21 @@ using System.Threading;
 
 namespace UnrealBuildTool
 {
+	struct AdditionalPluginData
+	{
+		public AdditionalPluginData(List<FileReference> InPluginFiles)
+		{
+			PluginFiles = InPluginFiles;
+			ReferencingProjects = new List<FileReference>();
+			ModuleToSourceFiles = new Dictionary<FileReference, List<FileReference>>();
+			
+		}
+
+		public List<FileReference> PluginFiles;
+		public List<FileReference> ReferencingProjects;
+		public Dictionary<FileReference, List<FileReference>> ModuleToSourceFiles;
+	}
+	
 	/// <summary>
 	/// Represents a folder within the primary project (e.g. Visual Studio solution)
 	/// </summary>
@@ -2264,6 +2279,44 @@ namespace UnrealBuildTool
 		/// <param name="Logger">Logger for output</param>
 		protected void AddProjectsForAllModules(List<FileReference> AllGames, Dictionary<ProjectFile, FileReference> ProjectFileToUProjectFile, Dictionary<FileReference, ProjectFile> ProgramProjects, List<ProjectFile> ModProjects, List<FileReference> AllModuleFiles, Dictionary<FileReference, List<DirectoryReference>> AdditionalSearchPaths, bool bGatherThirdPartySource, ILogger Logger)
 		{
+			Dictionary<DirectoryReference, AdditionalPluginData> AdditionalPluginsByDirectory = new Dictionary<DirectoryReference, AdditionalPluginData>();
+			Dictionary<FileReference, AdditionalPluginData> ModuleToAdditionalPlugin = new Dictionary<FileReference, AdditionalPluginData>();
+
+			if (IncludeCppSource)
+			{
+				foreach ((FileReference ProjectPath, List<DirectoryReference> AdditionalPluginDirectories) in AdditionalSearchPaths)
+				{
+					foreach (DirectoryReference AdditionalPluginDirectory in AdditionalPluginDirectories)
+					{
+						AdditionalPluginData PluginData;
+						if (!AdditionalPluginsByDirectory.TryGetValue(AdditionalPluginDirectory, out PluginData))
+						{
+							PluginData = new AdditionalPluginData(PluginsBase.EnumeratePlugins(AdditionalPluginDirectory).ToList());
+							AdditionalPluginsByDirectory.Add(AdditionalPluginDirectory, PluginData);
+						}
+						
+						PluginData.ReferencingProjects.Add(ProjectPath);
+					}
+				}
+
+				foreach (FileReference ModuleFile in AllModuleFiles)
+				{
+					if (ModuleToAdditionalPlugin.ContainsKey(ModuleFile))
+					{
+						continue;
+					}
+
+					foreach ((DirectoryReference AdditionalPluginDirectory, AdditionalPluginData PluginData) in AdditionalPluginsByDirectory)
+					{
+						if (ModuleFile.IsUnderDirectory(AdditionalPluginDirectory))
+						{
+							ModuleToAdditionalPlugin.Add(ModuleFile, PluginData);
+							PluginData.ModuleToSourceFiles[ModuleFile] = new List<FileReference>();
+						}
+					}
+				}
+			}
+			
 			HashSet<ProjectFile> ProjectsWithPlugins = new HashSet<ProjectFile>();
 			foreach (FileReference CurModuleFile in AllModuleFiles)
 			{
@@ -2297,7 +2350,7 @@ namespace UnrealBuildTool
 				if (WantProjectFileForModule)
 				{
 					DirectoryReference BaseFolder;
-					ProjectFile ProjectFile = FindProjectForModule(CurModuleFile, TargetType.Editor, AllGames, ProgramProjects, ModProjects, AdditionalSearchPaths, out BaseFolder);
+					ProjectFile ProjectFile = FindProjectForModule(CurModuleFile, TargetType.Editor, AllGames, ProgramProjects, ModProjects, ModuleToAdditionalPlugin, out BaseFolder);
 
 					// Update our module map
 					ModuleToEditorProjectFileMap[CurModuleFile] = ProjectFile;
@@ -2317,6 +2370,13 @@ namespace UnrealBuildTool
 					List<FileReference> FoundFiles = SourceFileSearch.FindModuleSourceFiles(CurModuleFile, SearchSubdirectories: SearchSubdirectories);
 					// remove any target files, they are technically not part of the module and are explicitly added when the project is created
 					FoundFiles.RemoveAll(f => f.FullName.EndsWith(".Target.cs"));
+
+					if (ModuleToAdditionalPlugin.ContainsKey(CurModuleFile))
+					{
+						ModuleToAdditionalPlugin[CurModuleFile].ModuleToSourceFiles[CurModuleFile] = FoundFiles;
+						continue;
+					}
+
 					ProjectFile.AddFilesToProject(FoundFiles, BaseFolder);
 
 					// Check if there's a plugin directory here
@@ -2334,22 +2394,30 @@ namespace UnrealBuildTool
 							}
 						}
 
+						ProjectsWithPlugins.Add(ProjectFile);
+					}
+				}
+			}
+			
+			if (IncludeCppSource)
+			{
+				foreach ((DirectoryReference PluginDirectory, AdditionalPluginData PluginData) in AdditionalPluginsByDirectory)
+				{
+					foreach (FileReference ReferencingProject in PluginData.ReferencingProjects)
+					{
+						ProjectFile ProjectFile = FindOrAddProjectHelper(ReferencingProject.GetFileNameWithoutExtension(), ReferencingProject.Directory);
 						FileReference? GameUProject = ProjectFileToUProjectFile.GetValueOrDefault(ProjectFile);
-						if (GameUProject != null && AdditionalSearchPaths.ContainsKey(GameUProject!))
+						if (GameUProject != null)
 						{
-							foreach (DirectoryReference PluginFolder in AdditionalSearchPaths[GameUProject])
+							PluginData.PluginFiles.ForEach(PluginFile => AddPluginFilesToProject(PluginFile, ProjectFile.BaseDir, ProjectFile));
+							foreach ((FileReference _, List<FileReference> ModuleSourceFiles) in PluginData.ModuleToSourceFiles)
 							{
-								foreach (FileReference PluginFileName in PluginsBase.EnumeratePlugins(PluginFolder))
+								if (ModuleSourceFiles != null && ModuleSourceFiles.Count > 0)
 								{
-									if (!ModProjects.Any(x => x.BaseDir == PluginFileName.Directory))
-									{
-										AddPluginFilesToProject(PluginFileName, BaseFolder, ProjectFile);
-									}
+									ProjectFile.AddFilesToProject(ModuleSourceFiles, ProjectFile.BaseDir);
 								}
 							}
 						}
-
-						ProjectsWithPlugins.Add(ProjectFile);
 					}
 				}
 			}
@@ -2393,7 +2461,7 @@ namespace UnrealBuildTool
 			return FindOrAddProject(ProjectFileName, InBaseFolder, IncludeInGeneratedProjects: true, bAlreadyExisted: out _);
 		}
 
-		private ProjectFile FindProjectForModule(FileReference CurModuleFile, TargetType TargetType, List<FileReference> AllGames, Dictionary<FileReference, ProjectFile> ProgramProjects, List<ProjectFile> ModProjects, Dictionary<FileReference, List<DirectoryReference>> AdditionalSearchPaths, out DirectoryReference BaseFolder)
+		private ProjectFile FindProjectForModule(FileReference CurModuleFile, TargetType TargetType, List<FileReference> AllGames, Dictionary<FileReference, ProjectFile> ProgramProjects, List<ProjectFile> ModProjects, Dictionary<FileReference, AdditionalPluginData> ModuleToAdditionalPlugin, out DirectoryReference BaseFolder)
 		{
 			// Starting at the base directory of the module find a project which has the same directory as base, walking up the directory hierarchy until a match is found
 
@@ -2463,16 +2531,11 @@ namespace UnrealBuildTool
 				Path = Path.ParentDirectory!;
 			}
 
-			foreach ((FileReference ProjectPath, List<DirectoryReference> AdditionalPluginDirectories) in AdditionalSearchPaths)
+			if (ModuleToAdditionalPlugin.TryGetValue(CurModuleFile, out AdditionalPluginData PluginData))
 			{
-				foreach (DirectoryReference AdditionalPluginDirectory in AdditionalPluginDirectories)
-				{
-					if (CurModuleFile.IsUnderDirectory(AdditionalPluginDirectory))
-					{
-						BaseFolder = ProjectPath.Directory;
-						return FindOrAddProjectHelper(ProjectPath.GetFileNameWithoutExtension(), BaseFolder);
-					}
-				}
+				FileReference ProjectFileRef = PluginData.ReferencingProjects[0];
+				BaseFolder = ProjectFileRef.Directory;
+				return FindOrAddProjectHelper(ProjectFileRef.GetFileNameWithoutExtension(), BaseFolder);
 			}
 
 			throw new BuildException("Found a module file (" + CurModuleFile + ") that did not exist within any of the known game folders or other source locations");
