Add missing copyright boilerplate

#rnx

diff --git a/Engine/Source/Programs/UnrealBuildTool/ToolChain/ClangWarnings.cs b/Engine/Source/Programs/UnrealBuildTool/ToolChain/ClangWarnings.cs
--- a/Engine/Source/Programs/UnrealBuildTool/ToolChain/ClangWarnings.cs
+++ b/Engine/Source/Programs/UnrealBuildTool/ToolChain/ClangWarnings.cs
@@ -1,4 +1,6 @@
-﻿using System.Collections.Generic;
+﻿// Copyright Epic Games, Inc. All Rights Reserved.
+
+using System.Collections.Generic;
 
 namespace UnrealBuildTool
 {
