Updated UBT to throw warnings or errors(based a build config xml property) if there are PCH performance issues detected in the buildgraph.

diff --git a/Engine/Source/Programs/UnrealBuildTool/Configuration/TargetRules.cs b/Engine/Source/Programs/UnrealBuildTool/Configuration/TargetRules.cs
--- a/Engine/Source/Programs/UnrealBuildTool/Configuration/TargetRules.cs
+++ b/Engine/Source/Programs/UnrealBuildTool/Configuration/TargetRules.cs
@@ -1561,6 +1561,13 @@ namespace UnrealBuildTool
 		[RequiresUniqueBuildEnvironment]
 		public bool bUndefinedIdentifierErrors = true;
 
+		/// <summary>
+		/// Indicates what warning/error level to treat potential PCH performance issues.
+		/// </summary>
+		[XmlConfigFile(Category = "BuildConfiguration")]
+		[CommandLine("-PCHPerformanceIssueWarningLevel=")]
+		public WarningLevel PCHPerformanceIssueWarningLevel = WarningLevel.Off;
+
 		/// <summary>
 		/// How to treat general module include path validation messages
 		/// </summary>
@@ -3464,6 +3471,11 @@ namespace UnrealBuildTool
 			get { return Inner.bUndefinedIdentifierErrors; }
 		}
 
+		public WarningLevel PCHPerformanceIssueWarningLevel
+		{
+			get { return Inner.PCHPerformanceIssueWarningLevel; }
+		}
+
 		public WarningLevel ModuleIncludePathWarningLevel
 		{
 			get { return Inner.ModuleIncludePathWarningLevel; }
diff --git a/Engine/Source/Programs/UnrealBuildTool/Configuration/UEBuildModuleCPP.cs b/Engine/Source/Programs/UnrealBuildTool/Configuration/UEBuildModuleCPP.cs
--- a/Engine/Source/Programs/UnrealBuildTool/Configuration/UEBuildModuleCPP.cs
+++ b/Engine/Source/Programs/UnrealBuildTool/Configuration/UEBuildModuleCPP.cs
@@ -1024,7 +1024,7 @@ namespace UnrealBuildTool
 		/// <param name="ModuleCompileEnvironment">The module compile environment</param>
 		/// <param name="CompileEnvironment">The shared PCH compile environment</param>
 		/// <returns>True if the two compile enviroments are compatible</returns>
-		static private bool IsCompatibleForSharedPCH(CppCompileEnvironment ModuleCompileEnvironment, CppCompileEnvironment CompileEnvironment)
+		static internal bool IsCompatibleForSharedPCH(CppCompileEnvironment ModuleCompileEnvironment, CppCompileEnvironment CompileEnvironment)
 		{
 			if(ModuleCompileEnvironment.bOptimizeCode != CompileEnvironment.bOptimizeCode)
 			{
diff --git a/Engine/Source/Programs/UnrealBuildTool/Configuration/UEBuildTarget.cs b/Engine/Source/Programs/UnrealBuildTool/Configuration/UEBuildTarget.cs
--- a/Engine/Source/Programs/UnrealBuildTool/Configuration/UEBuildTarget.cs
+++ b/Engine/Source/Programs/UnrealBuildTool/Configuration/UEBuildTarget.cs
@@ -22,6 +22,7 @@ using Microsoft.Extensions.Logging;
 using System.Runtime.InteropServices;
 using System.Threading.Tasks;
 using UnrealBuildTool;
+using System.Xml.Linq;
 
 
 namespace UnrealBuildTool
@@ -3497,6 +3498,54 @@ namespace UnrealBuildTool
 					ModuleInfo.Item1.FindOrCreateSharedPCH(ToolChain, GlobalCompileEnvironment.SharedPCHs[SharedPCHIndex], ModuleInfo.Item2, Graph);
 				}
 			}
+
+			// Throw warnings/errors if there are PCH performance issues
+			if (Target.PCHPerformanceIssueWarningLevel != WarningLevel.Off)
+			{
+				bool bFoundIssues = false;
+				LogLevel PCHLoggerLevel = Target.PCHPerformanceIssueWarningLevel == WarningLevel.Warning ? LogLevel.Warning : LogLevel.Error;
+				string Prefix = Target.PCHPerformanceIssueWarningLevel == WarningLevel.Warning ? "Warning:" : "Error:";
+				for (int SharedPCHIndex = NumSharedPCHs - 1; SharedPCHIndex >= 0; SharedPCHIndex--)
+				{
+					// Throw warnings or errors if there are module settings that could result in PCH performance issues
+					PrecompiledHeaderTemplate SharedPCH = GlobalCompileEnvironment.SharedPCHs[SharedPCHIndex];
+					if (SharedPCH.Instances.Any())
+					{
+						// UnsafeTypeCastWarningLevel
+						foreach (PrecompiledHeaderInstance UnsafeTypeCastWarningInstance in SharedPCH.Instances.Where(instance => instance.Modules.First().Rules.UnsafeTypeCastWarningLevel == WarningLevel.Warning))
+						{
+							CppCompileEnvironment UpdatedCppCompileEnvironment = new CppCompileEnvironment(UnsafeTypeCastWarningInstance.CompileEnvironment);
+							UpdatedCppCompileEnvironment.UnsafeTypeCastWarningLevel = WarningLevel.Error;
+							foreach (PrecompiledHeaderInstance UnsafeTypeCastErrorInstance in SharedPCH.Instances.Where(instance => instance.Modules.First().Rules.UnsafeTypeCastWarningLevel == WarningLevel.Error))
+							{
+								if (UEBuildModuleCPP.IsCompatibleForSharedPCH(UnsafeTypeCastErrorInstance.CompileEnvironment, UpdatedCppCompileEnvironment))
+								{
+									foreach (UEBuildModuleCPP Module in UnsafeTypeCastWarningInstance.Modules)
+									{
+										Logger.Log(PCHLoggerLevel, $"{Prefix} Module '{Module.Name}': Please set 'UnsafeTypeCastWarningLevel' to 'WarningLevel.Error' instead of 'WarningLevel.Warning'. This creates PCH permutations.");
+										bFoundIssues = true;
+									}
+								}
+							}
+						}
+
+						// OptimizeCode
+						foreach (PrecompiledHeaderInstance Instance in SharedPCH.Instances.Where(instance => instance.Modules.First().Rules.OptimizeCode == ModuleRules.CodeOptimization.Never))
+						{
+							foreach (var Module in Instance.Modules)
+							{
+								Logger.Log(PCHLoggerLevel, $"{Prefix} Module '{Module.Name}': Do not set 'OptimizeCode' to 'CodeOptimization.Never'. This creates PCH permutations.");
+								bFoundIssues = true;
+							}
+						}
+					}
+				}
+
+				if (bFoundIssues && PCHLoggerLevel == LogLevel.Error)
+				{
+					throw new BuildException("PCH performance issues were found.");
+				}
+			}
 		}
 
 		/// <summary>
