Ignore pch performance issues if the module file is checked out/writable.

diff --git a/Engine/Source/Programs/UnrealBuildTool/Configuration/UEBuildTarget.cs b/Engine/Source/Programs/UnrealBuildTool/Configuration/UEBuildTarget.cs
--- a/Engine/Source/Programs/UnrealBuildTool/Configuration/UEBuildTarget.cs
+++ b/Engine/Source/Programs/UnrealBuildTool/Configuration/UEBuildTarget.cs
@@ -3522,8 +3522,11 @@ namespace UnrealBuildTool
 								{
 									foreach (UEBuildModuleCPP Module in UnsafeTypeCastWarningInstance.Modules)
 									{
-										Logger.Log(PCHLoggerLevel, $"{Prefix} Module '{Module.Name}': Please set 'UnsafeTypeCastWarningLevel' to 'WarningLevel.Error' instead of 'WarningLevel.Warning'. This creates PCH permutations.");
-										bFoundIssues = true;
+										if (Module.RulesFile.ToFileInfo().IsReadOnly) // support for local module changes
+										{
+											Logger.Log(PCHLoggerLevel, $"{Prefix} Module '{Module.Name}': Please set 'UnsafeTypeCastWarningLevel' to 'WarningLevel.Error' instead of 'WarningLevel.Warning'. This creates PCH permutations.");
+											bFoundIssues = true;
+										}
 									}
 								}
 							}
@@ -3534,8 +3537,11 @@ namespace UnrealBuildTool
 						{
 							foreach (var Module in Instance.Modules)
 							{
-								Logger.Log(PCHLoggerLevel, $"{Prefix} Module '{Module.Name}': Do not set 'OptimizeCode' to 'CodeOptimization.Never'. This creates PCH permutations.");
-								bFoundIssues = true;
+								if (Module.RulesFile.ToFileInfo().IsReadOnly) // support for local module changes
+								{
+									Logger.Log(PCHLoggerLevel, $"{Prefix} Module '{Module.Name}': Do not set 'OptimizeCode' to 'CodeOptimization.Never'. This creates PCH permutations.");
+									bFoundIssues = true;
+								}
 							}
 						}
 					}
