UnrealBuildTool: Use a release channel enum instead of a boolean to check if a toolchain is preview or not, so we support additional release channels

diff --git a/Engine/Source/Programs/UnrealBuildTool/Platform/Windows/MicrosoftPlatformSDK.cs b/Engine/Source/Programs/UnrealBuildTool/Platform/Windows/MicrosoftPlatformSDK.cs
--- a/Engine/Source/Programs/UnrealBuildTool/Platform/Windows/MicrosoftPlatformSDK.cs
+++ b/Engine/Source/Programs/UnrealBuildTool/Platform/Windows/MicrosoftPlatformSDK.cs
@@ -514,7 +514,7 @@ namespace UnrealBuildTool
 
 			// Figure out the actual version number that we want
 			ToolChainInstallation? ToolChain = null;
-			if (CompilerVersion == null)
+			if (string.IsNullOrEmpty(CompilerVersion))
 			{
 				ToolChain = SelectToolChain(ToolChains, x => x, Architecture);
 				if (ToolChain == null)
@@ -525,36 +525,25 @@ namespace UnrealBuildTool
 					return false;
 				}
 			}
-			else if (String.Compare(CompilerVersion, "Latest", StringComparison.InvariantCultureIgnoreCase) == 0)
-			{
-				ToolChain = SelectToolChain(ToolChains.Where(x => !x.IsPreview), x => x.ThenByDescending(x => x.Version), Architecture);
-				if (ToolChain == null)
-				{
-					DumpToolChains(ToolChains, x => x.ThenBy(x => x.IsPreview).ThenByDescending(x => x.Version), Architecture, Logger);
-					throw new BuildException("Unable to find valid latest C++ toolchain for {0} {1}", Compiler, Architecture.ToString());
-				}
-			}
-			else if (String.Compare(CompilerVersion, "Preview", StringComparison.InvariantCultureIgnoreCase) == 0)
-			{
-				ToolChain = SelectToolChain(ToolChains.Where(x => x.IsPreview), x => x.ThenByDescending(x => x.Version), Architecture);
-				if (ToolChain == null)
-				{
-					DumpToolChains(ToolChains, x => x.ThenByDescending(x => x.IsPreview).ThenByDescending(x => x.Version), Architecture, Logger);
-					throw new BuildException("Unable to find valid preview toolchain for {0} {1}", Compiler, Architecture.ToString());
-				}
-			}
-			else if (VersionNumber.TryParse(CompilerVersion, out VersionNumber? ToolChainVersion))
+
+			foreach (WindowsCompilerChannel Channel in Enum.GetValues(typeof(WindowsCompilerChannel)))
 			{
-				ToolChain = SelectToolChain(ToolChains, x => x.ThenByDescending(x => x.Version == ToolChainVersion).ThenByDescending(x => x.Family == ToolChainVersion), Architecture);
-				if (ToolChain == null || !(ToolChain.Version == ToolChainVersion || ToolChain.Family == ToolChainVersion))
+				if (String.Compare(CompilerVersion, Channel.ToString(), StringComparison.InvariantCultureIgnoreCase) == 0)
 				{
-					DumpToolChains(ToolChains, x => x.ThenByDescending(x => x.Version == ToolChainVersion).ThenByDescending(x => x.Family == ToolChainVersion), Architecture, Logger);
-					throw new BuildException("Unable to find valid {0} toolchain for {1} {2}", ToolChainVersion, Compiler, Architecture.ToString());
+					ToolChain = SelectToolChain(ToolChains.Where(x => x.ReleaseChannel == Channel), x => x.ThenByDescending(x => x.Version), Architecture);
+					if (ToolChain == null)
+					{
+						DumpToolChains(ToolChains, x => x.ThenBy(x => x.ReleaseChannel).ThenByDescending(x => x.Version), Architecture, Logger);
+						throw new BuildException("Unable to find valid {0} C++ toolchain for {1} {2}", Channel, GetCompilerName(Compiler), Architecture);
+					}
+					break;
 				}
 			}
-			else
+
+			if (ToolChain == null)
 			{
-				throw new BuildException("Unable to find {0} toolchain; '{1}' is an invalid version", GetCompilerName(Compiler), CompilerVersion);
+				DumpToolChains(ToolChains, x => x.ThenBy(x => x.ReleaseChannel).ThenByDescending(x => x.Version), Architecture, Logger);
+				throw new BuildException("Unable to find {0} {1} C++ toolchain; '{2}' is an invalid version", GetCompilerName(Compiler), Architecture, CompilerVersion);
 			}
 
 			// Get the actual directory for this version
@@ -598,7 +587,7 @@ namespace UnrealBuildTool
 			ToolChainInstallation? ToolChain = Preference(ToolChains.Where(x => x.Architecture == Architecture)
 				.OrderByDescending(x => x.Error == null))
 				.ThenByDescending(x => x.Is64Bit)
-				.ThenBy(x => x.IsPreview)
+				.ThenBy(x => x.ReleaseChannel)
 				.ThenBy(x => x.FamilyRank)
 				.ThenByDescending(x => x.IsAutoSdk)
 				.ThenByDescending(x => x.Version)
@@ -629,7 +618,7 @@ namespace UnrealBuildTool
 			var SortedToolChains = Preference(ToolChains.Where(x => x.Architecture == Architecture)
 				.OrderByDescending(x => x.Error == null))
 				.ThenByDescending(x => x.Is64Bit)
-				.ThenBy(x => x.IsPreview)
+				.ThenBy(x => x.ReleaseChannel)
 				.ThenBy(x => x.FamilyRank)
 				.ThenByDescending(x => x.IsAutoSdk)
 				.ThenByDescending(x => x.Version);
@@ -639,7 +628,7 @@ namespace UnrealBuildTool
 				Logger.LogInformation("Available {Architecture} toolchains ({Count}):", Architecture, SortedToolChains.Count());
 				foreach (ToolChainInstallation ToolChain in SortedToolChains)
 				{
-					Logger.LogInformation(" * {ToolChainDir}\n    (Family={Family}, FamilyRank={FamilyRank}, Version={Version}, Is64Bit={Is64Bit}, Preview={Preview}, Architecture={Arch}, Error={Error})", ToolChain.BaseDir, ToolChain.Family, ToolChain.FamilyRank, ToolChain.Version, ToolChain.Is64Bit, ToolChain.IsPreview, ToolChain.Architecture, ToolChain.Error != null);
+					Logger.LogInformation(" * {ToolChainDir}\n    (Family={Family}, FamilyRank={FamilyRank}, Version={Version}, Is64Bit={Is64Bit}, ReleaseChannel={ReleaseChannel}, Architecture={Arch}, Error={Error})", ToolChain.BaseDir, ToolChain.Family, ToolChain.FamilyRank, ToolChain.Version, ToolChain.Is64Bit, ToolChain.ReleaseChannel, ToolChain.Architecture, ToolChain.Error != null);
 				}
 			}
 			else
@@ -701,6 +690,20 @@ namespace UnrealBuildTool
 						// the AutoRTFM Clang compiler is in-tree
 						DirectoryReference AutoRTFMDir = DirectoryReference.Combine(Unreal.EngineDirectory, "Restricted", "NotForLicensees", "Binaries", "Win64", "AutoRTFM");
 						AddClangToolChain(Compiler, AutoRTFMDir, ToolChains, IsAutoSdk: false, Logger);
+
+						// Check for AutoSDK paths
+						DirectoryReference? AutoSdkDir;
+						if (UEBuildPlatformSDK.TryGetHostPlatformAutoSDKDir(out AutoSdkDir))
+						{
+							DirectoryReference ClangBaseDir = DirectoryReference.Combine(AutoSdkDir, "Win64", "AutoRTFM");
+							if (DirectoryReference.Exists(ClangBaseDir))
+							{
+								foreach (DirectoryReference ToolChainDir in DirectoryReference.EnumerateDirectories(ClangBaseDir))
+								{
+									AddClangToolChain(Compiler, ToolChainDir, ToolChains, IsAutoSdk: true, Logger);
+								}
+							}
+						}
 					}
 					else if (Compiler == WindowsCompiler.Intel)
 					{
@@ -729,7 +732,7 @@ namespace UnrealBuildTool
 						{
 							DirectoryReference ToolChainBaseDir = DirectoryReference.Combine(Installation.BaseDir, "VC", "Tools", "MSVC");
 							DirectoryReference RedistBaseDir = DirectoryReference.Combine(Installation.BaseDir, "VC", "Redist", "MSVC");
-							FindVisualStudioToolChains(ToolChainBaseDir, RedistBaseDir, Installation.bPreview, ToolChains, IsAutoSdk: false, Logger);
+							FindVisualStudioToolChains(ToolChainBaseDir, RedistBaseDir, Installation.ReleaseChannel, ToolChains, IsAutoSdk: false, Logger);
 						}
 
 						// Enumerate all the AutoSDK toolchains
@@ -746,10 +749,10 @@ namespace UnrealBuildTool
 							if (!string.IsNullOrEmpty(VSDir))
 							{
 								DirectoryReference ReleaseBaseDir = DirectoryReference.Combine(PlatformDir, "Win64", VSDir);
-								FindVisualStudioToolChains(ReleaseBaseDir, null, false, ToolChains, IsAutoSdk: true, Logger);
+								FindVisualStudioToolChains(ReleaseBaseDir, null, WindowsCompilerChannel.Latest, ToolChains, IsAutoSdk: true, Logger);
 
 								DirectoryReference PreviewBaseDir = DirectoryReference.Combine(PlatformDir, "Win64", $"{VSDir}-Preview");
-								FindVisualStudioToolChains(PreviewBaseDir, null, true, ToolChains, IsAutoSdk: true, Logger);
+								FindVisualStudioToolChains(PreviewBaseDir, null, WindowsCompilerChannel.Preview, ToolChains, IsAutoSdk: true, Logger);
 							}
 						}
 					}
@@ -822,20 +825,20 @@ namespace UnrealBuildTool
 						}
 
 						ISetupInstanceCatalog? Catalog = Instance as ISetupInstanceCatalog;
-						bool bPreview = Catalog != null && Catalog.IsPrerelease();
+						WindowsCompilerChannel ReleaseChannel = Catalog?.IsPrerelease() == true ? WindowsCompilerChannel.Preview : WindowsCompilerChannel.Latest;
 
 						string ProductId = Instance.GetProduct().GetId();
 						bool bCommunity = ProductId.Equals("Microsoft.VisualStudio.Product.Community", StringComparison.Ordinal);
 
 						DirectoryReference BaseDir = new DirectoryReference(Instance.GetInstallationPath());
-						Installations.Add(new VisualStudioInstallation(Compiler, Version, BaseDir, bCommunity, bPreview));
+						Installations.Add(new VisualStudioInstallation(Compiler, Version, BaseDir, bCommunity, ReleaseChannel));
 
 						Logger.LogDebug("Found Visual Studio installation: {BaseDir} (Product={ProductId}, Version={Version})", BaseDir, ProductId, Version);
 					}
 
 					Installations = Installations.OrderByDescending(x => x.Compiler)
 						.ThenBy(x => x.bCommunity)
-						.ThenBy(x => x.bPreview)
+						.ThenBy(x => x.ReleaseChannel)
 						.ThenByDescending(x => x.Version)
 						.ToList();
 				}
@@ -864,11 +867,11 @@ namespace UnrealBuildTool
 		/// </summary>
 		/// <param name="BaseDir">Base directory to search</param>
 		/// <param name="OptionalRedistDir">Optional directory for redistributable components (DLLs etc)</param>
-		/// <param name="bPreview">Whether this is a preview installation</param>
+		/// <param name="ReleaseChannel">The release channel of this installation</param>
 		/// <param name="ToolChains">Map of tool chain version to installation info</param>
 		/// <param name="IsAutoSdk">Whether this folder contains AutoSDK entries</param>
 		/// <param name="Logger">Logger for output</param>
-		static void FindVisualStudioToolChains(DirectoryReference BaseDir, DirectoryReference? OptionalRedistDir, bool bPreview, List<ToolChainInstallation> ToolChains, bool IsAutoSdk, ILogger Logger)
+		static void FindVisualStudioToolChains(DirectoryReference BaseDir, DirectoryReference? OptionalRedistDir, WindowsCompilerChannel ReleaseChannel, List<ToolChainInstallation> ToolChains, bool IsAutoSdk, ILogger Logger)
 		{
 			if (DirectoryReference.Exists(BaseDir))
 			{
@@ -879,7 +882,7 @@ namespace UnrealBuildTool
 					{
 						DirectoryReference? RedistDir = FindVisualStudioRedistForToolChain(ToolChainDir, OptionalRedistDir, Version);
 
-						AddVisualCppToolChain(Version, bPreview, ToolChainDir, RedistDir, ToolChains, IsAutoSdk, Logger);
+						AddVisualCppToolChain(Version, ReleaseChannel, ToolChainDir, RedistDir, ToolChains, IsAutoSdk, Logger);
 					}
 				}
 			}
@@ -926,13 +929,13 @@ namespace UnrealBuildTool
 		/// Adds a Visual C++ toolchain to a list of installations
 		/// </summary>
 		/// <param name="Version"></param>
-		/// <param name="bPreview"></param>
+		/// <param name="ReleaseChannel"></param>
 		/// <param name="ToolChainDir"></param>
 		/// <param name="RedistDir"></param>
 		/// <param name="ToolChains"></param>
 		/// <param name="IsAutoSdk"></param>
 		/// <param name="Logger"></param>
-		static void AddVisualCppToolChain(VersionNumber Version, bool bPreview, DirectoryReference ToolChainDir, DirectoryReference? RedistDir, List<ToolChainInstallation> ToolChains, bool IsAutoSdk, ILogger Logger)
+		static void AddVisualCppToolChain(VersionNumber Version, WindowsCompilerChannel ReleaseChannel, DirectoryReference ToolChainDir, DirectoryReference? RedistDir, List<ToolChainInstallation> ToolChains, bool IsAutoSdk, ILogger Logger)
 		{
 			bool Is64Bit = Has64BitToolChain(ToolChainDir);
 
@@ -956,18 +959,18 @@ namespace UnrealBuildTool
 				Error = $"UnrealBuildTool has banned the MSVC {Banned} toolchains due to compiler issues. Please install a different toolchain such as {PreferredVisualCppVersions.Select(x => x.Min).Max()} by opening the generated solution and installing recommended components or from the Visual Studio installer.";
 			}
 
-			Logger.LogDebug("Found Visual Studio toolchain: {ToolChainDir} (Family={Family}, FamilyRank={FamilyRank}, Version={Version}, Is64Bit={Is64Bit}, Preview={Preview}, Architecture={Arch}, Error={Error}, Redist={RedistDir})", ToolChainDir, Family, FamilyRank, Version, Is64Bit, bPreview, UnrealArch.X64.ToString(), Error != null, RedistDir);
-			ToolChains.Add(new ToolChainInstallation(Family, FamilyRank, Version, Is64Bit, bPreview, UnrealArch.X64, Error, ToolChainDir, RedistDir, IsAutoSdk));
+			Logger.LogDebug("Found Visual Studio toolchain: {ToolChainDir} (Family={Family}, FamilyRank={FamilyRank}, Version={Version}, Is64Bit={Is64Bit}, ReleaseChannel={ReleaseChannel}, Architecture={Arch}, Error={Error}, Redist={RedistDir})", ToolChainDir, Family, FamilyRank, Version, Is64Bit, ReleaseChannel, UnrealArch.X64.ToString(), Error != null, RedistDir);
+			ToolChains.Add(new ToolChainInstallation(Family, FamilyRank, Version, Is64Bit, ReleaseChannel, UnrealArch.X64, Error, ToolChainDir, RedistDir, IsAutoSdk));
 
 			if (HasArm64ToolChain(ToolChainDir))
 			{
-				Logger.LogDebug("Found Visual Studio toolchain: {ToolChainDir} (Family={Family}, FamilyRank={FamilyRank}, Version={Version}, Is64Bit={Is64Bit}, Preview={Preview}, Architecture={Arch}, Error={Error}, Redist={RedistDir})", ToolChainDir, Family, FamilyRank, Version, Is64Bit, bPreview, UnrealArch.Arm64.ToString(), Error != null, RedistDir);
-				ToolChains.Add(new ToolChainInstallation(Family, FamilyRank, Version, Is64Bit, bPreview, UnrealArch.Arm64, Error, ToolChainDir, RedistDir, IsAutoSdk));
+				Logger.LogDebug("Found Visual Studio toolchain: {ToolChainDir} (Family={Family}, FamilyRank={FamilyRank}, Version={Version}, Is64Bit={Is64Bit}, ReleaseChannel={ReleaseChannel}, Architecture={Arch}, Error={Error}, Redist={RedistDir})", ToolChainDir, Family, FamilyRank, Version, Is64Bit, ReleaseChannel, UnrealArch.Arm64.ToString(), Error != null, RedistDir);
+				ToolChains.Add(new ToolChainInstallation(Family, FamilyRank, Version, Is64Bit, ReleaseChannel, UnrealArch.Arm64, Error, ToolChainDir, RedistDir, IsAutoSdk));
 				
 				if (HasArm64ECToolChain(ToolChainDir))
 				{
-					Logger.LogDebug("Found Visual Studio toolchain: {ToolChainDir} (Family={Family}, FamilyRank={FamilyRank}, Version={Version}, Is64Bit={Is64Bit}, Preview={Preview}, Architecture={Arch}, Error={Error}, Redist={RedistDir})", ToolChainDir, Family, FamilyRank, Version, Is64Bit, bPreview, UnrealArch.Arm64ec.ToString(), Error != null, RedistDir);
-					ToolChains.Add(new ToolChainInstallation(Family, FamilyRank, Version, Is64Bit, bPreview, UnrealArch.Arm64ec, Error, ToolChainDir, RedistDir, IsAutoSdk));
+					Logger.LogDebug("Found Visual Studio toolchain: {ToolChainDir} (Family={Family}, FamilyRank={FamilyRank}, Version={Version}, Is64Bit={Is64Bit}, ReleaseChannel={ReleaseChannel}, Architecture={Arch}, Error={Error}, Redist={RedistDir})", ToolChainDir, Family, FamilyRank, Version, Is64Bit, ReleaseChannel, UnrealArch.Arm64ec.ToString(), Error != null, RedistDir);
+					ToolChains.Add(new ToolChainInstallation(Family, FamilyRank, Version, Is64Bit, ReleaseChannel, UnrealArch.Arm64ec, Error, ToolChainDir, RedistDir, IsAutoSdk));
 				}
 			}
 		}
@@ -1000,7 +1003,7 @@ namespace UnrealBuildTool
 				}
 
 				Logger.LogDebug("Found Clang toolchain: {ToolChainDir} (Version={Version}, Is64Bit={Is64Bit}, Rank={Rank}, Error={Error})", ToolChainDir, Version, Is64Bit, Rank, Error != null);
-				ToolChains.Add(new ToolChainInstallation(Family, Rank, Version, Is64Bit, false, UnrealArch.X64, Error, ToolChainDir, null, IsAutoSdk));
+				ToolChains.Add(new ToolChainInstallation(Family, Rank, Version, Is64Bit, WindowsCompilerChannel.Latest, UnrealArch.X64, Error, ToolChainDir, null, IsAutoSdk));
 			}
 		}
 
@@ -1037,7 +1040,7 @@ namespace UnrealBuildTool
 				}
 
 				Logger.LogDebug("Found Intel OneAPI toolchain: {ToolChainDir} (Version={Version}, Is64Bit={Is64Bit}, Rank={Rank}, Error={Error})", ToolChainDir, Version, Is64Bit, Rank, Error != null);
-				ToolChains.Add(new ToolChainInstallation(Family, Rank, Version, Is64Bit, false, UnrealArch.X64, Error, ToolChainDir, null, IsAutoSdk));
+				ToolChains.Add(new ToolChainInstallation(Family, Rank, Version, Is64Bit, WindowsCompilerChannel.Latest, UnrealArch.X64, Error, ToolChainDir, null, IsAutoSdk));
 			}
 		}
 
@@ -1329,7 +1332,6 @@ namespace UnrealBuildTool
 		#endregion // Windows Specific SDK 
 	}
 
-
 	/// <summary>
 	/// Information about a particular toolchain installation
 	/// </summary>
@@ -1357,9 +1359,9 @@ namespace UnrealBuildTool
 		public bool Is64Bit { get; }
 
 		/// <summary>
-		/// Whether it's a pre-release version of the toolchain.
+		/// The release channel of this toolchain
 		/// </summary>
-		public bool IsPreview { get; }
+		public WindowsCompilerChannel ReleaseChannel { get; }
 
 		/// <summary>
 		/// The architecture of this ToolChainInstallation (multiple ToolChainInstallation instances may be created, one per architecture).
@@ -1393,19 +1395,19 @@ namespace UnrealBuildTool
 		/// <param name="FamilyRank"></param>
 		/// <param name="Version"></param>
 		/// <param name="Is64Bit"></param>
-		/// <param name="IsPreview">Whether it's a pre-release version of the toolchain</param>
+		/// <param name="ReleaseChannel">The release channel of this toolchain</param>
 		/// <param name="Architecture"></param>
 		/// <param name="Error"></param>
 		/// <param name="BaseDir">Base directory for the toolchain</param>
 		/// <param name="RedistDir">Optional directory for redistributable components (DLLs etc)</param>
 		/// <param name="IsAutoSdk">Whether this toolchain comes from AutoSDK</param>
-		public ToolChainInstallation(VersionNumber Family, int FamilyRank, VersionNumber Version, bool Is64Bit, bool IsPreview, UnrealArch Architecture, string? Error, DirectoryReference BaseDir, DirectoryReference? RedistDir, bool IsAutoSdk)
+		public ToolChainInstallation(VersionNumber Family, int FamilyRank, VersionNumber Version, bool Is64Bit, WindowsCompilerChannel ReleaseChannel, UnrealArch Architecture, string? Error, DirectoryReference BaseDir, DirectoryReference? RedistDir, bool IsAutoSdk)
 		{
 			this.Family = Family;
 			this.FamilyRank = FamilyRank;
 			this.Version = Version;
 			this.Is64Bit = Is64Bit;
-			this.IsPreview = IsPreview;
+			this.ReleaseChannel = ReleaseChannel;
 			this.Architecture = Architecture;
 			this.Error = Error;
 			this.BaseDir = BaseDir;
diff --git a/Engine/Source/Programs/UnrealBuildTool/Platform/Windows/UEBuildWindows.cs b/Engine/Source/Programs/UnrealBuildTool/Platform/Windows/UEBuildWindows.cs
--- a/Engine/Source/Programs/UnrealBuildTool/Platform/Windows/UEBuildWindows.cs
+++ b/Engine/Source/Programs/UnrealBuildTool/Platform/Windows/UEBuildWindows.cs
@@ -127,6 +127,15 @@ namespace UnrealBuildTool
 		VisualStudio2022,
 	}
 
+	/// <summary>
+	/// Enum describing the release channel of a compiler
+	/// </summary>
+	internal enum WindowsCompilerChannel
+	{
+		Latest,
+		Preview,
+	}
+
 
 	/// <summary>
 	/// Extension methods for WindowsCompilier enum
@@ -984,20 +993,20 @@ namespace UnrealBuildTool
 		public bool bCommunity { get; }
 
 		/// <summary>
-		/// Whether it's a pre-release version of the IDE.
+		/// The release channel of this installation
 		/// </summary>
-		public bool bPreview { get; }
+		public WindowsCompilerChannel ReleaseChannel { get; }
 
 		/// <summary>
 		/// Constructor
 		/// </summary>
-		public VisualStudioInstallation(WindowsCompiler Compiler, VersionNumber Version, DirectoryReference BaseDir, bool bCommunity, bool bPreview)
+		public VisualStudioInstallation(WindowsCompiler Compiler, VersionNumber Version, DirectoryReference BaseDir, bool bCommunity, WindowsCompilerChannel ReleaseChannel)
 		{
 			this.Compiler = Compiler;
 			this.Version = Version;
 			this.BaseDir = BaseDir;
 			this.bCommunity = bCommunity;
-			this.bPreview = bPreview;
+			this.ReleaseChannel = ReleaseChannel;
 		}
 	}
 
