UnrealBuildTool: Support scanning for experimental toolchains if available

#rnx

diff --git a/Engine/Source/Programs/UnrealBuildTool/Platform/Windows/MicrosoftPlatformSDK.cs b/Engine/Source/Programs/UnrealBuildTool/Platform/Windows/MicrosoftPlatformSDK.cs
--- a/Engine/Source/Programs/UnrealBuildTool/Platform/Windows/MicrosoftPlatformSDK.cs
+++ b/Engine/Source/Programs/UnrealBuildTool/Platform/Windows/MicrosoftPlatformSDK.cs
@@ -753,6 +753,9 @@ namespace UnrealBuildTool
 
 								DirectoryReference PreviewBaseDir = DirectoryReference.Combine(PlatformDir, "Win64", $"{VSDir}-Preview");
 								FindVisualStudioToolChains(PreviewBaseDir, null, WindowsCompilerChannel.Preview, ToolChains, IsAutoSdk: true, Logger);
+
+								DirectoryReference ExperimentalBaseDir = DirectoryReference.Combine(PlatformDir, "Win64", $"{VSDir}-Experimental");
+								FindVisualStudioToolChains(PreviewBaseDir, null, WindowsCompilerChannel.Experimental, ToolChains, IsAutoSdk: true, Logger);
 							}
 						}
 					}
diff --git a/Engine/Source/Programs/UnrealBuildTool/Platform/Windows/UEBuildWindows.cs b/Engine/Source/Programs/UnrealBuildTool/Platform/Windows/UEBuildWindows.cs
--- a/Engine/Source/Programs/UnrealBuildTool/Platform/Windows/UEBuildWindows.cs
+++ b/Engine/Source/Programs/UnrealBuildTool/Platform/Windows/UEBuildWindows.cs
@@ -134,6 +134,7 @@ namespace UnrealBuildTool
 	{
 		Latest,
 		Preview,
+		Experimental,
 	}
 
 
