UBT: Deprecated TragetInfo::bUseMallocProfiler toggle. No longer used as MallocProfiler2 and FMallocProfiler runtime is removed in favor of UnrealInsights.

diff --git a/Engine/Source/Programs/UnrealBuildTool/Configuration/TargetRules.cs b/Engine/Source/Programs/UnrealBuildTool/Configuration/TargetRules.cs
--- a/Engine/Source/Programs/UnrealBuildTool/Configuration/TargetRules.cs
+++ b/Engine/Source/Programs/UnrealBuildTool/Configuration/TargetRules.cs
@@ -1853,6 +1853,7 @@ namespace UnrealBuildTool
 		/// <summary>
 		/// If true, then enable memory profiling in the build (defines USE_MALLOC_PROFILER=1 and forces bOmitFramePointers=false).
 		/// </summary>
+		[Obsolete("Deprecated in UE5.3 - No longer used as MallocProfiler is removed in favor of UnrealInsights.")]
 		[XmlConfigFile(Category = "BuildConfiguration")]
 		public bool bUseMallocProfiler = false;
 
@@ -3650,6 +3651,7 @@ namespace UnrealBuildTool
 			get { return Inner.bEnableProcessPriorityControl; }
 		}
 
+		[Obsolete("Deprecated in UE5.3 - No longer used as MallocProfiler is removed in favor of UnrealInsights.")]
 		public bool bUseMallocProfiler
 		{
 			get { return Inner.bUseMallocProfiler; }
diff --git a/Engine/Source/Programs/UnrealBuildTool/Configuration/TestTargetRules.cs b/Engine/Source/Programs/UnrealBuildTool/Configuration/TestTargetRules.cs
--- a/Engine/Source/Programs/UnrealBuildTool/Configuration/TestTargetRules.cs
+++ b/Engine/Source/Programs/UnrealBuildTool/Configuration/TestTargetRules.cs
@@ -282,9 +282,6 @@ namespace UnrealBuildTool
 			bBuildRequiresCookedData = true;
 			bBuildDeveloperTools = false;
 
-			// Disable malloc profiling in tests
-			bUseMallocProfiler = false;
-
 			// Useful for debugging test failures
 			if (Target.Configuration == UnrealTargetConfiguration.Debug)
 			{
diff --git a/Engine/Source/Programs/UnrealBuildTool/Platform/IOS/IOSToolChain.cs b/Engine/Source/Programs/UnrealBuildTool/Platform/IOS/IOSToolChain.cs
--- a/Engine/Source/Programs/UnrealBuildTool/Platform/IOS/IOSToolChain.cs
+++ b/Engine/Source/Programs/UnrealBuildTool/Platform/IOS/IOSToolChain.cs
@@ -172,7 +172,7 @@ namespace UnrealBuildTool
 					BuildProducts.Add(AssetFile, BuildProductType.RequiredResource);
 				}
 			}
-			if (Target.IOSPlatform.bGeneratedSYM && (ProjectSettings.bGenerateCrashReportSymbols || Target.bUseMallocProfiler) && Binary.Type == UEBuildBinaryType.StaticLibrary)
+			if (Target.IOSPlatform.bGeneratedSYM && ProjectSettings.bGenerateCrashReportSymbols && Binary.Type == UEBuildBinaryType.StaticLibrary)
 			{
 				FileReference DebugFile = FileReference.Combine(Binary.OutputFilePath.Directory, Binary.OutputFilePath.GetFileNameWithoutExtension() + ".udebugsymbols");
 				BuildProducts.Add(DebugFile, BuildProductType.SymbolFile);
@@ -1096,7 +1096,7 @@ namespace UnrealBuildTool
 				{
 					OutputFiles.Add(item);
 				}
-				if (ProjectSettings.bGenerateCrashReportSymbols || Target.bUseMallocProfiler)
+				if (ProjectSettings.bGenerateCrashReportSymbols)
 				{
 					OutputFiles.Add(GeneratePseudoPDB(Executable, Graph));
 				}
diff --git a/Engine/Source/Programs/UnrealBuildTool/System/RulesAssembly.cs b/Engine/Source/Programs/UnrealBuildTool/System/RulesAssembly.cs
--- a/Engine/Source/Programs/UnrealBuildTool/System/RulesAssembly.cs
+++ b/Engine/Source/Programs/UnrealBuildTool/System/RulesAssembly.cs
@@ -652,13 +652,6 @@ namespace UnrealBuildTool
 				Rules.bOmitPCDebugInfoInDevelopment = false;
 			}
 
-			// Setup the malloc profiler
-			if (Rules.bUseMallocProfiler)
-			{
-				Rules.bOmitFramePointers = false;
-				Rules.GlobalDefinitions.Add("USE_MALLOC_PROFILER=1");
-			}
-
 			// Setup utrace for Shader Compiler Worker
 			if (Rules.bShaderCompilerWorkerTrace)
             {
