fix for localized AppXManifest resource files not being coped

 - project-specific localized image would only be copied if there were project-specific default images configured

diff --git a/Engine/Source/Programs/UnrealBuildTool/Platform/Windows/AppXManifestGeneratorBase.cs b/Engine/Source/Programs/UnrealBuildTool/Platform/Windows/AppXManifestGeneratorBase.cs
--- a/Engine/Source/Programs/UnrealBuildTool/Platform/Windows/AppXManifestGeneratorBase.cs
+++ b/Engine/Source/Programs/UnrealBuildTool/Platform/Windows/AppXManifestGeneratorBase.cs
@@ -326,8 +326,13 @@ namespace UnrealBuildTool
 
 		private bool RemoveStaleResourceFiles()
 		{
-			// remove all resource files that should not be included
 			string TargetResourceDir = Path.Combine(OutputPath!, BuildResourceSubPath);
+			if (!Directory.Exists(TargetResourceDir))
+			{
+				return false;
+			}
+
+			// remove all resource files that should not be included
 			var TargetResourceInstances = Directory.EnumerateFiles(TargetResourceDir, "*.*", SearchOption.AllDirectories);
 
 			var StaleResourceFiles = TargetResourceInstances.Where(X => !ManifestFiles!.ContainsKey(X)).ToList();
@@ -360,30 +365,30 @@ namespace UnrealBuildTool
 		/// <summary>
 		/// Attempts to locate the given resource binary file in several known folder locations
 		/// </summary>
-		protected virtual bool FindResourceBinaryFile( out string SourcePath, string ResourceFileName, bool AllowEngineFallback = true)
+		protected virtual bool FindResourceBinaryFile( out string SourceFilePath, string ResourceFileName, string CultureId = "", bool AllowEngineFallback = true)
 		{
 			// look in project normal Build location
-			SourcePath = Path.Combine(ProjectPath!, "Build", Platform.ToString(), BuildResourceSubPath);
-			bool bFileExists = File.Exists(Path.Combine(SourcePath, ResourceFileName));
+			SourceFilePath = Path.Combine(ProjectPath!, "Build", Platform.ToString(), BuildResourceSubPath, CultureId, ResourceFileName);
+			bool bFileExists = File.Exists(SourceFilePath);
 
 			// look in Platform Extensions next
 			if (!bFileExists)
 			{
-				SourcePath = Path.Combine(ProjectPath!, "Platforms", Platform.ToString(), "Build", BuildResourceSubPath);
-				bFileExists = File.Exists(Path.Combine(SourcePath, ResourceFileName));
+				SourceFilePath = Path.Combine(ProjectPath!, "Platforms", Platform.ToString(), "Build", BuildResourceSubPath, CultureId, ResourceFileName);
+				bFileExists = File.Exists(SourceFilePath);
 			}
 
 			// look in Engine, if allowed
 			if (!bFileExists && AllowEngineFallback)
 			{
-				SourcePath = Path.Combine(Unreal.EngineDirectory.FullName, "Build", Platform.ToString(), EngineResourceSubPath);
-				bFileExists = File.Exists(Path.Combine(SourcePath, ResourceFileName));
+				SourceFilePath = Path.Combine(Unreal.EngineDirectory.FullName, "Build", Platform.ToString(), EngineResourceSubPath, CultureId, ResourceFileName);
+				bFileExists = File.Exists(SourceFilePath);
 
 				// look in Platform extensions too
 				if (!bFileExists)
 				{
-					SourcePath = Path.Combine(Unreal.EngineDirectory.FullName, "Platforms", Platform.ToString(), "Build", EngineResourceSubPath);
-					bFileExists = File.Exists(Path.Combine(SourcePath, ResourceFileName));
+					SourceFilePath = Path.Combine(Unreal.EngineDirectory.FullName, "Platforms", Platform.ToString(), "Build", EngineResourceSubPath, CultureId, ResourceFileName);
+					bFileExists = File.Exists(SourceFilePath);
 				}
 			}
 
@@ -393,10 +398,10 @@ namespace UnrealBuildTool
 		/// <summary>
 		/// Determines whether the given resource binary file can be found in one of the known folder locations
 		/// </summary>
-		protected bool DoesResourceBinaryFileExist(string ResourceFileName, bool AllowEngineFallback = true)
+		protected bool DoesResourceBinaryFileExist(string ResourceFileName, string CultureId = "", bool AllowEngineFallback = true)
 		{
 			string SourcePath;
-			return FindResourceBinaryFile( out SourcePath, ResourceFileName, AllowEngineFallback );
+			return FindResourceBinaryFile( out SourcePath, ResourceFileName, CultureId, AllowEngineFallback );
 		}
 
 
@@ -406,8 +411,8 @@ namespace UnrealBuildTool
 		protected bool AddResourceBinaryFileReference(string ResourceFileName, bool AllowEngineFallback = true)
 		{
 			string TargetPath = Path.Combine(OutputPath!, BuildResourceSubPath);
-			string SourcePath;
-			bool bFileExists = FindResourceBinaryFile(out SourcePath, ResourceFileName, AllowEngineFallback);
+			string SourceFilePath;
+			bool bFileExists = FindResourceBinaryFile(out SourceFilePath, ResourceFileName, CultureId:"", AllowEngineFallback);
 
 			// At least the default culture entry for any resource binary must always exist
 			if (!bFileExists)
@@ -415,34 +420,17 @@ namespace UnrealBuildTool
 				return false;
 			}
 
-			// If the target resource folder doesn't exist yet, create it
-			if (!CreateCheckDirectory(TargetPath, Logger))
-			{
-				return false;
-			}
+			// Copy the default resource file
+			AddFileReference(SourceFilePath, Path.Combine(TargetPath, ResourceFileName), bIsGeneratedFile:false);
 
-			// Find all copies of the resource file in the source directory (could be up to one for each culture and the default).
-			List<string> SourceResourceInstances = new List<string>();
-			SourceResourceInstances.Add( Path.Combine( SourcePath, ResourceFileName) );
+			// Copy all per-culture resource files
 			foreach( string CultureId in CulturesToStage!)
 			{
-				string CultureResourceFile = Path.Combine(SourcePath, CultureId, ResourceFileName);
-				if (File.Exists(CultureResourceFile))
-				{
-					SourceResourceInstances.Add( CultureResourceFile );
-				}
-			}
-
-			// Copy new resource files
-			foreach (string SourceResourceFile in SourceResourceInstances)
-			{
-				string TargetResourcePath = Path.Combine(TargetPath, SourceResourceFile.Substring(SourcePath.Length + 1));
-				if (!CreateCheckDirectory(Path.GetDirectoryName(TargetResourcePath)!, Logger))
+				bFileExists = FindResourceBinaryFile(out SourceFilePath, ResourceFileName, CultureId, AllowEngineFallback);
+				if (bFileExists)
 				{
-					Logger.LogError("Unable to create intermediate directory {IntDir}.", Path.GetDirectoryName(TargetResourcePath));
-					continue;
+					AddFileReference(SourceFilePath, Path.Combine(TargetPath, CultureId, ResourceFileName), bIsGeneratedFile:false);
 				}
-				AddFileReference(SourceResourceFile, TargetResourcePath, bIsGeneratedFile:false);
 			}
 
 			return true;
