Simplify raytracing opt-in preprocessor declaration for Windows platform extensions

#rnx

diff --git a/Engine/Source/Programs/UnrealBuildTool/Platform/Windows/UEBuildWindows.cs b/Engine/Source/Programs/UnrealBuildTool/Platform/Windows/UEBuildWindows.cs
--- a/Engine/Source/Programs/UnrealBuildTool/Platform/Windows/UEBuildWindows.cs
+++ b/Engine/Source/Programs/UnrealBuildTool/Platform/Windows/UEBuildWindows.cs
@@ -1590,8 +1590,7 @@ namespace UnrealBuildTool
 				CompileEnvironment.Definitions.Add(String.Format("OVERRIDE_PLATFORM_HEADER_NAME={0}", OverridePlatformHeaderName));
 			}
 
-			// Ray tracing only supported on 64-bit Windows.
-			if (Target.Platform == UnrealTargetPlatform.Win64 && Target.WindowsPlatform.bEnableRayTracing)
+			if (Target.IsInPlatformGroup(UnrealPlatformGroup.Windows) && Target.WindowsPlatform.bEnableRayTracing)
 			{
 				CompileEnvironment.Definitions.Add("RHI_RAYTRACING=1");
 			}
