UnrealBuildTool: Put back -CompilerVersion check for a specific version that I accidently removed

#rnx

diff --git a/Engine/Source/Programs/UnrealBuildTool/Platform/Windows/MicrosoftPlatformSDK.cs b/Engine/Source/Programs/UnrealBuildTool/Platform/Windows/MicrosoftPlatformSDK.cs
--- a/Engine/Source/Programs/UnrealBuildTool/Platform/Windows/MicrosoftPlatformSDK.cs
+++ b/Engine/Source/Programs/UnrealBuildTool/Platform/Windows/MicrosoftPlatformSDK.cs
@@ -525,6 +525,15 @@ namespace UnrealBuildTool
 					return false;
 				}
 			}
+			else if (VersionNumber.TryParse(CompilerVersion, out VersionNumber? ToolChainVersion))
+			{
+				ToolChain = SelectToolChain(ToolChains, x => x.ThenByDescending(x => x.Version == ToolChainVersion).ThenByDescending(x => x.Family == ToolChainVersion), Architecture);
+				if (ToolChain == null || !(ToolChain.Version == ToolChainVersion || ToolChain.Family == ToolChainVersion))
+				{
+					DumpToolChains(ToolChains, x => x.ThenByDescending(x => x.Version == ToolChainVersion).ThenByDescending(x => x.Family == ToolChainVersion), Architecture, Logger);
+					throw new BuildException("Unable to find valid {0} toolchain for {1} {2}", ToolChainVersion, Compiler, Architecture.ToString());
+				}
+			}
 
 			foreach (WindowsCompilerChannel Channel in Enum.GetValues(typeof(WindowsCompilerChannel)))
 			{
