- Fixed a crash with parallel file writing in modern xcode

diff --git a/Engine/Source/Programs/UnrealBuildTool/ProjectFiles/Xcode/XcodeProject.Utils.cs b/Engine/Source/Programs/UnrealBuildTool/ProjectFiles/Xcode/XcodeProject.Utils.cs
--- a/Engine/Source/Programs/UnrealBuildTool/ProjectFiles/Xcode/XcodeProject.Utils.cs
+++ b/Engine/Source/Programs/UnrealBuildTool/ProjectFiles/Xcode/XcodeProject.Utils.cs
@@ -126,13 +126,18 @@ namespace UnrealBuildTool.XcodeProjectXcconfig
 				return AllArchitectures;
 			}
 
-			// First time seeing this target?
-			if (!CachedMacProjectArchitectures.ContainsKey(TargetName))
+			UnrealArchitectures Arches;
+			lock (CachedMacProjectArchitectures)
 			{
-				CachedMacProjectArchitectures[TargetName] = UnrealArchitectureConfig.ForPlatform(UnrealTargetPlatform.Mac).ProjectSupportedArchitectures(InProjectFile, TargetName);
+				// First time seeing this target?
+				if (!CachedMacProjectArchitectures.ContainsKey(TargetName))
+				{
+					CachedMacProjectArchitectures[TargetName] = UnrealArchitectureConfig.ForPlatform(UnrealTargetPlatform.Mac).ProjectSupportedArchitectures(InProjectFile, TargetName);
+				}
+				Arches = CachedMacProjectArchitectures[TargetName];
 			}
 
-			return CachedMacProjectArchitectures[TargetName];
+			return Arches;
 		}
 
 
