- Allow modern xcode indexing of files in modules that are fully under another platform (like /IOS/IOSTargetPlatform), but disallow platform files in a shared module (Core/Private/IOS/IOSPlatformFile.cpp)

diff --git a/Engine/Source/Programs/UnrealBuildTool/ProjectFiles/Xcode/XcodeProject.File.cs b/Engine/Source/Programs/UnrealBuildTool/ProjectFiles/Xcode/XcodeProject.File.cs
--- a/Engine/Source/Programs/UnrealBuildTool/ProjectFiles/Xcode/XcodeProject.File.cs
+++ b/Engine/Source/Programs/UnrealBuildTool/ProjectFiles/Xcode/XcodeProject.File.cs
@@ -221,7 +221,7 @@ namespace UnrealBuildTool.XcodeProjectXcconfig
 			RootDirectory = UProjectLocation == null || UProjectLocation.ContainsName("Programs", 0) ? (Unreal.EngineDirectory) : UProjectLocation!.Directory;
 		}
 
-		public void ProcessFile(XcodeSourceFile File, bool bIsForBuild, bool bIsFolder, string? GroupName=null, string ExtraResourceSettings="")
+		public void ProcessFile(XcodeSourceFile File, bool bIsForBuild, bool bIsFolder, string? GroupName=null, string ExtraResourceSettings="", Dictionary<DirectoryReference, int>? SourceToBuildFileMap=null)
 		{
 			// remember all buildable files
 			if (bIsForBuild)
@@ -230,11 +230,40 @@ namespace UnrealBuildTool.XcodeProjectXcconfig
 
 				if (IsSourceCode(FileExtension))
 				{
-					// look if it contains any of the exluded names, and is so, don't include it
-					if (!ExcludedFolders.Any(x => File.Reference.ContainsName(x, 0)))
+					DirectoryReference? SourceDir = File.Reference.Directory;
+					int BuildFileOffset = 0;
+					if (SourceToBuildFileMap != null)
 					{
-						// will fill in the response file later
-						BuildableFilesToResponseFile[File] = null;
+						if (!SourceToBuildFileMap.TryGetValue(SourceDir, out BuildFileOffset))
+						{
+							while (SourceDir != null && DirectoryReference.EnumerateFiles(SourceDir, "*.Build.cs").Count() == 0)
+							{
+								SourceDir = SourceDir.ParentDirectory!;
+							}
+
+							// if we didn't find a 
+							if (SourceDir != null)
+							{
+								BuildFileOffset = SourceDir.FullName.Length;
+							}
+							else
+							{
+								// mark that we don't want to index this file, if we didn;t find a Build.cs file
+								BuildFileOffset = -1;
+							}
+							SourceToBuildFileMap[File.Reference.Directory] = BuildFileOffset;
+						}
+					}
+
+					if (BuildFileOffset != -1)
+					{
+						// first find the build.cs file 
+						// look if it contains any of the exluded names, and is so, don't include it
+						if (!ExcludedFolders.Any(x => File.Reference.ContainsName(x, BuildFileOffset)))
+						{
+							// will fill in the response file later
+							BuildableFilesToResponseFile[File] = null;
+						}
 					}
 				}
 				else if (IsResourceFile(FileExtension))
diff --git a/Engine/Source/Programs/UnrealBuildTool/ProjectFiles/Xcode/XcodeProject.cs b/Engine/Source/Programs/UnrealBuildTool/ProjectFiles/Xcode/XcodeProject.cs
--- a/Engine/Source/Programs/UnrealBuildTool/ProjectFiles/Xcode/XcodeProject.cs
+++ b/Engine/Source/Programs/UnrealBuildTool/ProjectFiles/Xcode/XcodeProject.cs
@@ -1986,9 +1986,10 @@ namespace UnrealBuildTool.XcodeProjectXcconfig
 			// process the files that came from UE/cross-platform land
 			SourceFiles.SortBy(x => x.Reference.FullName);
 
+			Dictionary<DirectoryReference, int> BuildFileMap = new();
 			foreach (XcodeSourceFile SourceFile in SourceFiles.OfType<XcodeSourceFile>())
 			{
-				SharedFileCollection.ProcessFile(SourceFile, bIsForBuild:IsGeneratedProject, false);
+				SharedFileCollection.ProcessFile(SourceFile, bIsForBuild:IsGeneratedProject, bIsFolder:false, SourceToBuildFileMap:BuildFileMap);
 			}
 
 			// cache the main group
