UnrealBuildTool: Minor comment fix

#rnx

diff --git a/Engine/Source/Programs/UnrealBuildTool/Platform/Windows/MicrosoftPlatformSDK.cs b/Engine/Source/Programs/UnrealBuildTool/Platform/Windows/MicrosoftPlatformSDK.cs
--- a/Engine/Source/Programs/UnrealBuildTool/Platform/Windows/MicrosoftPlatformSDK.cs
+++ b/Engine/Source/Programs/UnrealBuildTool/Platform/Windows/MicrosoftPlatformSDK.cs
@@ -531,7 +531,7 @@ namespace UnrealBuildTool
 				if (ToolChain == null || !(ToolChain.Version == ToolChainVersion || ToolChain.Family == ToolChainVersion))
 				{
 					DumpToolChains(ToolChains, x => x.ThenByDescending(x => x.Version == ToolChainVersion).ThenByDescending(x => x.Family == ToolChainVersion), Architecture, Logger);
-					throw new BuildException("Unable to find valid {0} toolchain for {1} {2}", ToolChainVersion, Compiler, Architecture.ToString());
+					throw new BuildException("Unable to find valid {0} C++ toolchain for {1} {2}", ToolChainVersion, Compiler, Architecture.ToString());
 				}
 			}
 
