UnrealBuildTool: Always compile engine modules with C++20, unless the platform doesn't support it

* Add MaxCppStandard option to the target, so the CppStd can be limited if it would be set higher for any reason. HoloLens for example does not support c++20 because it uses the /ZW flag to enable WinRT and these are incompatible. I will look into making WinRT only be added if it is required for a Module at a later time
* Any module that sets bTreatAsEngineModule will be upgraded to c++20 unless it has already overridden it to c++17, unless MaxCppStandard is set.
* Some modules, mostly in plugins, are set in the ModuleRules to always build with c++17. This is generally because third party source is included that are incompatible with c++20 but there are a few cases that should be looked into later
* VS2019 needed a few extra fixes but also supports c++20. Reminder that support for the 2019 toolchain is going to be removed in 5.5

diff --git a/Engine/Platforms/Hololens/Source/Programs/UnrealBuildTool/UEBuildHoloLens.cs b/Engine/Platforms/Hololens/Source/Programs/UnrealBuildTool/UEBuildHoloLens.cs
--- a/Engine/Platforms/Hololens/Source/Programs/UnrealBuildTool/UEBuildHoloLens.cs
+++ b/Engine/Platforms/Hololens/Source/Programs/UnrealBuildTool/UEBuildHoloLens.cs
@@ -345,10 +345,15 @@ namespace UnrealBuildTool
 			ConfigCache.ReadSettings(IniDirRef, Platform, Target.HoloLensPlatform);
 
 			// cl : Command line error D8016 : '/std:c++20' and '/ZW' command-line options are incompatible
+			if (Target.MaxCppStandard > CppStandardVersion.Cpp17)
+			{
+				Target.MaxCppStandard = CppStandardVersion.Cpp17;
+			}
 			if (Target.CppStandard > CppStandardVersion.Cpp17)
 			{
 				Target.CppStandard = CppStandardVersion.Cpp17;
 			}
+			Target.bEnableCppModules = false; // Cpp modules require C++20
 
 			if (Target.HoloLensPlatform.Compiler == WindowsCompiler.Default)
 			{
diff --git a/Engine/Source/Programs/UnrealBuildTool/Configuration/TargetRules.cs b/Engine/Source/Programs/UnrealBuildTool/Configuration/TargetRules.cs
--- a/Engine/Source/Programs/UnrealBuildTool/Configuration/TargetRules.cs
+++ b/Engine/Source/Programs/UnrealBuildTool/Configuration/TargetRules.cs
@@ -2069,13 +2069,21 @@ namespace UnrealBuildTool
 		private bool? bLegacyParentIncludePathsPrivate;
 
 		/// <summary>
-		/// Which C++ stanard to use for compiling this target
+		/// Which C++ standard to use for compiling this target
 		/// </summary>
 		[RequiresUniqueBuildEnvironment]
 		[CommandLine("-CppStd")]
 		[XmlConfigFile(Category = "BuildConfiguration")]
 		public CppStandardVersion CppStandard = CppStandardVersion.Default;
 
+		/// <summary>
+		/// Which C++ standard is the maximum allowed for compiling modules in this target if they override the Target's CppStandard for any reason.
+		/// </summary>
+		[RequiresUniqueBuildEnvironment]
+		[CommandLine("-MaxCppStd")]
+		[XmlConfigFile(Category = "BuildConfiguration")]
+		public CppStandardVersion MaxCppStandard = CppStandardVersion.Latest;
+
 		/// <summary>
 		/// Which C standard to use for compiling this target
 		/// </summary>
@@ -3812,6 +3820,11 @@ namespace UnrealBuildTool
 			get { return Inner.CppStandard; }
 		}
 
+		public CppStandardVersion MaxCppStandard
+		{
+			get { return Inner.MaxCppStandard; }
+		}
+
 		public CStandardVersion CStandard
 		{
 			get { return Inner.CStandard; }
diff --git a/Engine/Source/Programs/UnrealBuildTool/Configuration/UEBuildModuleCPP.cs b/Engine/Source/Programs/UnrealBuildTool/Configuration/UEBuildModuleCPP.cs
--- a/Engine/Source/Programs/UnrealBuildTool/Configuration/UEBuildModuleCPP.cs
+++ b/Engine/Source/Programs/UnrealBuildTool/Configuration/UEBuildModuleCPP.cs
@@ -1753,10 +1753,21 @@ namespace UnrealBuildTool
 			{
 				Result.CppStandard = (CppStandardVersion)Rules.CppStandard;
 			}
+			// Otherwise, for engine modules use the EngineDefault standard.
+			else if (Rules.bTreatAsEngineModule)
+			{
+				Result.CppStandard = CppStandardVersion.EngineDefault;
+			}
+			// CppModules require at least Cpp20.
 			if (Target.bEnableCppModules && Result.CppStandard < CppStandardVersion.Cpp20)
 			{
 				Result.CppStandard = CppStandardVersion.Cpp20;
 			}
+			// Downgrade CppStandard if a max version has been set.
+			if (Result.CppStandard > Target.MaxCppStandard)
+			{
+				Result.CppStandard = Target.MaxCppStandard;
+			}
 
 			// If the module overrides the C language version, override it on the compile environment
 			if (Rules.CStandard != null)
diff --git a/Engine/Source/Programs/UnrealBuildTool/System/CppCompileEnvironment.cs b/Engine/Source/Programs/UnrealBuildTool/System/CppCompileEnvironment.cs
--- a/Engine/Source/Programs/UnrealBuildTool/System/CppCompileEnvironment.cs
+++ b/Engine/Source/Programs/UnrealBuildTool/System/CppCompileEnvironment.cs
@@ -30,7 +30,7 @@ namespace UnrealBuildTool
 	public enum CppStandardVersion
 	{
 		/// <summary>
-		/// Supports C++14
+		/// Supports C++14. No longer maintained, will be removed in 5.5
 		/// </summary>
 		Cpp14,
 
@@ -53,6 +53,11 @@ namespace UnrealBuildTool
 		/// Use the default standard version
 		/// </summary>
 		Default = Cpp17,
+
+		/// <summary>
+		/// Use the default standard version for engine modules
+		/// </summary>
+		EngineDefault = Cpp20,
 	}
 	
 	/// <summary>
