Exposing xcode settings to UE settings UI (more work about dynamic disabling etc will follow)

diff --git a/Engine/Config/BaseEngine.ini b/Engine/Config/BaseEngine.ini
--- a/Engine/Config/BaseEngine.ini
+++ b/Engine/Config/BaseEngine.ini
@@ -2877,19 +2877,6 @@ MetalLanguageVersion=0
 bDisableMotionData=False
 bEnableAdvertisingIdentifier=True
 
-[XcodeConfiguration]
-bUseModernXcode=false
-bUseModernCodeSigning=true
-ModernSigningTeam=
-ModernSigningPrefix=com.YourCompany
-ModernBundleIdentifier=$(UE_SIGNING_PREFIX).$(UE_PRODUCT_NAME_STRIPPED)
-AppCategory=public.app-category.games
-TemplateMacPlist=/Engine/Build/Mac/Resources/Info.Template.plist
-TemplateMacEditorPlist=/Engine/Build/Mac/Resources/Info-Editor.Template.plist
-TemplateIOSPlist=/Game/Build/IOS/UBTGenerated/Info.Template.plist
-PremadeMacEntitlements=/Engine/Build/Mac/Resources/Sandbox.Server.entitlements
-ShippingSpecificMacEntitlements=/Engine/Build/Mac/Resources/Sandbox.NoNet.entitlements
-
 ; These are the defaults for Android settings, and they need to be in the .ini since UBT reads the .ini settings, without instantiating the class
 [/Script/AndroidRuntimeSettings.AndroidRuntimeSettings]
 SDKAPILevelOverride=
@@ -3037,6 +3024,19 @@ DefaultArchitecture=MacTargetArchitectureHost
 bBuildAllSupportedOnBuildMachine=true
 bEnableRayTracing=false
 
+[/Script/MacTargetPlatform.XcodeProjectSettings]
+bUseModernXcode=false
+bUseModernCodeSigning=true
+ModernSigningTeam=
+ModernSigningPrefix=com.YourCompany
+ModernBundleIdentifier=$(UE_SIGNING_PREFIX).$(UE_PRODUCT_NAME_STRIPPED)
+AppCategory=public.app-category.games
+TemplateMacPlist=/Engine/Build/Mac/Resources/Info.Template.plist
+TemplateMacEditorPlist=/Engine/Build/Mac/Resources/Info-Editor.Template.plist
+TemplateIOSPlist=/Game/Build/IOS/UBTGenerated/Info.Template.plist
+PremadeMacEntitlements=/Engine/Build/Mac/Resources/Sandbox.Server.entitlements
+ShippingSpecificMacEntitlements=/Engine/Build/Mac/Resources/Sandbox.NoNet.entitlements
+
 [HMDPluginPriority]
 ; Since SteamVR also works with the Oculus Rift, give priority to the native Oculus plugin before trying SteamVR.
 ; OpenXR is production ready and, if enabled, should be the priority plugin. For 4.27, we still need to prioritize
diff --git a/Engine/Source/Developer/Mac/MacPlatformEditor/Private/MacPlatformEditorModule.cpp b/Engine/Source/Developer/Mac/MacPlatformEditor/Private/MacPlatformEditorModule.cpp
--- a/Engine/Source/Developer/Mac/MacPlatformEditor/Private/MacPlatformEditorModule.cpp
+++ b/Engine/Source/Developer/Mac/MacPlatformEditor/Private/MacPlatformEditorModule.cpp
@@ -5,6 +5,7 @@
 #include "ISettingsModule.h"
 #include "Modules/ModuleManager.h"
 #include "MacTargetSettings.h"
+#include "XcodeProjectSettings.h"
 #include "UObject/WeakObjectPtr.h"
 #include "UObject/Class.h"
 
@@ -31,6 +32,11 @@ class FMacPlatformEditorModule
 				LOCTEXT("TargetSettingsDescription", "Settings for Mac target platform"),
 				GetMutableDefault<UMacTargetSettings>()
 			);
+            SettingsModule->RegisterSettings("Project", "Platforms", "Mac",
+                LOCTEXT("TargetSettingsName", "Xcode projects"),
+                LOCTEXT("TargetSettingsDescription", "Settings for Xcode projects"),
+                GetMutableDefault<UXcodeProjectSettings>()
+            );
 		}
 	}
 
diff --git a/Engine/Source/Developer/Mac/MacTargetPlatform/Classes/XcodeProjectSettings.h b/Engine/Source/Developer/Mac/MacTargetPlatform/Classes/XcodeProjectSettings.h
new file mode 100644
--- /dev/null
+++ b/Engine/Source/Developer/Mac/MacTargetPlatform/Classes/XcodeProjectSettings.h
@@ -0,0 +1,89 @@
+// Copyright Epic Games, Inc. All Rights Reserved.
+
+/*=============================================================================
+    XcodeProjectSettings.h: Declares the UXcodeProjectSettings class.
+=============================================================================*/
+
+#pragma once
+
+#include "CoreMinimal.h"
+#include "UObject/Object.h"
+#include "XcodeProjectSettings.generated.h"
+
+/**
+ * Implements the settings for Xcode projects
+ */
+UCLASS(config=Engine, defaultconfig)
+class MACTARGETPLATFORM_API UXcodeProjectSettings
+    : public UObject
+{
+public:
+    
+    GENERATED_UCLASS_BODY()
+    
+    /**
+     * Enable modernized Xcode, when building from Xcode, use native Xcode for bundle generation and archiving instead of UBT
+     */
+    UPROPERTY(EditAnywhere, config, Category=Xcode, meta = (DisplayName = "Modernized Xcode"))
+    bool bUseModernXcode;
+    
+    /**
+     * Enable native Xcode code signing
+     */
+    UPROPERTY(EditAnywhere, config, Category=Xcode, meta = (DisplayName = "Use Modern Code Signing"))
+    bool bUseModernCodeSigning;
+    
+    /**
+     * Team ID used for native Xcode code signing
+     */
+    UPROPERTY(EditAnywhere, config, Category=Xcode, meta = (DisplayName = "Modern Code Sign Team"))
+    FString ModernSigningTeam;
+    
+    /**
+     * Bundle ID prefix used for native Xcode code signing
+     */
+    UPROPERTY(EditAnywhere, config, Category=Xcode, meta = (DisplayName = "ModernSigningPrefix"))
+    FString ModernSigningPrefix;
+    
+    /**
+     * Bundle ID used for nativr Xcode code signing
+     */
+    UPROPERTY(EditAnywhere, config, Category=Xcode, meta = (DisplayName = "ModernBundleIdentifier"))
+    FString ModernBundleIdentifier;
+    
+    /**
+     * The App Category that will be used for Mac App Store submission
+     */
+    UPROPERTY(EditAnywhere, config, Category=Xcode, meta = (DisplayName = "AppCategory"))
+    FString AppCategory;
+    
+    /**
+     * The template info.plist used for Mac game targets
+     */
+    UPROPERTY(EditAnywhere, config, Category=Xcode, meta = (DisplayName = "TemplateMacPlist"))
+    FFilePath TemplateMacPlist;
+    
+    /**
+     * The template info.plist used for Mac editor targets
+     */
+    UPROPERTY(EditAnywhere, config, Category=Xcode, meta = (DisplayName = "TemplateMacEditorPlist"))
+    FFilePath TemplateMacEditorPlist;
+    
+    /**
+     * The template info.plist used for iOS game targets
+     */
+    UPROPERTY(EditAnywhere, config, Category=Xcode, meta = (DisplayName = "TemplateIOSPlist"))
+    FFilePath TemplateIOSPlist;
+    
+    /**
+     * The premade entitlement file used for development Mac builds
+     */
+    UPROPERTY(EditAnywhere, config, Category=Xcode, meta = (DisplayName = "PremadeMacEntitlements"))
+    FFilePath PremadeMacEntitlements;
+    
+    /**
+     * The premade entitlement file used for shipping Mac builds
+     */
+    UPROPERTY(EditAnywhere, config, Category=Xcode, meta = (DisplayName = "ShippingSpecificMacEntitlements"))
+    FFilePath ShippingSpecificMacEntitlements;
+};
diff --git a/Engine/Source/Developer/Mac/MacTargetPlatform/Private/MacTargetPlatformClasses.cpp b/Engine/Source/Developer/Mac/MacTargetPlatform/Private/MacTargetPlatformClasses.cpp
--- a/Engine/Source/Developer/Mac/MacTargetPlatform/Private/MacTargetPlatformClasses.cpp
+++ b/Engine/Source/Developer/Mac/MacTargetPlatform/Private/MacTargetPlatformClasses.cpp
@@ -6,6 +6,7 @@
 
 #include "MacTargetSettings.h"
 #include "GenericMacTargetPlatform.h"
+#include "XcodeProjectSettings.h"
 
 
 /* UMacTargetSettings structors
@@ -14,3 +15,10 @@
 UMacTargetSettings::UMacTargetSettings( const FObjectInitializer& ObjectInitializer )
 	: Super(ObjectInitializer)
 { }
+
+/* UXcodeProjectSettings structors
+ *****************************************************************************/
+
+UXcodeProjectSettings::UXcodeProjectSettings( const FObjectInitializer& ObjectInitializer )
+    : Super(ObjectInitializer)
+{ }
diff --git a/Engine/Source/Developer/Mac/MacTargetPlatform/Private/MacTargetPlatformModule.cpp b/Engine/Source/Developer/Mac/MacTargetPlatform/Private/MacTargetPlatformModule.cpp
--- a/Engine/Source/Developer/Mac/MacTargetPlatform/Private/MacTargetPlatformModule.cpp
+++ b/Engine/Source/Developer/Mac/MacTargetPlatform/Private/MacTargetPlatformModule.cpp
@@ -7,6 +7,7 @@
 #include "Interfaces/ITargetPlatformModule.h"
 #include "Modules/ModuleManager.h"
 #include "MacTargetSettings.h"
+#include "XcodeProjectSettings.h"
 #include "UObject/Package.h"
 #include "UObject/WeakObjectPtr.h"
 
@@ -68,6 +69,9 @@ public:
 		}
 		
 		TargetSettings->AddToRoot();
+        
+        ProjectSettings = NewObject<UXcodeProjectSettings>(GetTransientPackage(), "XcodeProjectSettings", RF_Standalone);
+        ProjectSettings->AddToRoot();
 
 		ISettingsModule* SettingsModule = FModuleManager::GetModulePtr<ISettingsModule>("Settings");
 
@@ -78,6 +82,11 @@ public:
 				LOCTEXT("TargetSettingsDescription", "Settings and resources for Mac platform"),
 				TargetSettings
 			);
+            SettingsModule->RegisterSettings("Project", "Platforms", "Mac",
+                LOCTEXT("TargetSettingsName", "Xcode projects"),
+                LOCTEXT("TargetSettingsDescription", "Settings for Xcode projects"),
+                ProjectSettings
+            );
 		}
 #endif
 	}
@@ -96,10 +105,12 @@ public:
 		{
 			// If we're in exit purge, this object has already been destroyed
 			TargetSettings->RemoveFromRoot();
+            ProjectSettings->RemoveFromRoot();
 		}
 		else
 		{
 			TargetSettings = NULL;
+            ProjectSettings = NULL;
 		}
 #endif
 	}
@@ -111,6 +122,8 @@ private:
 
 	// Holds the target settings.
 	UMacTargetSettings* TargetSettings;
+    
+    UXcodeProjectSettings* ProjectSettings;
 };
 
 
diff --git a/Engine/Source/Programs/AutomationTool/Apple/ApplePlatform.Automation.cs b/Engine/Source/Programs/AutomationTool/Apple/ApplePlatform.Automation.cs
--- a/Engine/Source/Programs/AutomationTool/Apple/ApplePlatform.Automation.cs
+++ b/Engine/Source/Programs/AutomationTool/Apple/ApplePlatform.Automation.cs
@@ -73,7 +73,7 @@ public abstract class ApplePlatform : Platform
 	{
 		ConfigHierarchy Ini = ConfigCache.ReadHierarchy(ConfigHierarchyType.Engine, Params.RawProjectPath.Directory!, SC.StageTargetPlatform.PlatformType);
 		bool bUseModernXcode;
-		Ini.TryGetValue("XcodeConfiguration", "bUseModernXcode", out bUseModernXcode);
+		Ini.TryGetValue("/Script/MacTargetPlatform.XcodeProjectSettings", "bUseModernXcode", out bUseModernXcode);
 
 		if (bUseModernXcode)
 		{
diff --git a/Engine/Source/Programs/AutomationTool/IOS/IOSPlatform.Automation.cs b/Engine/Source/Programs/AutomationTool/IOS/IOSPlatform.Automation.cs
--- a/Engine/Source/Programs/AutomationTool/IOS/IOSPlatform.Automation.cs
+++ b/Engine/Source/Programs/AutomationTool/IOS/IOSPlatform.Automation.cs
@@ -793,7 +793,7 @@ public class IOSPlatform : ApplePlatform
 
 		ConfigHierarchy Ini = ConfigCache.ReadHierarchy(ConfigHierarchyType.Engine, Params.RawProjectPath.Directory!, UnrealTargetPlatform.IOS);
 		bool bUseModernXcode;
-		Ini.TryGetValue("XcodeConfiguration", "bUseModernXcode", out bUseModernXcode);
+		Ini.TryGetValue("/Script/MacTargetPlatform.XcodeProjectSettings", "bUseModernXcode", out bUseModernXcode);
 
 		string MobileProvision;
 		string SigningCertificate;
@@ -801,9 +801,9 @@ public class IOSPlatform : ApplePlatform
 		bool bAutomaticSigning;
 		if (bUseModernXcode)
 		{
-			Ini.TryGetValue("XcodeConfiguration", "bUseModernCodeSigning", out bAutomaticSigning);
-			Ini.TryGetValue("XcodeConfiguration", "ModernSigningTeam", out TeamUUID);
-			//			Ini.TryGetValue("XcodeConfiguration", "ModernSigningPrefix", out SigningPrefix);
+			Ini.TryGetValue("/Script/MacTargetPlatform.XcodeProjectSettings", "bUseModernCodeSigning", out bAutomaticSigning);
+			Ini.TryGetValue("/Script/MacTargetPlatform.XcodeProjectSettings", "ModernSigningTeam", out TeamUUID);
+			//			Ini.TryGetValue("/Script/MacTargetPlatform.XcodeProjectSettings", "ModernSigningPrefix", out SigningPrefix);
 			if (!bAutomaticSigning)
 			{
 				throw new AutomationException("Currently only Modern (automatic) codesigning is supported with Modern xcode");
@@ -1117,7 +1117,7 @@ public class IOSPlatform : ApplePlatform
 		if (OperatingSystem.IsMacOS())
 		{
 			ConfigHierarchy Ini = ConfigCache.ReadHierarchy(ConfigHierarchyType.Engine, RawProjectPath.Directory, UnrealTargetPlatform.IOS);
-			Ini.TryGetValue("XcodeConfiguration", "bUseModernXcode", out bUseModernXcode);
+			Ini.TryGetValue("/Script/MacTargetPlatform.XcodeProjectSettings", "bUseModernXcode", out bUseModernXcode);
 		}
 
 		if (bUseModernXcode)
diff --git a/Engine/Source/Programs/AutomationTool/Mac/MacPlatform.Automation.cs b/Engine/Source/Programs/AutomationTool/Mac/MacPlatform.Automation.cs
--- a/Engine/Source/Programs/AutomationTool/Mac/MacPlatform.Automation.cs
+++ b/Engine/Source/Programs/AutomationTool/Mac/MacPlatform.Automation.cs
@@ -132,7 +132,7 @@ public class MacPlatform : ApplePlatform
 	{
 		ConfigHierarchy Ini = ConfigCache.ReadHierarchy(ConfigHierarchyType.Engine, Params.RawProjectPath.Directory!, UnrealTargetPlatform.IOS);
 		bool bUseModernXcode;
-		Ini.TryGetValue("XcodeConfiguration", "bUseModernXcode", out bUseModernXcode);
+		Ini.TryGetValue("/Script/MacTargetPlatform.XcodeProjectSettings", "bUseModernXcode", out bUseModernXcode);
 
 		// Stage all the build products
 		foreach (StageTarget Target in SC.StageTargets)
diff --git a/Engine/Source/Programs/UnrealBuildTool/Platform/IOS/UEBuildIOS.cs b/Engine/Source/Programs/UnrealBuildTool/Platform/IOS/UEBuildIOS.cs
--- a/Engine/Source/Programs/UnrealBuildTool/Platform/IOS/UEBuildIOS.cs
+++ b/Engine/Source/Programs/UnrealBuildTool/Platform/IOS/UEBuildIOS.cs
@@ -150,7 +150,7 @@ namespace UnrealBuildTool
 		/// <summary>
 		/// Are we using modern xcode in this project?
 		/// </summary>
-		[ConfigFile(ConfigHierarchyType.Engine, "XcodeConfiguration", "bUseModernXcode")]
+		[ConfigFile(ConfigHierarchyType.Engine, "/Script/MacTargetPlatform.XcodeProjectSettings", "bUseModernXcode")]
 		public bool bUseModernXcode = false;
 	}
 
diff --git a/Engine/Source/Programs/UnrealBuildTool/Platform/IOS/UEDeployIOS.cs b/Engine/Source/Programs/UnrealBuildTool/Platform/IOS/UEDeployIOS.cs
--- a/Engine/Source/Programs/UnrealBuildTool/Platform/IOS/UEDeployIOS.cs
+++ b/Engine/Source/Programs/UnrealBuildTool/Platform/IOS/UEDeployIOS.cs
@@ -469,7 +469,7 @@ namespace UnrealBuildTool
 			DirectoryReference? DirRef = bIsUnrealGame ? (!string.IsNullOrEmpty(UnrealBuildTool.GetRemoteIniPath()) ? new DirectoryReference(UnrealBuildTool.GetRemoteIniPath()!) : null) : new DirectoryReference(ProjectDirectory);
 			ConfigHierarchy Ini = ConfigCache.ReadHierarchy(ConfigHierarchyType.Engine, DirRef, UnrealTargetPlatform.Mac);
 			bool bUseModernXcode;
-			Ini.TryGetValue("XcodeConfiguration", "bUseModernXcode", out bUseModernXcode);
+			Ini.TryGetValue("/Script/MacTargetPlatform.XcodeProjectSettings", "bUseModernXcode", out bUseModernXcode);
 
 			if (!bUseModernXcode)
 			{
@@ -1103,7 +1103,7 @@ namespace UnrealBuildTool
 
 			ConfigHierarchy Ini = ConfigCache.ReadHierarchy(ConfigHierarchyType.Engine, ProjectFile?.Directory, UnrealTargetPlatform.IOS);
 			bool bUseModernXcode;
-			if (Ini.TryGetValue("XcodeConfiguration", "bUseModernXcode", out bUseModernXcode) && bUseModernXcode)
+			if (Ini.TryGetValue("/Script/MacTargetPlatform.XcodeProjectSettings", "bUseModernXcode", out bUseModernXcode) && bUseModernXcode)
 			{
 				Logger.LogInformation("Generating plist (only step needed when deploying with Modern Xcode)");
 				GeneratePList(ProjectFile, Config, InProjectDirectory, bIsUnrealGame, GameExeName, false, InProjectName, InEngineDir, AppDirectory, UPLScripts, BundleID, bBuildAsFramework);
diff --git a/Engine/Source/Programs/UnrealBuildTool/Platform/Mac/UEBuildMac.cs b/Engine/Source/Programs/UnrealBuildTool/Platform/Mac/UEBuildMac.cs
--- a/Engine/Source/Programs/UnrealBuildTool/Platform/Mac/UEBuildMac.cs
+++ b/Engine/Source/Programs/UnrealBuildTool/Platform/Mac/UEBuildMac.cs
@@ -70,7 +70,7 @@ namespace UnrealBuildTool
 		/// <summary>
 		/// Whether or not to use Modernized xcode mode
 		/// </summary>
-		[ConfigFile(ConfigHierarchyType.Engine, "XcodeConfiguration", "bUseModernXcode")]
+		[ConfigFile(ConfigHierarchyType.Engine, "/Script/MacTargetPlatform.XcodeProjectSettings", "bUseModernXcode")]
 		public bool bUseModernXcode = false;
 	}
 
diff --git a/Engine/Source/Programs/UnrealBuildTool/ProjectFiles/Xcode/XcodeProject.cs b/Engine/Source/Programs/UnrealBuildTool/ProjectFiles/Xcode/XcodeProject.cs
--- a/Engine/Source/Programs/UnrealBuildTool/ProjectFiles/Xcode/XcodeProject.cs
+++ b/Engine/Source/Programs/UnrealBuildTool/ProjectFiles/Xcode/XcodeProject.cs
@@ -171,7 +171,7 @@ namespace UnrealBuildTool.XcodeProjectXcconfig
 				if (Path.GetExtension(Tokens[1]) == "")
 				{
 					string? PremadeLocation;
-					if (Ini.TryGetValue("XcodeConfiguration", Tokens[1], out PremadeLocation) && PremadeLocation != "")
+					if (Ini.TryGetValue("/Script/MacTargetPlatform.XcodeProjectSettings", Tokens[1], out PremadeLocation) && PremadeLocation != "")
 					{
 						FileReference TestFile;
 						if (PremadeLocation.StartsWith("/Engine/"))
@@ -371,7 +371,7 @@ namespace UnrealBuildTool.XcodeProjectXcconfig
 			if (string.IsNullOrEmpty(BundleIdentifier))
 			{
 				ConfigHierarchy Ini = ConfigCache.ReadHierarchy(ConfigHierarchyType.Engine, UProjectFileLocation?.Directory, UnrealTargetPlatform.Mac);
-				Ini.GetString($"XcodeConfiguration", "ModernBundleIdentifier", out BundleIdentifier);
+				Ini.GetString($"/Script/MacTargetPlatform.XcodeProjectSettings", "ModernBundleIdentifier", out BundleIdentifier);
 			}
 			if (string.IsNullOrEmpty(BundleIdentifier))
 			{
@@ -452,7 +452,7 @@ namespace UnrealBuildTool.XcodeProjectXcconfig
 
 			// read setings from the configs, now that we have a project
 			ConfigHierarchy SharedPlatformIni = ConfigCache.ReadHierarchy(ConfigHierarchyType.Engine, UProjectFileLocation?.Directory, UnrealTargetPlatform.Mac);
-			SharedPlatformIni.TryGetValue("XcodeConfiguration", "bUseModernCodeSigning", out bUseAutomaticSigning);
+			SharedPlatformIni.TryGetValue("/Script/MacTargetPlatform.XcodeProjectSettings", "bUseModernCodeSigning", out bUseAutomaticSigning);
 
 			Metadata = new Metadata(ProductDirectory, XcodeProjectFileLocation, SharedPlatformIni, bSupportsMac, bSupportsIOS || bSupportsTVOS, Logger);
 		}
@@ -1489,10 +1489,10 @@ namespace UnrealBuildTool.XcodeProjectXcconfig
 			List<string> ExtraConfigLines = new();
 
 			// get codesigning info (all platforms)
-			PlatformIni.TryGetValue("XcodeConfiguration", "bUseModernCodeSigning", out bAutomaticSigning);
-			PlatformIni.TryGetValue("XcodeConfiguration", "ModernSigningTeam", out SigningTeam);
-			PlatformIni.TryGetValue("XcodeConfiguration", "ModernSigningPrefix", out SigningPrefix);
-			PlatformIni.TryGetValue("XcodeConfiguration", "AppCategory", out AppCategory);
+			PlatformIni.TryGetValue("/Script/MacTargetPlatform.XcodeProjectSettings", "bUseModernCodeSigning", out bAutomaticSigning);
+			PlatformIni.TryGetValue("/Script/MacTargetPlatform.XcodeProjectSettings", "ModernSigningTeam", out SigningTeam);
+			PlatformIni.TryGetValue("/Script/MacTargetPlatform.XcodeProjectSettings", "ModernSigningPrefix", out SigningPrefix);
+			PlatformIni.TryGetValue("/Script/MacTargetPlatform.XcodeProjectSettings", "AppCategory", out AppCategory);
 
 
 			if (Platform == UnrealTargetPlatform.Mac)
@@ -1953,7 +1953,7 @@ namespace UnrealBuildTool.XcodeProjectXcconfig
 				FileReference? UProjectFileLocation = UnrealData.FindUProjectFileLocation(this);
 				ConfigHierarchy Ini = ConfigCache.ReadHierarchy(ConfigHierarchyType.Engine, UProjectFileLocation?.Directory, UnrealTargetPlatform.Mac);
 				bool bUseModernXcode;
-				if (!(Ini.TryGetValue("XcodeConfiguration", "bUseModernXcode", out bUseModernXcode) && bUseModernXcode))
+				if (!(Ini.TryGetValue("/Script/MacTargetPlatform.XcodeProjectSettings", "bUseModernXcode", out bUseModernXcode) && bUseModernXcode))
 				{
 					LegacyProjectFile = new XcodeProjectLegacy.XcodeProjectFile(ProjectFilePath, BaseDir, UnrealData.bForDistribution, UnrealData.BundleIdentifier, UnrealData.AppName, UnrealData.bMakeProjectPerTarget);
 					LegacyProjectFile.ProjectTargets.AddRange(ProjectTargets);
diff --git a/Engine/Source/Programs/UnrealBuildTool/ToolChain/AppleToolChain.cs b/Engine/Source/Programs/UnrealBuildTool/ToolChain/AppleToolChain.cs
--- a/Engine/Source/Programs/UnrealBuildTool/ToolChain/AppleToolChain.cs
+++ b/Engine/Source/Programs/UnrealBuildTool/ToolChain/AppleToolChain.cs
@@ -183,7 +183,7 @@ namespace UnrealBuildTool
 			// Modern Xcode mode does this now
 			bool _bUseModernXcode = false;
 			ConfigHierarchy Ini = ConfigCache.ReadHierarchy(ConfigHierarchyType.Engine, ProjectFile?.Directory, UnrealTargetPlatform.Mac);
-			Ini.TryGetValue("XcodeConfiguration", "bUseModernXcode", out _bUseModernXcode);
+			Ini.TryGetValue("/Script/MacTargetPlatform.XcodeProjectSettings", "bUseModernXcode", out _bUseModernXcode);
 			return _bUseModernXcode;
 		}
 
