UBT: Fix single-instance mutex being disposed from a different thread.

diff --git a/Engine/Source/Programs/UnrealBuildTool/UnrealBuildTool.cs b/Engine/Source/Programs/UnrealBuildTool/UnrealBuildTool.cs
--- a/Engine/Source/Programs/UnrealBuildTool/UnrealBuildTool.cs
+++ b/Engine/Source/Programs/UnrealBuildTool/UnrealBuildTool.cs
@@ -410,10 +410,11 @@ namespace UnrealBuildTool
 
 		/// <summary>
 		/// Main entry point. Parses any global options and initializes the logging system, then invokes the appropriate command.
+		/// NB: That the entry point is deliberately NOT async, since we have a single-instance mutex that cannot be disposed from a different thread.
 		/// </summary>
 		/// <param name="ArgumentsArray">Command line arguments</param>
 		/// <returns>Zero on success, non-zero on error</returns>
-		private static async Task<int> Main(string[] ArgumentsArray)
+		private static int Main(string[] ArgumentsArray)
 		{
 			FileReference? RunFile = null;
 			DirectoryReference? TempDirectory = null;
@@ -652,7 +653,7 @@ namespace UnrealBuildTool
 				ToolMode Mode = (ToolMode)Activator.CreateInstance(ModeType)!;
 
 				// Execute the mode
-				int Result = await Mode.ExecuteAsync(Arguments, Logger);
+				int Result = Task.Run(() => Mode.ExecuteAsync(Arguments, Logger)).Result;
 				if((ModeOptions & ToolModeOptions.ShowExecutionTime) != 0)
 				{
 					Logger.LogInformation("Total execution time: {Time:0.00} seconds", Timeline.Elapsed.TotalSeconds);
