UnrealBuildTool: Add arch to command message when compiling ISPC

#rnx

diff --git a/Engine/Source/Programs/UnrealBuildTool/ToolChain/ISPCToolChain.cs b/Engine/Source/Programs/UnrealBuildTool/ToolChain/ISPCToolChain.cs
--- a/Engine/Source/Programs/UnrealBuildTool/ToolChain/ISPCToolChain.cs
+++ b/Engine/Source/Programs/UnrealBuildTool/ToolChain/ISPCToolChain.cs
@@ -411,9 +411,11 @@ namespace UnrealBuildTool
 				}
 			}
 
+			string ISPCArch = GetISPCArchTarget(CompileEnvironment.Platform, null);
+
 			// Build target triplet
 			GlobalArguments.Add($"--target-os={GetISPCOSTarget(CompileEnvironment.Platform)}");
-			GlobalArguments.Add($"--arch={GetISPCArchTarget(CompileEnvironment.Platform, null)}");
+			GlobalArguments.Add($"--arch={ISPCArch}");
 			GlobalArguments.Add($"--target={TargetString}");
 			GlobalArguments.Add($"--emit-{GetISPCObjectFileFormat(CompileEnvironment.Platform)}");
 
@@ -457,7 +459,7 @@ namespace UnrealBuildTool
 			foreach (FileItem ISPCFile in InputFiles)
 			{
 				Action CompileAction = Graph.CreateAction(ActionType.Compile);
-				CompileAction.CommandDescription = "Generate Header";
+				CompileAction.CommandDescription = $"Generate Header [{ISPCArch}]";
 				CompileAction.WorkingDirectory = Unreal.EngineSourceDirectory;
 				CompileAction.CommandPath = new FileReference(GetISPCHostCompilerPath(BuildHostPlatform.Current.Platform));
 				CompileAction.StatusDescription = Path.GetFileName(ISPCFile.AbsolutePath);
@@ -508,6 +510,7 @@ namespace UnrealBuildTool
 
 				// Fix interrupted build issue by copying header after generation completes
 				Action CopyAction = Graph.CreateCopyAction(ISPCIncludeHeaderFile, ISPCFinalHeaderFile);
+				CopyAction.CommandDescription = $"{CopyAction.CommandDescription} [{ISPCArch}]";
 				CopyAction.DeleteItems.Clear();
 				CopyAction.PrerequisiteItems.Add(ISPCFile);
 				CopyAction.bShouldOutputStatusDescription = false;
@@ -549,9 +552,10 @@ namespace UnrealBuildTool
 
 			// Build target triplet
 			string PlatformObjectFileFormat = GetISPCObjectFileFormat(CompileEnvironment.Platform);
+			string ISPCArch = GetISPCArchTarget(CompileEnvironment.Platform, null);
 
 			GlobalArguments.Add($"--target-os={GetISPCOSTarget(CompileEnvironment.Platform)}");
-			GlobalArguments.Add($"--arch={GetISPCArchTarget(CompileEnvironment.Platform, null)}");
+			GlobalArguments.Add($"--arch={ISPCArch}");
 			GlobalArguments.Add($"--target={TargetString}");
 			GlobalArguments.Add($"--emit-{PlatformObjectFileFormat}");
 
@@ -613,7 +617,7 @@ namespace UnrealBuildTool
 			foreach (FileItem ISPCFile in InputFiles)
 			{
 				Action CompileAction = Graph.CreateAction(ActionType.Compile);
-				CompileAction.CommandDescription = "Compile";
+				CompileAction.CommandDescription = $"Compile [{ISPCArch}]";
 				CompileAction.WorkingDirectory = Unreal.EngineSourceDirectory;
 				CompileAction.CommandPath = new FileReference(GetISPCHostCompilerPath(BuildHostPlatform.Current.Platform));
 				CompileAction.StatusDescription = Path.GetFileName(ISPCFile.AbsolutePath);
@@ -739,7 +743,7 @@ namespace UnrealBuildTool
 
 							PostCompileAction.PrerequisiteItems.Add(CompiledBytecodeObjFile);
 							PostCompileAction.ProducedItems.Add(FinalCompiledISPCObjFile);
-							PostCompileAction.CommandDescription = "CompileByteCode";
+							PostCompileAction.CommandDescription = $"CompileByteCode [{ISPCArch}]";
 							PostCompileAction.WorkingDirectory = Unreal.EngineSourceDirectory;
 							PostCompileAction.CommandPath = new FileReference(ByteCodeCompilerPath);
 							PostCompileAction.StatusDescription = Path.GetFileName(ISPCFile.AbsolutePath);
