- Fixed FilePath properties with new XcodeProjectSettings class

diff --git a/Engine/Config/BaseEngine.ini b/Engine/Config/BaseEngine.ini
--- a/Engine/Config/BaseEngine.ini
+++ b/Engine/Config/BaseEngine.ini
@@ -3031,11 +3031,11 @@ ModernSigningTeam=
 ModernSigningPrefix=com.YourCompany
 ModernBundleIdentifier=$(UE_SIGNING_PREFIX).$(UE_PRODUCT_NAME_STRIPPED)
 AppCategory=public.app-category.games
-TemplateMacPlist=/Engine/Build/Mac/Resources/Info.Template.plist
-TemplateMacEditorPlist=/Engine/Build/Mac/Resources/Info-Editor.Template.plist
-TemplateIOSPlist=/Game/Build/IOS/UBTGenerated/Info.Template.plist
-PremadeMacEntitlements=/Engine/Build/Mac/Resources/Sandbox.Server.entitlements
-ShippingSpecificMacEntitlements=/Engine/Build/Mac/Resources/Sandbox.NoNet.entitlements
+TemplateMacPlist=(FilePath="/Engine/Build/Mac/Resources/Info.Template.plist")
+TemplateMacEditorPlist=(FilePath="/Engine/Build/Mac/Resources/Info-Editor.Template.plist")
+TemplateIOSPlist=(FilePath="/Game/Build/IOS/UBTGenerated/Info.Template.plist")
+PremadeMacEntitlements=(FilePath="/Engine/Build/Mac/Resources/Sandbox.Server.entitlements")
+ShippingSpecificMacEntitlements=(FilePath="/Engine/Build/Mac/Resources/Sandbox.NoNet.entitlements")
 
 [HMDPluginPriority]
 ; Since SteamVR also works with the Oculus Rift, give priority to the native Oculus plugin before trying SteamVR.
@@ -3316,4 +3316,4 @@ MaterialParents=(("MI_Default_Opaque", "/Interchange/gltf/MaterialInstances/MI_D
 
 [BuildPatchServices]
 ; Enable/Disable mechanism that provides a disk overflow store for patch chunks.
-bEnableDiskOverflowStore=true
\ No newline at end of file
+bEnableDiskOverflowStore=true
diff --git a/Engine/Source/Programs/UnrealBuildTool/ProjectFiles/Xcode/XcodeProject.cs b/Engine/Source/Programs/UnrealBuildTool/ProjectFiles/Xcode/XcodeProject.cs
--- a/Engine/Source/Programs/UnrealBuildTool/ProjectFiles/Xcode/XcodeProject.cs
+++ b/Engine/Source/Programs/UnrealBuildTool/ProjectFiles/Xcode/XcodeProject.cs
@@ -167,35 +167,51 @@ namespace UnrealBuildTool.XcodeProjectXcconfig
 				string[] Tokens = Loc.Split(':');
 				Mode = (Tokens[0] == "premade") ? MetadataMode.UsePremade : MetadataMode.UpdateTemplate;
 
+				FileReference? TestFile = null;
+
 				// no extension means it's a .ini entry
 				if (Path.GetExtension(Tokens[1]) == "")
 				{
 					string? PremadeLocation;
 					if (Ini.TryGetValue("/Script/MacTargetPlatform.XcodeProjectSettings", Tokens[1], out PremadeLocation) && PremadeLocation != "")
 					{
-						FileReference TestFile;
-						if (PremadeLocation.StartsWith("/Engine/"))
-						{
-							TestFile = FileReference.Combine(Unreal.EngineDirectory, PremadeLocation.Substring(8));
-						}
-						else if (PremadeLocation.StartsWith("/Game/"))
+						if (PremadeLocation.StartsWith("(FilePath=", StringComparison.OrdinalIgnoreCase))
 						{
-							TestFile = FileReference.Combine(ProductDirectory, PremadeLocation.Substring(6));
+							PremadeLocation = ConfigHierarchy.GetStructEntry(PremadeLocation, "FilePath", false)!;
+							if (PremadeLocation.StartsWith("/Engine/"))
+							{
+								TestFile = FileReference.Combine(Unreal.EngineDirectory, PremadeLocation.Substring(8));
+							}
+							else if (PremadeLocation.StartsWith("/Game/"))
+							{
+								TestFile = FileReference.Combine(ProductDirectory, PremadeLocation.Substring(6));
+							}
+							else
+							{
+								TestFile = new FileReference(PremadeLocation);
+							}
 						}
 						else
 						{
-							TestFile = new FileReference(Tokens[1]);
-						}
-
-						// make sure the file (if it's a generated file, then it may not exist yet, so allow it through)
-						if (TestFile.ContainsName("UBTGenerated", 0) || FileReference.Exists(TestFile))
-						{
-							File = TestFile;
-							break;
+							TestFile = new FileReference(PremadeLocation);
 						}
+					}
+				}
+				else
+				{
+					TestFile = new FileReference(Tokens[1]);
+				}
 
-						throw new BuildException($"Metadata file {Tokens[1]} (resolved to {TestFile}) was specified in project settings, but does not exist. Ignoring...");
+				// make sure the file (if it's a generated file, then it may not exist yet, so allow it through)
+				if (TestFile != null)
+				{
+					if (TestFile.ContainsName("UBTGenerated", 0) || FileReference.Exists(TestFile))
+					{
+						File = TestFile;
+						break;
 					}
+
+					throw new BuildException($"Metadata file {Tokens[1]} (resolved to {TestFile}) was specified in project settings, but does not exist. Ignoring...");
 				}
 			}
 
