Simplification of manifest generation, including:

 - use a dictionary to represent the executables and their config, not a pair of arrays
 - no need to pass ProjectPath - it is only used for pri binary resources and if the project doesn't exist there's nothing to copy

#rnx

diff --git a/Engine/Source/Programs/UnrealBuildTool/Platform/Windows/AppXManifestGeneratorBase.cs b/Engine/Source/Programs/UnrealBuildTool/Platform/Windows/AppXManifestGeneratorBase.cs
--- a/Engine/Source/Programs/UnrealBuildTool/Platform/Windows/AppXManifestGeneratorBase.cs
+++ b/Engine/Source/Programs/UnrealBuildTool/Platform/Windows/AppXManifestGeneratorBase.cs
@@ -66,9 +66,6 @@ namespace UnrealBuildTool
 		/// project file to use
 		protected FileReference? ProjectFile;
 
-		/// directory containing the project
-		protected string? ProjectPath;
-
 		/// output path - where the manifest will be created
 		protected string? OutputPath;
 
@@ -367,15 +364,21 @@ namespace UnrealBuildTool
 		/// </summary>
 		protected virtual bool FindResourceBinaryFile( out string SourceFilePath, string ResourceFileName, string CultureId = "", bool AllowEngineFallback = true)
 		{
-			// look in project normal Build location
-			SourceFilePath = Path.Combine(ProjectPath!, "Build", Platform.ToString(), BuildResourceSubPath, CultureId, ResourceFileName);
-			bool bFileExists = File.Exists(SourceFilePath);
+			bool bFileExists = false;
+			SourceFilePath = "";
 
-			// look in Platform Extensions next
-			if (!bFileExists)
+			if (ProjectFile != null)
 			{
-				SourceFilePath = Path.Combine(ProjectPath!, "Platforms", Platform.ToString(), "Build", BuildResourceSubPath, CultureId, ResourceFileName);
+				// look in project normal Build location
+				SourceFilePath = Path.Combine(ProjectFile.FullName!, "Build", Platform.ToString(), BuildResourceSubPath, CultureId, ResourceFileName);
 				bFileExists = File.Exists(SourceFilePath);
+
+				// look in Platform Extensions next
+				if (!bFileExists)
+				{
+					SourceFilePath = Path.Combine(ProjectFile.FullName!, "Platforms", Platform.ToString(), "Build", BuildResourceSubPath, CultureId, ResourceFileName);
+					bFileExists = File.Exists(SourceFilePath);
+				}
 			}
 
 			// look in Engine, if allowed
@@ -781,27 +784,22 @@ namespace UnrealBuildTool
 		/// <summary>
 		/// Return the entire manifest element
 		/// </summary>
-		protected abstract XElement GetManifest(List<UnrealTargetConfiguration> TargetConfigs, List<string> Executables, string? TargetName, out string IdentityName);
+		protected abstract XElement GetManifest(Dictionary<UnrealTargetConfiguration,string> InExecutablePairs, string? TargetName, out string IdentityName);
 
 		/// <summary>
 		/// Perform any platform-specific processing on the manifest before it is saved
 		/// </summary>
-		protected virtual void ProcessManifest(List<UnrealTargetConfiguration> TargetConfigs, List<string> Executables, string ManifestName, string ManifestTargetPath, string ManifestIntermediatePath)
+		protected virtual void ProcessManifest(Dictionary<UnrealTargetConfiguration,string> InExecutablePairs, string ManifestName, string ManifestTargetPath, string ManifestIntermediatePath)
         {
 		}
 
 		/// <summary>
 		/// Create a manifest and return the list of modified files
 		/// </summary>
-		public List<string>? CreateManifest(string InManifestName, string InOutputPath, string InIntermediatePath, string? InTargetName, FileReference? InProjectFile, string InProjectDirectory, List<UnrealTargetConfiguration> InTargetConfigs, List<string> InExecutables)
+		public List<string>? CreateManifest(string InManifestName, string InOutputPath, string InIntermediatePath, string? InTargetName, FileReference? InProjectFile, Dictionary<UnrealTargetConfiguration,string> InExecutablePairs)
 		{
 			// Check parameter values are valid.
-			if (InTargetConfigs.Count != InExecutables.Count)
-			{
-				Logger.LogError("The number of target configurations ({NumConfigs}) and executables ({NumExes}) passed to manifest generation do not match.", InTargetConfigs.Count, InExecutables.Count);
-				return null;
-			}
-			if (InTargetConfigs.Count < 1)
+			if (InExecutablePairs.Count < 1)
 			{
 				Logger.LogError("The number of target configurations is zero, so we cannot generate a manifest.");
 				return null;
@@ -821,7 +819,6 @@ namespace UnrealBuildTool
 			OutputPath = InOutputPath;
 			IntermediatePath = InIntermediatePath;
 			ProjectFile = InProjectFile;
-			ProjectPath = InProjectDirectory;
 			ManifestFiles = new Dictionary<string, string>();
 
 			// Load up INI settings. We'll use engine settings to retrieve the manifest configuration, but these may reference
@@ -878,7 +875,7 @@ namespace UnrealBuildTool
 				}
 			}
 			// Only warn if shipping, we can run without translated cultures they're just needed for cert
-			else if (InTargetConfigs.Contains(UnrealTargetConfiguration.Shipping))
+			else if (InExecutablePairs.ContainsKey(UnrealTargetConfiguration.Shipping))
 			{
 				Logger.LogInformation("Staged culture mappings not setup in the editor. See Per Culture Resources in the {Platform} Target Settings.", Platform.ToString() );
 			}
@@ -920,7 +917,7 @@ namespace UnrealBuildTool
 
 			// Create the manifest document
 			string? IdentityName = null;
-			var ManifestXmlDocument = new XDocument(GetManifest(InTargetConfigs, InExecutables, InTargetName, out IdentityName));
+			var ManifestXmlDocument = new XDocument(GetManifest(InExecutablePairs, InTargetName, out IdentityName));
 
 			// Export manifest to the intermediate directory then compare the contents to any existing target manifest
 			// and replace if there are differences.
@@ -928,7 +925,7 @@ namespace UnrealBuildTool
 			string ManifestTargetPath = Path.Combine(OutputPath, InManifestName);
 			ManifestXmlDocument.Save(ManifestIntermediatePath);
 			AddFileReference(ManifestIntermediatePath, ManifestTargetPath, bIsGeneratedFile: true);
-			ProcessManifest(InTargetConfigs, InExecutables, InManifestName, ManifestTargetPath, ManifestIntermediatePath);
+			ProcessManifest(InExecutablePairs, InManifestName, ManifestTargetPath, ManifestIntermediatePath);
 
 			// Export the resource tables starting with the default culture
 			string DefaultResourceTargetPath = Path.Combine(OutputPath, BuildResourceSubPath, "resources.resw");
