[UBT]
* Added bCanExecuteRemotelyWithXGE so actions that can't execute remotely because of xge don't turn of remote execution for sndbs etc

diff --git a/Engine/Source/Programs/UnrealBuildTool/Executors/XGE.cs b/Engine/Source/Programs/UnrealBuildTool/Executors/XGE.cs
--- a/Engine/Source/Programs/UnrealBuildTool/Executors/XGE.cs
+++ b/Engine/Source/Programs/UnrealBuildTool/Executors/XGE.cs
@@ -481,7 +481,7 @@ namespace UnrealBuildTool
 				LinkedAction Action = Actions[ActionIndex];
 
 				// Don't allow remote linking if on VPN.
-				bool CanExecuteRemotely = Action.bCanExecuteRemotely;
+				bool CanExecuteRemotely = Action.bCanExecuteRemotely && Action.bCanExecuteRemotelyWithXGE;
 				if(CanExecuteRemotely && Action.ActionType == ActionType.Link)
 				{
 					if (HostOnVpn || !bAllowRemoteLinking)
diff --git a/Engine/Source/Programs/UnrealBuildTool/Platform/Windows/VCCompileAction.cs b/Engine/Source/Programs/UnrealBuildTool/Platform/Windows/VCCompileAction.cs
--- a/Engine/Source/Programs/UnrealBuildTool/Platform/Windows/VCCompileAction.cs
+++ b/Engine/Source/Programs/UnrealBuildTool/Platform/Windows/VCCompileAction.cs
@@ -175,6 +175,9 @@ namespace UnrealBuildTool
 		/// <inheritdoc/>
 		public bool bCanExecuteRemotelyWithSNDBS { get; set; }
 
+		/// <inheritdoc/>
+		public bool bCanExecuteRemotelyWithXGE { get; set; } = true;
+
 		/// <inheritdoc/>
 		public bool bUseActionHistory => true;
 
@@ -352,6 +355,7 @@ namespace UnrealBuildTool
 			bShowIncludes = InAction.bShowIncludes;
 			bCanExecuteRemotely = InAction.bCanExecuteRemotely;
 			bCanExecuteRemotelyWithSNDBS = InAction.bCanExecuteRemotelyWithSNDBS;
+			bCanExecuteRemotelyWithXGE = InAction.bCanExecuteRemotelyWithXGE;
 			Architecture = InAction.Architecture;
 			Weight = InAction.Weight;
 
@@ -391,6 +395,8 @@ namespace UnrealBuildTool
 			bShowIncludes = Reader.ReadBool();
 			bCanExecuteRemotely = Reader.ReadBool();
 			bCanExecuteRemotelyWithSNDBS = Reader.ReadBool();
+			bCanExecuteRemotelyWithXGE = Reader.ReadBool();
+
 
 			AdditionalPrerequisiteItems = Reader.ReadList(() => Reader.ReadFileItem())!;
 			AdditionalProducedItems = Reader.ReadList(() => Reader.ReadFileItem())!;
@@ -425,6 +431,7 @@ namespace UnrealBuildTool
 			Writer.WriteBool(bShowIncludes);
 			Writer.WriteBool(bCanExecuteRemotely);
 			Writer.WriteBool(bCanExecuteRemotelyWithSNDBS);
+			Writer.WriteBool(bCanExecuteRemotelyWithXGE);
 
 			Writer.WriteList(AdditionalPrerequisiteItems, Item => Writer.WriteFileItem(Item));
 			Writer.WriteList(AdditionalProducedItems, Item => Writer.WriteFileItem(Item));
diff --git a/Engine/Source/Programs/UnrealBuildTool/System/Action.cs b/Engine/Source/Programs/UnrealBuildTool/System/Action.cs
--- a/Engine/Source/Programs/UnrealBuildTool/System/Action.cs
+++ b/Engine/Source/Programs/UnrealBuildTool/System/Action.cs
@@ -113,6 +113,11 @@ namespace UnrealBuildTool
 		/// </summary>
 		bool bCanExecuteRemotelyWithSNDBS { get; }
 
+		/// <summary>
+		/// True if this action is allowed to be run on a remote machine with XGE. Files with #import directives must be compiled locally. Also requires bCanExecuteRemotely = true.
+		/// </summary>
+		bool bCanExecuteRemotelyWithXGE { get; }
+
 		/// <summary>
 		/// True if this action is using the GCC compiler.  Some build systems may be able to optimize for this case.
 		/// </summary>
@@ -225,6 +230,11 @@ namespace UnrealBuildTool
 		/// </summary>
 		public bool bCanExecuteRemotelyWithSNDBS { get; set; } = true;
 
+		/// <summary>
+		/// True if this action is allowed to be run on a remote machine with XGE. Files with #import directives must be compiled locally. Also requires bCanExecuteRemotely = true.
+		/// </summary>
+		public bool bCanExecuteRemotelyWithXGE { get; set; } = true;
+
 		/// <summary>
 		/// True if this action is using the GCC compiler.  Some build systems may be able to optimize for this case.
 		/// </summary>
@@ -280,6 +290,7 @@ namespace UnrealBuildTool
 			StatusDescription = InOther.StatusDescription;
 			bCanExecuteRemotely = InOther.bCanExecuteRemotely;
 			bCanExecuteRemotelyWithSNDBS = InOther.bCanExecuteRemotelyWithSNDBS;
+			bCanExecuteRemotelyWithXGE = InOther.bCanExecuteRemotelyWithXGE;
 			bCanCache = InOther.bCanCache;
 			bIsGCCCompiler = InOther.bIsGCCCompiler;
 			bShouldOutputStatusDescription = InOther.bShouldOutputStatusDescription;
@@ -300,6 +311,7 @@ namespace UnrealBuildTool
 			StatusDescription = Reader.ReadString()!;
 			bCanExecuteRemotely = Reader.ReadBool();
 			bCanExecuteRemotelyWithSNDBS = Reader.ReadBool();
+			bCanExecuteRemotelyWithXGE = Reader.ReadBool();
 			bCanCache = Reader.ReadBool();
 			bIsGCCCompiler = Reader.ReadBool();
 			bShouldOutputStatusDescription = Reader.ReadBool();
@@ -328,6 +340,7 @@ namespace UnrealBuildTool
 			Writer.WriteString(StatusDescription);
 			Writer.WriteBool(bCanExecuteRemotely);
 			Writer.WriteBool(bCanExecuteRemotelyWithSNDBS);
+			Writer.WriteBool(bCanExecuteRemotelyWithXGE);
 			Writer.WriteBool(bCanCache);
 			Writer.WriteBool(bIsGCCCompiler);
 			Writer.WriteBool(bShouldOutputStatusDescription);
@@ -397,6 +410,12 @@ namespace UnrealBuildTool
 				Action.bCanExecuteRemotelyWithSNDBS = bCanExecuteRemotelyWithSNDBS;
 			}
 
+			bool bCanExecuteRemotelyWithXGE;
+			if (Object.TryGetBoolField("bCanExecuteRemotelyWithXGE", out bCanExecuteRemotelyWithXGE))
+			{
+				Action.bCanExecuteRemotelyWithXGE = bCanExecuteRemotelyWithXGE;
+			}
+
 			bool bCanCache;
 			if (Object.TryGetBoolField("bCanCache", out bCanCache))
 			{
@@ -509,6 +528,7 @@ namespace UnrealBuildTool
 			Writer.WriteValue("StatusDescription", Action.StatusDescription);
 			Writer.WriteValue("bCanExecuteRemotely", Action.bCanExecuteRemotely);
 			Writer.WriteValue("bCanExecuteRemotelyWithSNDBS", Action.bCanExecuteRemotelyWithSNDBS);
+			Writer.WriteValue("bCanExecuteRemotelyWithXGE", Action.bCanExecuteRemotelyWithXGE);
 			Writer.WriteValue("bCanCache", Action.bCanCache);
 			Writer.WriteValue("bIsGCCCompiler", Action.bIsGCCCompiler);
 			Writer.WriteValue("bShouldOutputStatusDescription", Action.bShouldOutputStatusDescription);
@@ -630,6 +650,7 @@ namespace UnrealBuildTool
 		public bool bCanCache => Inner.bCanCache;
 		public bool bCanExecuteRemotely => Inner.bCanExecuteRemotely;
 		public bool bCanExecuteRemotelyWithSNDBS => Inner.bCanExecuteRemotelyWithSNDBS;
+		public bool bCanExecuteRemotelyWithXGE => Inner.bCanExecuteRemotelyWithXGE;
 		public bool bIsGCCCompiler => Inner.bIsGCCCompiler;
 		public bool bShouldOutputStatusDescription => Inner.bShouldOutputStatusDescription;
 		public bool bProducesImportLibrary => Inner.bProducesImportLibrary;
diff --git a/Engine/Source/Programs/UnrealBuildTool/ToolChain/ISPCToolChain.cs b/Engine/Source/Programs/UnrealBuildTool/ToolChain/ISPCToolChain.cs
--- a/Engine/Source/Programs/UnrealBuildTool/ToolChain/ISPCToolChain.cs
+++ b/Engine/Source/Programs/UnrealBuildTool/ToolChain/ISPCToolChain.cs
@@ -464,8 +464,13 @@ namespace UnrealBuildTool
 				CompileAction.CommandPath = new FileReference(GetISPCHostCompilerPath(BuildHostPlatform.Current.Platform));
 				CompileAction.StatusDescription = Path.GetFileName(ISPCFile.AbsolutePath);
 
+				CompileAction.bCanExecuteRemotely = true;
+
 				// Disable remote execution to workaround mismatched case on XGE
-				CompileAction.bCanExecuteRemotely = false;
+				CompileAction.bCanExecuteRemotelyWithXGE = false;
+
+				// TODO: Remove, might work
+				CompileAction.bCanExecuteRemotelyWithSNDBS = false;
 
 				List<string> Arguments = new List<string>();
 
@@ -622,8 +627,13 @@ namespace UnrealBuildTool
 				CompileAction.CommandPath = new FileReference(GetISPCHostCompilerPath(BuildHostPlatform.Current.Platform));
 				CompileAction.StatusDescription = Path.GetFileName(ISPCFile.AbsolutePath);
 
+				CompileAction.bCanExecuteRemotely = true;
+
 				// Disable remote execution to workaround mismatched case on XGE
-				CompileAction.bCanExecuteRemotely = false;
+				CompileAction.bCanExecuteRemotelyWithXGE = false;
+
+				// TODO: Remove, might work
+				CompileAction.bCanExecuteRemotelyWithSNDBS = false;
 
 				List<string> Arguments = new List<string>();
 
