Adding ModuleRules.IsPlatformAvailable so build scripts can check platform availability.

diff --git a/Engine/Source/Programs/UnrealBuildTool/Configuration/ModuleRules.cs b/Engine/Source/Programs/UnrealBuildTool/Configuration/ModuleRules.cs
--- a/Engine/Source/Programs/UnrealBuildTool/Configuration/ModuleRules.cs
+++ b/Engine/Source/Programs/UnrealBuildTool/Configuration/ModuleRules.cs
@@ -1682,6 +1682,16 @@ namespace UnrealBuildTool
 			return false;
 		}
 
+		/// <summary>
+		/// Determines whether a given platform is available in the context of the current Target
+		/// </summary>
+		/// <param name="InPlatform">The platform to check for</param>
+		/// <returns>True if it's available, false otherwise</returns>
+		protected bool IsPlatformAvailable(UnrealTargetPlatform InPlatform)
+		{
+			return UEBuildPlatform.IsPlatformAvailableForTarget(InPlatform, this.Target);
+		}
+
 		/// <summary>
 		/// Prepares a module for building a low level tests executable.
 		/// If we're building a module as part of a test module chain and there's a Tests folder with low level tests, then they require the LowLevelTestsRunner dependency.
