fix for compile failure on Android

diff --git a/Engine/Source/Programs/UnrealBuildTool/Platform/Android/AndroidToolChain.cs b/Engine/Source/Programs/UnrealBuildTool/Platform/Android/AndroidToolChain.cs
--- a/Engine/Source/Programs/UnrealBuildTool/Platform/Android/AndroidToolChain.cs
+++ b/Engine/Source/Programs/UnrealBuildTool/Platform/Android/AndroidToolChain.cs
@@ -1050,6 +1050,11 @@ namespace UnrealBuildTool
 			return base.CompileCPPFiles(CompileEnvironment, InputFiles, OutputDir, ModuleName, Graph);
 		}
 
+		protected override List<string> ExpandResponseFileContents(List<string> ResponseFileContents)
+		{
+			return ResponseFileContents.Select(x => Utils.ExpandVariables(x).Replace("\\", "\\\\")).ToList();
+		}
+
 		static public string InlineArchName(string Pathname, UnrealArch Arch, bool bUseShortNames = false)
 		{
 			string FinalArch = "-" + Arch.ToString().ToLower();
