UnrealBuildTool: Fix line that was accidently commented out

#rnx

diff --git a/Engine/Source/Programs/UnrealBuildTool/System/ActionGraph.cs b/Engine/Source/Programs/UnrealBuildTool/System/ActionGraph.cs
--- a/Engine/Source/Programs/UnrealBuildTool/System/ActionGraph.cs
+++ b/Engine/Source/Programs/UnrealBuildTool/System/ActionGraph.cs
@@ -357,7 +357,7 @@ namespace UnrealBuildTool
 #if __BOXEXECUTOR_AVAILABLE__
 			if (BuildConfiguration.bAllowBoxExecutor)
 			{
-//				return new BoxExecutor(BuildConfiguration.MaxParallelActions, BuildConfiguration.bAllCores, BuildConfiguration.bCompactOutput, Logger);
+				return new BoxExecutor(BuildConfiguration.MaxParallelActions, BuildConfiguration.bAllCores, BuildConfiguration.bCompactOutput, Logger);
 			}
 #endif // #if __BOXEXECUTOR_AVAILABLE__
 
