fix PGO utilizing inccorrect directory seperators
#tests local build

diff --git a/Engine/Source/Programs/UnrealBuildTool/Platform/Android/AndroidToolChain.cs b/Engine/Source/Programs/UnrealBuildTool/Platform/Android/AndroidToolChain.cs
--- a/Engine/Source/Programs/UnrealBuildTool/Platform/Android/AndroidToolChain.cs
+++ b/Engine/Source/Programs/UnrealBuildTool/Platform/Android/AndroidToolChain.cs
@@ -536,7 +536,7 @@ namespace UnrealBuildTool
 				// apparently there can be hashing conflicts with PGO which can result in:
 				// 'Function control flow change detected (hash mismatch)' warnings. 
 				Arguments.Add("-Wno-backend-plugin");
-				Arguments.Add(string.Format(" -fprofile-use=\"{0}.profdata\"", Path.Combine(CompileEnvironment.PGODirectory!, CompileEnvironment.PGOFilenamePrefix!)));
+				Arguments.Add(string.Format(" -fprofile-use=\"{0}.profdata\"", Path.Combine(CompileEnvironment.PGODirectory!, CompileEnvironment.PGOFilenamePrefix!).Replace("\\", "/")));
 			}
 			else if (CompileEnvironment.bPGOProfile)
 			{
@@ -783,7 +783,7 @@ namespace UnrealBuildTool
 				Result += " -Wno-profile-instr-out-of-date";
 				Result += " -Wno-profile-instr-unprofiled";
 
-				Result += string.Format(" -fprofile-use=\"{0}.profdata\"", Path.Combine(LinkEnvironment.PGODirectory!, LinkEnvironment.PGOFilenamePrefix!));
+				Result += string.Format(" -fprofile-use=\"{0}.profdata\"", Path.Combine(LinkEnvironment.PGODirectory!, LinkEnvironment.PGOFilenamePrefix!).Replace("\\", "/"));
 			}
 			else if (LinkEnvironment.bPGOProfile)
 			{
