Bump version number so MSVC [REDACTED_PHONE] toolchain doesnt fail at compilation time when ASAN is enabled

#rnx

diff --git a/Engine/Source/Programs/UnrealBuildTool/Platform/Windows/VCToolChain.cs b/Engine/Source/Programs/UnrealBuildTool/Platform/Windows/VCToolChain.cs
--- a/Engine/Source/Programs/UnrealBuildTool/Platform/Windows/VCToolChain.cs
+++ b/Engine/Source/Programs/UnrealBuildTool/Platform/Windows/VCToolChain.cs
@@ -584,7 +584,7 @@ namespace UnrealBuildTool
 					AddDefinition(Arguments, "USING_ADDRESS_SANITISER=1");
 
 					// Re-evalulate the necessity of these workarounds later
-					if (EnvVars.CompilerVersion < new VersionNumber(14, 35))
+					if (EnvVars.CompilerVersion < new VersionNumber(14, 36))
 					{
 						AddDefinition(Arguments, "_DISABLE_STRING_ANNOTATION=1");
 						AddDefinition(Arguments, "_DISABLE_VECTOR_ANNOTATION=1");
