UnrealBuildTool: Supress WIP arch mismatch

#rnx

diff --git a/Engine/Source/Programs/UnrealBuildTool/UnrealBuildTool.csproj b/Engine/Source/Programs/UnrealBuildTool/UnrealBuildTool.csproj
--- a/Engine/Source/Programs/UnrealBuildTool/UnrealBuildTool.csproj
+++ b/Engine/Source/Programs/UnrealBuildTool/UnrealBuildTool.csproj
@@ -84,6 +84,7 @@
       </ItemGroup>
       <PropertyGroup Condition="Exists('..\..\..\Restricted\NotForLicensees\Binaries\Win64\Box\BoxManaged.dll')">
         <DefineConstants>$(DefineConstants);__BOXEXECUTOR_AVAILABLE__</DefineConstants>
+        <ResolveAssemblyWarnOrErrorOnTargetArchitectureMismatch>None</ResolveAssemblyWarnOrErrorOnTargetArchitectureMismatch>
       </PropertyGroup>
     </When>
   </Choose>
