Add a Suffix when compiling with AutoRTFM compiler for Linux builds

diff --git a/Engine/Source/Programs/UnrealBuildTool/Platform/Linux/UEBuildLinux.cs b/Engine/Source/Programs/UnrealBuildTool/Platform/Linux/UEBuildLinux.cs
--- a/Engine/Source/Programs/UnrealBuildTool/Platform/Linux/UEBuildLinux.cs
+++ b/Engine/Source/Programs/UnrealBuildTool/Platform/Linux/UEBuildLinux.cs
@@ -251,6 +251,11 @@ namespace UnrealBuildTool
 			{
 				string? SanitizerSuffix = null;
 
+				if (Target.bUseAutoRTFMCompiler)
+				{
+					SanitizerSuffix = "AutoRTFM";
+				}
+
 				if (Target.LinuxPlatform.bEnableAddressSanitizer)
 				{
 					SanitizerSuffix = "ASan";
