UnrealBuildTool: Add various BuildConfiguration.xml files to the project for convenience

#rnx

diff --git a/Engine/Source/Programs/UnrealBuildTool/UnrealBuildTool.csproj b/Engine/Source/Programs/UnrealBuildTool/UnrealBuildTool.csproj
--- a/Engine/Source/Programs/UnrealBuildTool/UnrealBuildTool.csproj
+++ b/Engine/Source/Programs/UnrealBuildTool/UnrealBuildTool.csproj
@@ -151,6 +151,7 @@
     <PackageReference Include="Newtonsoft.Json" Version="13.0.1" />
   </ItemGroup>
 
+  <!-- Various batch files that build or run UnrealBuildTool -->
   <ItemGroup>
     <None Include="..\..\..\Build\BatchFiles\BuildUBT.bat"><Link>BatchFiles\BuildUBT.bat</Link></None>
     <None Include="..\..\..\Build\BatchFiles\BuildUBT.sh"><Link>BatchFiles\BuildUBT.sh</Link></None>
@@ -158,4 +159,12 @@
     <None Include="..\..\..\Build\BatchFiles\RunUBT.bat"><Link>BatchFiles\RunUBT.bat</Link></None>
     <None Include="..\..\..\Build\BatchFiles\RunUBT.sh"><Link>BatchFiles\RunUBT.sh</Link></None>
   </ItemGroup>
+
+  <!-- Various BuildConfiguration files -->
+  <ItemGroup>
+      <None Condition="Exists('..\..\..\Restricted\NotForLicensees\Programs\UnrealBuildTool\BuildConfiguration.xml')" Include="..\..\..\Restricted\NotForLicensees\Programs\UnrealBuildTool\BuildConfiguration.xml"><Link>BuildConfiguration\BuildConfiguration.xml (NotForLicensees)</Link></None>
+      <None Condition="Exists('..\..\..\Saved\UnrealBuildTool\BuildConfiguration.xml')" Include="..\..\..\Saved\UnrealBuildTool\BuildConfiguration.xml"><Link>BuildConfiguration\BuildConfiguration.xml (Saved)</Link></None>
+      <None Condition="Exists('$(AppData)\Unreal Engine\UnrealBuildTool\BuildConfiguration.xml')" Include="$(AppData)\Unreal Engine\UnrealBuildTool\BuildConfiguration.xml"><Link>BuildConfiguration\BuildConfiguration.xml (User AppData)</Link></None>
+      <None Condition="Exists('$(UserProfile)\Documents\Unreal Engine\UnrealBuildTool\BuildConfiguration.xml')" Include="$(UserProfile)\Documents\Unreal Engine\UnrealBuildTool\BuildConfiguration.xml"><Link>BuildConfiguration\BuildConfiguration.xml (User Documents)</Link></None>
+  </ItemGroup>
 </Project>
