UBT: Unwrap aggregate exceptions thrown by async tools.

diff --git a/Engine/Source/Programs/UnrealBuildTool/UnrealBuildTool.cs b/Engine/Source/Programs/UnrealBuildTool/UnrealBuildTool.cs
--- a/Engine/Source/Programs/UnrealBuildTool/UnrealBuildTool.cs
+++ b/Engine/Source/Programs/UnrealBuildTool/UnrealBuildTool.cs
@@ -653,11 +653,21 @@ namespace UnrealBuildTool
 				ToolMode Mode = (ToolMode)Activator.CreateInstance(ModeType)!;
 
 				// Execute the mode
-				int Result = Task.Run(() => Mode.ExecuteAsync(Arguments, Logger)).Result;
-				if((ModeOptions & ToolModeOptions.ShowExecutionTime) != 0)
+				int Result;
+				try
+				{
+					Result = Task.Run(() => Mode.ExecuteAsync(Arguments, Logger)).Result;
+				}
+				catch (AggregateException AggEx) when (AggEx.InnerExceptions.Count == 1)
+				{
+					throw AggEx.InnerExceptions[0];
+				}
+
+				if ((ModeOptions & ToolModeOptions.ShowExecutionTime) != 0)
 				{
 					Logger.LogInformation("Total execution time: {Time:0.00} seconds", Timeline.Elapsed.TotalSeconds);
 				}
+
 				return Result;
 			}
 			catch (CompilationResultException Ex)
