UnrealBuildTool: Change some LogDebug to LogTrace for very verbose messages

#rnx

diff --git a/Engine/Source/Programs/UnrealBuildTool/Configuration/UEBuildModuleCPP.cs b/Engine/Source/Programs/UnrealBuildTool/Configuration/UEBuildModuleCPP.cs
--- a/Engine/Source/Programs/UnrealBuildTool/Configuration/UEBuildModuleCPP.cs
+++ b/Engine/Source/Programs/UnrealBuildTool/Configuration/UEBuildModuleCPP.cs
@@ -624,28 +624,28 @@ namespace UnrealBuildTool
 
 				if (Target.bForceUnityBuild)
 				{
-					Logger.LogDebug("Module '{ModuleName}' using unity build mode (bForceUnityBuild enabled for this module)", this.Name);
+					Logger.LogTrace("Module '{ModuleName}' using unity build mode (bForceUnityBuild enabled for this module)", this.Name);
 					bModuleUsesUnityBuild = true;
 				}
 				else if (!Rules.bUseUnity)
 				{
-					Logger.LogDebug("Module '{ModuleName}' not using unity build mode (bUseUnity disabled for this module)", this.Name);
+					Logger.LogTrace("Module '{ModuleName}' not using unity build mode (bUseUnity disabled for this module)", this.Name);
 					bModuleUsesUnityBuild = false;
 				}
 				else if (FileCount < MinSourceFilesForUnityBuild)
 				{
-					Logger.LogDebug("Module '{ModuleName}' not using unity build mode (module with fewer than {NumFiles} source files)", this.Name, MinSourceFilesForUnityBuild);
+					Logger.LogTrace("Module '{ModuleName}' not using unity build mode (module with fewer than {NumFiles} source files)", this.Name, MinSourceFilesForUnityBuild);
 					bModuleUsesUnityBuild = false;
 				}
 				else
 				{
-					Logger.LogDebug("Module '{ModuleName}' using unity build mode", this.Name);
+					Logger.LogTrace("Module '{ModuleName}' using unity build mode", this.Name);
 					bModuleUsesUnityBuild = true;
 				}
 			}
 			else
 			{
-				Logger.LogDebug("Module '{ModuleName}' not using unity build mode", this.Name);
+				Logger.LogTrace("Module '{ModuleName}' not using unity build mode", this.Name);
 			}
 
 			// Compile Generated CPP Files
diff --git a/Engine/Source/Programs/UnrealBuildTool/Configuration/UEBuildTarget.cs b/Engine/Source/Programs/UnrealBuildTool/Configuration/UEBuildTarget.cs
--- a/Engine/Source/Programs/UnrealBuildTool/Configuration/UEBuildTarget.cs
+++ b/Engine/Source/Programs/UnrealBuildTool/Configuration/UEBuildTarget.cs
@@ -4190,26 +4190,26 @@ namespace UnrealBuildTool
 				// Check if the plugin is required for this platform
 				if (!Reference.IsEnabledForPlatform(Platform))
 				{
-					Logger.LogDebug("Ignoring plugin '{ReferenceName}' (referenced via {ReferenceChain}), not enabled for platform {Platform}", Reference.Name, ReferenceChain, Platform);
+					Logger.LogTrace("Ignoring plugin '{ReferenceName}' (referenced via {ReferenceChain}), not enabled for platform {Platform}", Reference.Name, ReferenceChain, Platform);
 					return null;
 				}
 
 				if (!Reference.IsEnabledForTargetConfiguration(Configuration))
 				{
-					Logger.LogDebug("Ignoring plugin '{ReferenceName}' (referenced via {ReferenceChain}), not enabled for target configuration {Configuration}", Reference.Name, ReferenceChain, Configuration);
+					Logger.LogTrace("Ignoring plugin '{ReferenceName}' (referenced via {ReferenceChain}), not enabled for target configuration {Configuration}", Reference.Name, ReferenceChain, Configuration);
 					return null;
 				}
 
 				if (!Reference.IsEnabledForTarget(TargetType))
 				{
-					Logger.LogDebug("Ignoring plugin '{ReferenceName}' (referenced via {ReferenceChain}), not enabled for target {TargetType}", Reference.Name, ReferenceChain, TargetType);
+					Logger.LogTrace("Ignoring plugin '{ReferenceName}' (referenced via {ReferenceChain}), not enabled for target {TargetType}", Reference.Name, ReferenceChain, TargetType);
 					return null;
 				}
 
 				// Disable any plugin reference which does not support the target platform
 				if (!Rules.bIncludePluginsForTargetPlatforms && !Reference.IsSupportedTargetPlatform(Platform))
 				{
-					Logger.LogDebug("Ignoring plugin '{ReferenceName}' (referenced via {ReferenceChain}) due to unsupported target platform.", Reference.Name, ReferenceChain);
+					Logger.LogTrace("Ignoring plugin '{ReferenceName}' (referenced via {ReferenceChain}) due to unsupported target platform.", Reference.Name, ReferenceChain);
 					return null;
 				}
 
@@ -4260,19 +4260,19 @@ namespace UnrealBuildTool
 				// Disable any plugin that requires the build platform
 				if (Info.Descriptor.bRequiresBuildPlatform && ShouldExcludePlugin(Info, ExcludeFolders))
 				{
-					Logger.LogDebug("Ignoring plugin '{ReferenceName}' (referenced via {ReferenceChain}) due to missing build platform", Reference.Name, ReferenceChain);
+					Logger.LogTrace("Ignoring plugin '{ReferenceName}' (referenced via {ReferenceChain}) due to missing build platform", Reference.Name, ReferenceChain);
 					return null;
 				}
 
 				// Disable any plugins that aren't compatible with this program
 				if (Rules.Type == TargetType.Program && (Info.Descriptor.SupportedPrograms == null || !Info.Descriptor.SupportedPrograms.Contains(AppName)))
 				{
-					Logger.LogDebug("Ignoring plugin '{ReferenceName}' (referenced via {ReferenceChain}) due to absence from supported programs list.", Reference.Name, ReferenceChain);
+					Logger.LogTrace("Ignoring plugin '{ReferenceName}' (referenced via {ReferenceChain}) due to absence from supported programs list.", Reference.Name, ReferenceChain);
 					return null;
 				}
 
 				// Create the new instance and add it to the cache
-				Logger.LogDebug("Enabling plugin '{ReferenceName}' (referenced via {ReferenceChain})", Reference.Name, ReferenceChain);
+				Logger.LogTrace("Enabling plugin '{ReferenceName}' (referenced via {ReferenceChain})", Reference.Name, ReferenceChain);
 				Instance = new UEBuildPlugin(Info, ReferenceChain);
 				Instance.bDescriptorReferencedExplicitly = !Reference.bOptional;
 				NameToInstance.Add(Info.Name, Instance);
