[UBT - BoxExecutor]
* Does now support running all actions that can run remotely (before it was only cl.exe/clang++.exe)
* Added visualizer (use with -Visualizer)

diff --git a/Engine/Source/Programs/UnrealBuildTool/UnrealBuildTool.csproj b/Engine/Source/Programs/UnrealBuildTool/UnrealBuildTool.csproj
--- a/Engine/Source/Programs/UnrealBuildTool/UnrealBuildTool.csproj
+++ b/Engine/Source/Programs/UnrealBuildTool/UnrealBuildTool.csproj
@@ -78,6 +78,9 @@
         <Content Include="..\..\..\Restricted\NotForLicensees\Binaries\Win64\Box\BoxDetours.dll">
           <CopyToOutputDirectory>PreserveNewest</CopyToOutputDirectory>
         </Content>
+        <Content Include="..\..\..\Restricted\NotForLicensees\Binaries\Win64\Box\BoxVisualizer.exe">
+          <CopyToOutputDirectory>PreserveNewest</CopyToOutputDirectory>
+        </Content>
         <Content Include="..\..\..\Restricted\NotForLicensees\Binaries\Win64\Box\Ijwhost.dll">
           <CopyToOutputDirectory>PreserveNewest</CopyToOutputDirectory>
         </Content>
