UnrealBuildTool: SetDefaultsForCookedEditor requires a unique build environment

#rnx

diff --git a/Engine/Source/Programs/UnrealBuildTool/Configuration/TargetRules.cs b/Engine/Source/Programs/UnrealBuildTool/Configuration/TargetRules.cs
--- a/Engine/Source/Programs/UnrealBuildTool/Configuration/TargetRules.cs
+++ b/Engine/Source/Programs/UnrealBuildTool/Configuration/TargetRules.cs
@@ -2526,6 +2526,7 @@ namespace UnrealBuildTool
 		protected void SetDefaultsForCookedEditor(bool bIsCookedCooker, bool bIsForExternalUse)
 		{
 			LinkType = TargetLinkType.Monolithic;
+			BuildEnvironment = TargetBuildEnvironment.Unique;
 
 			if (!bIsCookedCooker)
 			{
