- Workaround to get engine source code in IOS projects (force the UnrealEditor project in the worksapce, but as a "stub" project)

diff --git a/Engine/Source/Programs/UnrealBuildTool/Platform/IOS/IOSToolChain.cs b/Engine/Source/Programs/UnrealBuildTool/Platform/IOS/IOSToolChain.cs
--- a/Engine/Source/Programs/UnrealBuildTool/Platform/IOS/IOSToolChain.cs
+++ b/Engine/Source/Programs/UnrealBuildTool/Platform/IOS/IOSToolChain.cs
@@ -1518,7 +1518,7 @@ namespace UnrealBuildTool
 
 			if (UseModernXcode(Target.ProjectFile))
 			{
-				throw new BuildException("Do we get here?");
+				throw new BuildException("Modern xcode mode should not reach this, has broken the PostBuild setup");
 
 			}
 
diff --git a/Engine/Source/Programs/UnrealBuildTool/ProjectFiles/Xcode/XcodeProject.cs b/Engine/Source/Programs/UnrealBuildTool/ProjectFiles/Xcode/XcodeProject.cs
--- a/Engine/Source/Programs/UnrealBuildTool/ProjectFiles/Xcode/XcodeProject.cs
+++ b/Engine/Source/Programs/UnrealBuildTool/ProjectFiles/Xcode/XcodeProject.cs
@@ -1819,7 +1819,7 @@ namespace UnrealBuildTool.XcodeProjectXcconfig
 
 		public XcodeBuildConfigList ProjectBuildConfigs;
 
-		public XcodeProject(UnrealTargetPlatform? Platform, UnrealData UnrealData, XcodeFileCollection FileCollection, XcodeProjectFile ProjectFile, ILogger Logger)
+		public XcodeProject(UnrealTargetPlatform? Platform, UnrealData UnrealData, XcodeFileCollection FileCollection, XcodeProjectFile ProjectFile, bool bIsStubEditor, ILogger Logger)
 		{
 			this.Platform = Platform;
 			this.UnrealData = UnrealData;
@@ -1829,21 +1829,24 @@ namespace UnrealBuildTool.XcodeProjectXcconfig
 
 			// if we are run-only, then we don't need a build target (this is shared between platforms if we are doing multi-target)
 			XcodeBuildTarget? BuildTarget = null;
-			if (!XcodeProjectFileGenerator.bGenerateRunOnlyProject && !UnrealData.bIsStubProject)
+			if (!XcodeProjectFileGenerator.bGenerateRunOnlyProject && !UnrealData.bIsStubProject && !bIsStubEditor)
 			{
 				BuildTarget = new XcodeBuildTarget(this);
 			}
 
-			// create one run target for each platform if our platform is null (ie XcodeProjectGenerator.PerPlatformMode is RunTargetPerPlatform)
-			List<UnrealTargetPlatform>					 TargetPlatforms = Platform == null ? XcodeProjectFileGenerator.XcodePlatforms : new() { Platform.Value };
-			foreach (UnrealTargetPlatform TargetPlatform in TargetPlatforms)
+			if (!bIsStubEditor)
 			{
-				XcodeRunTarget RunTarget = new XcodeRunTarget(this, UnrealData.ProductName, UnrealData.AllConfigs[0].ProjectTarget!.TargetRules!.Type, TargetPlatform, BuildTarget, ProjectFile, Logger);
-				RunTargets.Add(RunTarget);
-				References.Add(RunTarget);
+				// create one run target for each platform if our platform is null (ie XcodeProjectGenerator.PerPlatformMode is RunTargetPerPlatform)
+				List<UnrealTargetPlatform> TargetPlatforms = Platform == null ? XcodeProjectFileGenerator.XcodePlatforms : new() { Platform.Value };
+				foreach (UnrealTargetPlatform TargetPlatform in TargetPlatforms)
+				{
+					XcodeRunTarget RunTarget = new XcodeRunTarget(this, UnrealData.ProductName, UnrealData.AllConfigs[0].ProjectTarget!.TargetRules!.Type, TargetPlatform, BuildTarget, ProjectFile, Logger);
+					RunTargets.Add(RunTarget);
+					References.Add(RunTarget);
+				}
 			}
 
-			ProjectBuildConfigs = new XcodeBuildConfigList(Platform, UnrealData.ProductName, UnrealData, bIncludeAllPlatformsInConfig: true);
+			ProjectBuildConfigs = new XcodeBuildConfigList(bIsStubEditor ? null : Platform, UnrealData.ProductName, UnrealData, bIncludeAllPlatformsInConfig: true);
 			References.Add(ProjectBuildConfigs);
 			
 			// make an indexing target if we aren't just a run-only project, and it has buildable source files
@@ -1881,7 +1884,11 @@ namespace UnrealBuildTool.XcodeProjectXcconfig
 			Content.WriteLine(4,		"en");
 			Content.WriteLine(3,	");");
 			Content.WriteLine(3,	$"mainGroup = {FileCollection.MainGroupGuid};");
-			Content.WriteLine(3,	$"productRefGroup = {FileCollection.GetProductGroupGuid()};");
+			// for stub editor projects, we don't have a run target, so we don't have a product folder
+			if (RunTargets.Count > 0)
+			{
+				Content.WriteLine(3, $"productRefGroup = {FileCollection.GetProductGroupGuid()};");
+			}
 			Content.WriteLine(3,	"projectDirPath = \"\";");
 			Content.WriteLine(3,	"projectRoot = \"\";");
 			Content.WriteLine(3,	"targets = (");
@@ -2086,8 +2093,9 @@ namespace UnrealBuildTool.XcodeProjectXcconfig
 			bool bSuccess = true;
 			foreach (UnrealTargetPlatform? Platform in XcodeProjectFileGenerator.WorkspacePlatforms)
 			{
+				bool bAddStubEditor = Platform != UnrealTargetPlatform.Mac && UnrealData.ProductName == "UnrealEditor";
 				// skip the platform if the project has no configurations for it
-				if (!UnrealData.AllConfigs.Any(x => x.Supports(Platform)))
+				if (!bAddStubEditor && !UnrealData.AllConfigs.Any(x => x.Supports(Platform)))
 				{
 					continue;
 				}
@@ -2095,7 +2103,7 @@ namespace UnrealBuildTool.XcodeProjectXcconfig
 
 				// now create the xcodeproject elements (project -> target -> buildconfigs, etc)
 				FileCollection = new XcodeFileCollection(SharedFileCollection);
-				XcodeProject RootProject = new XcodeProject(Platform, UnrealData, FileCollection, this, Logger);
+				XcodeProject RootProject = new XcodeProject(Platform, UnrealData, FileCollection, this, bAddStubEditor, Logger);
 				RootProjects[RootProject] = Platform;
 
 				if (FileReference.Exists(TemplateProject))
@@ -2134,7 +2142,7 @@ namespace UnrealBuildTool.XcodeProjectXcconfig
 					bSuccess = ProjectFileGenerator.WriteFileIfChanged(PBXFilePath.FullName, Content.ToString(), Logger, new UTF8Encoding()) && bSuccess;
 				}
 
-				bool bNeedScheme = XcodeUtils.ShouldIncludeProjectInWorkspace(this, Logger);
+				bool bNeedScheme = !bAddStubEditor && XcodeUtils.ShouldIncludeProjectInWorkspace(this, Logger);
 				if (bNeedScheme)
 				{
 					if (bSuccess)
