UBT: Prevent stack trace being overwritten when re-throwing exceptions in tool modes.

diff --git a/Engine/Source/Programs/UnrealBuildTool/UnrealBuildTool.cs b/Engine/Source/Programs/UnrealBuildTool/UnrealBuildTool.cs
--- a/Engine/Source/Programs/UnrealBuildTool/UnrealBuildTool.cs
+++ b/Engine/Source/Programs/UnrealBuildTool/UnrealBuildTool.cs
@@ -656,7 +656,7 @@ namespace UnrealBuildTool
 				int Result;
 				try
 				{
-					Result = Task.Run(() => Mode.ExecuteAsync(Arguments, Logger)).Result;
+					Result = Mode.ExecuteAsync(Arguments, Logger).GetAwaiter().GetResult();
 				}
 				catch (AggregateException AggEx) when (AggEx.InnerExceptions.Count == 1)
 				{
