- Restored some code that got lost in a refactor

diff --git a/Engine/Source/Programs/AutomationTool/Apple/ApplePlatform.Automation.cs b/Engine/Source/Programs/AutomationTool/Apple/ApplePlatform.Automation.cs
--- a/Engine/Source/Programs/AutomationTool/Apple/ApplePlatform.Automation.cs
+++ b/Engine/Source/Programs/AutomationTool/Apple/ApplePlatform.Automation.cs
@@ -76,7 +76,9 @@ public abstract class ApplePlatform : Platform
 		{
 			foreach (TargetReceipt Target in SC.StageTargets.Select(x => x.Receipt))
 			{
-				MacExports.BuildWithModernXcode(SC.RawProjectPath, Target.Platform, Target.Configuration, Target.TargetName, bArchiveForDistro: false, Logger);
+				// fiddle with some envvars to redirect the .app into root of Staging directory, without going under any build subdirectories
+				string ExtraOptions = $"SYMROOT=\"{SC.StageDirectory.ParentDirectory}\" EFFECTIVE_PLATFORM_NAME={SC.StageDirectory.GetDirectoryName()}";
+				MacExports.BuildWithModernXcode(SC.RawProjectPath, Target.Platform, Target.Configuration, Target.TargetName, bArchiveForDistro: false, Logger, ExtraOptions);
 			}
 		}
 	}
diff --git a/Engine/Source/Programs/UnrealBuildTool/Platform/Mac/MacExports.cs b/Engine/Source/Programs/UnrealBuildTool/Platform/Mac/MacExports.cs
--- a/Engine/Source/Programs/UnrealBuildTool/Platform/Mac/MacExports.cs
+++ b/Engine/Source/Programs/UnrealBuildTool/Platform/Mac/MacExports.cs
@@ -94,15 +94,16 @@ namespace UnrealBuildTool
 		/// <summary>
 		/// Generates a stub xcode project for the given project/platform/target combo, then builds or archives it
 		/// </summary>
-		/// <param name="ProjectFile"></param>
-		/// <param name="Platform"></param>
-		/// <param name="Configuration"></param>
-		/// <param name="TargetName"></param>
-		/// <param name="bArchiveForDistro"></param>
+		/// <param name="ProjectFile">Project to build</param>
+		/// <param name="Platform">Platform to build</param>
+		/// <param name="Configuration">Configuration to build</param>
+		/// <param name="TargetName">Target to build</param>
+		/// <param name="bArchiveForDistro">Whether or not to create a .xcarchive to prepare for distribution</param>
 		/// <param name="Logger"></param>
+		/// <param name="ExtraOptions">Any extra options to pass to xcodebuild</param>
 		/// <returns></returns>
 		public static bool BuildWithModernXcode(FileReference? ProjectFile, UnrealTargetPlatform Platform, UnrealTargetConfiguration Configuration, 
-			string TargetName, bool bArchiveForDistro, ILogger Logger)
+			string TargetName, bool bArchiveForDistro, ILogger Logger, string ExtraOptions="")
 		{
 			DirectoryReference? GeneratedProjectFile;
 			// we don't use distro flag when making a modern project
@@ -115,7 +116,7 @@ namespace UnrealBuildTool
 
 			// run xcodebuild on the generated project to make the .app
 			IOSExports.FinalizeAppWithModernXcode(GeneratedProjectFile!, Platform, TargetName, Configuration.ToString(),
-				bArchiveForDistro ? "archive" : "build", "", bArchiveForDistro, Logger);
+				bArchiveForDistro ? "archive" : "build", ExtraOptions, bArchiveForDistro, Logger);
 
 			return true;
 		}
