UnrealBuildTool: Correct naming on PCH variants for CppStandard

#rnx

diff --git a/Engine/Source/Programs/UnrealBuildTool/Configuration/UEBuildModuleCPP.cs b/Engine/Source/Programs/UnrealBuildTool/Configuration/UEBuildModuleCPP.cs
--- a/Engine/Source/Programs/UnrealBuildTool/Configuration/UEBuildModuleCPP.cs
+++ b/Engine/Source/Programs/UnrealBuildTool/Configuration/UEBuildModuleCPP.cs
@@ -1167,12 +1167,26 @@ namespace UnrealBuildTool
 
 			if (CompileEnvironment.CppStandard != BaseCompileEnvironment.CppStandard)
 			{
-				Variant += String.Format(".{0}", CompileEnvironment.CppStandard);
+				switch (CompileEnvironment.CppStandard)
+				{
+					case CppStandardVersion.Cpp14: Variant += ".Cpp14"; break;
+					case CppStandardVersion.Cpp17: Variant += ".Cpp17"; break;
+					case CppStandardVersion.Cpp20: Variant += ".Cpp20"; break;
+					case CppStandardVersion.Latest: Variant += ".CppLatest"; break;
+				}
 			}
 
 			if (CompileEnvironment.CStandard != BaseCompileEnvironment.CStandard)
 			{
-				Variant += String.Format(".{0}", CompileEnvironment.CStandard);
+				switch (CompileEnvironment.CStandard)
+				{
+					case CStandardVersion.None: Variant += ".CNone"; break;
+					case CStandardVersion.C89: Variant += ".C89"; break;
+					case CStandardVersion.C99: Variant += ".C99"; break;
+					case CStandardVersion.C11: Variant += ".C11"; break;
+					case CStandardVersion.C17: Variant += ".C17"; break;
+					case CStandardVersion.Latest: Variant += ".CLatest"; break;
+				}
 			}
 
 			if (CompileEnvironment.IncludeOrderVersion != BaseCompileEnvironment.IncludeOrderVersion)
