UnrealBuildTool: Ban older VS2022 compilers and update Intel preferred version for 5.3

#rnx

diff --git a/Engine/Source/Programs/UnrealBuildTool/Platform/Windows/MicrosoftPlatformSDK.Versions.cs b/Engine/Source/Programs/UnrealBuildTool/Platform/Windows/MicrosoftPlatformSDK.Versions.cs
--- a/Engine/Source/Programs/UnrealBuildTool/Platform/Windows/MicrosoftPlatformSDK.Versions.cs
+++ b/Engine/Source/Programs/UnrealBuildTool/Platform/Windows/MicrosoftPlatformSDK.Versions.cs
@@ -54,6 +54,7 @@ namespace UnrealBuildTool
 		/// </summary>
 		static readonly VersionNumberRange[] BannedVisualCppVersions = new VersionNumberRange[]
 		{
+			VersionNumberRange.Parse("14.30.0", "14.33.99999"), // VS2022 17.0.x - 17.3.x
 		};
 
 		static readonly VersionNumber MinimumVisualCppVersion = new VersionNumber(14, 29, 30133);
@@ -64,10 +65,10 @@ namespace UnrealBuildTool
 		/// </summary>
 		static readonly VersionNumberRange[] PreferredIntelOneApiVersions =
 		{
-			VersionNumberRange.Parse("2022.2.0", "2022.9999"),
+			VersionNumberRange.Parse("2023.1.0", "2023.9999"),
 		};
 
-		static readonly VersionNumber MinimumIntelOneApiVersion = new VersionNumber(2022, 2, 0);
+		static readonly VersionNumber MinimumIntelOneApiVersion = new VersionNumber(2023, 0, 0);
 
 		/// <inheritdoc/>
 		public override string GetMainVersion()
diff --git a/Engine/Source/Programs/UnrealBuildTool/Platform/Windows/VCToolChain.cs b/Engine/Source/Programs/UnrealBuildTool/Platform/Windows/VCToolChain.cs
--- a/Engine/Source/Programs/UnrealBuildTool/Platform/Windows/VCToolChain.cs
+++ b/Engine/Source/Programs/UnrealBuildTool/Platform/Windows/VCToolChain.cs
@@ -1008,6 +1008,11 @@ namespace UnrealBuildTool
 				// Additional disabled warnings for msvc. Everything below should be checked if it is necessary
 				ClangWarnings.GetVCDisabledWarnings(Arguments);
 			}
+
+			if (Target.WindowsPlatform.Compiler == WindowsCompiler.Intel)
+			{
+				ClangWarnings.GetIntelDisabledWarnings(Arguments);
+			}
 		}
 
 		static void AppendCLArguments_C(CppCompileEnvironment CompileEnvironment, List<string> Arguments)
diff --git a/Engine/Source/Programs/UnrealBuildTool/ToolChain/ClangWarnings.cs b/Engine/Source/Programs/UnrealBuildTool/ToolChain/ClangWarnings.cs
--- a/Engine/Source/Programs/UnrealBuildTool/ToolChain/ClangWarnings.cs
+++ b/Engine/Source/Programs/UnrealBuildTool/ToolChain/ClangWarnings.cs
@@ -184,5 +184,17 @@ namespace UnrealBuildTool
 			Arguments.Add("-Wno-unused-value");
 			Arguments.Add("-Wno-bitfield-enum-conversion");
 		}
+
+		// Additional disabled warnings for Intel. Everything below should be checked if it is necessary
+		internal static void GetIntelDisabledWarnings(List<string> Arguments)
+		{
+			Arguments.Add("-Wno-deprecated-copy");
+			Arguments.Add("-Wno-deprecated-copy-with-user-provided-copy");
+			Arguments.Add("-Wno-enum-constexpr-conversion");
+			Arguments.Add("-Wno-format");
+			Arguments.Add("-Wno-implicit-float-size-conversion");
+			Arguments.Add("-Wno-null-pointer-subtraction");
+			Arguments.Add("-Wno-single-bit-bitfield-constant-conversion");
+		}
 	}
 }
