UBT: Disable Box dependency in UBT on ARM64.

diff --git a/Engine/Source/Programs/UnrealBuildTool/UnrealBuildTool.csproj b/Engine/Source/Programs/UnrealBuildTool/UnrealBuildTool.csproj
--- a/Engine/Source/Programs/UnrealBuildTool/UnrealBuildTool.csproj
+++ b/Engine/Source/Programs/UnrealBuildTool/UnrealBuildTool.csproj
@@ -72,7 +72,7 @@
   </Choose>
 
   <Choose>
-    <When Condition="'$([System.Runtime.InteropServices.RuntimeInformation]::IsOSPlatform($([System.Runtime.InteropServices.OSPlatform]::Windows)))' == 'true'">
+    <When Condition="'$([System.Runtime.InteropServices.RuntimeInformation]::IsOSPlatform($([System.Runtime.InteropServices.OSPlatform]::Windows)))' == 'true' And '$([System.Runtime.InteropServices.RuntimeInformation]::Architecture' != 'AMD64'">
       <ItemGroup Condition="Exists('..\..\..\Restricted\NotForLicensees\Binaries\Win64\Box\BoxManaged.dll')">
         <Reference Include="..\..\..\Restricted\NotForLicensees\Binaries\Win64\Box\BoxManaged.dll" />
         <Content Include="..\..\..\Restricted\NotForLicensees\Binaries\Win64\Box\BoxDetours.dll">
