UBT: Fix ARM64 check in UBT project file.

diff --git a/Engine/Source/Programs/UnrealBuildTool/UnrealBuildTool.csproj b/Engine/Source/Programs/UnrealBuildTool/UnrealBuildTool.csproj
--- a/Engine/Source/Programs/UnrealBuildTool/UnrealBuildTool.csproj
+++ b/Engine/Source/Programs/UnrealBuildTool/UnrealBuildTool.csproj
@@ -72,7 +72,7 @@
   </Choose>
 
   <Choose>
-    <When Condition="'$([System.Runtime.InteropServices.RuntimeInformation]::IsOSPlatform($([System.Runtime.InteropServices.OSPlatform]::Windows)))' == 'true' And '$([System.Runtime.InteropServices.RuntimeInformation]::Architecture' != 'AMD64'">
+    <When Condition="'$([System.Runtime.InteropServices.RuntimeInformation]::IsOSPlatform($([System.Runtime.InteropServices.OSPlatform]::Windows)))' == 'true' And '$([System.Runtime.InteropServices.RuntimeInformation]::OSArchitecture)' != 'Arm64'">
       <ItemGroup Condition="Exists('..\..\..\Restricted\NotForLicensees\Binaries\Win64\Box\BoxManaged.dll')">
         <Reference Include="..\..\..\Restricted\NotForLicensees\Binaries\Win64\Box\BoxManaged.dll" />
         <Content Include="..\..\..\Restricted\NotForLicensees\Binaries\Win64\Box\BoxDetours.dll">
