Updating code analysis to 4.2 and text encoding to 6.0.

diff --git a/Engine/Source/Programs/UnrealBuildTool/UnrealBuildTool.csproj b/Engine/Source/Programs/UnrealBuildTool/UnrealBuildTool.csproj
--- a/Engine/Source/Programs/UnrealBuildTool/UnrealBuildTool.csproj
+++ b/Engine/Source/Programs/UnrealBuildTool/UnrealBuildTool.csproj
@@ -1,4 +1,4 @@
-<Project Sdk="Microsoft.NET.Sdk">
+﻿<Project Sdk="Microsoft.NET.Sdk">
   <Import Project="..\Shared\UnrealEngine.csproj.props" />
   
   <PropertyGroup>
