UnrealBuildTool: Update executor tasks to be long running

#rnx

diff --git a/Engine/Source/Programs/UnrealBuildTool/Executors/ParallelExecutor.cs b/Engine/Source/Programs/UnrealBuildTool/Executors/ParallelExecutor.cs
--- a/Engine/Source/Programs/UnrealBuildTool/Executors/ParallelExecutor.cs
+++ b/Engine/Source/Programs/UnrealBuildTool/Executors/ParallelExecutor.cs
@@ -303,7 +303,7 @@ namespace UnrealBuildTool
 
 					// Start new action and add it to the ActiveTasks
 					var ActionToRun = Actions[ActionToRunIndex];
-					Task<ExecuteResults> NewTask = Task.Run(() => RunAction(ActionToRun, ProcessGroup, CancellationToken));
+					Task<ExecuteResults> NewTask = Task.Factory.StartNew(() => RunAction(ActionToRun, ProcessGroup, CancellationToken), CancellationToken, TaskCreationOptions.LongRunning | TaskCreationOptions.PreferFairness, TaskScheduler.Default).Unwrap();
 					ActionToTaskLookup.Add(ActionToRun, NewTask);
 
 					ActiveTasks[ActiveTaskCount] = NewTask;
