- Error out when building Mac/IOS modern .app if no cooked&staged data is found (for non-editor, non-program projects)
- Fixed an issue where it would skip running the PostBuild stuff if it fails during the pulling in staged directory step

diff --git a/Engine/Source/Programs/UnrealBuildTool/ProjectFiles/Xcode/XcodeProject.cs b/Engine/Source/Programs/UnrealBuildTool/ProjectFiles/Xcode/XcodeProject.cs
--- a/Engine/Source/Programs/UnrealBuildTool/ProjectFiles/Xcode/XcodeProject.cs
+++ b/Engine/Source/Programs/UnrealBuildTool/ProjectFiles/Xcode/XcodeProject.cs
@@ -1253,9 +1253,15 @@
 
 			CopyScript.AddRange(new string[]
 			{
+				"# We run this often during ApplePostBuildSync, which looks at the PkgInfo file in the .app to determine if this is up to date,",
+				"# so delete it in case we have an error, we will re-run PostBuildSync next build",
+				"echo rm ${CONFIGURATION_BUILD_DIR}/${CONTENTS_FOLDER_PATH}/PkgInfo",
+				"rm ${CONFIGURATION_BUILD_DIR}/${CONTENTS_FOLDER_PATH}/PkgInfo",
+				"",
 				"# Copy the executable into .app if needed (-ef checks if two files/paths are equivalent)",
 				$"SRC_EXE=\\\"${{UE_BINARIES_DIR}}/${{{SourceEnvVar}}}\\\"",
 				$"DEST_EXE=\\\"${{CONFIGURATION_BUILD_DIR}}/${{EXECUTABLE_PATH}}\\\"",
+				"",
 				"echo ${SRC_EXE} /// ${DEST_EXE}",
 				"if [[ ! ${SRC_EXE} -ef ${DEST_EXE} ]]; then",
 				"  echo Copying executable...",
@@ -1272,8 +1278,8 @@
 				TargetPlatformName += TargetType.ToString();
 			}
 
-			// editor builds don't need staged content in them
-			if (TargetType != TargetType.Editor)
+			// editor and program builds don't need staged content in them
+			if (TargetType != TargetType.Editor && TargetType != TargetType.Program)
 			{
 				bool bIsEngineBuild = Project.UnrealData.UProjectFileLocation == null;
 				string DefaultStageDir = bIsEngineBuild ? "${UE_OVERRIDE_STAGE_DIR}" : $"${{UE_PROJECT_DIR}}/Saved/StagedBuilds/{TargetPlatformName}";
@@ -1287,7 +1293,15 @@
 					"# When building engine projects, like UnrealGame, we don't have data to stage unless something has specified UE_OVERRIDE_STAGE_DIR",
 					$"STAGED_DIR={DefaultStageDir}",
 					"if [[ -z ${STAGED_DIR} ]]; then exit 0; fi",
-					"if [[ ! -e ${STAGED_DIR} ]]; then exit 0; fi",
+					"# Make sure the staged directory exists and has files in it",
+					"if [[ ! -e ${STAGED_DIR} || `ls ${STAGED_DIR} | wc -l` -eq 0 ]]; then ",
+					"  echo =========================================================================================",
+					"  echo \\\"ERROR: You must have a valid staged build directory\\\"",
+					"  echo \\\"Use the editor's Platforms menu, or run:\\\"",
+					$"  echo \\\"./RunUAT.sh BuildCookRun -platform={Platform} -project=<project> -build -cook -stage -pak\\\"",
+					"  echo =========================================================================================",
+					"  exit -1 ",
+					"fi",
 					"",
 					"echo \\\"Syncing ${STAGED_DIR} to ${CONFIGURATION_BUILD_DIR}/${CONTENTS_FOLDER_PATH}\\\"",
 					$"rsync -a --exclude=Info.plist --exclude=/Manifest_* --exclude=*.app \\\"${{STAGED_DIR}}/\\\" \\\"${{CONFIGURATION_BUILD_DIR}}/${{CONTENTS_FOLDER_PATH}}{SyncSubdir}\\\"",
@@ -1307,7 +1321,7 @@
 
 			// run this script every time, but xcode will show a warning if there isn't _some_ output
 			string ScriptOutput = $"/dev/null";
-			XcodeShellScriptBuildPhase CopyScriptPhase = new("Copy UE Executable into .app", CopyScript, new string[] { }, new string[] { ScriptOutput });
+			XcodeShellScriptBuildPhase CopyScriptPhase = new("Copy Executable and Staged Data into .app", CopyScript, new string[] { }, new string[] { ScriptOutput });
 			BuildPhases.Add(CopyScriptPhase);
 			References.Add(CopyScriptPhase);
 
