Updated the PCH performance issue logic to take into account the local build config xml when throwing warnings/errors.

diff --git a/Engine/Source/Programs/UnrealBuildTool/Configuration/UEBuildTarget.cs b/Engine/Source/Programs/UnrealBuildTool/Configuration/UEBuildTarget.cs
--- a/Engine/Source/Programs/UnrealBuildTool/Configuration/UEBuildTarget.cs
+++ b/Engine/Source/Programs/UnrealBuildTool/Configuration/UEBuildTarget.cs
@@ -3537,7 +3537,7 @@ namespace UnrealBuildTool
 						{
 							foreach (var Module in Instance.Modules)
 							{
-								if (Module.RulesFile.ToFileInfo().IsReadOnly) // support for local module changes
+								if (Module.RulesFile.ToFileInfo().IsReadOnly && !(Target.DisableOptimizeCodeForModules?.Contains(Module.Name) ?? false)) // support for local module changes
 								{
 									Logger.Log(PCHLoggerLevel, $"{Prefix} Module '{Module.Name}': Do not set 'OptimizeCode' to 'CodeOptimization.Never'. This creates PCH permutations.");
 									bFoundIssues = true;
