UnrealBuildTool: Print unavailable toolchains in a separate list when dumping toolchain info

#rnx

diff --git a/Engine/Source/Programs/UnrealBuildTool/Platform/Windows/MicrosoftPlatformSDK.cs b/Engine/Source/Programs/UnrealBuildTool/Platform/Windows/MicrosoftPlatformSDK.cs
--- a/Engine/Source/Programs/UnrealBuildTool/Platform/Windows/MicrosoftPlatformSDK.cs
+++ b/Engine/Source/Programs/UnrealBuildTool/Platform/Windows/MicrosoftPlatformSDK.cs
@@ -624,7 +624,7 @@ namespace UnrealBuildTool
 		/// <param name="Logger">The ILogger interface to write to</param>
 		static void DumpToolChains(IEnumerable<ToolChainInstallation> ToolChains, Func<IOrderedEnumerable<ToolChainInstallation>, IOrderedEnumerable<ToolChainInstallation>> Preference, UnrealArch Architecture, ILogger Logger)
 		{
-			var SortedToolChains = Preference(ToolChains.Where(x => x.Architecture == Architecture)
+			IOrderedEnumerable<ToolChainInstallation> SortedToolChains = Preference(ToolChains.Where(x => x.Architecture == Architecture)
 				.OrderByDescending(x => x.Error == null))
 				.ThenByDescending(x => x.Is64Bit)
 				.ThenBy(x => x.ReleaseChannel)
@@ -634,10 +634,25 @@ namespace UnrealBuildTool
 
 			if (SortedToolChains.Count() > 0)
 			{
-				Logger.LogInformation("Available {Architecture} toolchains ({Count}):", Architecture, SortedToolChains.Count());
-				foreach (ToolChainInstallation ToolChain in SortedToolChains)
+				IEnumerable<ToolChainInstallation> AvailableToolChains = SortedToolChains.Where(x => x.Error == null);
+				IEnumerable<ToolChainInstallation> UnavailableToolChains = SortedToolChains.Where(x => x.Error != null);
+
+				if (AvailableToolChains.Any())
+				{
+					Logger.LogInformation("Available {Architecture} toolchains ({Count}):", Architecture, AvailableToolChains.Count());
+					foreach (ToolChainInstallation ToolChain in AvailableToolChains)
+					{
+						Logger.LogInformation(" * {ToolChainDir}\n    (Family={Family}, FamilyRank={FamilyRank}, Version={Version}, Is64Bit={Is64Bit}, ReleaseChannel={ReleaseChannel}, Architecture={Arch})", ToolChain.BaseDir, ToolChain.Family, ToolChain.FamilyRank, ToolChain.Version, ToolChain.Is64Bit, ToolChain.ReleaseChannel, ToolChain.Architecture);
+					}
+				}
+
+				if (UnavailableToolChains.Any())
 				{
-					Logger.LogInformation(" * {ToolChainDir}\n    (Family={Family}, FamilyRank={FamilyRank}, Version={Version}, Is64Bit={Is64Bit}, ReleaseChannel={ReleaseChannel}, Architecture={Arch}, Error={Error})", ToolChain.BaseDir, ToolChain.Family, ToolChain.FamilyRank, ToolChain.Version, ToolChain.Is64Bit, ToolChain.ReleaseChannel, ToolChain.Architecture, ToolChain.Error != null);
+					Logger.LogWarning("Unavailable {Architecture} toolchains ({Count}):", Architecture, UnavailableToolChains.Count());
+					foreach (ToolChainInstallation ToolChain in UnavailableToolChains.Where(x => x.Error != null))
+					{
+						Logger.LogWarning(" * {ToolChainDir}\n    (Family={Family}, FamilyRank={FamilyRank}, Version={Version}, Is64Bit={Is64Bit}, ReleaseChannel={ReleaseChannel}, Architecture={Arch}, Error=\"{Error}\")", ToolChain.BaseDir, ToolChain.Family, ToolChain.FamilyRank, ToolChain.Version, ToolChain.Is64Bit, ToolChain.ReleaseChannel, ToolChain.Architecture, ToolChain.Error);
+					}
 				}
 			}
 			else
