Fix for per-project packaging images.

#rnx

diff --git a/Engine/Source/Programs/UnrealBuildTool/Platform/Windows/AppXManifestGeneratorBase.cs b/Engine/Source/Programs/UnrealBuildTool/Platform/Windows/AppXManifestGeneratorBase.cs
--- a/Engine/Source/Programs/UnrealBuildTool/Platform/Windows/AppXManifestGeneratorBase.cs
+++ b/Engine/Source/Programs/UnrealBuildTool/Platform/Windows/AppXManifestGeneratorBase.cs
@@ -370,13 +370,13 @@ namespace UnrealBuildTool
 			if (ProjectFile != null)
 			{
 				// look in project normal Build location
-				SourceFilePath = Path.Combine(ProjectFile.FullName!, "Build", Platform.ToString(), BuildResourceSubPath, CultureId, ResourceFileName);
+				SourceFilePath = Path.Combine(ProjectFile.Directory.FullName!, "Build", Platform.ToString(), BuildResourceSubPath, CultureId, ResourceFileName);
 				bFileExists = File.Exists(SourceFilePath);
 
 				// look in Platform Extensions next
 				if (!bFileExists)
 				{
-					SourceFilePath = Path.Combine(ProjectFile.FullName!, "Platforms", Platform.ToString(), "Build", BuildResourceSubPath, CultureId, ResourceFileName);
+					SourceFilePath = Path.Combine(ProjectFile.Directory.FullName!, "Platforms", Platform.ToString(), "Build", BuildResourceSubPath, CultureId, ResourceFileName);
 					bFileExists = File.Exists(SourceFilePath);
 				}
 			}
