UBT: Remove explicit objC ABI flags.

diff --git a/Engine/Source/Programs/UnrealBuildTool/ToolChain/ClangToolChain.cs b/Engine/Source/Programs/UnrealBuildTool/ToolChain/ClangToolChain.cs
--- a/Engine/Source/Programs/UnrealBuildTool/ToolChain/ClangToolChain.cs
+++ b/Engine/Source/Programs/UnrealBuildTool/ToolChain/ClangToolChain.cs
@@ -406,16 +406,12 @@ namespace UnrealBuildTool
 		protected virtual void GetCompileArguments_MM(CppCompileEnvironment CompileEnvironment, List<string> Arguments)
 		{
 			Arguments.Add("-x objective-c++");
-			Arguments.Add("-fobjc-abi-version=2");
-			Arguments.Add("-fobjc-legacy-dispatch");
 			GetCppStandardCompileArgument(CompileEnvironment, Arguments);
 		}
 
 		protected virtual void GetCompileArguments_M(CppCompileEnvironment CompileEnvironment, List<string> Arguments)
 		{
 			Arguments.Add("-x objective-c");
-			Arguments.Add("-fobjc-abi-version=2");
-			Arguments.Add("-fobjc-legacy-dispatch");
 			GetCStandardCompileArgument(CompileEnvironment, Arguments);
 		}
 
