Use Gauntlet to run low level tests instead of BG's Spawn, move LowLevelTests.xml to public folder and remove all restricted platform names

diff --git a/Engine/Build/LowLevelTests.xml b/Engine/Build/LowLevelTests.xml
new file mode 100644
--- /dev/null
+++ b/Engine/Build/LowLevelTests.xml
@@ -0,0 +1,198 @@
+<?xml version='1.0' ?>
+<BuildGraph xmlns="http://www.epicgames.com/BuildGraph" xmlns:xsi="http://www.w3.org/2001/XMLSchema-instance" xsi:schemaLocation="http://www.epicgames.com/BuildGraph ../../../Build/Graph/Schema.xsd" >
+	<EnvVar Name="UE_HORDE_JOBID"/>
+
+	<Property Name="TestNames" Value=""/>
+
+	<Property Name="AllTestablePlatforms" Value="Win64;Mac;Linux;Android;IOS"/>
+	<Property Name="MobilePlatforms" Value="Android;IOS"/>
+	<Property Name="BuildPlatforms" Value="Win64;Mac;Linux"/>
+	<Property Name="PlatformsBuiltFromWin64" Value="Win64;Android"/>
+	<Property Name="PlatformsBuiltFromMac" Value="Mac;IOS"/>
+	<Property Name="PlatformsBuiltFromLinux" Value="Linux"/>
+
+	<Property Name="AlternateDevicePoolPlatformsProperty" Value=""/>
+
+	<Include Script="LowLevelTests/*.xml"/>
+	<Include Script="../Platforms/*/Build/LowLevelTests/*.xml"/>
+	<Include Script="../Restricted/NotForLicensees/Build/LowLevelTests/*.xml"/>
+	<Include Script="../Restricted/NotForLicensees/Platforms/*/Build/LowLevelTests/*.xml"/>
+
+	<Option Name="BuildCatch2Libs" DefaultValue="false" Description="True if building Catch2 static libraries"/>
+	<Option Name="TestCatch2Libs" DefaultValue="false" Description="True if additionally building Catch2 with extras and tests"/>
+	<Option Name="Catch2LibVariation" DefaultValue="VS2019" Description="Catch2 library variation, e.g. VS2022"/>
+	<Option Name="SubmitBuiltCatch2Libs" DefaultValue="false" Description="True if submitting the built Catch2 static libraries"/>
+	<Property Name="Catch2LatestVersion" Value="v3.0.1"/>
+
+	<!-- Check if this is a release branch -->
+	<Option Name="IsReleaseBranch" DefaultValue="false" Description="True if release branch, false otherwise"/>
+	<Option Name="PreflightChange" DefaultValue="" Description="Preflight CL number if preflight, empty otherwise"/>
+
+	<!-- Test Options -->
+	<Option Name="RunAllTests" DefaultValue="true" Description="Run All Available Tests (except Foundation)"/>
+	
+	<!-- If we set one of the Run*Tests option we don't want to run all the tests, just the selected ones. -->
+	<Property Name="RunsAllTests" Value="$(RunAllTests)"/>
+	<ForEach Name="TestName" Values="$(TestNames)">
+		<!-- Don't turn of all tests if only Foundation is in the list. -->
+		<Do If="'$(TestName)' != 'Foundation'">
+			<Property Name="RunsAllTests" Value="false" If="$(Run$(TestName)Tests)"/>
+		</Do>
+	</ForEach>
+
+	<!-- Platform Options -->
+	<Option Name="TestAllSupportedPlatforms" DefaultValue="true" Description="Test All Supported Platforms"/>
+	<Option Name="TestPlatformWin64" DefaultValue="false" Description="Test Platform Win64"/>
+	<Option Name="TestPlatformLinux" DefaultValue="false" Description="Test Platform Linux"/>
+	<Option Name="TestPlatformMac" DefaultValue="false" Description="Test Platform Mac"/>
+	<Option Name="TestPlatformAndroid" DefaultValue="false" Description="Test Platform Android"/>
+	<Option Name="TestPlatformIOS" DefaultValue="false" Description="Test Platform IOS"/>
+	
+	<!-- If we set one of the TestPlatform* option we don't want to run on all the platforms, just the selected ones. -->
+	<Property Name="TestsAllSupportedPlatforms" Value="$(TestAllSupportedPlatforms)"/>
+	<ForEach Name="Platform" Values="$(AllTestablePlatforms)">
+		<Property Name="TestsAllSupportedPlatforms" Value="false" If="$(TestPlatform$(Platform))"/>
+	</ForEach>
+	
+	<Property Name="GauntletDeviceService" Value="https://horde.devtools.epicgames.com" />
+	<Option Name="TestArgs" DefaultValue="" Description="Extra test arguments."/>
+
+	<Option Name="GauntletDevicePool" DefaultValue="UE4" Description="Gauntlet Device Pool to use to run target platform tests"/>
+	<Option Name="AlternateGauntletDevicePool" DefaultValue="UE5" Description="Alternate Gauntlet Device Pool to use to run target platform tests"/>
+	<Option Name="AlternateDevicePoolPlatforms" DefaultValue="$(AlternateDevicePoolPlatformsProperty);HoloLens" Description="Platforms that need to use the Alternate Gauntlet Device Pool"/>
+	<Option Name="Device" DefaultValue="" Description="Device address to run tests on"/>
+	<Option Name="Configuration" DefaultValue="Development" Description="Configuration type: development, debug, test, shipping etc."/>
+	<Option Name="AttachToDebugger" DefaultValue="false" Description="Wait for debugger to attach."/>
+
+	<Property Name="TestNodes" Value=""/>
+	<Property Name="AfterTestNodes" Value=""/>
+	<Property Name="Catch2BuildNodes" Value=""/>
+	
+	<!-- Override build tags -->
+	<Option Name="OverrideTags" DefaultValue="" Description="Override default tags for all tests, used for special types of tests"/>
+	
+	<!-- Special test configs -->
+	<Option Name="PerTestTimeout" DefaultValue="0" Description="Individual test timeout (not for entire test run). 0 means no timeout is captured."/>
+	
+	<!-- After Steps -->
+	<Macro Name="ExpandAfterSteps" Arguments="Platform;TestName">
+		<Node Name="$(TestName) AfterStep $(Platform)">
+			<Expand Name="$(TestName)AfterSteps" Platform="$(Platform)" />	
+		</Node>
+		<Label Category="After" Name="$($(TestName)Short) After $(Platform)" Requires="$(TestName) AfterStep $(Platform)"/>
+		<Property Name="AfterTestNodes" Value="$(AfterTestNodes);$(TestName) AfterStep $(Platform)" If="$(RunAllTests) Or $(Run$(TestName)Tests)"/>
+	</Macro>
+
+	<Macro Name="DeployAndTest" Arguments="TestName;Platform">
+		<Property Name="DevicePool" Value="$(GauntletDevicePool)" />
+		<Property Name="DevicePool" Value="$(AlternateGauntletDevicePool)" If="ContainsItem('$(AlternateDevicePoolPlatforms)', '$(Platform)', ';')" />
+		<Property Name="TestArgsExt" Value="$(TestArgs) $($(TestName)InitialExtraArgs) -deviceurl=$(GauntletDeviceService) -devicepool=$(DevicePool)"/>
+		<Property Name="TestModeSuffix" Value=""/>
+		<Property Name="TestTags" Value=""/>
+		<Property Name="TestTags" Value="$($(Platform)$(TestName)Tags)" If="'$($(Platform)$(TestName)Tags)' != ''"/>
+		<Property Name="TestTags" Value="$(OverrideTags)" If="'$(OverrideTags)' != ''"/>
+		<Property Name="TestArgsExt" Value="$(TestArgsExt) -tags=&quot;$(TestTags)&quot;"/>
+		<Property Name="TestArgsExt" Value="$(TestArgsExt) --buildmachine" If="$(IsBuildMachine)"/>
+		<Property Name="TestArgsExt" Value="$(TestArgsExt) -timeout=$(PerTestTimeout)"/>
+		<Property Name="TestArgsExt" Value="$(TestArgsExt) $($(Platform)ExtraArgs)"/>
+		<Do If="'$(Device)' != ''">
+			<Property Name="TestArgsExt" Value="$(TestArgsExt) -device=$(Device) "/>
+		</Do>
+		<Do If="'$(AttachToDebugger)' != 'false'">
+			<Property Name="TestArgsExt" Value="$(TestArgsExt) -AttachToDebugger "/>
+		</Do>
+		<Do If="'$($(TestName)ReportType)' != ''">
+			<Property Name="TestArgsExt" Value="$(TestArgsExt) -ReportType=$($(TestName)ReportType)"/>
+		</Do>
+		<Do If="'$($(TestName)UsesCatch2)' != 'True'">
+			<Property Name="TestArgsExt" Value=""/>
+		</Do>
+		<Command Name="RunLowLevelTests"
+				Arguments="-test=LowLevelTests -testapp=$($(TestName)Target)$(TestModeSuffix) -platform=$(Platform) -configuration=$(Configuration) $(TestArgsExt) -build=&quot;$($(TestName)BinariesRelative)\$(Platform)\$($(TestName)Target)$(TestModeSuffix)&quot; -VeryVerbose"/>
+	</Macro>
+
+	<ForEach Name="BuildPlatform" Values="$(BuildPlatforms)">
+		<ForEach Name="Platform" Values="$(PlatformsBuiltFrom$(BuildPlatform))">
+			<Property Name="AgentType" Value="$(BuildPlatform)"/>
+			<Do If="'$($(Platform)RequiresSpecialAgent)' == 'True'">
+				<Expand Name="SelectAgentType" Platform="$(Platform)" BuildPlatform="$(BuildPlatform)" />
+			</Do>
+			<Property Name="BuildAgentType" Value="$(AgentType)"/>
+			<Do If="'$(Platform)' == 'Linux'">
+				<Property Name="BuildAgentType" Value="Win64"/>
+			</Do>
+			<Agent Name="Catch2 Build $(Platform)" Type="$(BuildAgentType)">
+				<Node Name="Catch2 Build Library $(Platform)">
+					<Property Name="Generator" Value="Makefile"/>
+					<Property Name="TargetConfigs" Value="debug+release"/>
+					<Property Name="ExtraCMakeArgs" Value=""/>
+					<Property Name="TargetPlatform" Value="$(Platform)"/>
+					<Property Name="CommonExtraCmakeArgs" Value=""/>
+					<Property Name="CommonExtraCmakeArgs" Value="$(CommonExtraCmakeArgs) -DCATCH_DEVELOPMENT_BUILD=ON -DBUILD_TESTING=ON -DCATCH_BUILD_EXTRA_TESTS=ON -DCATCH_ENABLE_CONFIGURE_TESTS=ON " If="$(TestCatch2Libs)"/>
+					<Do If="'$(Platform)' == 'Android'">
+						<Property Name="ExtraCMakeArgs" Value="-TargetArchitecture=arm64"/>
+					</Do>
+					<Do If="'$(Platform)' == 'Linux'">
+						<Property Name="ExtraCMakeArgs" Value="-TargetArchitecture=x86_64-unknown-linux-gnu -CMakeAdditionalArguments=&quot;$(CommonExtraCmakeArgs)-DBUILD_WITH_LIBCXX=ON&quot;"/>
+					</Do>
+					<Do If="'$(Platform)' == 'Win64'">
+						<Property Name="Generator" Value="VS2019"/>
+						<Do If="'$(Catch2LibVariation)' == 'VS2022'">
+							<Property Name="Generator" Value="VS2022"/>
+						</Do>
+						<Property Name="ExtraCMakeArgs" Value="-TargetArchitecture=x64 -CMakeAdditionalArguments=&quot;$(CommonExtraCmakeArgs)-DSET_CONFIG_NO_COLOUR_WIN32=OFF&quot;"/>
+					</Do>
+					<Do If="'$($(Platform)RequiresCatch2BuildExtraCMakeArgs)' == 'True'">
+						<Expand Name="Catch2BuildAppendExtraCMakeArgs" Platform="$(Platform)" />
+					</Do>
+					<Do If="'$(BuildPlatform)' == 'Mac'">
+						<Property Name="Generator" Value="Xcode"/>
+						<Property Name="TargetConfigs" Value="Debug+Release"/>
+						<Property Name="ExtraCMakeArgs" Value="-CmakeAdditionalArguments=&quot;$(CommonExtraCmakeArgs)-DCMAKE_OSX_DEPLOYMENT_TARGET=10.15 -DCMAKE_OSX_ARCHITECTURES=arm64;x86_64&quot;"/>
+					</Do>
+					<Command Name="BuildCMakeLib"
+							 Arguments="-CMakeGenerator=$(Generator) -MakeTarget=Catch2 -TargetLib=Catch2 -TargetLibVersion=$(Catch2LatestVersion) -BinOutputPath=Binaries -LibOutputPath=lib -TargetPlatform=$(TargetPlatform) -TargetConfigs=$(TargetConfigs) -SkipCreateChangelist $(ExtraCMakeArgs)"/>
+				</Node>
+				<Property Name="Catch2BuildNodes" Value="$(Catch2BuildNodes);Catch2 Build Library $(Platform)"/>
+			</Agent>
+			<Agent Name="Low Level Tests $(Platform)" Type="$(AgentType)">
+				<ForEach Name="TestName" Values="$(TestNames)">
+					<Property Name="TestSupportedOnPlatform" Value="false"/>
+					<Property Name="TestRunsOnPlatform" Value="false"/>
+					<Property Name="RunsTestOnPlatform" Value="false"/>
+					
+					<Property Name="TestSupportedOnPlatform" Value="true" If="Contains('$($(TestName)SupportedPlatforms)', '$(Platform)')"/>
+					<Property Name="TestRunsOnPlatform" Value="true" If="($(TestSupportedOnPlatform) And $(TestsAllSupportedPlatforms)) Or $(TestPlatform$(Platform))"/>
+					<Property Name="RunsTestOnPlatform" Value="true" If="!$($(TestName)Disabled) And $(TestRunsOnPlatform) And ($(Run$(TestName)Tests) Or $(RunsAllTests))"/>
+
+					<Do If="$(RunsTestOnPlatform)">
+						<Do If="!$(TestSupportedOnPlatform)">
+							<Warning Message="Low level test $(TestName) not yet tested for support on platform $(Platform)."/>
+						</Do>
+
+						<Property Name="TestRequirements" Value=""/>
+						<Do If="$(BuildCatch2Libs)">
+							<Property Name="TestRequirements" Value="$(TestRequirements);Catch2 Build Library $(Platform)"/>
+						</Do>
+						
+						<Node Name="$($(TestName)Short) Tests $(Platform)" Requires="$(TestRequirements)">
+							<Compile Target="$($(TestName)Target)" Configuration="$(Configuration)" Platform="$(Platform)" AllowXGE="false" Arguments="$($(Platform)$(TestName)ExtraCompilationArgs)"/>
+							<Do If="'$($(Platform)$(TestName)RunSupported)' == 'True'">
+								<Expand Name="DeployAndTest" TestName="$(TestName)" Platform="$(Platform)"/>
+							</Do>
+						</Node>
+						<Label Category="Test" Name="$($(TestName)Short) $(Platform)" Requires="$($(TestName)Short) Tests $(Platform)"/>
+						<Property Name="TestNodes" Value="$(TestNodes);$($(TestName)Short) Tests $(Platform)" If="$(RunAllTests) Or $(Run$(TestName)Tests)"/>
+						
+						<Do If="'$($(TestName)HasAfterSteps)' == 'True'">
+							<Expand Name="ExpandAfterSteps" TestName="$(TestName)" Platform="$(Platform)" />
+						</Do>
+					</Do>
+				</ForEach>
+			</Agent>
+		</ForEach>
+	</ForEach>
+	
+	<Aggregate Name="Run Low Level Tests" Requires="$(TestNodes);$(AfterTestNodes)"/>
+	<Aggregate Name="Build Catch2 Libraries" Requires="$(Catch2BuildNodes)"/>
+</BuildGraph>
diff --git a/Engine/Build/LowLevelTests/General.xml b/Engine/Build/LowLevelTests/General.xml
new file mode 100644
--- /dev/null
+++ b/Engine/Build/LowLevelTests/General.xml
@@ -0,0 +1,21 @@
+<?xml version="1.0" encoding="utf-8"?>
+<BuildGraph xmlns:xsi="http://www.w3.org/2001/XMLSchema-instance" xsi:schemaLocation="http://www.epicgames.com/BuildGraph ../../Build/Graph/Schema.xsd" xmlns="http://www.epicgames.com/BuildGraph">
+	<Property Name="Win64ExtraArgs" Value=""/>
+	<Property Name="LinuxExtraArgs" Value=""/>
+	<Property Name="MacExtraArgs" Value=""/>
+	<Property Name="IOSExtraArgs" Value=""/>
+	<Property Name="AndroidExtraArgs" Value=""/>
+	<Property Name="HoloLensExtraArgs" Value=""/>
+	<Property Name="Win64RequiresSpecialAgent" Value="False"/>
+	<Property Name="LinuxRequiresSpecialAgent" Value="False"/>
+	<Property Name="MacRequiresSpecialAgent" Value="False"/>
+	<Property Name="IOSRequiresSpecialAgent" Value="False"/>
+	<Property Name="AndroidRequiresSpecialAgent" Value="False"/>
+	<Property Name="HoloLensRequiresSpecialAgent" Value="False"/>
+	<Property Name="Win64RequiresCatch2BuildExtraCMakeArgs" Value="False"/>
+	<Property Name="LinuxRequiresCatch2BuildExtraCMakeArgs" Value="False"/>
+	<Property Name="MacRequiresCatch2BuildExtraCMakeArgs" Value="False"/>
+	<Property Name="IOSRequiresCatch2BuildExtraCMakeArgs" Value="False"/>
+	<Property Name="AndroidRequiresCatch2BuildExtraCMakeArgs" Value="False"/>
+	<Property Name="HoloLensRequiresCatch2BuildExtraCMakeArgs" Value="False"/>
+</BuildGraph>
\ No newline at end of file
diff --git a/Engine/Build/LowLevelTests/ReplicationSystem.xml b/Engine/Build/LowLevelTests/ReplicationSystem.xml
--- a/Engine/Build/LowLevelTests/ReplicationSystem.xml
+++ b/Engine/Build/LowLevelTests/ReplicationSystem.xml
@@ -11,7 +11,7 @@
   <Property Name="ReplicationSystemHasAfterSteps" Value="False" />
   <Property Name="ReplicationSystemUsesCatch2" Value="True" />
   <Option Name="RunReplicationSystemTests" DefaultValue="False" Description="Run Replication System Tests" />
-  <Property Name="ReplicationSystemSupportedPlatforms" Value="Win64;Mac;Linux;Android" />
+  <Property Name="ReplicationSystemSupportedPlatforms" Value="Win64;Linux" />
   <Property Name="Win64ReplicationSystemTags" Value="" />
   <Property Name="MacReplicationSystemTags" Value="" />
   <Property Name="IOSReplicationSystemTags" Value="" />
diff --git a/Engine/Build/LowLevelTests_GenProps.xml b/Engine/Build/LowLevelTests_GenProps.xml
deleted file mode 100644
diff --git a/Engine/Source/Programs/AutomationTool/Gauntlet/Platform/Linux/Gauntlet.TargetDeviceLinux.cs b/Engine/Source/Programs/AutomationTool/Gauntlet/Platform/Linux/Gauntlet.TargetDeviceLinux.cs
--- a/Engine/Source/Programs/AutomationTool/Gauntlet/Platform/Linux/Gauntlet.TargetDeviceLinux.cs
+++ b/Engine/Source/Programs/AutomationTool/Gauntlet/Platform/Linux/Gauntlet.TargetDeviceLinux.cs
@@ -227,7 +227,13 @@ namespace Gauntlet
 
 				Log.Verbose("\t{0}", CmdLine);
 
-				Result = CommandUtils.Run(LinuxApp.ExecutablePath, CmdLine, Options: LinuxApp.RunOptions, SpewFilterCallback: new SpewFilterCallbackType(delegate(string M) { return null; }) /* make sure stderr does not spew in the stdout */);
+				bool bAllowSpew = LinuxApp.RunOptions.HasFlag(CommandUtils.ERunOptions.AllowSpew);
+
+				Result = CommandUtils.Run(LinuxApp.ExecutablePath,
+					CmdLine,
+					Options: LinuxApp.RunOptions,
+					SpewFilterCallback: new SpewFilterCallbackType(delegate(string M) { return bAllowSpew ? M : null; }) /* make sure stderr does not spew in the stdout */,
+					WorkingDir: LinuxApp.WorkingDirectory);
 
 				if (Result.HasExited && Result.ExitCode != 0)
 				{
@@ -272,6 +278,28 @@ namespace Gauntlet
 			}
 		}
 
+		protected IAppInstall InstallNativeStagedBuild(UnrealAppConfig AppConfig, NativeStagedBuild InBuild)
+		{
+			LinuxAppInstall LinuxApp = new LinuxAppInstall(AppConfig.Name, AppConfig.ProjectName, this);
+			LinuxApp.RunOptions = RunOptions;
+			if (Log.IsVeryVerbose)
+			{
+				LinuxApp.RunOptions |= CommandUtils.ERunOptions.AllowSpew;
+			}
+
+			LinuxApp.CommandArguments = AppConfig.CommandLine;
+
+			LinuxApp.ArtifactPath = Path.Combine(InBuild.BuildPath, AppConfig.ProjectName, @"Saved");
+			LinuxApp.CleanDeviceArtifacts();
+
+			CopyAdditionalFiles(AppConfig);
+
+			LinuxApp.ExecutablePath = Path.Combine(InBuild.BuildPath, InBuild.ExecutablePath);
+			LinuxApp.WorkingDirectory = InBuild.BuildPath;
+
+			return LinuxApp;
+		}
+
 		protected IAppInstall InstallStagedBuild(UnrealAppConfig AppConfig, StagedBuild InBuild)
 		{
 			bool SkipDeploy = Globals.Params.ParseParam("SkipDeploy");
@@ -365,7 +393,11 @@ namespace Gauntlet
 
 		public IAppInstall InstallApplication(UnrealAppConfig AppConfig)
 		{
-			if (AppConfig.Build is StagedBuild)
+			if (AppConfig.Build is NativeStagedBuild)
+			{
+				return InstallNativeStagedBuild(AppConfig, AppConfig.Build as NativeStagedBuild);
+			}
+			else if (AppConfig.Build is StagedBuild)
 			{
 				return InstallStagedBuild(AppConfig, AppConfig.Build as StagedBuild);
 			}
diff --git a/Engine/Source/Programs/AutomationTool/Gauntlet/Platform/Mac/Gauntlet.TargetDeviceMac.cs b/Engine/Source/Programs/AutomationTool/Gauntlet/Platform/Mac/Gauntlet.TargetDeviceMac.cs
--- a/Engine/Source/Programs/AutomationTool/Gauntlet/Platform/Mac/Gauntlet.TargetDeviceMac.cs
+++ b/Engine/Source/Programs/AutomationTool/Gauntlet/Platform/Mac/Gauntlet.TargetDeviceMac.cs
@@ -274,6 +274,26 @@ namespace Gauntlet
 			return OutBuildPath;
 		}
 
+		protected IAppInstall InstallNativeStagedBuild(UnrealAppConfig AppConfig, NativeStagedBuild InBuild)
+		{
+			MacAppInstall MacApp = new MacAppInstall(AppConfig.Name, this);
+
+			MacApp.RunOptions = RunOptions;
+			if (Log.IsVeryVerbose)
+			{
+				MacApp.RunOptions |= CommandUtils.ERunOptions.AllowSpew;
+			}
+
+			MacApp.WorkingDirectory = InBuild.BuildPath;
+			MacApp.ExecutablePath = Path.Combine(InBuild.BuildPath, InBuild.ExecutablePath);
+			MacApp.CommandArguments = AppConfig.CommandLine;
+
+			MacApp.ArtifactPath = Path.Combine(InBuild.BuildPath, AppConfig.ProjectName, @"Saved");
+			MacApp.CleanDeviceArtifacts();
+
+			return MacApp;
+		}
+
 		protected IAppInstall InstallBuild(UnrealAppConfig AppConfig, IBuild InBuild)
 		{
 			// Full path to the build
@@ -357,6 +377,10 @@ namespace Gauntlet
 
 		public IAppInstall InstallApplication(UnrealAppConfig AppConfig)
 		{
+			if (AppConfig.Build is NativeStagedBuild)
+			{
+				return InstallNativeStagedBuild(AppConfig, (NativeStagedBuild)AppConfig.Build);
+			}
 			if (AppConfig.Build is StagedBuild || AppConfig.Build is MacPackagedBuild)
 			{
 				return InstallBuild(AppConfig, AppConfig.Build);
@@ -411,7 +435,8 @@ namespace Gauntlet
 				Log.Info("Launching {0} on {1}", App.Name, ToString());
 				Log.Verbose("\t{0}", MacInstall.CommandArguments);
 
-				Result = CommandUtils.Run(MacInstall.ExecutablePath, MacInstall.CommandArguments, Options: MacInstall.RunOptions, SpewFilterCallback: new SpewFilterCallbackType(delegate (string M) { return null; }) /* make sure stderr does not spew in the stdout */);
+				bool bAllowSpew = MacInstall.RunOptions.HasFlag(CommandUtils.ERunOptions.AllowSpew);
+				Result = CommandUtils.Run(MacInstall.ExecutablePath, MacInstall.CommandArguments, Options: MacInstall.RunOptions, SpewFilterCallback: new SpewFilterCallbackType(delegate (string M) { return bAllowSpew ? M : null; }) /* make sure stderr does not spew in the stdout */);
 
 				if (Result.HasExited && Result.ExitCode != 0)
 				{
diff --git a/Engine/Source/Programs/AutomationTool/Gauntlet/Platform/Windows/Gauntlet.TargetDeviceWindows.cs b/Engine/Source/Programs/AutomationTool/Gauntlet/Platform/Windows/Gauntlet.TargetDeviceWindows.cs
--- a/Engine/Source/Programs/AutomationTool/Gauntlet/Platform/Windows/Gauntlet.TargetDeviceWindows.cs
+++ b/Engine/Source/Programs/AutomationTool/Gauntlet/Platform/Windows/Gauntlet.TargetDeviceWindows.cs
@@ -151,7 +151,24 @@ namespace Gauntlet
 
 		public string ExecutablePath;
 
-		public string CommandArguments;
+		public bool CanAlterCommandArgs = true;
+
+		public string CommandArguments
+		{
+			get { return CommandArgumentsPrivate; }
+			set
+			{
+				if (CanAlterCommandArgs || string.IsNullOrEmpty(CommandArgumentsPrivate))
+				{
+					CommandArgumentsPrivate = value;
+				}
+				else
+				{
+					Log.Info("Skipped setting command AppInstall line when CanAlterCommandArgs = false");
+				}
+			}
+		}
+		private string CommandArgumentsPrivate;
 
 		public string ArtifactPath;
 
@@ -333,19 +350,22 @@ namespace Gauntlet
 
 				string CmdLine = WinApp.CommandArguments;
 
-				// Look in app parameters if abslog is specified, if so use it
-				Regex CLRegex = new Regex(@"(--?[a-zA-Z]+)[:\s=]?([A-Z]:(?:\\[\w\s-]+)+\\?(?=\s-)|\""[^\""]*\""|[^-][^\s]*)?");
-				foreach (Match M in CLRegex.Matches(CmdLine))
+				if (WinApp.CanAlterCommandArgs)
 				{
-					if (M.Groups.Count == 3 && M.Groups[1].Value == "-abslog")
+					// Look in app parameters if abslog is specified, if so use it
+					Regex CLRegex = new Regex(@"(--?[a-zA-Z]+)[:\s=]?([A-Z]:(?:\\[\w\s-]+)+\\?(?=\s-)|\""[^\""]*\""|[^-][^\s]*)?");
+					foreach (Match M in CLRegex.Matches(CmdLine))
 					{
-						ProcessLogFile = M.Groups[2].Value;
+						if (M.Groups.Count == 3 && M.Groups[1].Value == "-abslog")
+						{
+							ProcessLogFile = M.Groups[2].Value;
+						}
 					}
 				}
 
 				// Explicitly set log file when not already defined if not build machine
 				// -abslog makes sure Unreal dymanicaly update the log window when using -log
-				if (!AutomationTool.Automation.IsBuildMachine && string.IsNullOrEmpty(ProcessLogFile))
+				if (WinApp.CanAlterCommandArgs && !AutomationTool.Automation.IsBuildMachine && string.IsNullOrEmpty(ProcessLogFile))
 				{
 					string LogFolder = string.Format(@"{0}\Logs", WinApp.ArtifactPath);
 
@@ -361,7 +381,7 @@ namespace Gauntlet
 				// cleanup any existing log file
 				try
 				{
-					if (File.Exists(ProcessLogFile))
+					if (ProcessLogFile != null && File.Exists(ProcessLogFile))
 					{
 						EpicGames.Core.FileUtils.ForceDeleteFile(ProcessLogFile);
 					}
@@ -374,7 +394,19 @@ namespace Gauntlet
 
 				Log.Verbose("\t{0}", CmdLine);
 
-				Result = CommandUtils.Run(WinApp.ExecutablePath, CmdLine, Options: WinApp.RunOptions | (ProcessLogFile != null ? CommandUtils.ERunOptions.NoStdOutRedirect : 0 ), SpewFilterCallback: new SpewFilterCallbackType(delegate (string M) { return null; }) /* make sure stderr does not spew in the stdout */);
+				CommandUtils.ERunOptions FinalRunOptions = WinApp.RunOptions;
+				if (!WinApp.CanAlterCommandArgs)
+				{
+					CmdLine = WinApp.CommandArguments;
+				}
+
+				FinalRunOptions = WinApp.RunOptions | (ProcessLogFile != null ? CommandUtils.ERunOptions.NoStdOutRedirect : 0);
+
+				Result = CommandUtils.Run(WinApp.ExecutablePath,
+					CmdLine,
+					Options: FinalRunOptions,
+					SpewFilterCallback: new SpewFilterCallbackType(M => { return ProcessLogFile == null ? M : null ; }) /* make sure stderr does not spew in the stdout */,
+					WorkingDir: WinApp.WorkingDirectory);
 
 				if (Result.HasExited && Result.ExitCode != 0)
 				{
@@ -419,6 +451,29 @@ namespace Gauntlet
 			}
 		}
 
+		protected IAppInstall InstallNativeStagedBuild(UnrealAppConfig AppConfig, NativeStagedBuild InBuild)
+		{
+			WindowsAppInstall WinApp = new WindowsAppInstall(AppConfig.Name, AppConfig.ProjectName, this);
+			WinApp.CanAlterCommandArgs = AppConfig.CanAlterCommandArgs;
+
+			WinApp.RunOptions = RunOptions;
+			if (Log.IsVeryVerbose)
+			{
+				WinApp.RunOptions |= CommandUtils.ERunOptions.AllowSpew;
+			}
+
+			WinApp.WorkingDirectory = InBuild.BuildPath;
+			WinApp.ExecutablePath = Path.Combine(InBuild.BuildPath, InBuild.ExecutablePath);
+			WinApp.CommandArguments = AppConfig.CommandLine;
+
+			WinApp.ArtifactPath = Path.Combine(InBuild.BuildPath, AppConfig.ProjectName, @"Saved");
+			WinApp.CleanDeviceArtifacts();
+
+			CopyAdditionalFiles(AppConfig);
+
+			return WinApp;
+		}
+
 		protected IAppInstall InstallStagedBuild(UnrealAppConfig AppConfig, StagedBuild InBuild)
 		{
 			bool SkipDeploy = Globals.Params.ParseParam("SkipDeploy");
@@ -446,6 +501,7 @@ namespace Gauntlet
 
 			WindowsAppInstall WinApp = new WindowsAppInstall(AppConfig.Name, AppConfig.ProjectName, this);
 			WinApp.RunOptions = RunOptions;
+			WinApp.CanAlterCommandArgs = AppConfig.CanAlterCommandArgs;
 
 			// Set commandline replace any InstallPath arguments with the path we use
 			WinApp.CommandArguments = Regex.Replace(AppConfig.CommandLine, @"\$\(InstallPath\)", BuildPath, RegexOptions.IgnoreCase);
@@ -512,7 +568,11 @@ namespace Gauntlet
 
 		public IAppInstall InstallApplication(UnrealAppConfig AppConfig)
 		{
-			if (AppConfig.Build is StagedBuild)
+			if (AppConfig.Build is NativeStagedBuild)
+			{
+				return InstallNativeStagedBuild(AppConfig, AppConfig.Build as NativeStagedBuild);
+			}
+			else if (AppConfig.Build is StagedBuild)
 			{
 				return InstallStagedBuild(AppConfig, AppConfig.Build as StagedBuild);
 			}
diff --git a/Engine/Source/Programs/AutomationTool/Gauntlet/Unreal/Base/Gauntlet.UnrealAppConfig.cs b/Engine/Source/Programs/AutomationTool/Gauntlet/Unreal/Base/Gauntlet.UnrealAppConfig.cs
--- a/Engine/Source/Programs/AutomationTool/Gauntlet/Unreal/Base/Gauntlet.UnrealAppConfig.cs
+++ b/Engine/Source/Programs/AutomationTool/Gauntlet/Unreal/Base/Gauntlet.UnrealAppConfig.cs
@@ -48,6 +48,16 @@ namespace Gauntlet
         /// </summary>
         public List<UnrealFileToCopy> FilesToCopy { get; set; }
 
+		/// <summary>
+		/// Some IAppInstall instances can alter command line after it has been copied from AppConfig.CommandLine
+		/// Use this to restrict this behavior if necessary.
+		/// </summary>
+		public bool CanAlterCommandArgs {
+			get { return CanAlterCommandArgsPrivate; }
+			set { CanAlterCommandArgsPrivate = value; }
+		}
+		private bool CanAlterCommandArgsPrivate = true;
+
 		/// <summary>
 		/// Arguments for this instance
 		/// </summary>
diff --git a/Engine/Source/Programs/AutomationTool/Gauntlet/Unreal/BuildSource/Gauntlet.StagedBuildSource.cs b/Engine/Source/Programs/AutomationTool/Gauntlet/Unreal/BuildSource/Gauntlet.StagedBuildSource.cs
--- a/Engine/Source/Programs/AutomationTool/Gauntlet/Unreal/BuildSource/Gauntlet.StagedBuildSource.cs
+++ b/Engine/Source/Programs/AutomationTool/Gauntlet/Unreal/BuildSource/Gauntlet.StagedBuildSource.cs
@@ -309,6 +309,14 @@ namespace Gauntlet
 		}
 	}
 
+	public class NativeStagedBuild : StagedBuild
+	{
+		public NativeStagedBuild(UnrealTargetPlatform InPlatform, UnrealTargetConfiguration InConfig, UnrealTargetRole InRole, string InBuildPath, string InExecutablePath, string InFlavor = "")
+			: base(InPlatform, InConfig, InRole, InBuildPath, InExecutablePath, InFlavor)
+		{
+			Flags = BuildFlags.CanReplaceCommandLine | BuildFlags.Loose;
+		}
+	}
 
 	public abstract class StagedBuildSource<T> : IFolderBuildSource 
 		where T : StagedBuild
diff --git a/Engine/Source/Programs/AutomationTool/LowLevelTests/RunLowLevelTests.cs b/Engine/Source/Programs/AutomationTool/LowLevelTests/RunLowLevelTests.cs
--- a/Engine/Source/Programs/AutomationTool/LowLevelTests/RunLowLevelTests.cs
+++ b/Engine/Source/Programs/AutomationTool/LowLevelTests/RunLowLevelTests.cs
@@ -174,11 +174,12 @@ namespace LowLevelTests
 		[AutoParam("")]
 		public string LogDir;
 
-		[AutoParam("console")]
 		public string ReportType;
 
 		public int Timeout;
 
+		public bool LogReportContents;
+
 		public Type BuildSourceType { get; protected set; }
 
 		[AutoParam(UnrealTargetConfiguration.Development)]
@@ -218,7 +219,9 @@ namespace LowLevelTests
 
 			Timeout = Params.ParseValue("timeout=", 0);
 
-			Build = Params.ParseValue("build=", null);
+			LogReportContents = Params.ParseParam("printreport");
+
+			Build = Params.ParseValue("build=", null).Replace('/', Path.DirectorySeparatorChar).Replace('\\', Path.DirectorySeparatorChar);
 			TestApp = Globals.Params.ParseValue("testapp=", "");
 
 			Tags = Params.ParseValue("tags=", null);
@@ -226,7 +229,7 @@ namespace LowLevelTests
 
 			SkipStage = Params.ParseParam("skipstage");
 
-			ReportType = Params.ParseValue("reporttype=", "console");
+			ReportType = Params.ParseValue("reporttype=", null);
 
 			string PlatformArgString = Params.ParseValue("platform=", null);
 			Platform = string.IsNullOrEmpty(PlatformArgString) ? BuildHostPlatform.Current.Platform : UnrealTargetPlatform.Parse(PlatformArgString);
@@ -564,11 +567,11 @@ namespace LowLevelTests
 		{
 			if (InPlatform.IsInGroup(UnrealPlatformGroup.Windows))
 			{
-				return @"\w+Tests(?:-\w+)?(?:-\w+)?.exe$";
+				return @"\w+(Tests)?(?:-\w+)?(?:-\w+)?.exe$";
 			}
 			else if (InPlatform == UnrealTargetPlatform.Linux || InPlatform == UnrealTargetPlatform.Mac)
 			{
-				return @"\w+Tests$";
+				return @"\w+(Tests)?$";
 			}
 			else
 			{
@@ -586,7 +589,7 @@ namespace LowLevelTests
 
 		public string GetTargetReportPath(UnrealTargetPlatform InPlatform, string InTestApp, string InBuildPath)
 		{
-			return string.Format("{0}LLTResults.out", InPlatform.ToString());
+			return Path.Combine(InBuildPath, string.Format("{0}LLTResults.out", InPlatform.ToString()));
 		}
 
 		public string CopyDeviceReportTo(IAppInstall InAppInstall, UnrealTargetPlatform InPlatform, string InTestApp, string InBuildPath, string InTargetDirectory)
@@ -594,7 +597,10 @@ namespace LowLevelTests
 			string ReportRelativePath = GetTargetReportPath(InPlatform, InTestApp, InBuildPath);
 			string ReportPath = Path.Combine(InBuildPath, ReportRelativePath);
 			string ExpectedLocalPath = Path.Combine(InTargetDirectory, ReportRelativePath);
-			File.Copy(ReportPath, ExpectedLocalPath);
+			if (!ExpectedLocalPath.Equals(ReportPath))
+			{
+				File.Copy(ReportPath, ExpectedLocalPath, true);
+			}
 			return ExpectedLocalPath;
 		}
 	}
@@ -661,8 +667,11 @@ namespace LowLevelTests
 				CachedConfig.FilesToCopy = new List<UnrealFileToCopy>();
 				// Set reporting options, filters etc
 				CachedConfig.CommandLineParams.AddRawCommandline("--durations=no");
-				CachedConfig.CommandLineParams.AddRawCommandline(string.Format("--reporter={0}", InReportType));
-				CachedConfig.CommandLineParams.AddRawCommandline(string.Format("--out={0}", LowLevelTestsReporting.GetTargetReportPath(Platform, TestApp, BuildPath)));
+				if (!string.IsNullOrEmpty(InReportType))
+				{
+					CachedConfig.CommandLineParams.AddRawCommandline(string.Format("--reporter={0}", InReportType));
+					CachedConfig.CommandLineParams.AddRawCommandline(string.Format("--out={0}", LowLevelTestsReporting.GetTargetReportPath(Platform, TestApp, BuildPath)));
+				}
 				CachedConfig.CommandLineParams.AddRawCommandline("--filenames-as-tags");
 				if (!string.IsNullOrEmpty(InTags))
 				{
@@ -694,6 +703,7 @@ namespace LowLevelTests
 						CachedConfig.CommandLineParams.AddRawCommandline(string.Format("--extra-args {0}", ExtraCmd));
 					}
 				}
+				CachedConfig.CanAlterCommandArgs = false; // No further changes by IAppInstall instances etc.
 			}
 			return CachedConfig;
 		}
@@ -706,12 +716,12 @@ namespace LowLevelTests
 		public string BuildName { get { return TestApp; } }
 	}
 
-	public class LowLevelTestsBuild : StagedBuild
+	public class LowLevelTestsBuild : NativeStagedBuild
 	{
 		public LowLevelTestsBuild(UnrealTargetPlatform InPlatform, UnrealTargetConfiguration InConfiguration, string InBuildPath, string InExecutablePath)
 			: base(InPlatform, InConfiguration, UnrealTargetRole.Client, InBuildPath, InExecutablePath)
 		{
-			Flags = BuildFlags.CanReplaceExecutable | BuildFlags.Loose;
+			Flags = BuildFlags.CanReplaceExecutable | BuildFlags.CanReplaceCommandLine;
 		}
 	}
 }
diff --git a/Engine/Source/Programs/AutomationTool/LowLevelTests/Tests/Gauntlet.LowLevelTestNode.cs b/Engine/Source/Programs/AutomationTool/LowLevelTests/Tests/Gauntlet.LowLevelTestNode.cs
--- a/Engine/Source/Programs/AutomationTool/LowLevelTests/Tests/Gauntlet.LowLevelTestNode.cs
+++ b/Engine/Source/Programs/AutomationTool/LowLevelTests/Tests/Gauntlet.LowLevelTestNode.cs
@@ -188,20 +188,24 @@ namespace LowLevelTests
 				}
 			}
 
-			ILowLevelTestsReporting LowLevelTestsReporting = Gauntlet.Utils.InterfaceHelpers.FindImplementations<ILowLevelTestsReporting>(true)
-				.Where(B => B.CanSupportPlatform(Context.Options.Platform))
-				.First();
-
+			bool? ReportCopied = null;
 			string ReportPath = null;
-			bool ReportCopied = false;
-			try
-			{
-				ReportPath = LowLevelTestsReporting.CopyDeviceReportTo(LowLevelTestsApp.Install, Context.Options.Platform, Context.Options.TestApp, Context.Options.Build, LogDir);
-				ReportCopied = true;
-			}
-			catch (Exception ex)
+			if (!string.IsNullOrEmpty(Context.Options.ReportType))
 			{
-				Log.Error("Failed to copy report: {0}", ex.ToString());
+				ILowLevelTestsReporting LowLevelTestsReporting = Gauntlet.Utils.InterfaceHelpers.FindImplementations<ILowLevelTestsReporting>(true)
+					.Where(B => B.CanSupportPlatform(Context.Options.Platform))
+					.First();
+
+				try
+				{
+					ReportPath = LowLevelTestsReporting.CopyDeviceReportTo(LowLevelTestsApp.Install, Context.Options.Platform, Context.Options.TestApp, Context.Options.Build, LogDir);
+					ReportCopied = true;
+				}
+				catch (Exception ex)
+				{
+					ReportCopied = false;
+					Log.Error("Failed to copy report: {0}", ex.ToString());
+				}
 			}
 
 
@@ -224,15 +228,18 @@ namespace LowLevelTests
 				LowLevelTestResult = TestResult.Failed;
 				ExitReason = $"Process exited with exit code {TestInstance.ExitCode}";
 			}
-			else if (!ReportCopied)
+			else if (ReportCopied.HasValue && !ReportCopied.Value)
 			{
 				LowLevelTestResult = TestResult.Failed;
 				ExitReason = "Uabled to read test report";
 			}
-			else
+			else if (ReportPath != null)
 			{
 				string ReportContents = File.ReadAllText(ReportPath);
-				Log.Info(ReportContents);
+				if (Context.Options.LogReportContents) // Some tests prefer to log report contents
+				{
+					Log.Info(ReportContents);
+				}
 				string ReportType = Context.Options.ReportType.ToLower();
 				if (ReportType == "console")
 				{
@@ -247,7 +254,8 @@ namespace LowLevelTests
 						LowLevelTestResult = TestResult.Failed;
 						ExitReason = "Tests failed";
 					}
-				} else if (ReportType == "xml")
+				}
+				else if (ReportType == "xml")
 				{
 					LowLevelTestsReportParser LowLevelTestsReportParser = new LowLevelTestsReportParser(ReportContents);
 					if (LowLevelTestsReportParser.HasPassed())
@@ -262,6 +270,19 @@ namespace LowLevelTests
 					}
 				}
 			}
+			else // ReportPath == null
+			{
+				if (TestInstance.ExitCode != 0)
+				{
+					LowLevelTestResult = TestResult.Failed;
+					ExitReason = "Tests failed (no report to parse)";
+				}
+				else
+				{
+					LowLevelTestResult = TestResult.Passed;
+					ExitReason = "Tests passed (no report to parse)";
+				}
+			}
 			Log.Info($"Low level test exited with code {TestInstance.ExitCode} and reason: {ExitReason}");
 		}
 
diff --git a/Engine/Source/Programs/HeadlessChaos/HeadlessChaos.Build.cs b/Engine/Source/Programs/HeadlessChaos/HeadlessChaos.Build.cs
--- a/Engine/Source/Programs/HeadlessChaos/HeadlessChaos.Build.cs
+++ b/Engine/Source/Programs/HeadlessChaos/HeadlessChaos.Build.cs
@@ -50,6 +50,6 @@ public class HeadlessChaos : TestModuleRules
 
 		PrivateDefinitions.Add("CHAOS_INCLUDE_LEVEL_1=1");
 
-		UpdateBuildGraphPropertiesFile(new Metadata("HeadlessChaos", "Headless Chaos"));
+		UpdateBuildGraphPropertiesFile(new Metadata() { TestName = "HeadlessChaos", TestShortName = "Headless Chaos", UsesCatch2 = false });
 	}
 }
diff --git a/Engine/Source/Programs/HeadlessChaos/HeadlessChaos.Target.cs b/Engine/Source/Programs/HeadlessChaos/HeadlessChaos.Target.cs
--- a/Engine/Source/Programs/HeadlessChaos/HeadlessChaos.Target.cs
+++ b/Engine/Source/Programs/HeadlessChaos/HeadlessChaos.Target.cs
@@ -7,8 +7,6 @@ public class HeadlessChaosTarget : TestTargetRules
 {
 	public HeadlessChaosTarget(TargetInfo Target) : base(Target)
 	{
-		ExeBinariesSubFolder = LaunchModuleName = "HeadlessChaos";
-
         bHasExports = false;
 
 		GlobalDefinitions.Add("CHAOS_SERIALIZE_OUT=1");
diff --git a/Engine/Source/Programs/LowLevelTests/FoundationTests.Build.cs b/Engine/Source/Programs/LowLevelTests/FoundationTests.Build.cs
--- a/Engine/Source/Programs/LowLevelTests/FoundationTests.Build.cs
+++ b/Engine/Source/Programs/LowLevelTests/FoundationTests.Build.cs
@@ -1,9 +1,18 @@
 // Copyright Epic Games, Inc. All Rights Reserved.
-using System;
 using UnrealBuildTool;
 
 public class FoundationTests : TestModuleRules
 {
+	protected Metadata FoundationTestsMetadata = new Metadata() { TestName = "Foundation", TestShortName = "Foundation", ReportType = "xml" };
+
+	/// <summary>
+	/// Test metadata to be used with BuildGraph
+	/// </summary>
+	public Metadata TestMetadata
+	{ 
+		get { return FoundationTestsMetadata; }
+	}
+
 	public FoundationTests(ReadOnlyTargetRules Target) : base(Target, true)
 	{
 		PrivateDependencyModuleNames.AddRange(
@@ -21,6 +30,19 @@ public class FoundationTests : TestModuleRules
 				});
 		}
 
-		UpdateBuildGraphPropertiesFile(new Metadata("Foundation"), false);
+		string PlatformCompilationArgs;
+		foreach (var Platform in UnrealTargetPlatform.GetValidPlatforms())
+		{
+			if (Platform == UnrealTargetPlatform.Android)
+			{
+				PlatformCompilationArgs = "-allmodules -architectures=arm64";
+			}
+			else
+			{
+				PlatformCompilationArgs = "-allmodules";
+			}
+			TestMetadata.PlatformCompilationExtraArgs.Add(Platform, PlatformCompilationArgs);
+		}
+		UpdateBuildGraphPropertiesFile(TestMetadata);
 	}
 }
\ No newline at end of file
diff --git a/Engine/Source/Programs/NetworkPredictionTests/NetworkPredictionTests.Build.cs b/Engine/Source/Programs/NetworkPredictionTests/NetworkPredictionTests.Build.cs
--- a/Engine/Source/Programs/NetworkPredictionTests/NetworkPredictionTests.Build.cs
+++ b/Engine/Source/Programs/NetworkPredictionTests/NetworkPredictionTests.Build.cs
@@ -16,6 +16,6 @@ public class NetworkPredictionTests : TestModuleRules
 			}
 		);
 
-		UpdateBuildGraphPropertiesFile(new Metadata("NetworkPredictionPlugin", "Net Prediction"));
+		UpdateBuildGraphPropertiesFile(new Metadata() { TestName = "NetworkPredictionPlugin", TestShortName = "Net Prediction" });
 	}
 }
diff --git a/Engine/Source/Programs/ReplicationSystemLowLevelTests/ReplicationSystemLowLevelTests.Build.cs b/Engine/Source/Programs/ReplicationSystemLowLevelTests/ReplicationSystemLowLevelTests.Build.cs
--- a/Engine/Source/Programs/ReplicationSystemLowLevelTests/ReplicationSystemLowLevelTests.Build.cs
+++ b/Engine/Source/Programs/ReplicationSystemLowLevelTests/ReplicationSystemLowLevelTests.Build.cs
@@ -27,6 +27,6 @@ public class ReplicationSystemLowLevelTests : TestModuleRules
 			}
 		);
 
-		UpdateBuildGraphPropertiesFile(new Metadata("ReplicationSystem", "Replication System"));
+		UpdateBuildGraphPropertiesFile(new Metadata() { TestName = "ReplicationSystem", TestShortName = "Replication System" });
 	}
 }
diff --git a/Engine/Source/Programs/UnrealBuildTool/Configuration/TestModuleRules.cs b/Engine/Source/Programs/UnrealBuildTool/Configuration/TestModuleRules.cs
--- a/Engine/Source/Programs/UnrealBuildTool/Configuration/TestModuleRules.cs
+++ b/Engine/Source/Programs/UnrealBuildTool/Configuration/TestModuleRules.cs
@@ -147,64 +147,57 @@ namespace UnrealBuildTool
 			}
 		}
 
+
+#pragma warning disable 8602
+#pragma warning disable 8604
 		/// <summary>
-		/// Generates or updates include file for LowLevelTests.xml containing test flags: name, short name, target name, relative binaries path, supported platforms etc.
+		/// Generates or updates metadata file for LowLevelTests.xml containing test flags: name, short name, target name, relative binaries path, supported platforms etc.
 		/// <paramref name="TestMetadata">The test metadata specifying name, short name etc used to populate the BuildGraph properties file.</paramref>
 		/// </summary>
-		protected void UpdateBuildGraphPropertiesFile(Metadata TestMetadata, bool AddToTestNames = true)
+		protected void UpdateBuildGraphPropertiesFile(Metadata TestMetadata)
 		{
-			if (!bUpdateBuildGraphPropertiesFile)
+			if (Environment.GetEnvironmentVariable("IsBuildMachine") != null || !bUpdateBuildGraphPropertiesFile || TestMetadata == null)
 			{
 				return;
 			}
 
-			bool IsPublic = false;
 			string GeneratedPropertiesScriptFile;
 
-			string RestrictedFolder = Path.Combine(Unreal.EngineDirectory.FullName, "Restricted");
-			string NotForLicenseesFolder = Path.Combine(RestrictedFolder, "NotForLicensees");
-			string NonPublicFolder = Path.Combine(NotForLicenseesFolder, "Build");
-			string NonPublicPath = Path.Combine(NonPublicFolder, "LowLevelTests_GenProps.xml");
+			string NonPublicPath = Path.Combine(Unreal.EngineDirectory.FullName, "Restricted", "NotForLicensees", "Build", "LowLevelTests", $"{TestMetadata.TestName}.xml");
+
+			bool ModuleInRestrictedPath = IsRestrictedPath(ModuleDirectory);
 
-			if (IsRestrictedPath(ModuleDirectory))
+			if (ModuleInRestrictedPath)
 			{
 				GeneratedPropertiesScriptFile = NonPublicPath;
 			}
 			else
 			{
-				IsPublic = true;
-				GeneratedPropertiesScriptFile = Path.Combine(Unreal.EngineDirectory.FullName, "Build", "LowLevelTests_GenProps.xml");
-			}
-
-			// UE-133126
-			if (System.IO.File.Exists(NonPublicPath) && System.IO.Directory.GetFileSystemEntries(NonPublicFolder).Length == 1)
-			{
-				System.IO.File.Delete(NonPublicPath);
-				System.IO.Directory.Delete(NonPublicFolder);
-				System.IO.Directory.Delete(NotForLicenseesFolder);
-				System.IO.Directory.Delete(RestrictedFolder);
+				GeneratedPropertiesScriptFile = Path.Combine(Unreal.EngineDirectory.FullName, "Build", "LowLevelTests", $"{TestMetadata.TestName}.xml");
 			}
 
 			if (!System.IO.File.Exists(GeneratedPropertiesScriptFile))
 			{
 				string? DirGenProps = Path.GetDirectoryName(GeneratedPropertiesScriptFile);
-				if (DirGenProps != null)
+				if (DirGenProps != null && !System.IO.Directory.Exists(DirGenProps))
 				{
 					System.IO.Directory.CreateDirectory(DirGenProps);
 				}
-				using FileStream FileStream = System.IO.File.Create(GeneratedPropertiesScriptFile);
-				XDocument XInitFile = new XDocument(new XElement(BuildGraphNamespace + "BuildGraph", new XAttribute(XNamespace.Xmlns + "xsi", SchemaInstance), new XAttribute(SchemaInstance + "schemaLocation", SchemaLocation)));
-				XInitFile.Root?.Add(
-					new XElement(
-						BuildGraphNamespace + "Property",
-						new XAttribute("Name", "TestNames" + (!IsPublic ? "Restricted" : "")),
-						new XAttribute("Value", string.Empty)));
-
-				XInitFile.Save(FileStream);
+				using (FileStream FileStream = System.IO.File.Create(GeneratedPropertiesScriptFile))
+				{
+					XDocument XInitFile = new XDocument(new XElement(BuildGraphNamespace + "BuildGraph", new XAttribute(XNamespace.Xmlns + "xsi", SchemaInstance), new XAttribute(SchemaInstance + "schemaLocation", SchemaLocation)));
+					XInitFile.Root?.Add(
+						new XElement(
+							BuildGraphNamespace + "Property",
+							new XAttribute("Name", "TestNames"),
+							new XAttribute("Value", "$(TestNames);" + TestMetadata.TestName)));
+
+					XInitFile.Save(FileStream);
+				}
 			}
 
 			// All relevant properties
-			string TestTargetName = Target.LaunchModuleName ?? "NoLaunchModule";
+			string TestTargetName = Target.LaunchModuleName ?? "Launch";
 			string TestBinariesPath = TryGetBinariesPath();
 
 			// Do not save full paths
@@ -221,36 +214,25 @@ namespace UnrealBuildTool
 				XElement TestNames = (XElement)Root.FirstNode;
 				if (TestNames != null)
 				{
-					string? AllTestNames = TestNames.Attribute("Value")?.Value;
-					if (AllTestNames != null && !AllTestNames.Contains(TestMetadata.TestName) && AddToTestNames)
-					{
-						if (string.IsNullOrEmpty(AllTestNames))
-						{
-							AllTestNames += TestMetadata.TestName;
-						}
-						else if (!AllTestNames.Contains(TestMetadata.TestName))
-						{
-							AllTestNames += ";" + TestMetadata.TestName;
-						}
-					}
-					TestNames.Attribute("Value")?.SetValue(AllTestNames ?? "");
-
-
 					XElement lastUpdatedNode = TestNames;
-					InsertOrUpdateTestFlag(ref lastUpdatedNode, TestMetadata.TestName, "Disabled", TestMetadata.Disabled.ToString());
-					InsertOrUpdateTestFlag(ref lastUpdatedNode, TestMetadata.TestName, "Short", TestMetadata.TestShortName);
-					InsertOrUpdateTestFlag(ref lastUpdatedNode, TestMetadata.TestName, "Target", TestTargetName);
-					InsertOrUpdateTestFlag(ref lastUpdatedNode, TestMetadata.TestName, "BinariesRelative", TestBinariesPath);
-					InsertOrUpdateTestFlag(ref lastUpdatedNode, TestMetadata.TestName, "ReportType", TestMetadata.ReportType.ToString());
-
-					InsertOrUpdateTestOption(ref lastUpdatedNode, TestMetadata.TestName, TestMetadata.TestShortName, "Run", "Tests", false.ToString());
+						
+					InsertOrUpdateTestFlagProperty(ref lastUpdatedNode, TestMetadata.TestName, "Disabled", Convert.ToString(TestMetadata.Disabled));
+					InsertOrUpdateTestFlagProperty(ref lastUpdatedNode, TestMetadata.TestName, "Short", Convert.ToString(TestMetadata.TestShortName));
+					InsertOrUpdateTestFlagProperty(ref lastUpdatedNode, TestMetadata.TestName, "Target", Convert.ToString(TestTargetName));
+					InsertOrUpdateTestFlagProperty(ref lastUpdatedNode, TestMetadata.TestName, "BinariesRelative", Convert.ToString(TestBinariesPath));
+					InsertOrUpdateTestFlagProperty(ref lastUpdatedNode, TestMetadata.TestName, "ReportType", Convert.ToString(TestMetadata.ReportType));
+					InsertOrUpdateTestFlagProperty(ref lastUpdatedNode, TestMetadata.TestName, "InitialExtraArgs", Convert.ToString(TestMetadata.InitialExtraArgs));
+					InsertOrUpdateTestFlagProperty(ref lastUpdatedNode, TestMetadata.TestName, "HasAfterSteps", Convert.ToString(TestMetadata.HasAfterSteps));
+					InsertOrUpdateTestFlagProperty(ref lastUpdatedNode, TestMetadata.TestName, "UsesCatch2", Convert.ToString(TestMetadata.UsesCatch2));
+
+					InsertOrUpdateTestOption(ref lastUpdatedNode, TestMetadata.TestName, "Run" + TestMetadata.TestShortName + "Tests", "Run", "Tests", false.ToString());
 
 					List<UnrealTargetPlatform> AllSupportedPlatforms = new List<UnrealTargetPlatform>();
 					var SupportedPlatforms = GetType().GetCustomAttributes(typeof(SupportedPlatformsAttribute), false);
-					// If none specified we assume all platforms are supported by default
+					// If none specified we assume Win64
 					if (SupportedPlatforms.Length == 0)
 					{
-						UnrealTargetPlatform[] SupportedByDefault = { UnrealTargetPlatform.Win64, UnrealTargetPlatform.Mac, UnrealTargetPlatform.Linux, UnrealTargetPlatform.Android };
+						UnrealTargetPlatform[] SupportedByDefault = { UnrealTargetPlatform.Win64 };
 						AllSupportedPlatforms.AddRange(SupportedByDefault);
 					}
 					else
@@ -261,21 +243,81 @@ namespace UnrealBuildTool
 						}
 					}
 
-					InsertOrUpdateTestFlag(ref lastUpdatedNode, TestMetadata.TestName, "SupportedPlatforms", AllSupportedPlatforms.Aggregate("", (current, next) => (current == "" ? next.ToString() : current + ";" + next.ToString())));
+					InsertOrUpdateTestFlagProperty(ref lastUpdatedNode, TestMetadata.TestName, "SupportedPlatforms", AllSupportedPlatforms.Aggregate("", (current, next) => (current == "" ? next.ToString() : current + ";" + next.ToString())));
+				}
+			}
+
+			GenPropsDoc.Save(GeneratedPropertiesScriptFile);
+
+
+			// Platform-specific configurations
+			string GeneratedPropertiesPlatformFile;
+
+			string NonPublicPathPlatform;
+
+			// Generate peroperty file for each platform
+			foreach (var ValidPlatform in UnrealTargetPlatform.GetValidPlatforms())
+			{
+				bool IsRestrictedPlatformName = IsPlatformRestricted(ValidPlatform);
+				if (IsRestrictedPlatformName)
+				{
+					NonPublicPathPlatform = Path.Combine(Unreal.EngineDirectory.FullName, "Restricted", "NotForLicensees", "Platforms", ValidPlatform.ToString(), "Build", "LowLevelTests", $"{TestMetadata.TestName}.xml");
+				}
+				else
+				{
+					NonPublicPathPlatform = Path.Combine(Unreal.EngineDirectory.FullName, "Restricted", "NotForLicensees", "Build", "LowLevelTests", $"{TestMetadata.TestName}.xml");
+				}
+
+				if (ModuleInRestrictedPath)
+				{
+					GeneratedPropertiesPlatformFile = NonPublicPathPlatform;
+				}
+				else
+				{
+					if (IsRestrictedPlatformName)
+					{
+						GeneratedPropertiesPlatformFile = Path.Combine(Unreal.EngineDirectory.FullName, "Platforms", ValidPlatform.ToString(), "Build", "LowLevelTests", $"{TestMetadata.TestName}.xml");
+					}
+					else
+					{
+						GeneratedPropertiesPlatformFile = Path.Combine(Unreal.EngineDirectory.FullName, "Build", "LowLevelTests", $"{TestMetadata.TestName}.xml");
+					}
+				}
 
-					try
+				if (!System.IO.File.Exists(GeneratedPropertiesPlatformFile))
+				{
+					string? DirGenPropsPlatforms = Path.GetDirectoryName(GeneratedPropertiesPlatformFile);
+					if (DirGenPropsPlatforms != null && !System.IO.Directory.Exists(DirGenPropsPlatforms))
 					{
-						GenPropsDoc.Save(GeneratedPropertiesScriptFile);
+						System.IO.Directory.CreateDirectory(DirGenPropsPlatforms);
 					}
-					catch (UnauthorizedAccessException)
+					using (FileStream FileStream = System.IO.File.Create(GeneratedPropertiesPlatformFile))
 					{
-						// Expected on build machines.
-						// TODO: Ability to build for generate files and runnable tests.
+						 new XDocument(new XElement(BuildGraphNamespace + "BuildGraph", new XAttribute(XNamespace.Xmlns + "xsi", SchemaInstance), new XAttribute(SchemaInstance + "schemaLocation", SchemaLocation))).Save(FileStream);
 					}
 				}
+
+				XDocument XInitPlatformFile = XDocument.Load(GeneratedPropertiesPlatformFile);
+
+				// Adding per-test and per-platform tags
+				string TagsValue = TestMetadata.PlatformTags.ContainsKey(ValidPlatform) ? TestMetadata.PlatformTags[ValidPlatform] : string.Empty;
+				AppendOrUpdateTestFlagProperty(ref XInitPlatformFile, TestMetadata.TestName, ValidPlatform.ToString(), "Tags", TagsValue);
+
+				string ExtraCompilationArgsValue = TestMetadata.PlatformCompilationExtraArgs.ContainsKey(ValidPlatform) ? TestMetadata.PlatformCompilationExtraArgs[ValidPlatform] : string.Empty;
+				AppendOrUpdateTestFlagProperty(ref XInitPlatformFile, TestMetadata.TestName, ValidPlatform.ToString(), "ExtraCompilationArgs", ExtraCompilationArgsValue);
+
+				string RunSupportedValue = TestMetadata.PlatformsRunUnsupported.Contains(ValidPlatform) ? "False" : "True";
+				AppendOrUpdateTestFlagProperty(ref XInitPlatformFile, TestMetadata.TestName, ValidPlatform.ToString(), "RunSupported", RunSupportedValue);
+
+				XInitPlatformFile.Save(GeneratedPropertiesPlatformFile);
 			}
 		}
 
+		private bool IsPlatformRestricted(UnrealTargetPlatform Platform)
+		{
+			return RestrictedFolder.GetNames().Contains(Platform.ToString());
+		}
+
 		private bool IsRestrictedPath(string ModuleDirectory)
 		{
 			foreach (string RestrictedFolderName in RestrictedFolder.GetNames())
@@ -299,8 +341,23 @@ namespace UnrealBuildTool
 			return ModuleDirectory.Substring(0, SourceFolderIndex) + "Binaries";
 		}
 
-#pragma warning disable 8602
-		private void InsertOrUpdateTestFlag(ref XElement ElementUpsertAfter, string TestName, string FlagSuffix, string FlagValue)
+		private void AppendOrUpdateTestFlagProperty(ref XDocument Document, string FlagRadix, string FlagPrefix, string FlagSuffix, string FlagValue)
+		{
+			XElement? Existing = Document.Root!.Elements().Where(element => element.Attribute("Name").Value == FlagPrefix + FlagRadix + FlagSuffix).FirstOrDefault();
+			if (Existing != null)
+			{
+				Existing!.SetAttributeValue("Value", FlagValue);
+			}
+			else
+			{
+				XElement ElementAppend = new XElement(BuildGraphNamespace + "Property");
+				ElementAppend.SetAttributeValue("Name", FlagPrefix + FlagRadix + FlagSuffix);
+				ElementAppend.SetAttributeValue("Value", FlagValue);
+				Document.Root!.Add(ElementAppend);
+			}
+		}
+
+		private void InsertOrUpdateTestFlagProperty(ref XElement ElementUpsertAfter, string TestName, string FlagSuffix, string FlagValue)
 		{
 			IEnumerable<XElement> NextChunk = ElementUpsertAfter.ElementsAfterSelf(BuildGraphNamespace + "Property")
 				.Where(prop => prop.Attribute("Name").Value.EndsWith(FlagSuffix));
@@ -321,90 +378,110 @@ namespace UnrealBuildTool
 			ElementUpsertAfter = NextChunk.Last();
 		}
 
-		private void InsertOrUpdateTestOption(ref XElement ElementUpsertAfter, string TestName, string TestShortName, string OptionPrefix, string OptionSuffix, string DefaultValue)
+		private void InsertOrUpdateTestOption(ref XElement ElementUpsertAfter, string OptionRadix, string Description, string OptionPrefix, string OptionSuffix, string DefaultValue)
 		{
 			IEnumerable<XElement> NextChunk = ElementUpsertAfter.ElementsAfterSelf(BuildGraphNamespace + "Option")
 				.Where(prop => prop.Attribute("Name").Value.StartsWith(OptionPrefix) && prop.Attribute("Name").Value.EndsWith(OptionSuffix));
 			if (NextChunk
-				.Where(prop => prop.Attribute("Name").Value == OptionPrefix + TestName + OptionSuffix)
+				.Where(prop => prop.Attribute("Name").Value == OptionPrefix + OptionRadix + OptionSuffix)
 				.Count() == 0)
 			{
 				XElement ElementInsert = new XElement(BuildGraphNamespace + "Option");
-				ElementInsert.SetAttributeValue("Name", OptionPrefix + TestName + OptionSuffix);
+				ElementInsert.SetAttributeValue("Name", OptionPrefix + OptionRadix + OptionSuffix);
 				ElementInsert.SetAttributeValue("DefaultValue", DefaultValue);
-				ElementInsert.SetAttributeValue("Description", string.Format("{0} {1} {2}", OptionPrefix, TestShortName, OptionSuffix));
+				ElementInsert.SetAttributeValue("Description", Description);
 				ElementUpsertAfter.AddAfterSelf(ElementInsert);
 			}
 			else
 			{
 				XElement ElementUpdate = NextChunk
-					.Where(prop => prop.Attribute("Name").Value == OptionPrefix + TestName + OptionSuffix).First();
-				ElementUpdate.SetAttributeValue("Description", string.Format("{0} {1} {2}", OptionPrefix, TestShortName, OptionSuffix));
+					.Where(prop => prop.Attribute("Name").Value == OptionPrefix + OptionRadix + OptionSuffix).First();
+				ElementUpdate.SetAttributeValue("Description", Description);
 				ElementUpdate.SetAttributeValue("DefaultValue", DefaultValue);
 			}
 			ElementUpsertAfter = NextChunk.Last();
 		}
+
+#pragma warning restore 8604
 #pragma warning restore 8602
 
+#pragma warning disable 8618
 		/// <summary>
 		/// Test metadata class.
 		/// </summary>
 		public class Metadata
 		{
-			/// <summary>
-			/// Catch2 report type console - prints results to stdout.
-			/// </summary>
-			public static readonly string ConsoleReportType = "console";
-
 			/// <summary>
 			/// Test long name.
 			/// </summary>
-			public string TestName { get; }
+			public string TestName { get; set; }
 
 			/// <summary>
 			/// Test short name used for display in build system.
 			/// </summary>
-			public string TestShortName { get; }
+			public string TestShortName { get; set; }
 
 			/// <summary>
 			/// Type of Catch2 report, defaults to console.
 			/// </summary>
-			public string ReportType { get; }
+			public string ReportType { get; set; }
 
 			/// <summary>
 			/// Is this test disabled?
 			/// </summary>
-			public bool Disabled { get; }
+			public bool Disabled { get; set; }
+
+			/// <summary>
+			/// Any initial extra args to be passed to the test executable
+			/// </summary>
+			public string InitialExtraArgs { get; set; }
+
+			/// <summary>
+			/// Whether there's a step that gets executed after the tests have finished.
+			/// Typically used for cleanup of resources.
+			/// </summary>
+			public bool HasAfterSteps { get; set; }
+
+			private bool UsesCatch2Private = true;
+			/// <summary>
+			/// Test built with a frakework other than Catch2
+			/// </summary>
+			public bool UsesCatch2
+			{
+				get { return UsesCatch2Private; }
+				set { UsesCatch2Private = value; }
+			}
 
+			private Dictionary<UnrealTargetPlatform, string> PlatformTagsPrivate = new Dictionary<UnrealTargetPlatform, string>();
 			/// <summary>
-			/// Ctor that sets short name to test name.
+			/// Per-platform tags.
 			/// </summary>
-			public Metadata(string InTestName)
+			public Dictionary<UnrealTargetPlatform, string> PlatformTags
 			{
-				TestShortName = TestName = InTestName;
-				ReportType = ConsoleReportType;
+				get { return PlatformTagsPrivate; }
+				set { PlatformTagsPrivate = value; }
 			}
 
+			private Dictionary<UnrealTargetPlatform, string> PlatformCompilationExtraArgsPrivate = new Dictionary<UnrealTargetPlatform, string>();
 			/// <summary>
-			/// Ctor that sets short name and test name.
+			/// Per-platform extra compilation arguments.
 			/// </summary>
-			public Metadata(string InTestName, string InTestShortName)
+			public Dictionary<UnrealTargetPlatform, string> PlatformCompilationExtraArgs
 			{
-				TestName = InTestName;
-				TestShortName = InTestShortName;
-				ReportType = ConsoleReportType;
+				get { return PlatformCompilationExtraArgsPrivate; }
+				set { PlatformCompilationExtraArgsPrivate = value; }
 			}
 
+			private List<UnrealTargetPlatform> PlatformsRunUnsupportedPrivate = new List<UnrealTargetPlatform>() { UnrealTargetPlatform.Android, UnrealTargetPlatform.IOS };
 			/// <summary>
-			/// Ctor that sets short name, test name, report type and disabled status.
+			/// List of platforms that cannot run tests.
 			/// </summary>
-			public Metadata(string InTestName, string InTestShortName, string InReportType, bool InDisabled)
+			public List<UnrealTargetPlatform> PlatformsRunUnsupported
 			{
-				TestName = InTestName;
-				TestShortName = InTestShortName;
-				ReportType = InReportType;
-				Disabled = InDisabled;
+				get { return PlatformsRunUnsupportedPrivate;}
+				set { PlatformsRunUnsupportedPrivate = value; }
 			}
 		}
+#pragma warning restore 8618
 	}
 }
\ No newline at end of file
diff --git a/Engine/Source/Runtime/Core/Private/Unix/UnixPlatformOutputDevices.cpp b/Engine/Source/Runtime/Core/Private/Unix/UnixPlatformOutputDevices.cpp
--- a/Engine/Source/Runtime/Core/Private/Unix/UnixPlatformOutputDevices.cpp
+++ b/Engine/Source/Runtime/Core/Private/Unix/UnixPlatformOutputDevices.cpp
@@ -8,6 +8,7 @@
 #include "Misc/Build.h"
 #include "Misc/CommandLine.h"
 #include "Misc/OutputDeviceConsole.h"
+#include "Misc/OutputDeviceDebug.h"
 #include "Misc/OutputDeviceRedirector.h"
 #include "Misc/Parse.h"
 #include "Misc/Paths.h"
@@ -34,7 +35,16 @@ void FUnixOutputDevices::SetupOutputDevices()
 		}
 	}
 
-	// debug and event logging is not really supported on Unix. 
+#if USE_DEBUG_LOGGING
+	// If the platform has a separate debug output channel (e.g. OutputDebugString) then add an output device
+	// unless logging is turned off
+		if (FPlatformMisc::HasSeparateChannelForDebugOutput() && !FParse::Param(FCommandLine::Get(), TEXT("NODEBUGOUTPUT")))
+		{
+			GLog->AddOutputDevice(new FOutputDeviceDebug());
+		}
+#endif // USE_DEBUG_LOGGING
+
+	GLog->AddOutputDevice(FPlatformOutputDevices::GetEventLog());
 }
 
 FString FUnixOutputDevices::GetAbsoluteLogFilename()
