UnrealBuildTool: Print time spent in executor before throwing any execeptions

#rnx

diff --git a/Engine/Source/Programs/UnrealBuildTool/System/ActionGraph.cs b/Engine/Source/Programs/UnrealBuildTool/System/ActionGraph.cs
--- a/Engine/Source/Programs/UnrealBuildTool/System/ActionGraph.cs
+++ b/Engine/Source/Programs/UnrealBuildTool/System/ActionGraph.cs
@@ -400,13 +400,15 @@ namespace UnrealBuildTool
 
 				// Execute the build
 				Stopwatch Timer = Stopwatch.StartNew();
-				if (!await Executor.ExecuteActionsAsync(ActionsToExecute, Logger))
+				bool Result = await Executor.ExecuteActionsAsync(ActionsToExecute, Logger);
+
+				Logger.LogInformation("Total time in {ExecutorName} executor: {TotalSeconds:0.00} seconds", Executor.Name, Timer.Elapsed.TotalSeconds);
+
+				if (!Result)
 				{
 					throw new CompilationResultException(CompilationResult.OtherCompilationError);
 				}
 
-				Logger.LogInformation("Total time in {ExecutorName} executor: {TotalSeconds:0.00} seconds", Executor.Name, Timer.Elapsed.TotalSeconds);
-
 				// Reset the file info for all the produced items
 				foreach (LinkedAction BuildAction in ActionsToExecute)
 				{
diff --git a/Engine/Source/Programs/UnrealBuildTool/UnrealBuildTool.cs b/Engine/Source/Programs/UnrealBuildTool/UnrealBuildTool.cs
--- a/Engine/Source/Programs/UnrealBuildTool/UnrealBuildTool.cs
+++ b/Engine/Source/Programs/UnrealBuildTool/UnrealBuildTool.cs
@@ -662,10 +662,12 @@ namespace UnrealBuildTool
 				{
 					throw AggEx.InnerExceptions[0];
 				}
-
-				if ((ModeOptions & ToolModeOptions.ShowExecutionTime) != 0)
+				finally
 				{
-					Logger.LogInformation("Total execution time: {Time:0.00} seconds", Timeline.Elapsed.TotalSeconds);
+					if ((ModeOptions & ToolModeOptions.ShowExecutionTime) != 0)
+					{
+						Logger.LogInformation("Total execution time: {Time:0.00} seconds", Timeline.Elapsed.TotalSeconds);
+					}
 				}
 
 				return Result;
