UnrealBuildTool: Add msvc conformance flags to generated projects for intellisense

#rnx

diff --git a/Engine/Source/Programs/UnrealBuildTool/ProjectFiles/VisualStudio/VCProject.cs b/Engine/Source/Programs/UnrealBuildTool/ProjectFiles/VisualStudio/VCProject.cs
--- a/Engine/Source/Programs/UnrealBuildTool/ProjectFiles/VisualStudio/VCProject.cs
+++ b/Engine/Source/Programs/UnrealBuildTool/ProjectFiles/VisualStudio/VCProject.cs
@@ -630,6 +630,52 @@ namespace UnrealBuildTool
 			}
 		}
 
+		private string GetConformanceCompileArguments(TargetRules Target)
+		{
+			if (Target.Platform.IsInGroup(UnrealPlatformGroup.Microsoft) && Target.WindowsPlatform.Compiler.IsMSVC())
+			{
+				VersionNumber CompilerVersion = Target.WindowsPlatform.Environment?.CompilerVersion ?? new VersionNumber(0);
+				List<string> Arguments = new();
+				if (Target.WindowsPlatform.bStrictConformanceMode)
+				{
+					// This define is needed to ensure that MSVC static analysis mode doesn't declare attributes that are incompatible with strict conformance mode
+					Arguments.Add("/DSAL_NO_ATTRIBUTE_DECLARATIONS=1");
+
+					Arguments.Add("/permissive-");
+					Arguments.Add("/Zc:strictStrings-"); // Have to disable strict const char* semantics due to Windows headers not being compliant.
+					if (CompilerVersion >= new VersionNumber(14, 32) && CompilerVersion < new VersionNumber(14, 33, 31629))
+					{
+						Arguments.Add("/Zc:lambda-");
+					}
+				}
+				else
+				{
+					Arguments.Add("/Zc:hiddenFriend");
+				}
+
+				if (Target.WindowsPlatform.bUpdatedCPPMacro)
+				{
+					Arguments.Add("/Zc:__cplusplus");
+				}
+
+				if (Target.WindowsPlatform.bStrictInlineConformance)
+				{
+					Arguments.Add("/Zc:inline");
+				}
+
+				if (Target.WindowsPlatform.bStrictPreprocessorConformance)
+				{
+					Arguments.Add("/Zc:preprocessor");
+				}
+
+				if (Target.WindowsPlatform.bStrictEnumTypesConformance && CompilerVersion >= new VersionNumber(14, 34, 31931))
+				{
+					Arguments.Add("/Zc:enumTypes");
+				}
+			}
+			return string.Empty;
+		}
+
 		/// <summary>
 		/// Gets compiler switch for specifying in AdditionalOptions in .vcxproj file for coroutines support
 		/// </summary>
@@ -2073,7 +2119,7 @@ namespace UnrealBuildTool
 					VCProjectFileContent.AppendLine("    <NMakeReBuildCommandLine>{0} {1}</NMakeReBuildCommandLine>", EscapePath(NormalizeProjectPath(Builder.RebuildScript)), BuildArguments);
 					VCProjectFileContent.AppendLine("    <NMakeCleanCommandLine>{0} {1}</NMakeCleanCommandLine>", EscapePath(NormalizeProjectPath(Builder.CleanScript)), BuildArguments);
 					VCProjectFileContent.AppendLine("    <NMakeOutput>{0}</NMakeOutput>", NormalizeProjectPath(NMakePath.FullName));
-					VCProjectFileContent.AppendLine("    <AdditionalOptions>{0}</AdditionalOptions>", GetCppStandardCompileArgument(TargetRulesObject.CppStandard));
+					VCProjectFileContent.AppendLine("    <AdditionalOptions>{0} {1} {2}</AdditionalOptions>", GetCppStandardCompileArgument(TargetRulesObject.CppStandard), GetEnableCoroutinesArgument(), GetConformanceCompileArguments(TargetRulesObject));
 
 					if (TargetRulesObject.Type == TargetType.Game || TargetRulesObject.Type == TargetType.Client || TargetRulesObject.Type == TargetType.Server)
 					{
