1) Addition of the artifact interfaces
2) Moved common code from executors into an immediate action queue.
3) Added support in the immediate action queue to read/write artifacts

diff --git a/Engine/Source/Programs/UnrealBuildTool/Artifacts/Artifacts.cs b/Engine/Source/Programs/UnrealBuildTool/Artifacts/Artifacts.cs
new file mode 100644
--- /dev/null
+++ b/Engine/Source/Programs/UnrealBuildTool/Artifacts/Artifacts.cs
@@ -0,0 +1,345 @@
+// Copyright Epic Games, Inc. All Rights Reserved.
+
+using EpicGames.Core;
+using EpicGames.Serialization;
+using System;
+using System.Collections.Generic;
+using System.Linq;
+using System.Threading;
+using System.Threading.Tasks;
+
+namespace UnrealBuildTool.Artifacts
+{
+
+	/// <summary>
+	/// Artifacts can exist in different directory roots.
+	/// </summary>
+	public enum ArtifactDirectoryTree : byte
+	{
+
+		/// <summary>
+		/// Absolute path
+		/// </summary>
+		Absolute,
+
+		/// <summary>
+		/// Input/Output exists in the engine directory tree
+		/// </summary>
+		Engine,
+
+		/// <summary>
+		/// Input/Outputs exists in the project directory tree
+		/// </summary>
+		Project,
+	}
+
+	/// <summary>
+	/// Type of the artifact
+	/// </summary>
+	public enum ArtifactType : byte
+	{
+
+		/// <summary>
+		/// Artifact is a source file
+		/// </summary>
+		Source,
+
+		/// <summary>
+		/// Artifact is an object file
+		/// </summary>
+		Object,
+	}
+
+	/// <summary>
+	/// Represents a single artifact
+	/// </summary>
+	public readonly struct Artifact
+	{
+
+		/// <summary>
+		/// Directory tree containing the artifact
+		/// </summary>
+		public readonly ArtifactDirectoryTree Tree;
+
+		/// <summary>
+		/// Type of the artifact
+		/// </summary>
+		public readonly ArtifactType Type;
+
+		/// <summary>
+		/// Name of the artifact
+		/// </summary>
+		public readonly Utf8String Name;
+
+		/// <summary>
+		/// Hash of the artifact contents
+		/// </summary>
+		public readonly IoHash ContentHash;
+
+		/// <summary>
+		/// Constructor
+		/// </summary>
+		/// <param name="tree">Directory tree containing the artifact</param>
+		/// <param name="type">Type of the artifact</param>
+		/// <param name="name">Name of the artifact</param>
+		/// <param name="contentHash">Hash of the artifact contents</param>
+		public Artifact(ArtifactDirectoryTree tree, ArtifactType type, Utf8String name, IoHash contentHash)
+		{
+			Tree = tree;
+			Type = type;
+			Name = name;
+			ContentHash = contentHash;
+		}
+	}
+
+	/// <summary>
+	/// Collection of input and output artifacts.
+	/// </summary>
+	public readonly struct ArtifactBundle : IEquatable<ArtifactBundle>
+	{
+
+		/// <summary>
+		/// The hash of the primary input and the environment
+		/// </summary>
+		public readonly IoHash Key;
+
+		/// <summary>
+		/// The unique hash for all inputs and the environment
+		/// </summary>
+		public readonly IoHash BundleKey;
+
+		/// <summary>
+		/// Information about all inputs
+		/// </summary>
+		public readonly Artifact[] Inputs;
+
+		/// <summary>
+		/// Information about all outputs
+		/// </summary>
+		public readonly Artifact[] Outputs;
+
+		/// <summary>
+		/// Constructor
+		/// </summary>
+		/// <param name="key">The hash of the primary input and the environment</param>
+		/// <param name="bundleKey">The unique hash for all inputs and the environment</param>
+		/// <param name="inputs">Information about all inputs</param>
+		/// <param name="outputs">Information about all outputs</param>
+		public ArtifactBundle(IoHash key, IoHash bundleKey, Artifact[] inputs, Artifact[] outputs)
+		{
+			Key = key;
+			BundleKey = bundleKey;
+			Inputs = inputs;
+			Outputs = outputs;
+		}
+
+		/// <inheritdoc/>
+		public bool Equals(ArtifactBundle other)
+		{
+			return Key == other.Key && BundleKey == other.BundleKey && Inputs.SequenceEqual(other.Inputs) && Outputs.SequenceEqual(other.Outputs);
+		}
+
+		/// <inheritdoc/>
+		public override bool Equals(object? obj)
+		{
+			return obj is ArtifactBundle artifactBundle && Equals(artifactBundle);
+		}
+
+		/// <inheritdoc/>
+		public override int GetHashCode()
+		{
+			return HashCode.Combine(Key.GetHashCode(), BundleKey.GetHashCode(), Inputs.GetHashCode(), Outputs.GetHashCode()); 
+		}
+
+		/// <inheritdoc/>
+		public static bool operator ==(ArtifactBundle left, ArtifactBundle right)
+		{
+			return left.Equals(right);
+		}
+
+		/// <inheritdoc/>
+		public static bool operator !=(ArtifactBundle left, ArtifactBundle right)
+		{
+			return !(left == right);
+		}
+	}
+
+	/// <summary>
+	/// Collection of helper extension methods for working with IMemoryReader/Writer
+	/// </summary>
+	static class ArtifactSerializationExtensions
+	{
+
+		/// <summary>
+		/// Read an artifact structure
+		/// </summary>
+		/// <param name="reader">Source reader</param>
+		/// <returns>Read artifact structure</returns>
+		public static Artifact ReadArtifact(this IMemoryReader reader)
+		{
+			ArtifactDirectoryTree tree = (ArtifactDirectoryTree)reader.ReadUInt8();
+			ArtifactType type = (ArtifactType)reader.ReadUInt8();
+			Utf8String name = reader.ReadUtf8String();
+			IoHash contentHash = reader.ReadIoHash();
+			return new Artifact(tree, type, name, contentHash);
+		}
+
+		/// <summary>
+		/// Write an artifact structure
+		/// </summary>
+		/// <param name="writer">Destination writer</param>
+		/// <param name="artifact">Artifact to be written</param>
+		public static void WriteArtifact(this IMemoryWriter writer, Artifact artifact)
+		{
+			writer.WriteUInt8((byte)artifact.Tree);
+			writer.WriteUInt8((byte)artifact.Type);
+			writer.WriteUtf8String(artifact.Name);
+			writer.WriteIoHash(artifact.ContentHash);
+		}
+
+		/// <summary>
+		/// Read an artifact bundle structure
+		/// </summary>
+		/// <param name="reader">Source reader</param>
+		/// <returns>Read artifact bundle structure</returns>
+		public static ArtifactBundle ReadArtifactBundle(this IMemoryReader reader)
+		{
+			IoHash key = reader.ReadIoHash();
+			IoHash bundleKey = reader.ReadIoHash();
+			Artifact[] inputs = reader.ReadVariableLengthArray(() => reader.ReadArtifact());
+			Artifact[] outputs = reader.ReadVariableLengthArray(() => reader.ReadArtifact());
+			return new ArtifactBundle(key, bundleKey, inputs, outputs);
+		}
+
+		/// <summary>
+		/// Write an artifact bundle structure
+		/// </summary>
+		/// <param name="writer">Destination writer</param>
+		/// <param name="artifactBundle">Artifact bundle to be written</param>
+		public static void WriteArtifactBundle(this IMemoryWriter writer, ArtifactBundle artifactBundle)
+		{
+			writer.WriteIoHash(artifactBundle.Key);
+			writer.WriteIoHash(artifactBundle.BundleKey);
+			writer.WriteVariableLengthArray(artifactBundle.Inputs, x => writer.WriteArtifact(x));
+			writer.WriteVariableLengthArray(artifactBundle.Outputs, x => writer.WriteArtifact(x));
+		}
+	}
+
+	/// <summary>
+	/// Represents the state of the cache.  It is expect that after construction, the cache
+	/// can be in pending state, but this is optional. From the pending state, the cache 
+	/// becomes available or unavailable.
+	/// </summary>
+	public enum ArtifactCacheState
+	{
+
+		/// <summary>
+		/// The cache is still initializing
+		/// </summary>
+		Pending,
+
+		/// <summary>
+		/// There has been some form of cache failure and it is unavailable
+		/// </summary>
+		Unavailable,
+
+		/// <summary>
+		/// The cache is functional and ready to process requests
+		/// </summary>
+		Available,
+	}
+
+	/// <summary>
+	/// Interface for querying and adding artifacts
+	/// </summary>
+	public interface IArtifactCache
+	{
+
+		/// <summary>
+		/// Return true if the cache is ready to process requests
+		/// </summary>
+		/// <returns></returns>
+		public ArtifactCacheState State { get; }
+
+		/// <summary>
+		/// Return task that waits for the cache to be ready
+		/// </summary>
+		/// <returns>State of the artifact cache</returns>
+		public Task<ArtifactCacheState> WaitForReadyAsync();
+
+		/// <summary>
+		/// Given a collection of partial keys return all matching bundles
+		/// </summary>
+		/// <param name="partialKeys">Source file key</param>
+		/// <param name="cancellationToken">Token to be used to cancel operations</param>
+		/// <returns>Collection of all known artifacts</returns>
+		public Task<ArtifactBundle[]> QueryArtifactBundlesAsync(IoHash[] partialKeys, CancellationToken cancellationToken);
+
+		/// <summary>
+		/// Query the actual contents of a group of artifacts
+		/// </summary>
+		/// <param name="bundles">Bundles to be read</param>
+		/// <param name="cancellationToken">Token to be used to cancel operations</param>
+		/// <returns>Dictionary of the artifacts</returns>
+		public Task<bool[]?> QueryArtifactOutputsAsync(ArtifactBundle[] bundles, CancellationToken cancellationToken);
+
+		/// <summary>
+		/// Save new artifact to the cache
+		/// </summary>
+		/// <param name="artifactBundles">Collection of artifacts to be saved</param>
+		/// <param name="cancellationToken">Token to be used to cancel operations</param>
+		public Task SaveArtifactBundlesAsync(ArtifactBundle[] artifactBundles, CancellationToken cancellationToken);
+
+		/// <summary>
+		/// Flush all updates in the cache asynchronously.
+		/// </summary>
+		/// <param name="cancellationToken">Token to be used to cancel operations</param>
+		/// <returns>Async task objects</returns>
+		public Task FlushChangesAsync(CancellationToken cancellationToken);
+	}
+
+	/// <summary>
+	/// Interface for action specific support of artifacts.
+	/// </summary>
+	interface IActionArtifactCache
+	{
+
+		/// <summary>
+		/// Underlying artifact cache
+		/// </summary>
+		public IArtifactCache? ArtifactCache { get; }
+
+		/// <summary>
+		/// If true, reads will be serviced
+		/// </summary>
+		public bool EnableReads { get; set; }
+
+		/// <summary>
+		/// If true, writes will be propagated to storage
+		/// </summary>
+		public bool EnableWrites { get; set; }
+
+		/// <summary>
+		/// Complete an action from existing cached data
+		/// </summary>
+		/// <param name="action">Action to be completed</param>
+		/// <param name="cancellationToken">Token to be used to cancel operations</param>
+		/// <returns>True if it has been completed, false if not</returns>
+		public Task<bool> CompleteActionFromCacheAsync(LinkedAction action, CancellationToken cancellationToken);
+
+		/// <summary>
+		/// Save the output for a completed action
+		/// </summary>
+		/// <param name="action">Completed action</param>
+		/// <param name="cancellationToken">Token to be used to cancel operations</param>
+		/// <returns>Async task object</returns>
+		public Task ActionCompleteAsync(LinkedAction action, CancellationToken cancellationToken);
+
+		/// <summary>
+		/// Flush all updates in the cache asynchronously.
+		/// </summary>
+		/// <param name="cancellationToken">Token to be used to cancel operations</param>
+		/// <returns>Async task object</returns>
+		public Task FlushChangesAsync(CancellationToken cancellationToken);
+	}
+}
diff --git a/Engine/Source/Programs/UnrealBuildTool/Executors/Experimental/FASTBuild.cs b/Engine/Source/Programs/UnrealBuildTool/Executors/Experimental/FASTBuild.cs
--- a/Engine/Source/Programs/UnrealBuildTool/Executors/Experimental/FASTBuild.cs
+++ b/Engine/Source/Programs/UnrealBuildTool/Executors/Experimental/FASTBuild.cs
@@ -25,6 +25,7 @@ using UnrealBuildBase;
 using System.Runtime.Versioning;
 using Microsoft.Extensions.Logging;
 using System.Threading.Tasks;
+using UnrealBuildTool.Artifacts;
 
 namespace UnrealBuildTool
 {
@@ -459,8 +460,9 @@ namespace UnrealBuildTool
 			}
 		}
 
+		/// <inheritdoc/>
 		[SupportedOSPlatform("windows")]
-		public override async Task<bool> ExecuteActionsAsync(IEnumerable<LinkedAction> Actions, ILogger Logger)
+		public override async Task<bool> ExecuteActionsAsync(IEnumerable<LinkedAction> Actions, ILogger Logger, IActionArtifactCache? actionArtifactCache)
 		{
 			if (!Actions.Any())
 			{
@@ -479,7 +481,7 @@ namespace UnrealBuildTool
 				return ExecuteBffFile(FASTBuildFilePath, Logger);
 			}
 
-			return await LocalExecutor.ExecuteActionsAsync(Actions, Logger);
+			return await LocalExecutor.ExecuteActionsAsync(Actions, Logger, actionArtifactCache);
 		}
 
 		private void AddText(string StringToWrite)
diff --git a/Engine/Source/Programs/UnrealBuildTool/Executors/HybridExecutor.cs b/Engine/Source/Programs/UnrealBuildTool/Executors/HybridExecutor.cs
--- a/Engine/Source/Programs/UnrealBuildTool/Executors/HybridExecutor.cs
+++ b/Engine/Source/Programs/UnrealBuildTool/Executors/HybridExecutor.cs
@@ -5,6 +5,7 @@ using Microsoft.Extensions.Logging;
 using System.Collections.Generic;
 using System.Linq;
 using System.Threading.Tasks;
+using UnrealBuildTool.Artifacts;
 
 namespace UnrealBuildTool
 {
@@ -72,13 +73,8 @@ namespace UnrealBuildTool
 		/// </summary>
 		public override string Name => $"hybrid ({LocalExecutor.Name}+{RemoteExecutor.Name})";
 
-		/// <summary>
-		/// Execute the given actions
-		/// </summary>
-		/// <param name="ActionsToExecute">Actions to be executed</param>
-		/// <param name="Logger">Logger for output</param>
-		/// <returns>True if the build succeeded, false otherwise</returns>
-		public override async Task<bool> ExecuteActionsAsync(IEnumerable<LinkedAction> ActionsToExecute, ILogger Logger)
+		/// <inheritdoc/>
+		public override async Task<bool> ExecuteActionsAsync(IEnumerable<LinkedAction> ActionsToExecute, ILogger Logger, IActionArtifactCache? actionArtifactCache)
 		{
 			if (!ActionsToExecute.Any())
 			{
@@ -144,13 +140,13 @@ namespace UnrealBuildTool
 			}
 
 			// Execute the remote actions
-			if (RemoteActionsToExecute.Count > 0 && !await RemoteExecutor.ExecuteActionsAsync(RemoteActionsToExecute, Logger))
+			if (RemoteActionsToExecute.Count > 0 && !await RemoteExecutor.ExecuteActionsAsync(RemoteActionsToExecute, Logger, actionArtifactCache))
 			{
 				return false;
 			}
 
 			// Pass all the local actions through to the parallel executor
-			if (LocalActionsToExecute.Count > 0 && !await LocalExecutor.ExecuteActionsAsync(LocalActionsToExecute, Logger))
+			if (LocalActionsToExecute.Count > 0 && !await LocalExecutor.ExecuteActionsAsync(LocalActionsToExecute, Logger, actionArtifactCache))
 			{
 				return false;
 			}
diff --git a/Engine/Source/Programs/UnrealBuildTool/Executors/ImmediateActionQueue.cs b/Engine/Source/Programs/UnrealBuildTool/Executors/ImmediateActionQueue.cs
new file mode 100644
--- /dev/null
+++ b/Engine/Source/Programs/UnrealBuildTool/Executors/ImmediateActionQueue.cs
@@ -0,0 +1,946 @@
+// Copyright Epic Games, Inc. All Rights Reserved.
+
+using EpicGames.Core;
+using Microsoft.Extensions.Logging;
+using System;
+using System.Collections.Generic;
+using System.Linq;
+using System.Threading;
+using System.Threading.Tasks;
+using UnrealBuildTool.Artifacts;
+
+namespace UnrealBuildTool
+{
+	/// <summary>
+	/// Results from a run action
+	/// </summary>
+	/// <param name="LogLines">Console log lines</param>
+	/// <param name="ExitCode">Process return code.  Zero is assumed to be success and all other values an error.</param>
+	/// <param name="ExecutionTime">Wall time</param>
+	/// <param name="ProcessorTime">CPU time</param>
+	/// <param name="AdditionalDescription">Additional description of action</param>
+	record ExecuteResults(List<string> LogLines, int ExitCode, TimeSpan ExecutionTime, TimeSpan ProcessorTime, string? AdditionalDescription = null);
+
+	/// <summary>
+	/// Defines the phase of execution
+	/// </summary>
+	enum ActionPhase : byte
+	{
+
+		/// <summary>
+		/// Check for already existing artifacts for the inputs
+		/// </summary>
+		ArtifactCheck,
+
+		/// <summary>
+		/// Compile the action
+		/// </summary>
+		Compile,
+	}
+
+	/// <summary>
+	/// Defines the type of the runner.
+	/// </summary>
+	enum ImmediateActionQueueRunnerType : byte
+	{
+
+		/// <summary>
+		/// Will be used to queue jobs as part of general requests
+		/// </summary>
+		Automatic,
+
+		/// <summary>
+		/// Will only be used for manual requests to queue jobs
+		/// </summary>
+		Manual,
+	}
+
+
+	/// <summary>
+	/// Actions are assigned a runner when they need executing
+	/// </summary>
+	/// <param name="Type">Type of the runner</param>
+	/// <param name="ActionPhase">The action phase that this step supports</param>
+	/// <param name="RunAction">Action to be run queue running a task</param>
+	/// <param name="UseActionWeights">If true, use the action weight as a secondary limit</param>
+	/// <param name="MaxActions">Maximum number of action actions</param>
+	/// <param name="MaxActionWeight">Maximum weight of actions</param>
+	record ImmediateActionQueueRunner(ImmediateActionQueueRunnerType Type, ActionPhase ActionPhase, Func<LinkedAction, System.Action?> RunAction, bool UseActionWeights = false, int MaxActions = int.MaxValue, double MaxActionWeight = int.MaxValue)
+	{
+		/// <summary>
+		/// Current number of active actions
+		/// </summary>
+		public int ActiveActions = 0;
+
+		/// <summary>
+		/// Current weight of actions
+		/// </summary>
+		public double ActiveActionWeight = 0;
+
+		/// <summary>
+		/// True if the current limits have not been reached.
+		/// </summary>
+		public bool IsUnderLimits => ActiveActions < MaxActions && (!UseActionWeights || ActiveActionWeight < MaxActionWeight);
+	}
+
+	/// <summary>
+	/// Helper class to manage the action queue
+	///
+	/// Running the queue can be done with a mixture of automatic and manual runners. Runners are responsible for performing
+	/// the work associated with an action. Automatic runners will have actions automatically queued to them up to the point
+	/// that any runtime limits aren't exceeded (such as maximum number of concurrent processes).  Manual runners must have
+	/// jobs queued to them by calling TryStartOneAction or StartManyActions with the runner specified.
+	/// 
+	/// For example:
+	/// 
+	///		ParallelExecutor uses an automatic runner exclusively.
+	///		BoxExecutor uses an automatic runner to run jobs locally and a manual runner to run jobs remotely as processes 
+	///			become available.
+	/// </summary>
+	class ImmediateActionQueue : IDisposable
+	{
+
+		/// <summary>
+		/// Running status of the action
+		/// </summary>
+		private enum ActionStatus : byte
+		{
+			/// <summary>
+			/// Queued waiting for execution
+			/// </summary>
+			Queued,
+
+			/// <summary>
+			/// Action is running
+			/// </summary>
+			Running,
+
+			/// <summary>
+			/// Action has successfully finished
+			/// </summary>
+			Finished,
+
+			/// <summary>
+			/// Action has finished with an error
+			/// </summary>
+			Error,
+		}
+
+		/// <summary>
+		/// Used to track the state of each action
+		/// </summary>
+		private struct ActionState
+		{
+			/// <summary>
+			/// Action to be executed
+			/// </summary>
+			public LinkedAction Action;
+
+			/// <summary>
+			/// Phase of the action
+			/// </summary>
+			public ActionPhase Phase;
+
+			/// <summary>
+			/// Current status of the execution
+			/// </summary>
+			public ActionStatus Status;
+
+			/// <summary>
+			/// Runner assigned to the action
+			/// </summary>
+			public ImmediateActionQueueRunner? Runner;
+
+			/// <summary>
+			/// Optional execution results
+			/// </summary>
+			public ExecuteResults? Results;
+		};
+
+		/// <summary>
+		/// State information about each action
+		/// </summary>
+		private readonly ActionState[] Actions;
+
+		/// <summary>
+		/// Output logging
+		/// </summary>
+		public readonly ILogger Logger;
+
+		/// <summary>
+		/// Total number of actions
+		/// </summary>
+		public int TotalActions => Actions.Length;
+
+		/// <summary>
+		/// Process group
+		/// </summary>		
+		public readonly ManagedProcessGroup ProcessGroup = new();
+
+		/// <summary>
+		/// Source for the cancellation token
+		/// </summary>
+		public readonly CancellationTokenSource CancellationTokenSource = new();
+
+		/// <summary>
+		/// Cancellation token
+		/// </summary>
+		public CancellationToken CancellationToken => CancellationTokenSource.Token;
+
+		/// <summary>
+		/// Progress writer
+		/// </summary>
+		public readonly ProgressWriter ProgressWriter;
+
+		/// <summary>
+		/// Overall success of the action queue
+		/// </summary>
+		public bool Success = true;
+
+		/// <summary>
+		/// Whether to show compilation times along with worst offenders or not.
+		/// </summary>
+		public bool ShowCompilationTimes = false;
+
+		/// <summary>
+		/// Whether to show CPU utilization after the work is complete.
+		/// </summary>
+		public bool ShowCPUUtilization = false;
+
+		/// <summary>
+		/// Add target names for each action executed
+		/// </summary>
+		public bool PrintActionTargetNames = false;
+
+		/// <summary>
+		/// Whether to log command lines for actions being executed
+		/// </summary>
+		public bool LogActionCommandLines = false;
+
+		/// <summary>
+		/// Whether to show compilation times for each executed action
+		/// </summary>
+		public bool ShowPerActionCompilationTimes = false;
+
+		/// <summary>
+		/// Collapse non-error output lines
+		/// </summary>
+		public bool CompactOutput = false;
+
+		/// <summary>
+		/// When enabled, will stop compiling targets after a compile error occurs.
+		/// </summary>
+		public bool StopCompilationAfterErrors = false;
+
+		/// <summary>
+		/// Return true if the queue is done
+		/// </summary>
+		public bool IsDone => _doneTaskSource.Task.IsCompleted;
+
+		/// <summary>
+		/// Collection of available runners
+		/// </summary>
+		private readonly List<ImmediateActionQueueRunner> _runners = new();
+
+		/// <summary>
+		/// First action to start scanning for actions to run
+		/// </summary>
+		private int _firstPendingAction;
+
+		/// <summary>
+		/// Action to invoke when writing tool output
+		/// </summary>
+		private readonly Action<string> _writeToolOutput;
+
+		/// <summary>
+		/// Timer used to collect CPU utilization
+		/// </summary>
+		private Timer? _cpuUtilizationTimer;
+
+		/// <summary>
+		/// Per-second logging of cpu utilization
+		/// </summary>
+		private List<float> _cpuUtilization = new();
+
+		/// <summary>
+		/// Collection of all actions remaining to be logged
+		/// </summary>
+		private readonly List<int> _actionsToLog = new();
+
+		/// <summary>
+		/// Task waiting to process logging
+		/// </summary>
+		private Task? _actionsToLogTask = null;
+
+		/// <summary>
+		/// Used only by the logger to track the [x,total] output
+		/// </summary>
+		private int _loggedCompletedActions = 0;
+
+		/// <summary>
+		/// Tracks the number of completed actions.
+		/// </summary>
+		private int _completedActions = 0;
+
+		/// <summary>
+		/// Used to terminate the run with status
+		/// </summary>
+		private readonly TaskCompletionSource _doneTaskSource = new();
+
+		/// <summary>
+		/// If set, artifact cache used to retrieve previously compiled results and save new results
+		/// </summary>
+		private IActionArtifactCache? _actionArtifactCache;
+
+		static ExecuteResults s_copiedFromCacheResults = new(new List<string>(), 0, TimeSpan.Zero, TimeSpan.Zero, "copied from cache");
+
+		/// <summary>
+		/// Construct a new instance of the action queue
+		/// </summary>
+		/// <param name="actions">Collection of actions</param>
+		/// <param name="actionArtifactCache">If true, the artifact cache is being used</param>
+		/// <param name="maxActionArtifactCacheTasks">Max number of concurrent artifact cache tasks</param>
+		/// <param name="progressWriterText">Text to be displayed with the progress writer</param>
+		/// <param name="writeToolOutput">Action to invoke when writing tool output</param>
+		/// <param name="logger">Logging interface</param>
+		public ImmediateActionQueue(IEnumerable<LinkedAction> actions, IActionArtifactCache? actionArtifactCache, int maxActionArtifactCacheTasks, string progressWriterText, Action<string> writeToolOutput, ILogger logger)
+		{
+			int count = actions.Count();
+			Actions = new ActionState[count];
+
+			Logger = logger;
+			ProgressWriter = new(progressWriterText, false, logger);
+			_actionArtifactCache = actionArtifactCache;
+			_writeToolOutput = writeToolOutput;
+
+			bool readArtifacts = _actionArtifactCache != null && _actionArtifactCache.EnableReads;
+			ActionPhase initialPhase = readArtifacts ? ActionPhase.ArtifactCheck : ActionPhase.Compile;
+			int index = 0;
+			foreach (LinkedAction action in actions)
+			{
+				action.SortIndex = index;
+				Actions[index++] = new ActionState { Action = action, Status = ActionStatus.Queued, Phase = initialPhase, Results = null };
+			}
+
+			if (readArtifacts)
+			{
+				var runAction = (LinkedAction action) =>
+				{
+					return new System.Action(async () =>
+					{
+						bool success = await _actionArtifactCache!.CompleteActionFromCacheAsync(action, CancellationToken);
+						if (success)
+						{
+							OnActionCompleted(action, success, s_copiedFromCacheResults);
+						}
+						else
+						{
+							RequeueAction(action);
+						}
+					});
+				};
+
+				_runners.Add(new(ImmediateActionQueueRunnerType.Automatic, ActionPhase.ArtifactCheck, runAction, false, maxActionArtifactCacheTasks, 0));
+			}
+		}
+
+		/// <summary>
+		/// Create a new automatic runner
+		/// </summary>
+		/// <param name="runAction">Action to be run queue running a task</param>
+		/// <param name="useActionWeights">If true, use the action weight as a secondary limit</param>
+		/// <param name="maxActions">Maximum number of action actions</param>
+		/// <param name="maxActionWeight">Maximum weight of actions</param>
+		/// <returns>Created runner</returns>
+		public ImmediateActionQueueRunner CreateAutomaticRunner(Func<LinkedAction, System.Action?> runAction, bool useActionWeights, int maxActions, double maxActionWeight)
+		{
+			ImmediateActionQueueRunner runner = new(ImmediateActionQueueRunnerType.Automatic, ActionPhase.Compile, runAction, useActionWeights, maxActions, maxActionWeight);
+			_runners.Add(runner);
+			return runner;
+		}
+
+		/// <summary>
+		/// Create a manual runner
+		/// </summary>
+		/// <param name="runAction">Action to be run queue running a task</param>
+		/// <param name="useActionWeights">If true, use the action weight as a secondary limit</param>
+		/// <param name="maxActions">Maximum number of action actions</param>
+		/// <param name="maxActionWeight">Maximum weight of actions</param>
+		/// <returns>Created runner</returns>
+		public ImmediateActionQueueRunner CreateManualRunner(Func<LinkedAction, System.Action?> runAction, bool useActionWeights = false, int maxActions = int.MaxValue, double maxActionWeight = double.MaxValue)
+		{
+			ImmediateActionQueueRunner runner = new(ImmediateActionQueueRunnerType.Manual, ActionPhase.Compile, runAction, useActionWeights, maxActions, maxActionWeight);
+			_runners.Add(runner);
+			return runner;
+		}
+
+		/// <summary>
+		/// Start the process of running all the actions
+		/// </summary>
+		public void Start()
+		{
+			if (ShowCPUUtilization)
+			{
+				_cpuUtilizationTimer = new(x =>
+				{
+					lock (_cpuUtilization)
+					{
+						if (Utils.GetTotalCpuUtilization(out float cpuUtilization))
+						{
+							_cpuUtilization.Add(cpuUtilization);
+						}
+					}
+				}, null, 1000, 1000);
+			}
+
+			Logger.LogInformation("------ Building {TotalActions} action(s) started ------", TotalActions);
+		}
+
+		/// <summary>
+		/// Run the actions until complete
+		/// </summary>
+		/// <returns>True if completed successfully</returns>
+		public async Task<bool> RunTillDone()
+		{
+			await _doneTaskSource.Task;
+			TraceSummary();
+			return Success;
+		}
+
+		/// <summary>
+		/// Return an enumeration of ready compile tasks.  This is not executed under a lock and 
+		/// does not modify the state of any actions.
+		/// </summary>
+		/// <returns>Enumerations of all ready to compile actions.</returns>
+		public IEnumerable<LinkedAction> EnumerateReadyToCompileActions()
+		{
+			foreach (ActionState actionState in Actions)
+			{
+				if (actionState.Status == ActionStatus.Queued && 
+					actionState.Phase == ActionPhase.Compile && 
+					GetActionReadyState(actionState.Action) == ActionReadyState.Ready)
+				{
+					yield return actionState.Action;
+				}				
+			}
+		}
+
+		/// <summary>
+		/// Try to start one action
+		/// </summary>
+		/// <param name="runner">If specified, tasks will only be queued to this runner.  Otherwise all manual runners will be used.</param>
+		/// <returns>True if an action was run, false if not</returns>
+		public bool TryStartOneAction(ImmediateActionQueueRunner? runner = null)
+		{
+			lock (Actions)
+			{
+				// If we are starting deeper in the action collection, then never set the first pending action location
+				bool foundFirstPending = false;
+
+				// Loop through the items
+				for (int actionIndex = _firstPendingAction; actionIndex != Actions.Length; ++actionIndex)
+				{
+
+					// If the action has already reached the compilation state, then just ignore
+					if (Actions[actionIndex].Status != ActionStatus.Queued)
+					{
+						continue;
+					}
+
+					// If needed, update the first valid slot for searching for actions to run
+					if (!foundFirstPending)
+					{
+						_firstPendingAction = actionIndex;
+						foundFirstPending = true;
+					}
+
+					// Based on the ready state, use this action or mark as an error
+					switch (GetActionReadyState(Actions[actionIndex].Action))
+					{
+						case ActionReadyState.NotReady:
+							break;
+
+						case ActionReadyState.Error:
+							Actions[actionIndex].Status = ActionStatus.Error;
+							IncrementCompletedActions();
+							break;
+
+						case ActionReadyState.Ready:
+							ImmediateActionQueueRunner? selectedRunner = null;
+							System.Action? action = null;
+							if (runner != null)
+							{
+								if (runner.IsUnderLimits && runner.ActionPhase == Actions[actionIndex].Phase)
+								{
+									action = runner.RunAction(Actions[actionIndex].Action);
+									if (action != null)
+									{
+										selectedRunner = runner;
+									}
+								}
+							}
+							else
+							{
+								foreach (ImmediateActionQueueRunner tryRunner in _runners)
+								{
+									if (tryRunner.IsUnderLimits && tryRunner.ActionPhase == Actions[actionIndex].Phase)
+									{
+										action = tryRunner.RunAction(Actions[actionIndex].Action);
+										if (action != null)
+										{
+											selectedRunner = tryRunner;
+											break;
+										}
+									}
+								}
+							}
+							
+							if (action != null && selectedRunner != null)
+							{
+								Actions[actionIndex].Status = ActionStatus.Running;
+								Actions[actionIndex].Runner = selectedRunner;
+								selectedRunner.ActiveActions++;
+								selectedRunner.ActiveActionWeight += Actions[actionIndex].Action.Weight;
+								Task.Factory.StartNew(() => action(), CancellationToken, TaskCreationOptions.LongRunning | TaskCreationOptions.PreferFairness, TaskScheduler.Default);
+								return true;
+							}
+							break;
+
+						default:
+							throw new ApplicationException("Unexpected action ready state");
+					}
+				}
+
+				// We found nothing, check to see if we have no active running tasks and no manual runners.
+				// If we don't, then it is assumed we can't queue any more work.
+				bool prematureDone = true;
+				foreach (ImmediateActionQueueRunner tryRunner in _runners)
+				{
+					if (tryRunner.Type == ImmediateActionQueueRunnerType.Manual || tryRunner.ActiveActions != 0)
+					{
+						prematureDone = false;
+						break;
+					}
+				}
+				if (prematureDone && !_doneTaskSource.Task.IsCompleted)
+				{
+					_doneTaskSource.SetResult();
+				}
+				return false;
+			}
+		}
+
+		/// <summary>
+		/// Start as many actions as possible.
+		/// </summary>
+		/// <param name="runner">If specified, all actions will be limited to the runner</param>
+		public void StartManyActions(ImmediateActionQueueRunner? runner = null)
+		{
+			while (TryStartOneAction(runner))
+			{ }
+		}
+
+		/// <summary>
+		/// Dispose the object
+		/// </summary>
+		public void Dispose()
+		{
+			if (_cpuUtilizationTimer != null)
+			{
+				_cpuUtilizationTimer.Dispose();
+			}
+			ProcessGroup.Dispose();
+		}
+
+		/// <summary>
+		/// Reset the status of the given action to the queued state.
+		/// </summary>
+		/// <param name="action">Action being re-queued</param>
+		public void RequeueAction(LinkedAction action)
+		{
+			SetActionState(action, ActionStatus.Queued, null);
+		}
+
+		/// <summary>
+		/// Notify the system that an action has been completed.
+		/// </summary>
+		/// <param name="action">Action being completed</param>
+		/// <param name="success">If true, the action succeeded or false if it failed.</param>
+		/// <param name="results">Optional execution results</param>
+		public void OnActionCompleted(LinkedAction action, bool success, ExecuteResults? results)
+		{
+			SetActionState(action, success ? ActionStatus.Finished : ActionStatus.Error, results);
+		}
+
+		/// <summary>
+		/// Set the new state of an action.  Can only set state to Queued, Finished, or Error
+		/// </summary>
+		/// <param name="action">Action being set</param>
+		/// <param name="status">Status to set</param>
+		/// <param name="results">Optional results</param>
+		/// <exception cref="BuildException"></exception>
+		private void SetActionState(LinkedAction action, ActionStatus status, ExecuteResults? results)
+		{
+			int actionIndex = action.SortIndex;
+
+			// Update the actions data
+			lock (Actions)
+			{
+				ImmediateActionQueueRunner runner = Actions[actionIndex].Runner ?? throw new BuildException("Attempting to update action state but runner isn't set");
+				runner.ActiveActions--;
+				runner.ActiveActionWeight -= Actions[actionIndex].Action.Weight;
+
+				// If we are doing an artifact check, then move to compile phase
+				if (Actions[actionIndex].Phase == ActionPhase.ArtifactCheck)
+				{
+					Actions[actionIndex].Phase = ActionPhase.Compile;
+					if (status != ActionStatus.Finished)
+					{
+						status = ActionStatus.Queued;
+					}
+				}
+
+				Actions[actionIndex].Status = status;
+				Actions[actionIndex].Runner = null;
+				if (results != null)
+				{
+					Actions[actionIndex].Results = results;
+				}
+
+				// Add the action to the logging queue
+				if (status != ActionStatus.Queued)
+				{
+					lock (_actionsToLog)
+					{
+						_actionsToLog.Add(action.SortIndex);
+						if (_actionsToLogTask == null)
+						{
+							_actionsToLogTask = Task.Run(LogActions);
+						}
+					}
+				}
+
+				switch (status)
+				{
+					case ActionStatus.Queued:
+						if (actionIndex < _firstPendingAction)
+						{
+							_firstPendingAction = actionIndex;
+						}
+						break;
+
+					case ActionStatus.Finished:
+						// Notify the artifact handler of the action completing.  We don't wait on the resulting task.  The
+						// cache is required to remember any pending saves and a final call to Flush will wait for everything to complete.
+						_actionArtifactCache?.ActionCompleteAsync(action, CancellationToken);
+						IncrementCompletedActions();
+						break;
+
+					case ActionStatus.Error:
+						Success = false;
+						IncrementCompletedActions();
+						break;
+
+					default:
+						throw new BuildException("Unexpected action status set");
+				}
+			}
+
+			// Since something has been completed or returned to the queue, try to run actions again
+			StartManyActions();
+		}
+
+		/// <summary>
+		/// Increment the number of completed actions and signal done if all actions complete.
+		/// This must be executed under the lock of Actions.
+		/// </summary>
+		private void IncrementCompletedActions()
+		{
+			_completedActions++;
+			if (_completedActions >= Actions.Length && !_doneTaskSource.Task.IsCompleted)
+			{
+				_doneTaskSource.SetResult();
+			}
+		}
+
+		/// <summary>
+		/// Purge the pending logging actions
+		/// </summary>
+		private void LogActions()
+		{
+			for (; ; )
+			{
+				int[]? actionsToLog = null;
+				lock (_actionsToLog)
+				{
+					if (_actionsToLog.Count == 0)
+					{
+						_actionsToLogTask = null;
+					}
+					else
+					{
+						actionsToLog = _actionsToLog.ToArray();
+						_actionsToLog.Clear();
+					}
+				}
+
+				if (actionsToLog == null)
+				{
+					return;
+				}
+
+				foreach (int index in actionsToLog)
+				{
+					LogAction(Actions[index].Action, Actions[index].Results);
+				}
+			}
+		}
+
+		private static int s_previousLineLength = -1;
+
+		/// <summary>
+		/// Log an action that has completed
+		/// </summary>
+		/// <param name="action">Action that has completed</param>
+		/// <param name="executeTaskResult">Results of the action</param>
+		private void LogAction(LinkedAction action, ExecuteResults? executeTaskResult)
+		{
+			List<string>? logLines = null;
+			int exitCode = int.MaxValue;
+			TimeSpan executionTime = TimeSpan.Zero;
+			TimeSpan processorTime = TimeSpan.Zero;
+			string? additionalDescription = null;
+			if (executeTaskResult != null)
+			{
+				logLines = executeTaskResult.LogLines;
+				exitCode = executeTaskResult.ExitCode;
+				executionTime = executeTaskResult.ExecutionTime;
+				processorTime = executeTaskResult.ProcessorTime;
+				additionalDescription = executeTaskResult.AdditionalDescription;
+			}
+
+			// Write it to the log
+			string description = string.Empty;
+			if (action.bShouldOutputStatusDescription || (logLines != null && logLines.Count == 0))
+			{
+				description = $"{(action.CommandDescription ?? action.CommandPath.GetFileNameWithoutExtension())} {action.StatusDescription}".Trim();
+			}
+			else if (logLines != null && logLines.Count > 0)
+			{
+				description = $"{(action.CommandDescription ?? action.CommandPath.GetFileNameWithoutExtension())} {logLines[0]}".Trim();
+			}
+			if (!string.IsNullOrEmpty(additionalDescription))
+			{
+				description = $"{description} {additionalDescription}";
+			}
+
+			lock (ProgressWriter)
+			{
+				int totalActions = Actions.Length;
+				int completedActions = Interlocked.Increment(ref _loggedCompletedActions);
+				ProgressWriter.Write(completedActions, Actions.Length);
+
+				// Canceled
+				if (exitCode == int.MaxValue)
+				{
+					Logger.LogInformation("[{CompletedActions}/{TotalActions}] {Description} canceled", completedActions, totalActions, description);
+					return;
+				}
+
+				string targetDetails = "";
+				TargetDescriptor? target = action.Target;
+				if (PrintActionTargetNames && target != null)
+				{
+					targetDetails = $"[{target.Name} {target.Platform} {target.Configuration}]";
+				}
+
+				if (LogActionCommandLines)
+				{
+					Logger.LogDebug("[{CompletedActions}/{TotalActions}]{TargetDetails} Command: {CommandPath} {CommandArguments}", completedActions, totalActions, targetDetails, action.CommandPath, action.CommandArguments);
+				}
+
+				string compilationTimes = "";
+
+				if (ShowPerActionCompilationTimes)
+				{
+					if (processorTime.Ticks > 0)
+					{
+						compilationTimes = $" (Wall: {executionTime.TotalSeconds:0.00}s CPU: {processorTime.TotalSeconds:0.00}s)";
+					}
+					else
+					{
+						compilationTimes = $" (Wall: {executionTime.TotalSeconds:0.00}s)";
+					}
+				}
+
+				string message = ($"[{completedActions}/{totalActions}]{targetDetails}{compilationTimes} {description}");
+
+				if (CompactOutput)
+				{
+					if (s_previousLineLength > 0)
+					{
+						// move the cursor to the far left position, one line back
+						Console.CursorLeft = 0;
+						Console.CursorTop -= 1;
+						// clear the line
+						Console.Write("".PadRight(s_previousLineLength));
+						// move the cursor back to the left, so output is written to the desired location
+						Console.CursorLeft = 0;
+					}
+				}
+
+				s_previousLineLength = message.Length;
+
+				_writeToolOutput(message);
+				if (logLines != null)
+				{
+					foreach (string Line in logLines.Skip(action.bShouldOutputStatusDescription ? 0 : 1))
+					{
+						// suppress library creation messages when writing compact output
+						if (CompactOutput && Line.StartsWith("   Creating library ") && Line.EndsWith(".exp"))
+						{
+							continue;
+						}
+
+						_writeToolOutput(Line);
+
+						// Prevent overwriting of logged lines
+						s_previousLineLength = -1;
+					}
+				}
+
+				if (exitCode != 0)
+				{
+					// BEGIN TEMPORARY TO CATCH PVS-STUDIO ISSUES
+					if (logLines == null || logLines.Count == 0)
+					{
+						Logger.LogError("{TargetDetails} {Description}: Exited with error code {ExitCode}", targetDetails, description, exitCode);
+						Logger.LogInformation("{TargetDetails} {Description}: WorkingDirectory {WorkingDirectory}", targetDetails, description, action.WorkingDirectory);
+						Logger.LogInformation("{TargetDetails} {Description}: {CommandPath} {CommandArguments}", targetDetails, description, action.CommandPath, action.CommandArguments);
+					}
+					// END TEMPORARY
+
+					// prevent overwriting of error text
+					s_previousLineLength = -1;
+
+					// Cancel all other pending tasks
+					if (StopCompilationAfterErrors)
+					{
+						CancellationTokenSource.Cancel();
+						if (!_doneTaskSource.Task.IsCompleted)
+						{
+							_doneTaskSource.SetResult();
+						}
+					}
+				}
+			}
+		}
+
+		/// <summary>
+		/// Generate the final summary display
+		/// </summary>
+		private void TraceSummary()
+		{
+
+			// Wait for logging to complete
+			Task? loggingTask = null;
+			lock (_actionsToLog)
+			{
+				loggingTask = _actionsToLogTask;
+			}
+			loggingTask?.Wait();
+
+			if (ShowCPUUtilization)
+			{
+				lock (_cpuUtilization)
+				{
+					if (_cpuUtilization.Count > 0)
+					{
+						Logger.LogInformation("");
+						Logger.LogInformation("Average CPU Utilization: {CPUPercentage}%", (int)(_cpuUtilization.Average()));
+					}
+				}
+			}
+
+			if (!ShowCompilationTimes)
+			{
+				return;
+			}
+
+			Logger.LogInformation("");
+			if (ProcessGroup.TotalProcessorTime.Ticks > 0)
+			{
+				Logger.LogInformation("Total CPU Time: {TotalSeconds} s", ProcessGroup.TotalProcessorTime.TotalSeconds);
+				Logger.LogInformation("");
+			}
+
+			IEnumerable<int> CompletedActions = Enumerable.Range(0, Actions.Length)
+				.Where(x => Actions[x].Results != null && Actions[x].Results!.ExecutionTime > TimeSpan.Zero)
+				.OrderByDescending(x => Actions[x].Results!.ExecutionTime)
+				.Take(20);
+
+			if (CompletedActions.Any())
+			{
+				Logger.LogInformation("Compilation Time Top {CompletedTaskCount}", CompletedActions.Count());
+				Logger.LogInformation("");
+				foreach (int Index in CompletedActions)
+				{
+					IExternalAction Action = Actions[Index].Action.Inner;
+					ExecuteResults Result = Actions[Index].Results!;
+
+					string Description = $"{(Action.CommandDescription ?? Action.CommandPath.GetFileNameWithoutExtension())} {Action.StatusDescription}".Trim();
+					if (Result.ProcessorTime.Ticks > 0)
+					{
+						Logger.LogInformation("{Description} [ Wall Time {ExecutionTime:0.00} s / CPU Time {ProcessorTime:0.00} s ]", Description, Result.ExecutionTime.TotalSeconds, Result.ProcessorTime.TotalSeconds);
+					}
+					else
+					{
+						Logger.LogInformation("{Description} [ Time {ExecutionTime:0.00} s ]", Description, Result.ExecutionTime.TotalSeconds);
+					}
+
+				}
+				Logger.LogInformation("");
+			}
+		}
+
+		private enum ActionReadyState
+		{
+			NotReady,
+			Error,
+			Ready,
+		}
+
+		/// <summary>
+		/// Get the ready state of an action
+		/// </summary>
+		/// <param name="action">Action in question</param>
+		/// <returns>Action ready state</returns>
+		private ActionReadyState GetActionReadyState(LinkedAction action)
+		{
+			foreach (LinkedAction prereq in action.PrerequisiteActions)
+			{
+
+				// To avoid doing artifact checks on actions that might need compiling,
+				// we first make sure the action is in the compile phase
+				if (Actions[prereq.SortIndex].Phase != ActionPhase.Compile)
+				{
+					return ActionReadyState.NotReady;
+				}
+
+				// Respect the compile status of the action
+				switch (Actions[prereq.SortIndex].Status)
+				{
+					case ActionStatus.Finished:
+						continue;
+
+					case ActionStatus.Error:
+						return ActionReadyState.Error;
+
+					default:
+						return ActionReadyState.NotReady;
+				}
+			}
+			return ActionReadyState.Ready;
+		}
+	}
+}
diff --git a/Engine/Source/Programs/UnrealBuildTool/Executors/ParallelExecutor.cs b/Engine/Source/Programs/UnrealBuildTool/Executors/ParallelExecutor.cs
--- a/Engine/Source/Programs/UnrealBuildTool/Executors/ParallelExecutor.cs
+++ b/Engine/Source/Programs/UnrealBuildTool/Executors/ParallelExecutor.cs
@@ -9,9 +9,11 @@ using System.IO;
 using System.Linq;
 using System.Threading;
 using System.Threading.Tasks;
+using UnrealBuildTool.Artifacts;
 
 namespace UnrealBuildTool
 {
+
 	/// <summary>
 	/// This executor uses async Tasks to process the action graph
 	/// </summary>
@@ -150,222 +152,69 @@ namespace UnrealBuildTool
 			return true;
 		}
 
-		protected class ExecuteResults
+		/// <summary>
+		/// Create an action queue
+		/// </summary>
+		/// <param name="actionsToExecute">Actions to be executed</param>
+		/// <param name="actionArtifactCache">Artifact cache</param>
+		/// <param name="logger">Logging interface</param>
+		/// <returns>Action queue</returns>
+		public ImmediateActionQueue CreateActionQueue(IEnumerable<LinkedAction> actionsToExecute, IActionArtifactCache? actionArtifactCache, ILogger logger)
 		{
-			public List<string> LogLines { get; private set; }
-			public int ExitCode { get; private set; }
-			public TimeSpan ExecutionTime { get; private set; }
-			public TimeSpan ProcessorTime { get; private set; }
-			public string? AdditionalDescription { get; protected set; } = null;
-
-			public ExecuteResults(List<string> LogLines, int ExitCode, TimeSpan ExecutionTime, TimeSpan ProcessorTime, string? AdditionalDescription = null)
+			return new(actionsToExecute, actionArtifactCache, NumParallelProcesses, "Compiling C++ source code...", x => WriteToolOutput(x), logger)
 			{
-				this.LogLines = LogLines;
-				this.ExitCode = ExitCode;
-				this.ProcessorTime = ProcessorTime;
-				this.ExecutionTime = ExecutionTime;
-				this.AdditionalDescription = AdditionalDescription;
-			}
-			public ExecuteResults(List<string> LogLines, int ExitCode)
-			{
-				this.LogLines = LogLines;
-				this.ExitCode = ExitCode;
-			}
+				ShowCompilationTimes = bShowCompilationTimes,
+				ShowCPUUtilization = bShowCPUUtilization,
+				PrintActionTargetNames = bPrintActionTargetNames,
+				LogActionCommandLines = bLogActionCommandLines,
+				ShowPerActionCompilationTimes = bShowPerActionCompilationTimes,
+				CompactOutput = bCompactOutput,
+				StopCompilationAfterErrors = bStopCompilationAfterErrors,
+			};
 		}
 
-		internal enum ActionStatus : byte
-		{
-			Queued,
-			Running,
-			Finished,
-			Error,
-		}
-
-		public override Task<bool> ExecuteActionsAsync(IEnumerable<LinkedAction> ActionsToExecute, ILogger Logger)
-		{
-			return Task.FromResult(ExecuteActions(ActionsToExecute, Logger));
-		}
-
-		/// <summary>
-		/// Executes the specified actions locally.
-		/// </summary>
-		/// <returns>True if all the tasks successfully executed, or false if any of them failed.</returns>
-		bool ExecuteActions(IEnumerable<LinkedAction> InputActions, ILogger Logger)
+		/// <inheritdoc/>
+		public override async Task<bool> ExecuteActionsAsync(IEnumerable<LinkedAction> ActionsToExecute, ILogger Logger, IActionArtifactCache? actionArtifactCache)
 		{
-			if (!InputActions.Any())
+			if (!ActionsToExecute.Any())
 			{
 				return true;
 			}
 
-			int TotalActions = InputActions.Count();
-			int NumCompletedActions = 0;
-
-			int ActualNumParallelProcesses = Math.Min(TotalActions, NumParallelProcesses);
-
-			using ManagedProcessGroup ProcessGroup = new ManagedProcessGroup();
-			CancellationTokenSource CancellationTokenSource = new CancellationTokenSource();
-			CancellationToken CancellationToken = CancellationTokenSource.Token;
-			using ProgressWriter ProgressWriter = new ProgressWriter("Compiling C++ source code...", false, Logger);
-
-			Dictionary<LinkedAction, Task<ExecuteResults>> ActionToTaskLookup = new();
-
-			Logger.LogInformation($"------ Building {TotalActions} action(s) started ------");
-
-			var Actions = new LinkedAction[TotalActions];
-			var ActionsStatus = new ActionStatus[TotalActions];
-			int Index = 0;
-			foreach (var Action in InputActions)
+			// The "useAutomaticQueue" should always be true unless manual queue is being tested
+			bool useAutomaticQueue = true;
+			if (useAutomaticQueue)
 			{
-				Actions[Index] = Action;
-				ActionsStatus[Index] = ActionStatus.Queued;
-				Action.SortIndex = Index++;
+				using ImmediateActionQueue queue = CreateActionQueue(ActionsToExecute, actionArtifactCache, Logger);
+				int actionLimit = Math.Min(NumParallelProcesses, queue.TotalActions);
+				queue.CreateAutomaticRunner(action => RunAction(queue, action), bUseActionWeights, actionLimit, actionLimit);
+				queue.Start();
+				queue.StartManyActions();
+				return await queue.RunTillDone();
 			}
-
-			var ActiveTasks = new Task<ExecuteResults>[ActualNumParallelProcesses];
-			var ActiveActionIndices = new int[ActualNumParallelProcesses];
-			double MaxParallelProcessWeight = ActualNumParallelProcesses;
-			double ActiveTaskWeight = 0.0;
-			int ActiveTaskCount = 0;
-			bool Success = true;
-			List<float> LoggedCpuUtilization = new List<float>();
-			Stopwatch CpuQueryUtilization = new Stopwatch();
-			CpuQueryUtilization.Start();
-
-			int FirstQueuedAction = 0;
-
-			Task LastLoggingTask = Task.CompletedTask;
-
-			// Loop until all actions has been processed and finished
-			while (!CancellationTokenSource.IsCancellationRequested)
+			else
 			{
-				// Always fill up to allowed number of processes
-				while (bUseActionWeights ? 
-					(ActiveTaskWeight < MaxParallelProcessWeight && ActiveTaskCount < ActualNumParallelProcesses)
-					: ActiveTaskCount < ActualNumParallelProcesses)
-				{
-					// Always traverse from first not started/finished action. We always want to respect the sorting.
-					// Note that there might be multiple actions we need to search passed to find the one we can run next
-					int ActionToRunIndex = -1;
-					bool FoundFirstQueued = false;
-					for (int i = FirstQueuedAction; i!= Actions.Length; ++i)
-					{
-						if (ActionsStatus[i] != ActionStatus.Queued)
-						{
-							continue;
-						}
-
-						if (!FoundFirstQueued) // We found the first queued action, let's move our search start index
-						{
-							FirstQueuedAction = i;
-							FoundFirstQueued = true;
-						}
-
-						// Let's see if all prerequisite actions are finished
-						bool ReadyToRun = true;
-						bool HasError = false;
-						foreach (var Prereq in Actions[i].PrerequisiteActions)
-						{
-							ActionStatus PrereqStatus = ActionsStatus[Prereq.SortIndex];
-							ReadyToRun = ReadyToRun && PrereqStatus == ActionStatus.Finished;
-							HasError = HasError || PrereqStatus == ActionStatus.Error;
-						}
-
-						// Not ready, look for next queued action
-						if (!ReadyToRun)
-						{
-							// If Prereq has error, this action inherits the error
-							if (HasError)
-							{
-								ActionsStatus[i] = ActionStatus.Error;
-								continue;
-							}
-							continue;
-						}
-						ActionToRunIndex = i;
-						break;
-					}
-
-					if (bShowCPUUtilization && CpuQueryUtilization.Elapsed.Seconds > 1)
-					{
-						if (Utils.GetTotalCpuUtilization(out float CpuUtilization))
-						{
-							LoggedCpuUtilization.Add(CpuUtilization);
-						}
-						
-						CpuQueryUtilization.Restart();
-					}
-
-					// Didn't find an action we can run. Either there are none left or all depends on currently running actions
-					if (ActionToRunIndex == -1)
-					{
-						break;
-					}
-
-					// Start new action and add it to the ActiveTasks
-					var ActionToRun = Actions[ActionToRunIndex];
-					Task<ExecuteResults> NewTask = Task.Factory.StartNew(() => RunAction(ActionToRun, ProcessGroup, CancellationToken), CancellationToken, TaskCreationOptions.LongRunning | TaskCreationOptions.PreferFairness, TaskScheduler.Default).Unwrap();
-					ActionToTaskLookup.Add(ActionToRun, NewTask);
-
-					ActiveTasks[ActiveTaskCount] = NewTask;
-					ActionsStatus[ActionToRunIndex] = ActionStatus.Running;
-					ActiveActionIndices[ActiveTaskCount] = ActionToRunIndex;
-					ActiveTaskWeight += ActionToRun.Weight;
-					++ActiveTaskCount;
-				}
-
-				// We tried starting new actions and none was added and none is running, that means that we are done!
-				if (ActiveTaskCount == 0)
-				{
-					break;
-				}
-
-				// Wait for all tasks and continue if any is done.
-				Task<ExecuteResults>[] TasksToWaitOn = ActiveTasks;
-				if (ActiveTaskCount < ActiveTasks.Length)
-				{
-					TasksToWaitOn = new Task<ExecuteResults>[ActiveTaskCount];
-					Array.Copy(ActiveTasks, TasksToWaitOn, ActiveTaskCount);
-				}
-				int ActiveTaskIndex = Task.WaitAny(TasksToWaitOn, Timeout.Infinite, CancellationToken);
-
-				// Timeout? This should never happen.. We have a check here just in case
-				if (ActiveTaskIndex == -1)
-				{
-					Logger.LogError("Task timeout? This should never happen. Report to Epic if you see this");
-					Success = false;
-					break;
-				}
-
-				int FinishedActionIndex = ActiveActionIndices[ActiveTaskIndex];
-				LinkedAction FinishedAction = Actions[FinishedActionIndex];
-
-				Task<ExecuteResults> FinishedTask = ActiveTasks[ActiveTaskIndex];
-				bool ActionSuccess = FinishedTask.Result.ExitCode == 0;
-				ActionsStatus[FinishedActionIndex] = ActionSuccess ? ActionStatus.Finished : ActionStatus.Error;
-
-				// Log (chain log entries together so we don't get annoying overlapping logging)
-				LastLoggingTask = LastLoggingTask.ContinueWith((foo) =>
-				{
-					LogCompletedAction(FinishedAction, FinishedTask, CancellationTokenSource, ProgressWriter, TotalActions, ref NumCompletedActions, Logger);
-				});
-
-				// Swap finished task with last and pop last.
-				--ActiveTaskCount;
-				ActiveTasks[ActiveTaskIndex] = ActiveTasks[ActiveTaskCount];
-				ActiveActionIndices[ActiveTaskIndex] = ActiveActionIndices[ActiveTaskCount];
-				ActiveTaskWeight -= FinishedAction.Weight;
-
-				// Update Success status
-				Success = Success && ActionSuccess;
+				using ImmediateActionQueue queue = CreateActionQueue(ActionsToExecute, actionArtifactCache, Logger);
+				int actionLimit = Math.Min(NumParallelProcesses, queue.TotalActions);
+				ImmediateActionQueueRunner runner = queue.CreateManualRunner(action => RunAction(queue, action), bUseActionWeights, actionLimit, actionLimit);
+				queue.Start();
+				using Timer timer = new((_) => queue.StartManyActions(runner), null, 0, 500);
+				queue.StartManyActions();
+				return await queue.RunTillDone();
 			}
+		}
 
-			// Wait for last logging entry before showing summary
-			LastLoggingTask.Wait();
-
-			// Summary
-			TraceSummary(ActionToTaskLookup, ProcessGroup, Logger, LoggedCpuUtilization);
+		private static System.Action? RunAction(ImmediateActionQueue queue, LinkedAction action)
+		{
+			return async () =>
+			{
+				Task<ExecuteResults> task = RunAction(action, queue.ProcessGroup, queue.CancellationToken);
+				await task;
 
-			return Success;
+				ExecuteResults? results = task.Status == TaskStatus.RanToCompletion ? task.Result : null;
+				bool success = results?.ExitCode == 0;
+				queue.OnActionCompleted(action, success, results);
+			};
 		}
 
 		protected static async Task<ExecuteResults> RunAction(LinkedAction Action, ManagedProcessGroup ProcessGroup, CancellationToken CancellationToken)
@@ -387,180 +236,6 @@ namespace UnrealBuildTool
 			TimeSpan ExecutionTime = Process.ExitTime - Process.StartTime;
 			return new ExecuteResults(LogLines, ExitCode, ExecutionTime, ProcessorTime);
 		}
-
-		private static int s_previousLineLength = -1;
-
-		protected void LogCompletedAction(LinkedAction Action, Task<ExecuteResults> ExecuteTask, CancellationTokenSource CancellationTokenSource, ProgressWriter ProgressWriter, int TotalActions, ref int NumCompletedActions, ILogger Logger)
-		{
-			List<string> LogLines = new List<string>();
-			int ExitCode = int.MaxValue;
-			TimeSpan ExecutionTime = TimeSpan.Zero;
-			TimeSpan ProcessorTime = TimeSpan.Zero;
-			string? AdditionalDescription = null;
-			if (ExecuteTask.Status == TaskStatus.RanToCompletion)
-			{
-				ExecuteResults ExecuteTaskResult = ExecuteTask.Result;
-				LogLines = ExecuteTaskResult.LogLines;
-				ExitCode = ExecuteTaskResult.ExitCode;
-				ExecutionTime = ExecuteTaskResult.ExecutionTime;
-				ProcessorTime = ExecuteTaskResult.ProcessorTime;
-				AdditionalDescription = ExecuteTaskResult.AdditionalDescription;
-			}
-
-			// Write it to the log
-			string Description = string.Empty;
-			if (Action.bShouldOutputStatusDescription || LogLines.Count == 0)
-			{
-				Description = $"{(Action.CommandDescription ?? Action.CommandPath.GetFileNameWithoutExtension())} {Action.StatusDescription}".Trim();
-			}
-			else if (LogLines.Count > 0)
-			{
-				Description = $"{(Action.CommandDescription ?? Action.CommandPath.GetFileNameWithoutExtension())} {LogLines[0]}".Trim();
-			}
-			if (!string.IsNullOrEmpty(AdditionalDescription))
-            {
-				Description = $"{Description} {AdditionalDescription}";
-			}
-
-			lock (ProgressWriter)
-			{
-				int CompletedActions;
-				CompletedActions = Interlocked.Increment(ref NumCompletedActions);
-				ProgressWriter.Write(CompletedActions, TotalActions);
-
-				// Cancelled
-				if (ExitCode == int.MaxValue)
-				{
-					Logger.LogInformation("[{CompletedActions}/{TotalActions}] {Description} cancelled", CompletedActions, TotalActions, Description);
-					return;
-				}
-
-				string TargetDetails = "";
-				TargetDescriptor? Target = Action.Target;
-				if (bPrintActionTargetNames && Target != null)
-				{
-					TargetDetails = $"[{Target.Name} {Target.Platform} {Target.Configuration}]";
-				}
-
-				if (bLogActionCommandLines)
-				{
-					Logger.LogDebug("[{CompletedActions}/{TotalActions}]{TargetDetails} Command: {CommandPath} {CommandArguments}", CompletedActions, TotalActions, TargetDetails, Action.CommandPath, Action.CommandArguments);
-				}
-
-				string CompilationTimes = "";
-
-				if (bShowPerActionCompilationTimes)
-				{
-					if (ProcessorTime.Ticks > 0)
-					{
-						CompilationTimes = $" (Wall: {ExecutionTime.TotalSeconds:0.00}s CPU: {ProcessorTime.TotalSeconds:0.00}s)";
-					}
-					else
-					{
-						CompilationTimes = $" (Wall: {ExecutionTime.TotalSeconds:0.00}s)";
-					}
-				}
-
-				string message = ($"[{CompletedActions}/{TotalActions}]{TargetDetails}{CompilationTimes} {Description}");
-				
-				if (bCompactOutput)
-				{
-					if (s_previousLineLength > 0)
-					{
-						// move the cursor to the far left position, one line back
-						Console.CursorLeft = 0;
-						Console.CursorTop -= 1;
-						// clear the line
-						Console.Write("".PadRight(s_previousLineLength));
-						// move the cursor back to the left, so output is written to the desired location
-						Console.CursorLeft = 0;
-					}
-				}
-
-				s_previousLineLength = message.Length;
-
-				LogEventParser Parser = new LogEventParser(Logger);
-
-				WriteToolOutput(message);
-				foreach (string Line in LogLines.Skip(Action.bShouldOutputStatusDescription ? 0 : 1))
-				{
-					// suppress library creation messages when writing compact output
-					if (bCompactOutput && Line.StartsWith("   Creating library ") && Line.EndsWith(".exp"))
-					{
-						continue;
-					}
-
-					WriteToolOutput(Line);
-
-					// Prevent overwriting of logged lines
-					s_previousLineLength = -1;
-				}
-
-				if (ExitCode != 0)
-				{
-					// BEGIN TEMPORARY TO CATCH PVS-STUDIO ISSUES
-					if (LogLines.Count == 0)
-					{
-						Logger.LogError("{TargetDetails} {Description}: Exited with error code {ExitCode}", TargetDetails, Description, ExitCode);
-						Logger.LogInformation("{TargetDetails} {Description}: WorkingDirectory {WorkingDirectory}", TargetDetails, Description, Action.WorkingDirectory);
-						Logger.LogInformation("{TargetDetails} {Description}: {CommandPath} {CommandArguments}", TargetDetails, Description, Action.CommandPath, Action.CommandArguments);
-					}
-					// END TEMPORARY
-
-					// prevent overwriting of error text
-					s_previousLineLength = -1;
-
-					// Cancel all other pending tasks
-					if (bStopCompilationAfterErrors)
-					{
-						CancellationTokenSource.Cancel();
-					}
-				}
-			}
-		}
-
-		protected static void TraceSummary(Dictionary<LinkedAction, Task<ExecuteResults>> Tasks, ManagedProcessGroup ProcessGroup, ILogger Logger, IList<float>? LoggedCpuUtilization)
-		{
-			if (bShowCPUUtilization && LoggedCpuUtilization != null && LoggedCpuUtilization.Any())
-			{
-				Logger.LogInformation("");
-				Logger.LogInformation("Average CPU Utilization: {CPUPercentage}%",(int)(LoggedCpuUtilization.Average()));
-			}
-
-			if (!bShowCompilationTimes)
-			{
-				return;
-			}
-
-			Logger.LogInformation("");
-			if (ProcessGroup.TotalProcessorTime.Ticks > 0)
-			{
-				Logger.LogInformation("Total CPU Time: {TotalSeconds} s", ProcessGroup.TotalProcessorTime.TotalSeconds);
-				Logger.LogInformation("");
-			}
-
-			var CompletedTasks = Tasks.Where(x => x.Value.Status == TaskStatus.RanToCompletion).OrderByDescending(x => x.Value.Result.ExecutionTime).Take(20);
-
-			if (CompletedTasks.Any())
-			{
-				Logger.LogInformation("Compilation Time Top {CompletedTaskCount}", CompletedTasks.Count());
-				Logger.LogInformation("");
-				foreach (var Pair in CompletedTasks)
-				{
-					string Description = $"{(Pair.Key.Inner.CommandDescription ?? Pair.Key.Inner.CommandPath.GetFileNameWithoutExtension())} {Pair.Key.Inner.StatusDescription}".Trim();
-					if (Pair.Value.Result.ProcessorTime.Ticks > 0)
-					{
-						Logger.LogInformation("{Description} [ Wall Time {ExecutionTime:0.00} s / CPU Time {ProcessorTime:0.00} s ]", Description, Pair.Value.Result.ExecutionTime.TotalSeconds, Pair.Value.Result.ProcessorTime.TotalSeconds);
-					}
-					else
-					{
-						Logger.LogInformation("{Description} [ Time {ExecutionTime:0.00} s ]", Description, Pair.Value.Result.ExecutionTime.TotalSeconds);
-					}
-
-				}
-				Logger.LogInformation("");
-			}
-		}
 	}
 
 	/// <summary>
diff --git a/Engine/Source/Programs/UnrealBuildTool/Executors/SNDBS.cs b/Engine/Source/Programs/UnrealBuildTool/Executors/SNDBS.cs
--- a/Engine/Source/Programs/UnrealBuildTool/Executors/SNDBS.cs
+++ b/Engine/Source/Programs/UnrealBuildTool/Executors/SNDBS.cs
@@ -18,6 +18,7 @@ using System.Threading.Tasks;
 using EpicGames.Core;
 using Microsoft.Extensions.Logging;
 using UnrealBuildBase;
+using UnrealBuildTool.Artifacts;
 
 namespace UnrealBuildTool
 {
@@ -215,7 +216,8 @@ namespace UnrealBuildTool
 			return true;
 		}
 
-		public override Task<bool> ExecuteActionsAsync(IEnumerable<LinkedAction> ActionsToExecute, ILogger Logger)
+		/// <inheritdoc/>
+		public override Task<bool> ExecuteActionsAsync(IEnumerable<LinkedAction> ActionsToExecute, ILogger Logger, IActionArtifactCache? actionArtifactCache)
 		{
 			return Task.FromResult(ExecuteActions(ActionsToExecute, Logger));
 		}
diff --git a/Engine/Source/Programs/UnrealBuildTool/Executors/XGE.cs b/Engine/Source/Programs/UnrealBuildTool/Executors/XGE.cs
--- a/Engine/Source/Programs/UnrealBuildTool/Executors/XGE.cs
+++ b/Engine/Source/Programs/UnrealBuildTool/Executors/XGE.cs
@@ -20,6 +20,7 @@ using System.Diagnostics.CodeAnalysis;
 using System.Runtime.Versioning;
 using Microsoft.Extensions.Logging;
 using System.Threading.Tasks;
+using UnrealBuildTool.Artifacts;
 
 namespace UnrealBuildTool
 {
@@ -396,7 +397,8 @@ namespace UnrealBuildTool
 			Logger.LogInformation("XGEEXPORT: Exported '{OutFile}'", OutFile);
 		}
 
-		public override Task<bool> ExecuteActionsAsync(IEnumerable<LinkedAction> ActionsToExecute, ILogger Logger)
+		/// <inheritdoc/>
+		public override Task<bool> ExecuteActionsAsync(IEnumerable<LinkedAction> ActionsToExecute, ILogger Logger, IActionArtifactCache? actionArtifactCache)
 		{
 			return Task.FromResult(ExecuteActions(ActionsToExecute, Logger));
 		}
diff --git a/Engine/Source/Programs/UnrealBuildTool/Modes/BuildMode.cs b/Engine/Source/Programs/UnrealBuildTool/Modes/BuildMode.cs
--- a/Engine/Source/Programs/UnrealBuildTool/Modes/BuildMode.cs
+++ b/Engine/Source/Programs/UnrealBuildTool/Modes/BuildMode.cs
@@ -11,6 +11,7 @@ using UnrealBuildBase;
 using Microsoft.Extensions.Logging;
 using System.Collections.Concurrent;
 using System.Threading.Tasks;
+using UnrealBuildTool.Artifacts;
 
 namespace UnrealBuildTool
 {
@@ -532,6 +533,10 @@ namespace UnrealBuildTool
 					}
 				}
 
+				// Now that we have the dependencies object, optionally create the artifact cache
+				// For future work
+				IActionArtifactCache? actionArtifactCache = null;
+
 				// Pre-process module interfaces to generate dependency files
 				List<LinkedAction> ModuleDependencyActions = PrerequisiteActions.Where(x => x.ActionType == ActionType.GatherModuleDependencies).ToList();
 				if (ModuleDependencyActions.Count > 0)
@@ -767,7 +772,19 @@ namespace UnrealBuildTool
 
 						using (GlobalTracer.Instance.BuildSpan("ActionGraph.ExecuteActions()").StartActive())
 						{
-							await ActionGraph.ExecuteActionsAsync(BuildConfiguration, MergedActionsToExecute, TargetDescriptors, Logger);
+
+							// We need to wait for the cache to be ready.  If we fail to fetch from the cache for PCH files, they are
+							// non-deterministic and will cause rippled cache misses.  Changing the system to assume a PCH file is the same
+							// if the inputs are the same might be a good way to avoid this issue.
+							if (actionArtifactCache != null && actionArtifactCache.ArtifactCache != null)
+							{
+								await actionArtifactCache.ArtifactCache.WaitForReadyAsync();
+							}
+							await ActionGraph.ExecuteActionsAsync(BuildConfiguration, MergedActionsToExecute, TargetDescriptors, Logger, actionArtifactCache);
+							if (actionArtifactCache != null)
+							{
+								await actionArtifactCache.FlushChangesAsync(default); //ETSTODO
+							}
 						}
 					}
 
diff --git a/Engine/Source/Programs/UnrealBuildTool/System/ActionExecutor.cs b/Engine/Source/Programs/UnrealBuildTool/System/ActionExecutor.cs
--- a/Engine/Source/Programs/UnrealBuildTool/System/ActionExecutor.cs
+++ b/Engine/Source/Programs/UnrealBuildTool/System/ActionExecutor.cs
@@ -6,6 +6,7 @@ using System;
 using System.Collections.Generic;
 using System.Reflection;
 using System.Threading.Tasks;
+using UnrealBuildTool.Artifacts;
 
 namespace UnrealBuildTool
 {
@@ -54,7 +55,14 @@ namespace UnrealBuildTool
 			MemoryPerActionBytesOverride = Math.Max(MemoryPerActionBytesOverride, MemoryPerActionOverrideGB * 1024 * 1024 * 1024);
 		}
 
-		public abstract Task<bool> ExecuteActionsAsync(IEnumerable<LinkedAction> ActionsToExecute, ILogger Logger);
+		/// <summary>
+		/// Execute the given actions
+		/// </summary>
+		/// <param name="ActionsToExecute">Actions to be executed</param>
+		/// <param name="Logger">Logger for output</param>
+		/// <param name="actionArtifactCache">Cache used to read/write action outputs.</param>
+		/// <returns>True if the build succeeded, false otherwise</returns>
+		public abstract Task<bool> ExecuteActionsAsync(IEnumerable<LinkedAction> ActionsToExecute, ILogger Logger, IActionArtifactCache? actionArtifactCache = null);
 
 		protected void WriteToolOutput(string Line)
 		{
diff --git a/Engine/Source/Programs/UnrealBuildTool/System/ActionGraph.cs b/Engine/Source/Programs/UnrealBuildTool/System/ActionGraph.cs
--- a/Engine/Source/Programs/UnrealBuildTool/System/ActionGraph.cs
+++ b/Engine/Source/Programs/UnrealBuildTool/System/ActionGraph.cs
@@ -19,6 +19,7 @@ using System.Collections.Concurrent;
 using System.Collections.Immutable;
 using System.Reflection;
 using System.Runtime.CompilerServices;
+using UnrealBuildTool.Artifacts;
 
 namespace UnrealBuildTool
 {
@@ -387,7 +388,7 @@ namespace UnrealBuildTool
 		/// <summary>
 		/// Executes a list of actions.
 		/// </summary>
-		public static async Task ExecuteActionsAsync(BuildConfiguration BuildConfiguration, List<LinkedAction> ActionsToExecute, List<TargetDescriptor> TargetDescriptors, ILogger Logger)
+		public static async Task ExecuteActionsAsync(BuildConfiguration BuildConfiguration, List<LinkedAction> ActionsToExecute, List<TargetDescriptor> TargetDescriptors, ILogger Logger, IActionArtifactCache? actionArtifactCache = null)
 		{
 			if (ActionsToExecute.Count == 0)
 			{
@@ -400,7 +401,7 @@ namespace UnrealBuildTool
 
 				// Execute the build
 				Stopwatch Timer = Stopwatch.StartNew();
-				bool Result = await Executor.ExecuteActionsAsync(ActionsToExecute, Logger);
+				bool Result = await Executor.ExecuteActionsAsync(ActionsToExecute, Logger, actionArtifactCache);
 
 				Logger.LogInformation("Total time in {ExecutorName} executor: {TotalSeconds:0.00} seconds", Executor.Name, Timer.Elapsed.TotalSeconds);
 
