Allow temporarily disabling Box executor on command line with -NoBox

diff --git a/Engine/Source/Programs/UnrealBuildTool/Configuration/BuildConfiguration.cs b/Engine/Source/Programs/UnrealBuildTool/Configuration/BuildConfiguration.cs
--- a/Engine/Source/Programs/UnrealBuildTool/Configuration/BuildConfiguration.cs
+++ b/Engine/Source/Programs/UnrealBuildTool/Configuration/BuildConfiguration.cs
@@ -48,6 +48,7 @@ namespace UnrealBuildTool
 		/// </summary>
 		[XmlConfigFile]
 		[CommandLine("-Box", Value = "true")]
+		[CommandLine("-NoBox", Value = "false")]
 		public bool bAllowBoxExecutor = false;
 #endif // #if __BOXEXECUTOR_AVAILABLE__
 
