[UBT]
* Fixed so manual runners are not used in TryStartOneAction if no runner is provided

diff --git a/Engine/Source/Programs/UnrealBuildTool/Executors/ImmediateActionQueue.cs b/Engine/Source/Programs/UnrealBuildTool/Executors/ImmediateActionQueue.cs
--- a/Engine/Source/Programs/UnrealBuildTool/Executors/ImmediateActionQueue.cs
+++ b/Engine/Source/Programs/UnrealBuildTool/Executors/ImmediateActionQueue.cs
@@ -483,7 +483,8 @@ namespace UnrealBuildTool
 							{
 								foreach (ImmediateActionQueueRunner tryRunner in _runners)
 								{
-									if (tryRunner.IsUnderLimits && tryRunner.ActionPhase == Actions[actionIndex].Phase)
+									if (tryRunner.Type == ImmediateActionQueueRunnerType.Automatic &&                                   
+										tryRunner.IsUnderLimits && tryRunner.ActionPhase == Actions[actionIndex].Phase)
 									{
 										action = tryRunner.RunAction(Actions[actionIndex].Action);
 										if (action != null)
