[UBT]
* Fixed deadlock when done event is set (where task can recursively pickup done task which can cause a deadlock with some executors

diff --git a/Engine/Source/Programs/UnrealBuildTool/Executors/ImmediateActionQueue.cs b/Engine/Source/Programs/UnrealBuildTool/Executors/ImmediateActionQueue.cs
--- a/Engine/Source/Programs/UnrealBuildTool/Executors/ImmediateActionQueue.cs
+++ b/Engine/Source/Programs/UnrealBuildTool/Executors/ImmediateActionQueue.cs
@@ -162,6 +162,8 @@ namespace UnrealBuildTool
 		/// </summary>
 		private readonly ActionState[] Actions;
 
+		private readonly ReaderWriterLockSlim ActionsLock;
+
 		/// <summary>
 		/// Output logging
 		/// </summary>
@@ -307,6 +309,7 @@ namespace UnrealBuildTool
 		{
 			int count = actions.Count();
 			Actions = new ActionState[count];
+			ActionsLock = new ReaderWriterLockSlim();
 
 			Logger = logger;
 			ProgressWriter = new(progressWriterText, false, logger);
@@ -425,6 +428,31 @@ namespace UnrealBuildTool
 			}
 		}
 
+		public struct WriteLockScope : IDisposable
+		{
+			public WriteLockScope(ReaderWriterLockSlim lockObject)
+			{
+				_lockObject = lockObject;
+				_lockObject.EnterWriteLock();
+			}
+
+			public void Exit()
+			{
+				if (!_inside)
+					return;
+				_lockObject.ExitWriteLock();
+				_inside = false;
+			}
+
+			public void Dispose()
+			{
+				Exit();
+			}
+
+			ReaderWriterLockSlim _lockObject;
+			bool _inside = true;
+		}
+
 		/// <summary>
 		/// Try to start one action
 		/// </summary>
@@ -432,7 +460,7 @@ namespace UnrealBuildTool
 		/// <returns>True if an action was run, false if not</returns>
 		public bool TryStartOneAction(ImmediateActionQueueRunner? runner = null)
 		{
-			lock (Actions)
+			using (var actionsLock = new WriteLockScope(ActionsLock))
 			{
 				// If we are starting deeper in the action collection, then never set the first pending action location
 				bool foundFirstPending = false;
@@ -462,6 +490,7 @@ namespace UnrealBuildTool
 
 						case ActionReadyState.Error:
 							Actions[actionIndex].Status = ActionStatus.Error;
+							actionsLock.Exit();
 							IncrementCompletedActions();
 							break;
 
@@ -502,6 +531,7 @@ namespace UnrealBuildTool
 								Actions[actionIndex].Runner = selectedRunner;
 								selectedRunner.ActiveActions++;
 								selectedRunner.ActiveActionWeight += Actions[actionIndex].Action.Weight;
+								actionsLock.Exit();
 								Task.Factory.StartNew(() => action(), CancellationToken, TaskCreationOptions.LongRunning | TaskCreationOptions.PreferFairness, TaskScheduler.Default);
 								return true;
 							}
@@ -585,7 +615,7 @@ namespace UnrealBuildTool
 			int actionIndex = action.SortIndex;
 
 			// Update the actions data
-			lock (Actions)
+			using (var actionsLock = new WriteLockScope(ActionsLock))
 			{
 				ImmediateActionQueueRunner runner = Actions[actionIndex].Runner ?? throw new BuildException("Attempting to update action state but runner isn't set");
 				runner.ActiveActions--;
@@ -634,11 +664,13 @@ namespace UnrealBuildTool
 						// Notify the artifact handler of the action completing.  We don't wait on the resulting task.  The
 						// cache is required to remember any pending saves and a final call to Flush will wait for everything to complete.
 						_actionArtifactCache?.ActionCompleteAsync(action, CancellationToken);
+						actionsLock.Exit();
 						IncrementCompletedActions();
 						break;
 
 					case ActionStatus.Error:
 						Success = false;
+						actionsLock.Exit();
 						IncrementCompletedActions();
 						break;
 
@@ -657,8 +689,8 @@ namespace UnrealBuildTool
 		/// </summary>
 		private void IncrementCompletedActions()
 		{
-			_completedActions++;
-			if (_completedActions >= Actions.Length && !_doneTaskSource.Task.IsCompleted)
+			var completedActions = Interlocked.Increment(ref _completedActions);
+			if (completedActions >= Actions.Length && !_doneTaskSource.Task.IsCompleted)
 			{
 				_doneTaskSource.SetResult();
 			}
