UnrealBuildTool: Remove obsolete static analyzer code

#rnx

diff --git a/Engine/Source/Programs/UnrealBuildTool/Platform/Windows/UEBuildWindows.cs b/Engine/Source/Programs/UnrealBuildTool/Platform/Windows/UEBuildWindows.cs
--- a/Engine/Source/Programs/UnrealBuildTool/Platform/Windows/UEBuildWindows.cs
+++ b/Engine/Source/Programs/UnrealBuildTool/Platform/Windows/UEBuildWindows.cs
@@ -165,62 +165,6 @@ namespace UnrealBuildTool
 	}
 
 
-	/// <summary>
-	/// Which static analyzer to use
-	/// </summary>
-	[Obsolete("Replace with StaticAnalyzer. Will be removed in 5.2")]
-	public enum WindowsStaticAnalyzer
-	{
-		/// <summary>
-		/// Do not perform static analysis
-		/// </summary>
-		None,
-
-		/// <summary>
-		/// Use the default static analyzer for the selected compiler, if it has one. For
-		/// Visual Studio and Clang, this means using their built-in static analysis tools.
-		/// Any compiler that doesn't support static analysis will ignore this option.
-		/// </summary>
-		Default, 
-
-		/// <summary>
-		/// Use the built-in Visual C++ static analyzer
-		/// </summary>
-		VisualCpp = Default,
-
-		/// <summary>
-		/// Use PVS-Studio for static analysis
-		/// </summary>
-		PVSStudio,
-		
-		/// <summary>
-		/// Use clang for static analysis. This forces the compiler to clang.
-		/// </summary>
-		Clang,
-	}
-
-	/// <summary>
-	/// Output type for the static analyzer. This currently only works for the Clang static analyzer.
-	/// The Clang static analyzer can do either Text, which prints the analysis to stdout, or
-	/// html, where it writes out a navigable HTML page for each issue that it finds, per file.
-	/// The HTML is output in the same directory as the object fil that would otherwise have
-	/// been generated. 
-	/// All other analyzers default automatically to Text. 
-	/// </summary>
-	[Obsolete("Replace with StaticAnalyzerOutputType. Will be removed in 5.2")]
-	public enum WindowsStaticAnalyzerOutputType
-	{
-		/// <summary>
-		/// Output the analysis to stdout.
-		/// </summary>
-		Text,
-
-		/// <summary>
-		/// Output the analysis to an HTML file in the object folder.
-		/// </summary>
-		Html,
-	}
-
 	/// <summary>
 	/// Windows-specific target settings
 	/// </summary>
@@ -351,22 +295,6 @@ namespace UnrealBuildTool
 		[ConfigFile(ConfigHierarchyType.Game, "/Script/EngineSettings.GeneralProjectSettings", "ProjectName")]
 		public string? ProductName;
 
-		/// <summary>
-		/// The static analyzer to use.
-		/// </summary>
-		[XmlConfigFile(Category = "WindowsPlatform")]
-		[CommandLine("-StaticAnalyzer")]
-		[Obsolete("Replace with TargetRuies.StaticAnalyzer (Xml Category BuildConfiguration). Will be removed in 5.2")]
-		public WindowsStaticAnalyzer StaticAnalyzer = WindowsStaticAnalyzer.None;
-
-		/// <summary>
-		/// The output type to use for the static analyzer.
-		/// </summary>
-		[XmlConfigFile(Category = "WindowsPlatform")]
-		[CommandLine("-StaticAnalyzerOutputType")]
-		[Obsolete("Replace with TargetRuies.StaticAnalyzerOutputType (Xml Category BuildConfiguration). Will be removed in 5.2")]
-		public WindowsStaticAnalyzerOutputType StaticAnalyzerOutputType = WindowsStaticAnalyzerOutputType.Text;
-
 		/// <summary>
 		/// Enables address sanitizer (ASan). Only supported for Visual Studio 2019 16.7.0 and up.
 		/// </summary>
@@ -771,18 +699,6 @@ namespace UnrealBuildTool
 			get { return Inner.ProductName; }
 		}
 
-		[Obsolete("Replace with TargetRuies.StaticAnalyzer (Xml Category BuildConfiguration). Will be removed in 5.2")]
-		public WindowsStaticAnalyzer StaticAnalyzer
-		{
-			get { return Inner.StaticAnalyzer; }
-		}
-
-		[Obsolete("Replace with TargetRuies.StaticAnalyzerOutputType (Xml Category BuildConfiguration). Will be removed in 5.2")]
-		public WindowsStaticAnalyzerOutputType StaticAnalyzerOutputType
-		{
-			get { return Inner.StaticAnalyzerOutputType; }
-		}
-
 		public bool bEnableAddressSanitizer
 		{ 
 			get { return Inner.bEnableAddressSanitizer; }
