UnrealBuildTool: Add ini and command line control for WindowsSDKVersion

#rnx

diff --git a/Engine/Source/Programs/UnrealBuildTool/Platform/Windows/UEBuildWindows.cs b/Engine/Source/Programs/UnrealBuildTool/Platform/Windows/UEBuildWindows.cs
--- a/Engine/Source/Programs/UnrealBuildTool/Platform/Windows/UEBuildWindows.cs
+++ b/Engine/Source/Programs/UnrealBuildTool/Platform/Windows/UEBuildWindows.cs
@@ -240,12 +240,14 @@ namespace UnrealBuildTool
 		[XmlConfigFile(Category = "WindowsPlatform")]
 		[CommandLine("-ClangLinker")]
 		public bool bAllowClangLinker = false;
-		
+
 		/// <summary>
 		/// The specific Windows SDK version to use. This may be a specific version number (for example, "8.1", "10.0" or "10.0.10150.0"), or the string "Latest", to select the newest available version.
 		/// By default, and if it is available, we use the Windows SDK version indicated by WindowsPlatform.DefaultWindowsSdkVersion (otherwise, we use the latest version).
 		/// </summary>
+		[ConfigFile(ConfigHierarchyType.Engine, "/Script/WindowsTargetPlatform.WindowsTargetSettings", "WindowsSDKVersion")]
 		[XmlConfigFile(Category = "WindowsPlatform")]
+		[CommandLine("-WindowsSDKVersion")]
 		public string? WindowsSdkVersion = null;
 
 		/// <summary>
