UnrealBuildTool: If a specific windows sdk is request and is not found, print out all available sdks

#rnx

diff --git a/Engine/Source/Programs/UnrealBuildTool/Platform/Windows/MicrosoftPlatformSDK.cs b/Engine/Source/Programs/UnrealBuildTool/Platform/Windows/MicrosoftPlatformSDK.cs
--- a/Engine/Source/Programs/UnrealBuildTool/Platform/Windows/MicrosoftPlatformSDK.cs
+++ b/Engine/Source/Programs/UnrealBuildTool/Platform/Windows/MicrosoftPlatformSDK.cs
@@ -136,6 +136,29 @@ namespace UnrealBuildTool
 			return CachedWindowsSdkDirs!;
 		}
 
+		/// <summary>
+		/// List all the installed Windows SDK versions
+		/// </summary>
+		/// <param name="Logger">Logger for output</param>
+		[SupportedOSPlatform("windows")]
+		public static void DumpWindowsSdkDirs(ILogger Logger)
+		{
+			IOrderedEnumerable<KeyValuePair<VersionNumber, DirectoryReference>> SortedSdkDirs = FindWindowsSdkDirs(Logger).OrderByDescending(x => x.Key);
+
+			if (SortedSdkDirs.Any())
+			{
+				Logger.LogInformation("Available Windows SDKs ({Count}):", SortedSdkDirs.Count());
+				foreach (KeyValuePair<VersionNumber, DirectoryReference> Item in SortedSdkDirs)
+				{
+					Logger.LogInformation(" * {SdkDir}\n    (Version={Version})", Item.Value, Item.Key);
+				}
+			}
+			else
+			{
+				Logger.LogInformation("No available Windows SDKs found");
+			}
+		}
+
 		/// <summary>
 		/// Finds all the installed Universal CRT versions
 		/// </summary>
diff --git a/Engine/Source/Programs/UnrealBuildTool/Platform/Windows/VCEnvironment.cs b/Engine/Source/Programs/UnrealBuildTool/Platform/Windows/VCEnvironment.cs
--- a/Engine/Source/Programs/UnrealBuildTool/Platform/Windows/VCEnvironment.cs
+++ b/Engine/Source/Programs/UnrealBuildTool/Platform/Windows/VCEnvironment.cs
@@ -518,6 +518,7 @@ namespace UnrealBuildTool
 			{
 				if (!WindowsPlatform.TryGetWindowsSdkDir(WindowsSdkVersion, Logger, out SelectedWindowsSdkVersion, out SelectedWindowsSdkDir))
 				{
+					MicrosoftPlatformSDK.DumpWindowsSdkDirs(Logger);
 					throw new BuildException("Windows SDK{0} must be installed in order to build this target.", String.IsNullOrEmpty(WindowsSdkVersion) ? "" : String.Format(" ({0})", WindowsSdkVersion));
 				}
 			}
