1) Fixed BoxExecutor to not use have any async methods due to Stop needing to happen on the same thread.
2) Fixed issue where Task.Run and SetResult were being invoked under a  lock which could result in deadlocks.

diff --git a/Engine/Source/Programs/UnrealBuildTool/Executors/ImmediateActionQueue.cs b/Engine/Source/Programs/UnrealBuildTool/Executors/ImmediateActionQueue.cs
--- a/Engine/Source/Programs/UnrealBuildTool/Executors/ImmediateActionQueue.cs
+++ b/Engine/Source/Programs/UnrealBuildTool/Executors/ImmediateActionQueue.cs
@@ -1,6 +1,7 @@
 // Copyright Epic Games, Inc. All Rights Reserved.
 
 using EpicGames.Core;
+using Microsoft.CodeAnalysis;
 using Microsoft.Extensions.Logging;
 using System;
 using System.Collections.Generic;
@@ -162,8 +163,6 @@ namespace UnrealBuildTool
 		/// </summary>
 		private readonly ActionState[] Actions;
 
-		private readonly ReaderWriterLockSlim ActionsLock;
-
 		/// <summary>
 		/// Output logging
 		/// </summary>
@@ -309,7 +308,6 @@ namespace UnrealBuildTool
 		{
 			int count = actions.Count();
 			Actions = new ActionState[count];
-			ActionsLock = new ReaderWriterLockSlim();
 
 			Logger = logger;
 			ProgressWriter = new(progressWriterText, false, logger);
@@ -428,39 +426,51 @@ namespace UnrealBuildTool
 			}
 		}
 
-		public struct WriteLockScope : IDisposable
+		/// <summary>
+		/// Try to start one action
+		/// </summary>
+		/// <param name="runner">If specified, tasks will only be queued to this runner.  Otherwise all manual runners will be used.</param>
+		/// <returns>True if an action was run, false if not</returns>
+		public bool TryStartOneAction(ImmediateActionQueueRunner? runner = null)
 		{
-			public WriteLockScope(ReaderWriterLockSlim lockObject)
-			{
-				_lockObject = lockObject;
-				_lockObject.EnterWriteLock();
-			}
+			(System.Action? action, int completedActions) = TryStartOneActionInternal(runner);
 
-			public void Exit()
+			if (action == null)
 			{
-				if (!_inside)
-					return;
-				_lockObject.ExitWriteLock();
-				_inside = false;
+				// We found nothing, check to see if we have no active running tasks and no manual runners.
+				// If we don't, then it is assumed we can't queue any more work.
+				bool prematureDone = true;
+				foreach (ImmediateActionQueueRunner tryRunner in _runners)
+				{
+					if (tryRunner.Type == ImmediateActionQueueRunnerType.Manual || tryRunner.ActiveActions != 0)
+					{
+						prematureDone = false;
+						break;
+					}
+				}
+				if (prematureDone)
+				{
+					completedActions = int.MaxValue;
+				}
 			}
-
-			public void Dispose()
+			else
 			{
-				Exit();
+				Task.Factory.StartNew(() => action(), CancellationToken, TaskCreationOptions.LongRunning | TaskCreationOptions.PreferFairness, TaskScheduler.Default);
 			}
 
-			ReaderWriterLockSlim _lockObject;
-			bool _inside = true;
+			AddCompletedActions(completedActions);
+			return action != null;
 		}
 
 		/// <summary>
-		/// Try to start one action
+		/// Try to start one action.
 		/// </summary>
 		/// <param name="runner">If specified, tasks will only be queued to this runner.  Otherwise all manual runners will be used.</param>
-		/// <returns>True if an action was run, false if not</returns>
-		public bool TryStartOneAction(ImmediateActionQueueRunner? runner = null)
+		/// <returns>Action to be run if something was queued and the number of completed actions</returns>
+		public (System.Action?, int) TryStartOneActionInternal(ImmediateActionQueueRunner? runner = null)
 		{
-			using (var actionsLock = new WriteLockScope(ActionsLock))
+			int completedActions = 0;
+			lock (Actions)
 			{
 				// If we are starting deeper in the action collection, then never set the first pending action location
 				bool foundFirstPending = false;
@@ -490,8 +500,7 @@ namespace UnrealBuildTool
 
 						case ActionReadyState.Error:
 							Actions[actionIndex].Status = ActionStatus.Error;
-							actionsLock.Exit();
-							IncrementCompletedActions();
+							completedActions++;
 							break;
 
 						case ActionReadyState.Ready:
@@ -512,7 +521,7 @@ namespace UnrealBuildTool
 							{
 								foreach (ImmediateActionQueueRunner tryRunner in _runners)
 								{
-									if (tryRunner.Type == ImmediateActionQueueRunnerType.Automatic &&                                   
+									if (tryRunner.Type == ImmediateActionQueueRunnerType.Automatic &&
 										tryRunner.IsUnderLimits && tryRunner.ActionPhase == Actions[actionIndex].Phase)
 									{
 										action = tryRunner.RunAction(Actions[actionIndex].Action);
@@ -524,16 +533,14 @@ namespace UnrealBuildTool
 									}
 								}
 							}
-							
+
 							if (action != null && selectedRunner != null)
 							{
 								Actions[actionIndex].Status = ActionStatus.Running;
 								Actions[actionIndex].Runner = selectedRunner;
 								selectedRunner.ActiveActions++;
 								selectedRunner.ActiveActionWeight += Actions[actionIndex].Action.Weight;
-								actionsLock.Exit();
-								Task.Factory.StartNew(() => action(), CancellationToken, TaskCreationOptions.LongRunning | TaskCreationOptions.PreferFairness, TaskScheduler.Default);
-								return true;
+								return (action, completedActions);
 							}
 							break;
 
@@ -541,24 +548,8 @@ namespace UnrealBuildTool
 							throw new ApplicationException("Unexpected action ready state");
 					}
 				}
-
-				// We found nothing, check to see if we have no active running tasks and no manual runners.
-				// If we don't, then it is assumed we can't queue any more work.
-				bool prematureDone = true;
-				foreach (ImmediateActionQueueRunner tryRunner in _runners)
-				{
-					if (tryRunner.Type == ImmediateActionQueueRunnerType.Manual || tryRunner.ActiveActions != 0)
-					{
-						prematureDone = false;
-						break;
-					}
-				}
-				if (prematureDone && !_doneTaskSource.Task.IsCompleted)
-				{
-					_doneTaskSource.SetResult();
-				}
-				return false;
 			}
+			return (null, completedActions);
 		}
 
 		/// <summary>
@@ -613,9 +604,10 @@ namespace UnrealBuildTool
 		private void SetActionState(LinkedAction action, ActionStatus status, ExecuteResults? results)
 		{
 			int actionIndex = action.SortIndex;
+			int completedActions = 0;
 
 			// Update the actions data
-			using (var actionsLock = new WriteLockScope(ActionsLock))
+			lock (Actions)
 			{
 				ImmediateActionQueueRunner runner = Actions[actionIndex].Runner ?? throw new BuildException("Attempting to update action state but runner isn't set");
 				runner.ActiveActions--;
@@ -664,14 +656,12 @@ namespace UnrealBuildTool
 						// Notify the artifact handler of the action completing.  We don't wait on the resulting task.  The
 						// cache is required to remember any pending saves and a final call to Flush will wait for everything to complete.
 						_actionArtifactCache?.ActionCompleteAsync(action, CancellationToken);
-						actionsLock.Exit();
-						IncrementCompletedActions();
+						completedActions = 1;
 						break;
 
 					case ActionStatus.Error:
 						Success = false;
-						actionsLock.Exit();
-						IncrementCompletedActions();
+						completedActions = 1;
 						break;
 
 					default:
@@ -679,20 +669,36 @@ namespace UnrealBuildTool
 				}
 			}
 
+			// Outside of lock, update the completed actions
+			AddCompletedActions(completedActions);
+
 			// Since something has been completed or returned to the queue, try to run actions again
 			StartManyActions();
 		}
 
 		/// <summary>
-		/// Increment the number of completed actions and signal done if all actions complete.
-		/// This must be executed under the lock of Actions.
+		/// Add the number of completed actions and signal done if all actions complete.
 		/// </summary>
-		private void IncrementCompletedActions()
+		/// <param name="count">Number of completed actions to add.  If int.MaxValue, then the number is immediately set to total actions</param>
+		private void AddCompletedActions(int count)
 		{
-			var completedActions = Interlocked.Increment(ref _completedActions);
-			if (completedActions >= Actions.Length && !_doneTaskSource.Task.IsCompleted)
+			if (count == 0)
 			{
-				_doneTaskSource.SetResult();
+				// do nothing
+			}
+			else if (count != int.MaxValue)
+			{
+				if (Interlocked.Add(ref _completedActions, count) == TotalActions)
+				{
+					_doneTaskSource.SetResult();
+				}
+			}
+			else
+			{
+				if (Interlocked.Exchange(ref _completedActions, TotalActions) != TotalActions)
+				{
+					_doneTaskSource.SetResult();
+				}
 			}
 		}
 
@@ -855,16 +861,16 @@ namespace UnrealBuildTool
 
 					// prevent overwriting of error text
 					s_previousLineLength = -1;
+				}
+			}
 
-					// Cancel all other pending tasks
-					if (StopCompilationAfterErrors)
-					{
-						CancellationTokenSource.Cancel();
-						if (!_doneTaskSource.Task.IsCompleted)
-						{
-							_doneTaskSource.SetResult();
-						}
-					}
+			if (exitCode != 0 && exitCode != int.MaxValue)
+			{
+				// Cancel all other pending tasks
+				if (StopCompilationAfterErrors)
+				{
+					CancellationTokenSource.Cancel();
+					AddCompletedActions(int.MaxValue);
 				}
 			}
 		}
