- Fixed generatingXcode temporary projects (hybrid content-only-as-code projects)
- Updated the temp .Target.cs to use modern include order to avoid some warning spam

diff --git a/Engine/Source/Programs/AutomationTool/AutomationUtils/ProjectUtils.cs b/Engine/Source/Programs/AutomationTool/AutomationUtils/ProjectUtils.cs
--- a/Engine/Source/Programs/AutomationTool/AutomationUtils/ProjectUtils.cs
+++ b/Engine/Source/Programs/AutomationTool/AutomationUtils/ProjectUtils.cs
@@ -530,7 +530,8 @@ namespace AutomationTool
 				Writer.WriteLine("{");
 				Writer.WriteLine("\tpublic {0}Target(TargetInfo Target) : base(Target)", ProjectName);
 				Writer.WriteLine("\t{");
-				Writer.WriteLine("\t\tDefaultBuildSettings = BuildSettingsVersion.V2;");
+				Writer.WriteLine("\t\tDefaultBuildSettings = BuildSettingsVersion.V3;");
+				Writer.WriteLine("\t\tIncludeOrderVersion = EngineIncludeOrderVersion.Latest;");
 				Writer.WriteLine("\t\tType = TargetType.Game;");
 				Writer.WriteLine("\t\tExtraModuleNames.Add(\"{0}\");", ProjectName);
 				Writer.WriteLine("\t}");
diff --git a/Engine/Source/Programs/UnrealBuildTool/ProjectFiles/ProjectFileGenerator.cs b/Engine/Source/Programs/UnrealBuildTool/ProjectFiles/ProjectFileGenerator.cs
--- a/Engine/Source/Programs/UnrealBuildTool/ProjectFiles/ProjectFileGenerator.cs
+++ b/Engine/Source/Programs/UnrealBuildTool/ProjectFiles/ProjectFileGenerator.cs
@@ -2469,6 +2469,8 @@ namespace UnrealBuildTool
 			// Starting at the base directory of the module find a project which has the same directory as base, walking up the directory hierarchy until a match is found
 
 			DirectoryReference Path = CurModuleFile.Directory;
+			bool bIsTemporaryTarget = Path.ContainsName("Intermediate", 0);
+
 			while (!Path.IsRootDirectory())
 			{
 				// Figure out which game project this target belongs to
@@ -2480,7 +2482,7 @@ namespace UnrealBuildTool
 						FileReference ProjectInfo = Game;
 						BaseFolder = ProjectInfo.Directory;
 
-						if (bMakeProjectPerTarget)
+						if (bMakeProjectPerTarget && !bIsTemporaryTarget)
 						{
 							// find the project that the module is under, and has a TargetType target (useful with bMakeProjectPerTarget)
 							foreach (KeyValuePair<FileReference, ProjectFile> Pair in ProjectFileMap)
