UnrealBuildTool: Suppress "warning C5049: Embedding a full path may result in machine-dependent output" when determinisic mode is enabled until a proper fix is added

#rnx

diff --git a/Engine/Source/Programs/UnrealBuildTool/Platform/Windows/VCToolChain.cs b/Engine/Source/Programs/UnrealBuildTool/Platform/Windows/VCToolChain.cs
--- a/Engine/Source/Programs/UnrealBuildTool/Platform/Windows/VCToolChain.cs
+++ b/Engine/Source/Programs/UnrealBuildTool/Platform/Windows/VCToolChain.cs
@@ -547,6 +547,11 @@ namespace UnrealBuildTool
 				if (Target.WindowsPlatform.Compiler.IsMSVC())
 				{
 					Arguments.Add("/experimental:deterministic");
+					// TODO: the following occurs when /GL is also passed for link time code gen:
+					// warning C5049: Embedding a full path may result in machine-dependent output
+					// note: Use /pathmap to map the path prefix to a fixed value to ensure the full path can be reconstructed at link time
+					Arguments.Add("/wd5049");
+
 					if (CompileEnvironment.DeterministicWarningLevel == WarningLevel.Error)
 					{
 						Arguments.Add("/we5048");
