UnrealBuildTool: If /LTCG is enabled also pass it when generating libs

#rnx

diff --git a/Engine/Source/Programs/UnrealBuildTool/Platform/Windows/VCToolChain.cs b/Engine/Source/Programs/UnrealBuildTool/Platform/Windows/VCToolChain.cs
--- a/Engine/Source/Programs/UnrealBuildTool/Platform/Windows/VCToolChain.cs
+++ b/Engine/Source/Programs/UnrealBuildTool/Platform/Windows/VCToolChain.cs
@@ -1304,7 +1304,7 @@ namespace UnrealBuildTool
 			//
 			//	Shipping & LTCG
 			//
-			if (LinkEnvironment.Configuration == CppConfiguration.Shipping)
+			if (LinkEnvironment.bAllowLTCG || LinkEnvironment.Configuration == CppConfiguration.Shipping)
 			{
 				// Use link-time code generation.
 				Arguments.Add("/LTCG");
