Implementation of the action artifact cache.

diff --git a/Engine/Source/Programs/UnrealBuildTool.Tests/ArtifactsTest.cs b/Engine/Source/Programs/UnrealBuildTool.Tests/ArtifactsTest.cs
new file mode 100644
--- /dev/null
+++ b/Engine/Source/Programs/UnrealBuildTool.Tests/ArtifactsTest.cs
@@ -0,0 +1,97 @@
+// Copyright Epic Games, Inc. All Rights Reserved.
+
+using System.Threading;
+using EpicGames.Core;
+using Microsoft.Extensions.Logging.Abstractions;
+using Microsoft.VisualStudio.TestTools.UnitTesting;
+using UnrealBuildTool.Artifacts;
+
+#nullable enable
+
+namespace UnrealBuildToolTests
+{
+
+	[TestClass]
+	public class ArtifactsTest
+	{
+		static readonly Utf8String s_input1Data = new Utf8String("This is some sample input");
+		static readonly Utf8String s_input2Data = new Utf8String("This is some more sample input");
+		static readonly Utf8String s_output1Data = new Utf8String("This is the sample output");
+
+		private static Artifact MakeInput1()
+		{
+			return new Artifact(ArtifactDirectoryTree.Absolute, ArtifactType.Source, "Input1", IoHash.Compute(s_input1Data.Span));
+		}
+
+		private static Artifact MakeInput2()
+		{
+			return new Artifact(ArtifactDirectoryTree.Absolute, ArtifactType.Source, "Input2", IoHash.Compute(s_input2Data.Span));
+		}
+
+		private static Artifact MakeOutput1()
+		{
+			return new Artifact(ArtifactDirectoryTree.Absolute, ArtifactType.Object, "Output1", IoHash.Compute(s_output1Data.Span));
+		}
+
+		public static ArtifactMapping MakeBundle1()
+		{
+			return new ArtifactMapping(
+				IoHash.Compute(new Utf8String("SampleKey")),
+				IoHash.Compute(new Utf8String("SampleBundleKey")),
+				new Artifact[] { MakeInput1(), MakeInput2() },
+				new Artifact[] { MakeOutput1() }
+			);
+		}
+
+		public static ArtifactMapping MakeBundle2()
+		{
+			return new ArtifactMapping(
+				IoHash.Compute(new Utf8String("SampleKey")),
+				IoHash.Compute(new Utf8String("SampleBundleKeyV2")),
+				new Artifact[] { MakeInput1(), MakeInput2() },
+				new Artifact[] { MakeOutput1() }
+			);
+		}
+
+		[TestMethod]
+		public void ArtifactBundleStorageTest1()
+		{
+			CancellationToken cancellationToken = default;
+			
+			IArtifactCache cache = HordeStorageArtifactCache.CreateMemoryCache(NullLogger.Instance);
+
+			_ = cache.WaitForReadyAsync().Result;
+			Assert.AreEqual(ArtifactCacheState.Available, cache.State);
+
+			ArtifactMapping bundle1 = MakeBundle1();
+			cache.SaveArtifactMappingsAsync(new ArtifactMapping[] { bundle1 }, cancellationToken).Wait();
+
+			ArtifactMapping[] readBack1 = cache.QueryArtifactMappingsAsync(new IoHash[] { bundle1.Key }, cancellationToken).Result;
+			Assert.AreEqual(1, readBack1.Length);
+
+			ArtifactMapping bundle2 = MakeBundle2();
+			cache.SaveArtifactMappingsAsync(new ArtifactMapping[] { bundle2 }, cancellationToken).Wait();
+
+			ArtifactMapping[] readBack2 = cache.QueryArtifactMappingsAsync(new IoHash[] { bundle1.Key }, cancellationToken).Result;
+			Assert.AreEqual(2, readBack2.Length);
+
+			cache.FlushChangesAsync(cancellationToken).Wait();
+		}
+
+#if DISABLED
+		[TestMethod]
+		public void ArtifactBundleStorageTest2()
+		{
+			CancellationToken cancellationToken = default;
+
+			IArtifactCache cache = HordeStorageArtifactCache.CreateMemoryCache(NullLogger.Instance);
+
+			cache.WaitForReadyAsync().Wait(cancellationToken);
+
+			ArtifactBundle bundle1 = MakeBundle1();
+			ArtifactBundle[] readBack2 = cache.QueryArtifactBundles(new IoHash[] { bundle1.Key }, cancellationToken);
+			Assert.AreEqual(2, readBack2.Length);
+		}
+#endif
+	}
+}
diff --git a/Engine/Source/Programs/UnrealBuildTool.Tests/UnrealBuildTool.Tests.csproj b/Engine/Source/Programs/UnrealBuildTool.Tests/UnrealBuildTool.Tests.csproj
--- a/Engine/Source/Programs/UnrealBuildTool.Tests/UnrealBuildTool.Tests.csproj
+++ b/Engine/Source/Programs/UnrealBuildTool.Tests/UnrealBuildTool.Tests.csproj
@@ -15,6 +15,7 @@
   </ItemGroup>
 
   <ItemGroup>
+    <ProjectReference Include="..\Shared\EpicGames.Horde\EpicGames.Horde.csproj" />
     <ProjectReference Include="..\UnrealBuildTool\UnrealBuildTool.csproj" />
   </ItemGroup>
 
diff --git a/Engine/Source/Programs/UnrealBuildTool/Artifacts/ActionArtifactCache.cs b/Engine/Source/Programs/UnrealBuildTool/Artifacts/ActionArtifactCache.cs
new file mode 100644
--- /dev/null
+++ b/Engine/Source/Programs/UnrealBuildTool/Artifacts/ActionArtifactCache.cs
@@ -0,0 +1,289 @@
+// Copyright Epic Games, Inc. All Rights Reserved.
+
+using EpicGames.Core;
+using Microsoft.Extensions.Logging;
+using Microsoft.Extensions.Logging.Abstractions;
+using System.Collections.Generic;
+using System.Linq;
+using System.Text;
+using System.Threading;
+using System.Threading.Tasks;
+using UnrealBuildBase;
+
+namespace UnrealBuildTool.Artifacts
+{
+
+	/// <summary>
+	/// Generic handler for artifacts
+	/// </summary>
+	internal class ActionArtifactCache : IActionArtifactCache
+	{
+
+		/// <inheritdoc/>
+		public bool EnableReads { get; set; } = true;
+
+		/// <inheritdoc/>
+		public bool EnableWrites { get; set; } = true;
+
+		/// <inheritdoc/>
+		public IArtifactCache ArtifactCache { get; init; }
+
+		/// <summary>
+		/// Logging device
+		/// </summary>
+		private readonly ILogger _logger;
+
+		/// <summary>
+		/// Cache of dependency files.
+		/// </summary>
+		private readonly CppDependencyCache _cppDependencyCache;
+
+		/// <summary>
+		/// Cache for file hashes
+		/// </summary>
+		private readonly FileHasher _fileHasher;
+
+		/// <summary>
+		/// Construct a new artifact handler object
+		/// </summary>
+		/// <param name="artifactCache">Artifact cache instance</param>
+		/// <param name="cppDependencyCache">Previously created dependency cache</param>
+		/// <param name="logger">Logging device</param>
+		private ActionArtifactCache(IArtifactCache artifactCache, CppDependencyCache cppDependencyCache, ILogger logger)
+		{
+			_logger = logger;
+			_cppDependencyCache = cppDependencyCache;
+			ArtifactCache = artifactCache;
+			_fileHasher = new(NullLogger.Instance);
+		}
+
+		/// <summary>
+		/// Create a new action artifact cache using horde file based storage
+		/// </summary>
+		/// <param name="directory">Directory for the cache</param>
+		/// <param name="cppDependencyCache">Previously created dependency cache</param>
+		/// <param name="logger">Logging device</param>
+		/// <returns>Action artifact cache object</returns>
+		public static IActionArtifactCache CreateHordeFileCache(DirectoryReference directory, CppDependencyCache cppDependencyCache, ILogger logger)
+		{
+			IArtifactCache artifactCache = HordeStorageArtifactCache.CreateFileCache(directory, /*logger*/ NullLogger.Instance, false);
+			return new ActionArtifactCache(artifactCache, cppDependencyCache, logger);
+		}
+
+		/// <summary>
+		/// Create a new action artifact cache using horde memory based storage
+		/// </summary>
+		/// <param name="cppDependencyCache">Previously created dependency cache</param>
+		/// <param name="logger">Logging device</param>
+		/// <returns>Action artifact cache object</returns>
+		public static IActionArtifactCache CreateHordeMemoryCache(CppDependencyCache cppDependencyCache, ILogger logger)
+		{
+			IArtifactCache artifactCache = HordeStorageArtifactCache.CreateMemoryCache(logger);
+			return new ActionArtifactCache(artifactCache, cppDependencyCache, logger);
+		}
+
+		/// <inheritdoc/>
+		public async Task<bool> CompleteActionFromCacheAsync(LinkedAction action, CancellationToken cancellationToken)
+		{
+			if (!EnableReads)
+			{
+				return false;
+			}
+
+			IoHash key = await GetKeyAsync(action);
+
+			ArtifactMapping[] mappings = await ArtifactCache.QueryArtifactMappingsAsync(new IoHash[] { key }, cancellationToken);
+
+			if (mappings.Length == 0)
+			{
+				return false;
+			}
+
+			foreach (ArtifactMapping mapping in mappings)
+			{
+				bool match = true;
+
+				foreach (Artifact input in mapping.Inputs)
+				{
+					string name = input.Name.ToString();
+					FileItem item = FileItem.GetItemByPath(name);
+					if (!item.Exists)
+					{
+						match = false;
+						break;
+					}
+					if (input.ContentHash != await _fileHasher.GetDigestAsync(item, cancellationToken))
+					{
+						match = false;
+						break;
+					}
+				}
+
+				if (match)
+				{
+					bool[]? readResults = await ArtifactCache.QueryArtifactOutputsAsync(new[] { mapping }, cancellationToken);
+
+					if (readResults == null || readResults.Length == 0 || !readResults[0])
+					{
+						return false;
+					}
+					else
+					{
+						foreach (Artifact output in mapping.Outputs)
+						{
+							string outputName = output.Name.ToString();
+							FileItem item = FileItem.GetItemByPath(outputName);
+							item.ResetCachedInfo(); // newly created outputs need refreshing
+							_fileHasher.SetDigest(item, output.ContentHash);
+						}
+						return true;
+					}
+				}
+			}
+			return false;
+		}
+
+		/// <inheritdoc/>
+		public async Task ActionCompleteAsync(LinkedAction action, CancellationToken cancellationToken)
+		{
+			if (!EnableWrites)
+			{
+				return;
+			}
+
+			ArtifactMapping[] mappings = new ArtifactMapping[] { await CreateArtifactMappingAsync(action, cancellationToken) };
+			await ArtifactCache.SaveArtifactMappingsAsync(mappings, cancellationToken);
+		}
+
+		/// <inheritdoc/>
+		public async Task FlushChangesAsync(CancellationToken cancellationToken)
+		{
+			await ArtifactCache.FlushChangesAsync(cancellationToken);
+		}
+
+		/// <summary>
+		/// Create a new artifact mapping that represents the input and output of the action
+		/// </summary>
+		/// <param name="action">Source action</param>
+		/// <param name="cancellationToken">Cancellation token</param>
+		/// <returns>Artifact mapping</returns>
+		private async Task<ArtifactMapping> CreateArtifactMappingAsync(LinkedAction action, CancellationToken cancellationToken)
+		{
+			(IoHash key, IoHash mappingKey) = await GetKeyAndMappingKeyAsync(action);
+
+			// We gather the output files first to make sure that all the generated files (including the dependency file) gets
+			// their cached FileInfo reset.
+			List<Artifact> outputs = new();
+			foreach (FileItem output in action.ProducedItems)
+			{
+				output.ResetCachedInfo();
+				IoHash hash = await _fileHasher.GetDigestAsync(output, cancellationToken);
+				Artifact artifact = new(ArtifactDirectoryTree.Absolute, ArtifactType.Source, new Utf8String(output.AbsolutePath), hash);
+				outputs.Add(artifact);
+			}
+
+			List<Artifact> inputs = new();
+			foreach (FileItem input in action.PrerequisiteItems)
+			{
+				IoHash hash = await _fileHasher.GetDigestAsync(input, cancellationToken);
+				Artifact artifact = new(ArtifactDirectoryTree.Absolute, ArtifactType.Source, new Utf8String(input.AbsolutePath), hash);
+				inputs.Add(artifact);
+			}
+
+			if (action.DependencyListFile != null)
+			{
+				if (_cppDependencyCache.TryGetDependencies(action.DependencyListFile, _logger, out List<FileItem>? depedencies))
+				{
+					foreach (FileItem dep in depedencies)
+					{
+						IoHash hash = await _fileHasher.GetDigestAsync(dep, cancellationToken);
+						Artifact artifact = new(ArtifactDirectoryTree.Absolute, ArtifactType.Source, new Utf8String(dep.AbsolutePath), hash);
+						inputs.Add(artifact);
+					}
+				}
+			}
+
+			return new(key, mappingKey, inputs.ToArray(), outputs.ToArray());
+		}
+
+		/// <summary>
+		/// Get the key has for the action
+		/// </summary>
+		/// <param name="action">Source action</param>
+		/// <returns>Task returning the key</returns>
+		private async Task<IoHash> GetKeyAsync(LinkedAction action)
+		{
+			StringBuilder builder = new();
+			await AppendKeyAsync(builder, action);
+			IoHash key = IoHash.Compute(new Utf8String(builder.ToString()));
+			return key;
+		}
+
+		/// <summary>
+		/// Generate the key and mapping key hashes for the action
+		/// </summary>
+		/// <param name="action">Source action</param>
+		/// <returns>Task object with the key and mapping key</returns>
+		private async Task<(IoHash, IoHash)> GetKeyAndMappingKeyAsync(LinkedAction action)
+		{
+			StringBuilder builder = new();
+			await AppendKeyAsync(builder, action);
+			IoHash key = IoHash.Compute(new Utf8String(builder.ToString()));
+			await AppendMappingKeyAsync(builder, action);
+			IoHash mappingKey = IoHash.Compute(new Utf8String(builder.ToString()));
+			return (key, mappingKey);
+		}
+
+		/// <summary>
+		/// Generate the lookup key.  This key is generated from the action's inputs.
+		/// </summary>
+		/// <param name="builder">Destination builder</param>
+		/// <param name="action">Source action</param>
+		/// <returns>Task object</returns>
+		private async Task AppendKeyAsync(StringBuilder builder, LinkedAction action)
+		{
+			builder.AppendLine(action.CommandVersion);
+			builder.AppendLine(action.CommandArguments);
+
+			FileItem[] inputs = action.PrerequisiteItems.ToArray();
+			Task<IoHash>[] waits = new Task<IoHash>[inputs.Length];
+			for (int index = 0; index < inputs.Length; index++)
+			{
+				builder.AppendLine(inputs[index].FullName);
+				waits[index] = _fileHasher.GetDigestAsync(inputs[index]);
+			}
+			await Task.WhenAll(waits);
+			foreach (Task<IoHash> wait in waits)
+			{
+				builder.AppendLine(wait.Result.ToString());
+			}
+		}
+
+		/// <summary>
+		/// Generate the full mapping key.  This contains the hashes for the action's inputs and the dependent files.
+		/// </summary>
+		/// <param name="builder">Destination builder</param>
+		/// <param name="action">Source action</param>
+		/// <returns>Task object</returns>
+		private async Task AppendMappingKeyAsync(StringBuilder builder, LinkedAction action)
+		{
+			if (action.DependencyListFile != null)
+			{
+				if (_cppDependencyCache.TryGetDependencies(action.DependencyListFile, _logger, out List<FileItem>? dependencies))
+				{
+					Task<IoHash>[] waits = new Task<IoHash>[dependencies.Count];
+					for (int index = 0; index < dependencies.Count; index++)
+					{
+						builder.AppendLine(dependencies[index].FullName);
+						waits[index] = _fileHasher.GetDigestAsync(dependencies[index]);
+					}
+					await Task.WhenAll(waits);
+					foreach (Task<IoHash> wait in waits)
+					{
+						builder.AppendLine(wait.Result.ToString());
+					}
+				}
+			}
+		}
+	}
+}
diff --git a/Engine/Source/Programs/UnrealBuildTool/Artifacts/Artifacts.cs b/Engine/Source/Programs/UnrealBuildTool/Artifacts/Artifacts.cs
--- a/Engine/Source/Programs/UnrealBuildTool/Artifacts/Artifacts.cs
+++ b/Engine/Source/Programs/UnrealBuildTool/Artifacts/Artifacts.cs
@@ -93,9 +93,9 @@ namespace UnrealBuildTool.Artifacts
 	}
 
 	/// <summary>
-	/// Collection of input and output artifacts.
+	/// Given a set of inputs, provide a list of outputs
 	/// </summary>
-	public readonly struct ArtifactBundle : IEquatable<ArtifactBundle>
+	public readonly struct ArtifactMapping : IEquatable<ArtifactMapping>
 	{
 
 		/// <summary>
@@ -106,7 +106,7 @@ namespace UnrealBuildTool.Artifacts
 		/// <summary>
 		/// The unique hash for all inputs and the environment
 		/// </summary>
-		public readonly IoHash BundleKey;
+		public readonly IoHash MappingKey;
 
 		/// <summary>
 		/// Information about all inputs
@@ -122,43 +122,43 @@ namespace UnrealBuildTool.Artifacts
 		/// Constructor
 		/// </summary>
 		/// <param name="key">The hash of the primary input and the environment</param>
-		/// <param name="bundleKey">The unique hash for all inputs and the environment</param>
+		/// <param name="mappingKey">The unique hash for all inputs and the environment</param>
 		/// <param name="inputs">Information about all inputs</param>
 		/// <param name="outputs">Information about all outputs</param>
-		public ArtifactBundle(IoHash key, IoHash bundleKey, Artifact[] inputs, Artifact[] outputs)
+		public ArtifactMapping(IoHash key, IoHash mappingKey, Artifact[] inputs, Artifact[] outputs)
 		{
 			Key = key;
-			BundleKey = bundleKey;
+			MappingKey = mappingKey;
 			Inputs = inputs;
 			Outputs = outputs;
 		}
 
 		/// <inheritdoc/>
-		public bool Equals(ArtifactBundle other)
+		public bool Equals(ArtifactMapping other)
 		{
-			return Key == other.Key && BundleKey == other.BundleKey && Inputs.SequenceEqual(other.Inputs) && Outputs.SequenceEqual(other.Outputs);
+			return Key == other.Key && MappingKey == other.MappingKey && Inputs.SequenceEqual(other.Inputs) && Outputs.SequenceEqual(other.Outputs);
 		}
 
 		/// <inheritdoc/>
 		public override bool Equals(object? obj)
 		{
-			return obj is ArtifactBundle artifactBundle && Equals(artifactBundle);
+			return obj is ArtifactMapping artifactMapping && Equals(artifactMapping);
 		}
 
 		/// <inheritdoc/>
 		public override int GetHashCode()
 		{
-			return HashCode.Combine(Key.GetHashCode(), BundleKey.GetHashCode(), Inputs.GetHashCode(), Outputs.GetHashCode()); 
+			return HashCode.Combine(Key.GetHashCode(), MappingKey.GetHashCode(), Inputs.GetHashCode(), Outputs.GetHashCode()); 
 		}
 
 		/// <inheritdoc/>
-		public static bool operator ==(ArtifactBundle left, ArtifactBundle right)
+		public static bool operator ==(ArtifactMapping left, ArtifactMapping right)
 		{
 			return left.Equals(right);
 		}
 
 		/// <inheritdoc/>
-		public static bool operator !=(ArtifactBundle left, ArtifactBundle right)
+		public static bool operator !=(ArtifactMapping left, ArtifactMapping right)
 		{
 			return !(left == right);
 		}
@@ -198,30 +198,30 @@ namespace UnrealBuildTool.Artifacts
 		}
 
 		/// <summary>
-		/// Read an artifact bundle structure
+		/// Read an artifact mapping structure
 		/// </summary>
 		/// <param name="reader">Source reader</param>
-		/// <returns>Read artifact bundle structure</returns>
-		public static ArtifactBundle ReadArtifactBundle(this IMemoryReader reader)
+		/// <returns>Read artifact mapping structure</returns>
+		public static ArtifactMapping ReadArtifactMapping(this IMemoryReader reader)
 		{
 			IoHash key = reader.ReadIoHash();
-			IoHash bundleKey = reader.ReadIoHash();
+			IoHash mappingKey = reader.ReadIoHash();
 			Artifact[] inputs = reader.ReadVariableLengthArray(() => reader.ReadArtifact());
 			Artifact[] outputs = reader.ReadVariableLengthArray(() => reader.ReadArtifact());
-			return new ArtifactBundle(key, bundleKey, inputs, outputs);
+			return new ArtifactMapping(key, mappingKey, inputs, outputs);
 		}
 
 		/// <summary>
-		/// Write an artifact bundle structure
+		/// Write an artifact mapping structure
 		/// </summary>
 		/// <param name="writer">Destination writer</param>
-		/// <param name="artifactBundle">Artifact bundle to be written</param>
-		public static void WriteArtifactBundle(this IMemoryWriter writer, ArtifactBundle artifactBundle)
+		/// <param name="artifactMapping">Artifact mapping to be written</param>
+		public static void WriteArtifactMapping(this IMemoryWriter writer, ArtifactMapping artifactMapping)
 		{
-			writer.WriteIoHash(artifactBundle.Key);
-			writer.WriteIoHash(artifactBundle.BundleKey);
-			writer.WriteVariableLengthArray(artifactBundle.Inputs, x => writer.WriteArtifact(x));
-			writer.WriteVariableLengthArray(artifactBundle.Outputs, x => writer.WriteArtifact(x));
+			writer.WriteIoHash(artifactMapping.Key);
+			writer.WriteIoHash(artifactMapping.MappingKey);
+			writer.WriteVariableLengthArray(artifactMapping.Inputs, x => writer.WriteArtifact(x));
+			writer.WriteVariableLengthArray(artifactMapping.Outputs, x => writer.WriteArtifact(x));
 		}
 	}
 
@@ -268,33 +268,34 @@ namespace UnrealBuildTool.Artifacts
 		public Task<ArtifactCacheState> WaitForReadyAsync();
 
 		/// <summary>
-		/// Given a collection of partial keys return all matching bundles
+		/// Given a collection of partial keys return all matching mappings
 		/// </summary>
 		/// <param name="partialKeys">Source file key</param>
 		/// <param name="cancellationToken">Token to be used to cancel operations</param>
 		/// <returns>Collection of all known artifacts</returns>
-		public Task<ArtifactBundle[]> QueryArtifactBundlesAsync(IoHash[] partialKeys, CancellationToken cancellationToken);
+		public Task<ArtifactMapping[]> QueryArtifactMappingsAsync(IoHash[] partialKeys, CancellationToken cancellationToken);
 
 		/// <summary>
 		/// Query the actual contents of a group of artifacts
 		/// </summary>
-		/// <param name="bundles">Bundles to be read</param>
+		/// <param name="artifactMappings">Mappings to be read</param>
 		/// <param name="cancellationToken">Token to be used to cancel operations</param>
 		/// <returns>Dictionary of the artifacts</returns>
-		public Task<bool[]?> QueryArtifactOutputsAsync(ArtifactBundle[] bundles, CancellationToken cancellationToken);
+		public Task<bool[]?> QueryArtifactOutputsAsync(ArtifactMapping[] artifactMappings, CancellationToken cancellationToken);
 
 		/// <summary>
 		/// Save new artifact to the cache
 		/// </summary>
-		/// <param name="artifactBundles">Collection of artifacts to be saved</param>
+		/// <param name="artifactMappings">Collection of artifacts to be saved</param>
 		/// <param name="cancellationToken">Token to be used to cancel operations</param>
-		public Task SaveArtifactBundlesAsync(ArtifactBundle[] artifactBundles, CancellationToken cancellationToken);
+		/// <returns>Asynchronous task objects</returns>
+		public Task SaveArtifactMappingsAsync(ArtifactMapping[] artifactMappings, CancellationToken cancellationToken);
 
 		/// <summary>
 		/// Flush all updates in the cache asynchronously.
 		/// </summary>
 		/// <param name="cancellationToken">Token to be used to cancel operations</param>
-		/// <returns>Async task objects</returns>
+		/// <returns>Asynchronous task objects</returns>
 		public Task FlushChangesAsync(CancellationToken cancellationToken);
 	}
 
@@ -307,7 +308,7 @@ namespace UnrealBuildTool.Artifacts
 		/// <summary>
 		/// Underlying artifact cache
 		/// </summary>
-		public IArtifactCache? ArtifactCache { get; }
+		public IArtifactCache ArtifactCache { get; }
 
 		/// <summary>
 		/// If true, reads will be serviced
@@ -332,14 +333,14 @@ namespace UnrealBuildTool.Artifacts
 		/// </summary>
 		/// <param name="action">Completed action</param>
 		/// <param name="cancellationToken">Token to be used to cancel operations</param>
-		/// <returns>Async task object</returns>
+		/// <returns>Asynchronous task object</returns>
 		public Task ActionCompleteAsync(LinkedAction action, CancellationToken cancellationToken);
 
 		/// <summary>
 		/// Flush all updates in the cache asynchronously.
 		/// </summary>
 		/// <param name="cancellationToken">Token to be used to cancel operations</param>
-		/// <returns>Async task object</returns>
+		/// <returns>Asynchronous task object</returns>
 		public Task FlushChangesAsync(CancellationToken cancellationToken);
 	}
 }
diff --git a/Engine/Source/Programs/UnrealBuildTool/Artifacts/HordeStorageArtifactCache.cs b/Engine/Source/Programs/UnrealBuildTool/Artifacts/HordeStorageArtifactCache.cs
new file mode 100644
--- /dev/null
+++ b/Engine/Source/Programs/UnrealBuildTool/Artifacts/HordeStorageArtifactCache.cs
@@ -0,0 +1,583 @@
+// Copyright Epic Games, Inc. All Rights Reserved.
+
+using EpicGames.Core;
+using EpicGames.Horde.Storage;
+using EpicGames.Horde.Storage.Backends;
+using EpicGames.Horde.Storage.Nodes;
+using Microsoft.Extensions.Logging;
+using System;
+using System.Collections.Generic;
+using System.IO;
+using System.Linq;
+using System.Threading;
+using System.Threading.Tasks;
+
+namespace UnrealBuildTool.Artifacts
+{
+	/// <summary>
+	/// Horde specific artifact mapping structure that also contains the file nodes for the outputs
+	/// </summary>
+	readonly struct HordeArtifactMapping
+	{
+
+		/// <summary>
+		/// Artifact mapping
+		/// </summary>
+		public readonly ArtifactMapping ArtifactMapping;
+
+		/// <summary>
+		/// Collection of output file references.  There should be exactly the same number
+		/// of file references as outputs in the mapping
+		/// </summary>
+		public readonly TreeNodeRef<FileNode>[] OutputRefs;
+
+		/// <summary>
+		/// Construct a new horde artifact number
+		/// </summary>
+		/// <param name="artifactMapping">Artifact mapping</param>
+		/// <exception cref="ArgumentException"></exception>
+		public HordeArtifactMapping(ArtifactMapping artifactMapping)
+		{
+			ArtifactMapping = artifactMapping;
+			OutputRefs = new TreeNodeRef<FileNode>[ArtifactMapping.Outputs.Length];
+		}
+
+		/// <summary>
+		/// Construct a new artifact mapping from the reader
+		/// </summary>
+		/// <param name="reader">Source reader</param>
+		public HordeArtifactMapping(ITreeNodeReader reader)
+		{
+			ArtifactMapping = reader.ReadArtifactMapping();
+			OutputRefs = reader.ReadVariableLengthArray(() => new TreeNodeRef<FileNode>(reader));
+		}
+
+		/// <summary>
+		/// Serialize the artifact mapping 
+		/// </summary>
+		/// <param name="writer">Destination writer</param>
+		public void Serialize(ITreeNodeWriter writer)
+		{
+			writer.WriteArtifactMapping(ArtifactMapping);
+			writer.WriteVariableLengthArray(OutputRefs, x => x.Serialize(writer));
+		}
+
+		/// <summary>
+		/// Return an enumeration of the references
+		/// </summary>
+		/// <returns></returns>
+		public IEnumerable<TreeNodeRef> EnumerateRefs() => OutputRefs;
+
+		/// <summary>
+		/// Write all the files to disk
+		/// </summary>
+		/// <param name="writer">Destination writer</param>
+		/// <param name="cancellationToken">Cancellation token</param>
+		/// <returns>Task</returns>
+		public async Task WriteFilesAsync(TreeWriter writer, CancellationToken cancellationToken)
+		{
+			ChunkingOptionsForNodeType nodeTypeOptions = new(512 * 1024, 2 * 1024 * 1024, 1 * 1024 * 1024);
+			ChunkingOptions options = new(){ LeafOptions = nodeTypeOptions, InteriorOptions = nodeTypeOptions };
+
+			FileNodeWriter fileWriter = new(writer, options);
+			int index = 0;
+			foreach (Artifact artifact in ArtifactMapping.Outputs)
+			{
+				string outputName = artifact.Name.ToString();
+				using FileStream stream = new(outputName, FileMode.Open, FileAccess.Read, FileShare.Read);
+				OutputRefs[index++] = new TreeNodeRef<FileNode>(await fileWriter.CreateAsync(stream, nodeTypeOptions.TargetSize, cancellationToken)); 
+			}
+		}
+	}
+
+	/// <summary>
+	/// Series of helper methods for serialization
+	/// </summary>
+	static class HordeArtifactReaderWriterExtensions
+	{
+
+		/// <summary>
+		/// Read a horde artifact mapping
+		/// </summary>
+		/// <param name="reader">Source reader</param>
+		/// <returns>Created artifact mapping</returns>
+		public static HordeArtifactMapping ReadHordeArtifactMapping(this ITreeNodeReader reader)
+		{
+			return new HordeArtifactMapping(reader);
+		}
+
+		/// <summary>
+		/// Write a horde artifact mapping
+		/// </summary>
+		/// <param name="writer">Destination writer</param>
+		/// <param name="artifactMapping">Artifact mapping to write</param>
+		public static void WriteHordeArtifactMapping(this ITreeNodeWriter writer, HordeArtifactMapping artifactMapping)
+		{
+			artifactMapping.Serialize(writer);
+		}
+	}
+
+	/// <summary>
+	/// Horde node that represents a collection of mapping nodes 
+	/// </summary>
+	[TreeNode("{E8DBCD77-861D-4CAE-B77F-5807D26E2533}")]
+	class ArtifactMappingCollectionNode : TreeNode
+	{
+
+		/// <summary>
+		/// Collection of mappings
+		/// </summary>
+		public Dictionary<IoHash, HordeArtifactMapping> Mappings = new();
+
+		/// <summary>
+		/// Construct a new collection
+		/// </summary>
+		public ArtifactMappingCollectionNode()
+		{
+		}
+
+		/// <summary>
+		/// Construct a new mapping collection from the source reader
+		/// </summary>
+		/// <param name="reader">Source reader</param>
+		public ArtifactMappingCollectionNode(ITreeNodeReader reader)
+		{
+			Mappings = reader.ReadDictionary<IoHash, HordeArtifactMapping>(() => reader.ReadIoHash(), () => reader.ReadHordeArtifactMapping());
+		}
+
+		/// <inheritdoc/>
+		public override void Serialize(ITreeNodeWriter writer)
+		{
+			writer.WriteDictionary<IoHash, HordeArtifactMapping>(Mappings, (x) => writer.WriteIoHash(x), (x) => writer.WriteHordeArtifactMapping(x));
+		}
+
+		/// <inheritdoc/>
+		public override IEnumerable<TreeNodeRef> EnumerateRefs()
+		{
+			foreach (HordeArtifactMapping artifactMapping in Mappings.Values)
+			{
+				foreach (TreeNodeRef outputRef in artifactMapping.OutputRefs)
+				{
+					yield return outputRef;
+				}
+			}
+		}
+
+		/// <summary>
+		/// Mark the node as dirty
+		/// </summary>
+		public new void MarkAsDirty()
+		{
+			base.MarkAsDirty();
+		}
+	}
+
+	/// <summary>
+	/// Class for managing artifacts using horde storage
+	/// </summary>
+	public class HordeStorageArtifactCache : IArtifactCache
+	{
+		/// <summary>
+		/// Defines the theoretical max number of pending mappings to write
+		/// </summary>
+		const int MaxPendingSize = 128;
+
+		/// <summary>
+		/// Underlying storage object
+		/// </summary>
+		private IStorageClient? _store = null;
+
+		/// <summary>
+		/// Logger to be used
+		/// </summary>
+		private readonly ILogger _logger;
+
+		/// <summary>
+		/// Task used to wait on ready state
+		/// </summary>
+		private Task<ArtifactCacheState>? _readyTask = null;
+
+		/// <summary>
+		/// Ready state
+		/// </summary>
+		private int _state = (int)ArtifactCacheState.Pending;
+
+		/// <summary>
+		/// Collection of tree reader options. This is needed for the extra node types.
+		/// </summary>
+		private readonly TreeReaderOptions _treeReaderOptions;
+
+		/// <summary>
+		/// Collection of mappings waiting to be written
+		/// </summary>
+		private readonly List<ArtifactMapping> _pendingMappings;
+
+		/// <summary>
+		/// Task for any pending flush
+		/// </summary>
+		private Task? _pendingFlushTask = null;
+
+		/// <summary>
+		/// Controls access to shared data structures
+		/// </summary>
+		private readonly SemaphoreSlim _semaphore = new(1);
+
+		/// <summary>
+		/// Test to see if the cache is ready
+		/// </summary>
+		public ArtifactCacheState State
+		{
+			get => (ArtifactCacheState) Interlocked.Add(ref _state, 0);
+			private set => Interlocked.Exchange(ref _state, (int)value);
+		}
+
+		/// <summary>
+		/// Create a memory only cache
+		/// </summary>
+		/// <returns>Storage client instance</returns>
+		public static IArtifactCache CreateMemoryCache(ILogger logger)
+		{
+			HordeStorageArtifactCache cache = new(new MemoryStorageClient(), logger)
+			{
+				State = ArtifactCacheState.Available
+			};
+			return cache;
+		}
+
+		/// <summary>
+		/// Create a file based cache
+		/// </summary>
+		/// <param name="directory">Destination directory</param>
+		/// <param name="logger">Logging object</param>
+		/// <param name="cleanDirectory">If true, clean the directory</param>
+		/// <returns>Storage client instance</returns>
+		public static IArtifactCache CreateFileCache(DirectoryReference directory, ILogger logger, bool cleanDirectory)
+		{
+			HordeStorageArtifactCache cache = new(null, logger);
+			cache._readyTask = Task.Run(() => cache.InitFileCache(directory, logger, cleanDirectory));
+			return cache;
+		}
+
+		/// <summary>
+		/// Constructor
+		/// </summary>
+		/// <param name="storage">Storage object to use</param>
+		/// <param name="logger">Logging destination</param>
+		private HordeStorageArtifactCache(IStorageClient? storage, ILogger logger)
+		{
+			_store = storage;
+			_logger = logger;
+			_treeReaderOptions = new TreeReaderOptions();
+			_treeReaderOptions.Types.Add(typeof(ArtifactMappingCollectionNode));
+			_pendingMappings = new(MaxPendingSize);
+		}
+
+		/// <inheritdoc/>
+		public Task<ArtifactCacheState> WaitForReadyAsync()
+		{
+			return _readyTask ?? Task.FromResult<ArtifactCacheState>(State);
+		}
+
+		/// <inheritdoc/>
+		public async Task<ArtifactMapping[]> QueryArtifactMappingsAsync(IoHash[] partialKeys, CancellationToken cancellationToken)
+		{
+			if (State != ArtifactCacheState.Available || _store == null)
+			{
+				return Array.Empty<ArtifactMapping>();
+			}
+
+			List<ArtifactMapping> mappings = new();
+			await _semaphore.WaitAsync(cancellationToken);
+			try
+			{
+				TreeReader reader = CreateTreeReader();
+				foreach (IoHash key in partialKeys)
+				{
+					lock (_pendingMappings)
+					{
+						mappings.AddRange(_pendingMappings.Where(x => x.Key == key));
+					}
+					ArtifactMappingCollectionNode? node = await reader.TryReadNodeAsync<ArtifactMappingCollectionNode>(GetRefName(key), default, cancellationToken);
+					if (node != null)
+					{
+						foreach (HordeArtifactMapping mapping in node.Mappings.Values)
+						{
+							mappings.Add(mapping.ArtifactMapping);
+						}
+					}
+				}
+			}
+			finally
+			{
+				_semaphore.Release();
+			}
+			return mappings.ToArray();
+		}
+
+		/// <inheritdoc/>
+		public async Task<bool[]?> QueryArtifactOutputsAsync(ArtifactMapping[] mappings, CancellationToken cancellationToken)
+		{
+			if (State != ArtifactCacheState.Available || _store == null)
+			{
+				return null;
+			}
+
+			TreeReader reader = CreateTreeReader();
+			bool[] output = new bool[mappings.Length];
+			Array.Fill(output, false);
+
+			for (int index = 0; index < mappings.Length; index++)
+			{
+				output[index] = false;
+				ArtifactMapping mapping = mappings[index];
+				ArtifactMappingCollectionNode? node = await reader.TryReadNodeAsync<ArtifactMappingCollectionNode>(GetRefName(mapping.Key), default, cancellationToken);
+				if (node != null)
+				{
+					if (node.Mappings.TryGetValue(mapping.MappingKey, out HordeArtifactMapping hordeArtifactMapping))
+					{
+						output[index] = true;
+
+						int refIndex = 0;
+						foreach (TreeNodeRef<FileNode> artifactRef in hordeArtifactMapping.OutputRefs)
+						{
+							if (artifactRef.Handle == null)
+							{
+								output[index] = false;
+								break;
+							}
+							try
+							{
+								string outputName = hordeArtifactMapping.ArtifactMapping.Outputs[refIndex++].Name.ToString();
+								using FileStream stream = new(outputName, FileMode.Create, FileAccess.Write, FileShare.ReadWrite);
+								await FileNode.CopyToStreamAsync(reader, artifactRef.Handle.Locator, stream, cancellationToken);
+							}
+							catch (Exception)
+							{
+								output[index] = false;
+								break;
+							}
+						}
+
+						if (!output[index])
+						{
+							foreach (Artifact artifact in hordeArtifactMapping.ArtifactMapping.Outputs)
+							{
+								string outputName = artifact.Name.ToString();
+								if (File.Exists(outputName))
+								{
+									try
+									{
+										File.Delete(outputName);
+									}
+									catch (Exception)
+									{
+									}
+								}
+							}
+						}
+					}
+				}
+			}
+			return output;
+		}
+
+		/// <inheritdoc/>
+		public async Task SaveArtifactMappingsAsync(ArtifactMapping[] artifactMappings, CancellationToken cancellationToken)
+		{
+			if (State != ArtifactCacheState.Available || _store == null)
+			{
+				return;
+			}
+
+			lock (_pendingMappings)
+			{
+				_pendingMappings.AddRange(artifactMappings);
+			}
+
+			Task? task = FlushChangesInternalAsync(false, cancellationToken);
+			if (task != null)
+			{
+				await task;
+			}
+		}
+
+		/// <inheritdoc/>
+		public async Task FlushChangesAsync(CancellationToken cancellationToken)
+		{
+			if (State != ArtifactCacheState.Available || _store == null)
+			{
+				return;
+			}
+
+			Task? task = FlushChangesInternalAsync(true, cancellationToken);
+			if (task != null)
+			{
+				await task;
+			}
+		}
+
+		/// <summary>
+		/// Optionally flush all pending mappings
+		/// </summary>
+		/// <param name="force">If true, force a flush</param>
+		/// <param name="cancellationToken">Cancellation token</param>
+		private Task? FlushChangesInternalAsync(bool force, CancellationToken cancellationToken)
+		{
+			Task? pendingFlushTask = null;
+			lock (_pendingMappings)
+			{
+
+				// If any prior flush task has completed, then forget it
+				if (_pendingFlushTask != null && _pendingFlushTask.IsCompleted)
+				{
+					_pendingFlushTask = null;
+				}
+
+				// We start a new flush under the following condition
+				//
+				// 1) Mappings must be pending
+				// 2) Create a new task if force is specified
+				// 3) -OR- Create a new task if there is no current task and we have reached the limit 
+				if (_pendingMappings.Count > 0 && (force || (_pendingMappings.Count >= MaxPendingSize && _pendingFlushTask == null)))
+				{
+					ArtifactMapping[] mappingsToFlush = _pendingMappings.ToArray();
+					Task? priorTask = _pendingFlushTask;
+					_pendingMappings.Clear();
+					pendingFlushTask = _pendingFlushTask = new(async () =>
+					{
+
+						// When forcing, we might have a prior flush task in progress.  Wait for it to complete
+						if (priorTask != null)
+						{
+							await priorTask;
+						}
+
+						// Block reading while we update the mappings
+						await _semaphore.WaitAsync(cancellationToken);
+						try
+						{
+							List<Task> tasks = CommitMappings(mappingsToFlush, cancellationToken);
+							await Task.WhenAll(tasks);
+						}
+						finally
+						{
+							_semaphore.Release();
+						}
+					}, cancellationToken);
+				}
+			}
+
+			// Start the task outside of the lock
+			pendingFlushTask?.Start();
+			return pendingFlushTask;
+		}
+
+		/// <summary>
+		/// Add a group of mappings to a new or existing source
+		/// </summary>
+		/// <param name="mappings">New mappings to add</param>
+		/// <param name="cancellationToken">Token to be used to cancel operations</param>
+		/// <returns>List of tasks</returns>
+		private List<Task> CommitMappings(ArtifactMapping[] mappings, CancellationToken cancellationToken)
+		{
+			List<Task> tasks = new();
+			if (mappings.Length == 0)
+			{
+				return tasks;
+			}
+
+			TreeReader reader = CreateTreeReader();
+
+			// Loop through the mappings
+			foreach (ArtifactMapping mapping in mappings)
+			{
+
+				// Create the task to write the files
+				tasks.Add(Task.Run(async () =>
+				{
+					RefName refName = GetRefName(mapping.Key);
+
+					// Locate the destination collection for this key
+					ArtifactMappingCollectionNode? node = reader.TryReadNodeAsync<ArtifactMappingCollectionNode>(refName, default, cancellationToken).Result;
+					node ??= new ArtifactMappingCollectionNode();
+
+					// Update the mapping collection
+					HordeArtifactMapping hordeArtifactMapping = new(mapping);
+					node.Mappings[mapping.MappingKey] = hordeArtifactMapping;
+					node.MarkAsDirty();
+
+					// Save the mapping file
+					using TreeWriter writer = CreateTreeWriter();
+					await hordeArtifactMapping.WriteFilesAsync(writer, cancellationToken);
+
+					// Save the collection
+					NodeHandle _ = await writer.WriteAsync(refName, node, null, cancellationToken);
+				}, cancellationToken));
+			}
+			return tasks;
+		}
+
+		/// <summary>
+		/// Initialize a file based cache
+		/// </summary>
+		/// <param name="directory">Destination directory</param>
+		/// <param name="logger">Logger</param>
+		/// <param name="cleanDirectory">If true, clean the directory</param>
+		/// <returns>Cache state</returns>
+		private ArtifactCacheState InitFileCache(DirectoryReference directory, ILogger logger, bool cleanDirectory)
+		{
+			try
+			{
+				if (cleanDirectory)
+				{
+					// Clear the output directory
+					try
+					{
+						Directory.Delete(directory.FullName, true);
+					}
+					catch (Exception)
+					{ }
+				}
+				Directory.CreateDirectory(directory.FullName);
+
+				_store = new FileStorageClient(directory, logger);
+
+				State = ArtifactCacheState.Available;
+				return State;
+			}
+			catch (Exception)
+			{
+				State = ArtifactCacheState.Unavailable;
+				throw;
+			}
+		}
+
+		/// <summary>
+		/// Return the ref name for horde storage given a mapping collection key
+		/// </summary>
+		/// <param name="key">Mapping collection key</param>
+		/// <returns>The reference name</returns>
+		private static RefName GetRefName(IoHash key)
+		{
+			return new RefName($"action_artifact_v1_{key}");
+		}
+
+		/// <summary>
+		/// Create a new tree reader
+		/// </summary>
+		/// <returns>TreeReader</returns>
+		private TreeReader CreateTreeReader()
+		{
+			return new(_store!, null, _treeReaderOptions, _logger);
+		}
+
+		/// <summary>
+		/// Create a new tree writer 
+		/// </summary>
+		/// <returns>TreeWriter</returns>
+		private TreeWriter CreateTreeWriter()
+		{
+			return new(_store!, null, default, null, _logger);
+		}
+	}
+}
diff --git a/Engine/Source/Programs/UnrealBuildTool/Configuration/BuildConfiguration.cs b/Engine/Source/Programs/UnrealBuildTool/Configuration/BuildConfiguration.cs
--- a/Engine/Source/Programs/UnrealBuildTool/Configuration/BuildConfiguration.cs
+++ b/Engine/Source/Programs/UnrealBuildTool/Configuration/BuildConfiguration.cs
@@ -203,6 +203,29 @@ namespace UnrealBuildTool
 		[CommandLine("-CompactOutput")]
 		private bool bCompactOutputCommandLine = false;
 
+		/// <summary>
+		/// If set, artifacts will be read
+		/// </summary>
+		[XmlConfigFile(Category = "BuildConfiguration")]
+		[CommandLine("-ArtifactReads", Value = "True")]
+		[CommandLine("-NoArtifactReads", Value = "False")]
+		public bool bArtifactRead = true;
+
+		/// <summary>
+		/// If set, artifacts will be written
+		/// </summary>
+		[XmlConfigFile(Category = "BuildConfiguration")]
+		[CommandLine("-ArtifactWrites", Value = "True")]
+		[CommandLine("-NoArtifactWrites", Value = "False")]
+		public bool bArtifactWrites = true;
+
+		/// <summary>
+		/// Location to store the artifacts.
+		/// </summary>
+		[XmlConfigFile(Category = "BuildConfiguration")]
+		[CommandLine("-ArtifactDirectory=")]
+		public string ArtifactDirectory = String.Empty;
+
 		/// <summary>
 		/// Instruct the executor to write compact output e.g. only errors, if supported by the executor,
 		/// and only if output is not being redirected e.g. during a build from within Visual Studio
diff --git a/Engine/Source/Programs/UnrealBuildTool/Modes/BuildMode.cs b/Engine/Source/Programs/UnrealBuildTool/Modes/BuildMode.cs
--- a/Engine/Source/Programs/UnrealBuildTool/Modes/BuildMode.cs
+++ b/Engine/Source/Programs/UnrealBuildTool/Modes/BuildMode.cs
@@ -534,8 +534,25 @@ namespace UnrealBuildTool
 				}
 
 				// Now that we have the dependencies object, optionally create the artifact cache
-				// For future work
 				IActionArtifactCache? actionArtifactCache = null;
+				if (!string.IsNullOrEmpty(BuildConfiguration.ArtifactDirectory) && (BuildConfiguration.bArtifactRead || BuildConfiguration.bArtifactWrites))
+				{
+					DirectoryReference artifactDirectory = new(BuildConfiguration.ArtifactDirectory);
+					DirectoryReference.CreateDirectory(artifactDirectory);
+
+					string mode = (BuildConfiguration.bArtifactRead, BuildConfiguration.bArtifactWrites) switch
+					{
+						(false, true) => "write only",
+						(true, false) => "read only",
+						_ => "read/write",
+					};
+
+					Logger.LogInformation("Experimental artifact system using '{directory}' in {mode} mode", artifactDirectory.FullName, mode);
+				
+					actionArtifactCache = ActionArtifactCache.CreateHordeFileCache(artifactDirectory, CppDependencies, Logger);
+					actionArtifactCache.EnableReads = BuildConfiguration.bArtifactRead;
+					actionArtifactCache.EnableWrites = BuildConfiguration.bArtifactWrites;
+				}
 
 				// Pre-process module interfaces to generate dependency files
 				List<LinkedAction> ModuleDependencyActions = PrerequisiteActions.Where(x => x.ActionType == ActionType.GatherModuleDependencies).ToList();
diff --git a/Engine/Source/Programs/UnrealBuildTool/UnrealBuildTool.csproj b/Engine/Source/Programs/UnrealBuildTool/UnrealBuildTool.csproj
--- a/Engine/Source/Programs/UnrealBuildTool/UnrealBuildTool.csproj
+++ b/Engine/Source/Programs/UnrealBuildTool/UnrealBuildTool.csproj
@@ -1,4 +1,4 @@
-﻿<Project Sdk="Microsoft.NET.Sdk">
+<Project Sdk="Microsoft.NET.Sdk">
   <Import Project="..\Shared\UnrealEngine.csproj.props" />
   
   <PropertyGroup>
@@ -126,6 +126,7 @@
     <!-- If you add a project reference, BuildUAT.bat/.sh and BuildUBT.bat/.sh need to be updated to also test for those projects.-->
     <ProjectReference Include="..\Shared\EpicGames.Build\EpicGames.Build.csproj" />
     <ProjectReference Include="..\Shared\EpicGames.Core\EpicGames.Core.csproj" />
+    <ProjectReference Include="..\Shared\EpicGames.Horde\EpicGames.Horde.csproj" />
     <ProjectReference Include="..\Shared\EpicGames.IoHash\EpicGames.IoHash.csproj" />
     <ProjectReference Include="..\Shared\EpicGames.Serialization\EpicGames.Serialization.csproj" />
     <ProjectReference Include="..\Shared\EpicGames.UHT\EpicGames.UHT.csproj" />
