Restore functionality of the "-ini:" command line argument in UBT.  Also restrict the scope of command line arguments parsed in ConfigCache::ReadHierarchy to prevent reading arguments from other "Target" command lines.

diff --git a/Engine/Source/Programs/UnrealBuildTool/Configuration/TargetRules.cs b/Engine/Source/Programs/UnrealBuildTool/Configuration/TargetRules.cs
--- a/Engine/Source/Programs/UnrealBuildTool/Configuration/TargetRules.cs
+++ b/Engine/Source/Programs/UnrealBuildTool/Configuration/TargetRules.cs
@@ -2426,7 +2426,7 @@ namespace UnrealBuildTool
 			Dictionary<ConfigDependencyKey, IReadOnlyList<string>?> ConfigValues = new Dictionary<ConfigDependencyKey, IReadOnlyList<string>?>();
 			foreach (object ConfigurableObject in GetConfigurableObjects())
 			{
-				ConfigCache.ReadSettings(DirectoryReference.FromFile(ProjectFile), Platform, ConfigurableObject, ConfigValues);
+				ConfigCache.ReadSettings(DirectoryReference.FromFile(ProjectFile), Platform, ConfigurableObject, ConfigValues, Target.Arguments);
 				XmlConfig.ApplyTo(ConfigurableObject);
 				if (Target.Arguments != null)
 				{
diff --git a/Engine/Source/Programs/UnrealBuildTool/System/ConfigCache.cs b/Engine/Source/Programs/UnrealBuildTool/System/ConfigCache.cs
--- a/Engine/Source/Programs/UnrealBuildTool/System/ConfigCache.cs
+++ b/Engine/Source/Programs/UnrealBuildTool/System/ConfigCache.cs
@@ -271,30 +271,48 @@ namespace UnrealBuildTool
 		/// <param name="ProjectDir">The project directory to read the hierarchy for</param>
 		/// <param name="Platform">Which platform to read platform-specific config files for</param>
 		/// <param name="CustomConfig">Optional override config directory to search, for support of multiple target types</param>
-		/// <param name="CustomArgs">Optional list of command line arguments</param>
+		/// <param name="CustomArgs">Optional list of command line arguments added to the existing command line arguments</param>
 		/// <returns>The requested config hierarchy</returns>
 		public static ConfigHierarchy ReadHierarchy(ConfigHierarchyType Type, DirectoryReference? ProjectDir, UnrealTargetPlatform Platform, string CustomConfig = "", string[]? CustomArgs = null)
 		{
-			// Handle command line overrides
-			List<String> OverrideStrings = new List<String>();
-			List<String> CmdLine = Environment.GetCommandLineArgs().ToList();
+			CommandLineArguments CombinedArgs = new CommandLineArguments(Environment.GetCommandLineArgs());
+
 			if (CustomArgs != null)
 			{
 				foreach (string CustomArg in CustomArgs)
 				{
 					if (CustomArg.StartsWith("-", StringComparison.InvariantCultureIgnoreCase))
 					{
-						CmdLine.Add(CustomArg);
+						CombinedArgs.Append(CustomArg);
 					}
 					else
 					{
-						CmdLine.Add("-" + CustomArg);
+						CombinedArgs.Append("-" + CustomArg);
 					}
 				}
 			}
+
+			return ReadHierarchy(Type, ProjectDir, Platform, CustomConfig, CombinedArgs);
+		}
+
+		/// <summary>
+		/// Reads a config hierarchy (or retrieve it from the cache)
+		/// </summary>
+		/// <param name="Type">The type of hierarchy to read</param>
+		/// <param name="ProjectDir">The project directory to read the hierarchy for</param>
+		/// <param name="Platform">Which platform to read platform-specific config files for</param>
+		/// <param name="CustomConfig">Optional override config directory to search, for support of multiple target types</param>
+		/// <param name="CmdLineArgs">The command line arguments to parse</param>
+		/// <returns>The requested config hierarchy</returns>
+		public static ConfigHierarchy ReadHierarchy(ConfigHierarchyType Type, DirectoryReference? ProjectDir, UnrealTargetPlatform Platform, 
+			string CustomConfig, CommandLineArguments CmdLineArgs)
+		{
+			// Handle command line overrides
+			List<String> OverrideStrings = new List<String>();
+
 			string IniConfigArgPrefix = "-ini:" + Enum.GetName(typeof(ConfigHierarchyType), Type) + ":";
 			string CustomConfigPrefix = "-CustomConfig=";
-			foreach (string CmdLineArg in CmdLine)
+			foreach (string CmdLineArg in CmdLineArgs)
 			{
 				if (CmdLineArg.StartsWith(IniConfigArgPrefix, StringComparison.InvariantCultureIgnoreCase))
 				{
@@ -399,7 +417,7 @@ namespace UnrealBuildTool
 		/// <param name="TargetObject">Object to receive the settings</param>
 		public static void ReadSettings(DirectoryReference? ProjectDir, UnrealTargetPlatform Platform, object TargetObject)
 		{
-			ReadSettings(ProjectDir, Platform, TargetObject, null);
+			ReadSettings(ProjectDir, Platform, TargetObject, null, null);
 		}
 
 		/// <summary>
@@ -409,13 +427,24 @@ namespace UnrealBuildTool
 		/// <param name="Platform">The platform being built</param>
 		/// <param name="TargetObject">Object to receive the settings</param>
 		/// <param name="ConfigValues">Will be populated with config values that were retrieved. May be null.</param>
-		internal static void ReadSettings(DirectoryReference? ProjectDir, UnrealTargetPlatform Platform, object TargetObject, Dictionary<ConfigDependencyKey, IReadOnlyList<string>?>? ConfigValues)
+		/// <param name="CmdLineArgs">Command line arguments, if null the application's command line arguments will be used</param>
+		internal static void ReadSettings(DirectoryReference? ProjectDir, UnrealTargetPlatform Platform, object TargetObject, Dictionary<ConfigDependencyKey, IReadOnlyList<string>?>? ConfigValues, CommandLineArguments? CmdLineArgs)
 		{
 			List<ConfigMember> Members = FindConfigMembersForType(TargetObject.GetType());
-			foreach(ConfigMember Member in Members)
+
+			foreach (ConfigMember Member in Members)
 			{
 				// Read the hierarchy listed
-				ConfigHierarchy Hierarchy = ReadHierarchy(Member.Attribute.ConfigType, ProjectDir, Platform);
+				ConfigHierarchy Hierarchy;
+
+				if (CmdLineArgs == null)
+				{
+					Hierarchy = ReadHierarchy(Member.Attribute.ConfigType, ProjectDir, Platform);
+				}
+				else
+				{
+					Hierarchy = ReadHierarchy(Member.Attribute.ConfigType, ProjectDir, Platform, "", CmdLineArgs);
+				}
 
 				// Get the key name
 				string KeyName = Member.Attribute.KeyName ?? Member.MemberInfo.Name;
