[Box]
* Added Horde support to BoxExecutor! Implementation made by [REDACTED_PERSON]
* New box binaries

diff --git a/Engine/Source/Programs/UnrealBuildTool/UnrealBuildTool.csproj b/Engine/Source/Programs/UnrealBuildTool/UnrealBuildTool.csproj
--- a/Engine/Source/Programs/UnrealBuildTool/UnrealBuildTool.csproj
+++ b/Engine/Source/Programs/UnrealBuildTool/UnrealBuildTool.csproj
@@ -128,6 +128,7 @@
     <ProjectReference Include="..\Shared\EpicGames.Core\EpicGames.Core.csproj" />
     <ProjectReference Include="..\Shared\EpicGames.Horde\EpicGames.Horde.csproj" />
     <ProjectReference Include="..\Shared\EpicGames.IoHash\EpicGames.IoHash.csproj" />
+    <ProjectReference Include="..\Shared\EpicGames.OIDC\EpicGames.OIDC.csproj" />
     <ProjectReference Include="..\Shared\EpicGames.Serialization\EpicGames.Serialization.csproj" />
     <ProjectReference Include="..\Shared\EpicGames.UHT\EpicGames.UHT.csproj" />
   </ItemGroup>
