Add new references to batch file UAT UBT dependencies

#rnx

diff --git a/Engine/Build/BatchFiles/BuildUAT.bat b/Engine/Build/BatchFiles/BuildUAT.bat
--- a/Engine/Build/BatchFiles/BuildUAT.bat
+++ b/Engine/Build/BatchFiles/BuildUAT.bat
@@ -25,17 +25,21 @@ dir /s ^
  Programs\Shared\EpicGames.Build\*.csproj ^
  Programs\Shared\EpicGames.Core\*.cs ^
  Programs\Shared\EpicGames.Core\*.csproj ^
+ Programs\Shared\EpicGames.Horde\*.cs ^
+ Programs\Shared\EpicGames.Horde\*.csproj ^
  Programs\Shared\EpicGames.IoHash\*.cs ^
  Programs\Shared\EpicGames.IoHash\*.csproj ^
  Programs\Shared\EpicGames.MsBuild\*.cs ^
  Programs\Shared\EpicGames.MsBuild\*.csproj ^
+ Programs\Shared\EpicGames.OIDC\*.cs ^
+ Programs\Shared\EpicGames.OIDC\*.csproj ^
  Programs\Shared\EpicGames.Serialization\*.cs ^
  Programs\Shared\EpicGames.Serialization\*.csproj ^
  Programs\Shared\EpicGames.UHT\*.cs ^
  Programs\Shared\EpicGames.UHT\*.csproj ^
  Programs\UnrealBuildTool\*.cs ^
  Programs\UnrealBuildTool\*.csproj ^
- | find ".cs" >..\Intermediate\Build\AutomationToolFiles.txt
+ | find ".cs" > ..\Intermediate\Build\AutomationToolFiles.txt
 
 if not exist ..\Platforms goto NoPlatforms
 for /d %%D in (..\Platforms\*) do (
diff --git a/Engine/Build/BatchFiles/BuildUAT.sh b/Engine/Build/BatchFiles/BuildUAT.sh
--- a/Engine/Build/BatchFiles/BuildUAT.sh
+++ b/Engine/Build/BatchFiles/BuildUAT.sh
@@ -39,8 +39,10 @@ elif [ -f ../Intermediate/Build/AutomationToolLastBuildTime ]; then
   UPDATED_DEP_FILES="$(find \
     Programs/Shared/EpicGames.Build \
     Programs/Shared/EpicGames.Core \
+    Programs/Shared/EpicGames.Horde \
     Programs/Shared/EpicGames.IoHash \
     Programs/Shared/EpicGames.MsBuild \
+    Programs/Shared/EpicGames.OIDC \
     Programs/Shared/EpicGames.Serialization \
     Programs/Shared/EpicGames.UHT \
     Programs/UnrealBuildTool \
diff --git a/Engine/Build/BatchFiles/BuildUBT.bat b/Engine/Build/BatchFiles/BuildUBT.bat
--- a/Engine/Build/BatchFiles/BuildUBT.bat
+++ b/Engine/Build/BatchFiles/BuildUBT.bat
@@ -25,10 +25,14 @@ dir /s ^
  Programs\Shared\EpicGames.Build\*.csproj ^
  Programs\Shared\EpicGames.Core\*.cs ^
  Programs\Shared\EpicGames.Core\*.csproj ^
+ Programs\Shared\EpicGames.Horde\*.cs ^
+ Programs\Shared\EpicGames.Horde\*.csproj ^
  Programs\Shared\EpicGames.IoHash\*.cs ^
  Programs\Shared\EpicGames.IoHash\*.csproj ^
  Programs\Shared\EpicGames.MsBuild\*.cs ^
  Programs\Shared\EpicGames.MsBuild\*.csproj ^
+ Programs\Shared\EpicGames.OIDC\*.cs ^
+ Programs\Shared\EpicGames.OIDC\*.csproj ^
  Programs\Shared\EpicGames.Serialization\*.cs ^
  Programs\Shared\EpicGames.Serialization\*.csproj ^
  Programs\Shared\EpicGames.UHT\*.cs ^
diff --git a/Engine/Build/BatchFiles/BuildUBT.sh b/Engine/Build/BatchFiles/BuildUBT.sh
--- a/Engine/Build/BatchFiles/BuildUBT.sh
+++ b/Engine/Build/BatchFiles/BuildUBT.sh
@@ -39,8 +39,10 @@ elif [ -f ../Intermediate/Build/UnrealBuildToolLastBuildTime ]; then
   UPDATED_DEP_FILES="$(find \
     Programs/Shared/EpicGames.Build \
     Programs/Shared/EpicGames.Core \
+    Programs/Shared/EpicGames.Horde \
     Programs/Shared/EpicGames.IoHash \
     Programs/Shared/EpicGames.MsBuild \
+    Programs/Shared/EpicGames.OIDC \
     Programs/Shared/EpicGames.Serialization \
     Programs/Shared/EpicGames.UHT \
     Programs/UnrealBuildTool \
diff --git a/Engine/Source/Programs/UnrealBuildTool/UnrealBuildTool.csproj b/Engine/Source/Programs/UnrealBuildTool/UnrealBuildTool.csproj
--- a/Engine/Source/Programs/UnrealBuildTool/UnrealBuildTool.csproj
+++ b/Engine/Source/Programs/UnrealBuildTool/UnrealBuildTool.csproj
@@ -158,11 +158,20 @@
 
   <!-- Various batch files that build or run UnrealBuildTool -->
   <ItemGroup>
+    <None Include="..\..\..\Build\BatchFiles\BuildUAT.bat"><Link>BatchFiles\BuildUAT.bat</Link></None>
+    <None Include="..\..\..\Build\BatchFiles\BuildUAT.sh"><Link>BatchFiles\BuildUAT.sh</Link></None>
     <None Include="..\..\..\Build\BatchFiles\BuildUBT.bat"><Link>BatchFiles\BuildUBT.bat</Link></None>
     <None Include="..\..\..\Build\BatchFiles\BuildUBT.sh"><Link>BatchFiles\BuildUBT.sh</Link></None>
     <None Include="..\..\..\Build\BatchFiles\GenerateProjectFiles.bat"><Link>BatchFiles\GenerateProjectFiles.bat</Link></None>
+    <None Include="..\..\..\Build\BatchFiles\RunUAT.bat"><Link>BatchFiles\RunUAT.bat</Link></None>
+    <None Include="..\..\..\Build\BatchFiles\RunUAT.sh"><Link>BatchFiles\RunUAT.sh</Link></None>
     <None Include="..\..\..\Build\BatchFiles\RunUBT.bat"><Link>BatchFiles\RunUBT.bat</Link></None>
     <None Include="..\..\..\Build\BatchFiles\RunUBT.sh"><Link>BatchFiles\RunUBT.sh</Link></None>
+
+    <!-- Visual Studio batch files -->
+    <None Include="..\..\..\Build\BatchFiles\Build.bat"><Link>BatchFiles\Build.bat</Link></None>
+    <None Include="..\..\..\Build\BatchFiles\Clean.bat"><Link>BatchFiles\Clean.bat</Link></None>
+    <None Include="..\..\..\Build\BatchFiles\Rebuild.bat"><Link>BatchFiles\Rebuild.bat</Link></None>
   </ItemGroup>
 
   <!-- Various BuildConfiguration files -->
