1) Added missing exception handling in the tasks
2) Reworked how "stop after error" works
3) Fixed  issue where error propagation might result in a failure to exit.

diff --git a/Engine/Source/Programs/UnrealBuildTool/Executors/ImmediateActionQueue.cs b/Engine/Source/Programs/UnrealBuildTool/Executors/ImmediateActionQueue.cs
--- a/Engine/Source/Programs/UnrealBuildTool/Executors/ImmediateActionQueue.cs
+++ b/Engine/Source/Programs/UnrealBuildTool/Executors/ImmediateActionQueue.cs
@@ -2,8 +2,10 @@
 
 using EpicGames.Core;
 using Microsoft.CodeAnalysis;
+using Microsoft.CodeAnalysis.CSharp.Syntax;
 using Microsoft.Extensions.Logging;
 using System;
+using System.Collections;
 using System.Collections.Generic;
 using System.Linq;
 using System.Threading;
@@ -56,17 +58,16 @@ namespace UnrealBuildTool
 		Manual,
 	}
 
-
 	/// <summary>
 	/// Actions are assigned a runner when they need executing
 	/// </summary>
 	/// <param name="Type">Type of the runner</param>
 	/// <param name="ActionPhase">The action phase that this step supports</param>
-	/// <param name="RunAction">Action to be run queue running a task</param>
+	/// <param name="RunAction">Function to be run queue running a task</param>
 	/// <param name="UseActionWeights">If true, use the action weight as a secondary limit</param>
 	/// <param name="MaxActions">Maximum number of action actions</param>
 	/// <param name="MaxActionWeight">Maximum weight of actions</param>
-	record ImmediateActionQueueRunner(ImmediateActionQueueRunnerType Type, ActionPhase ActionPhase, Func<LinkedAction, System.Action?> RunAction, bool UseActionWeights = false, int MaxActions = int.MaxValue, double MaxActionWeight = int.MaxValue)
+	record ImmediateActionQueueRunner(ImmediateActionQueueRunnerType Type, ActionPhase ActionPhase, Func<LinkedAction, Func<Task>?> RunAction, bool UseActionWeights = false, int MaxActions = int.MaxValue, double MaxActionWeight = int.MaxValue)
 	{
 		/// <summary>
 		/// Current number of active actions
@@ -327,7 +328,7 @@ namespace UnrealBuildTool
 			{
 				var runAction = (LinkedAction action) =>
 				{
-					return new System.Action(async () =>
+					return new Func<Task>(async () =>
 					{
 						bool success = await _actionArtifactCache!.CompleteActionFromCacheAsync(action, CancellationToken);
 						if (success)
@@ -353,7 +354,7 @@ namespace UnrealBuildTool
 		/// <param name="maxActions">Maximum number of action actions</param>
 		/// <param name="maxActionWeight">Maximum weight of actions</param>
 		/// <returns>Created runner</returns>
-		public ImmediateActionQueueRunner CreateAutomaticRunner(Func<LinkedAction, System.Action?> runAction, bool useActionWeights, int maxActions, double maxActionWeight)
+		public ImmediateActionQueueRunner CreateAutomaticRunner(Func<LinkedAction, Func<Task>?> runAction, bool useActionWeights, int maxActions, double maxActionWeight)
 		{
 			ImmediateActionQueueRunner runner = new(ImmediateActionQueueRunnerType.Automatic, ActionPhase.Compile, runAction, useActionWeights, maxActions, maxActionWeight);
 			_runners.Add(runner);
@@ -368,7 +369,7 @@ namespace UnrealBuildTool
 		/// <param name="maxActions">Maximum number of action actions</param>
 		/// <param name="maxActionWeight">Maximum weight of actions</param>
 		/// <returns>Created runner</returns>
-		public ImmediateActionQueueRunner CreateManualRunner(Func<LinkedAction, System.Action?> runAction, bool useActionWeights = false, int maxActions = int.MaxValue, double maxActionWeight = double.MaxValue)
+		public ImmediateActionQueueRunner CreateManualRunner(Func<LinkedAction, Func<Task>?> runAction, bool useActionWeights = false, int maxActions = int.MaxValue, double maxActionWeight = double.MaxValue)
 		{
 			ImmediateActionQueueRunner runner = new(ImmediateActionQueueRunnerType.Manual, ActionPhase.Compile, runAction, useActionWeights, maxActions, maxActionWeight);
 			_runners.Add(runner);
@@ -433,33 +434,103 @@ namespace UnrealBuildTool
 		/// <returns>True if an action was run, false if not</returns>
 		public bool TryStartOneAction(ImmediateActionQueueRunner? runner = null)
 		{
-			(System.Action? action, int completedActions) = TryStartOneActionInternal(runner);
 
-			if (action == null)
+			// Loop until we have an action or no completed actions (i.e. no propagated errors)
+			for(; ; )
 			{
-				// We found nothing, check to see if we have no active running tasks and no manual runners.
-				// If we don't, then it is assumed we can't queue any more work.
-				bool prematureDone = true;
-				foreach (ImmediateActionQueueRunner tryRunner in _runners)
+				bool hasCanceled = false;
+
+				// Try to get the next action
+				Func<Task>? runAction = null;
+				LinkedAction? action = null;
+				int completedActions = 0;
+				lock (Actions)
 				{
-					if (tryRunner.Type == ImmediateActionQueueRunnerType.Manual || tryRunner.ActiveActions != 0)
+					hasCanceled = CancellationTokenSource.IsCancellationRequested;
+					if (!hasCanceled)
 					{
-						prematureDone = false;
-						break;
+						(runAction, action, completedActions) = TryStartOneActionInternal(runner);
+					}
+				}
+
+				// If we have an action, run it and account for any completed actions
+				if (runAction != null && action != null)
+				{
+					Task.Factory.StartNew(() =>
+					{
+						try
+						{
+							runAction().Wait();
+						}
+						catch (Exception ex)
+						{
+							HandleException(action, ex);
+						}
+					}, CancellationToken, TaskCreationOptions.LongRunning | TaskCreationOptions.PreferFairness, TaskScheduler.Default);
+					AddCompletedActions(completedActions);
+					return true;
+				}
+
+				// If we have no completed actions, then we need to check for a stall
+				if (completedActions == 0)
+				{
+					// We found nothing, check to see if we have no active running tasks and no manual runners.
+					// If we don't, then it is assumed we can't queue any more work.
+					bool prematureDone = true;
+					foreach (ImmediateActionQueueRunner tryRunner in _runners)
+					{
+						if ((tryRunner.Type == ImmediateActionQueueRunnerType.Manual && !hasCanceled) || tryRunner.ActiveActions != 0)
+						{
+							prematureDone = false;
+							break;
+						}
 					}
+					if (prematureDone)
+					{
+						AddCompletedActions(int.MaxValue);
+					}
+					return false;
 				}
-				if (prematureDone)
+
+				// Otherwise add the completed actions and test again for the possibility that errors still need propagating
+				AddCompletedActions(completedActions);
+			}
+		}
+
+		/// <summary>
+		/// Handle an exception running a task
+		/// </summary>
+		/// <param name="action">Action that has been run</param>
+		/// <param name="ex">Thrown exception</param>
+		void HandleException(LinkedAction action, Exception ex)
+		{
+			if (ex is AggregateException aggregateEx)
+			{
+				if (aggregateEx.InnerException != null)
 				{
-					completedActions = int.MaxValue;
+					HandleException(action, aggregateEx.InnerException);
+				}
+				else
+				{
+					ExecuteResults results = new(new List<string>(), int.MaxValue, TimeSpan.Zero, TimeSpan.Zero);
+					OnActionCompleted(action, false, results);
 				}
 			}
+			else if (ex is OperationCanceledException)
+			{
+				ExecuteResults results = new(new List<string>(), int.MaxValue, TimeSpan.Zero, TimeSpan.Zero);
+				OnActionCompleted(action, false, results);
+			}
 			else
 			{
-				Task.Factory.StartNew(() => action(), CancellationToken, TaskCreationOptions.LongRunning | TaskCreationOptions.PreferFairness, TaskScheduler.Default);
+				List<string> text = new()
+				{
+					ExceptionUtils.FormatException(ex),
+					ExceptionUtils.FormatExceptionDetails(ex),
+				};
+				ExecuteResults results = new(text, int.MaxValue, TimeSpan.Zero, TimeSpan.Zero);
+				OnActionCompleted(action, false, results);
 			}
-
-			AddCompletedActions(completedActions);
-			return action != null;
 		}
 
 		/// <summary>
@@ -467,89 +538,87 @@ namespace UnrealBuildTool
 		/// </summary>
 		/// <param name="runner">If specified, tasks will only be queued to this runner.  Otherwise all manual runners will be used.</param>
 		/// <returns>Action to be run if something was queued and the number of completed actions</returns>
-		public (System.Action?, int) TryStartOneActionInternal(ImmediateActionQueueRunner? runner = null)
+		public (Func<Task>?, LinkedAction?, int) TryStartOneActionInternal(ImmediateActionQueueRunner? runner = null)
 		{
 			int completedActions = 0;
-			lock (Actions)
+
+			// If we are starting deeper in the action collection, then never set the first pending action location
+			bool foundFirstPending = false;
+
+			// Loop through the items
+			for (int actionIndex = _firstPendingAction; actionIndex != Actions.Length; ++actionIndex)
 			{
-				// If we are starting deeper in the action collection, then never set the first pending action location
-				bool foundFirstPending = false;
 
-				// Loop through the items
-				for (int actionIndex = _firstPendingAction; actionIndex != Actions.Length; ++actionIndex)
+				// If the action has already reached the compilation state, then just ignore
+				if (Actions[actionIndex].Status != ActionStatus.Queued)
 				{
+					continue;
+				}
 
-					// If the action has already reached the compilation state, then just ignore
-					if (Actions[actionIndex].Status != ActionStatus.Queued)
-					{
-						continue;
-					}
-
-					// If needed, update the first valid slot for searching for actions to run
-					if (!foundFirstPending)
-					{
-						_firstPendingAction = actionIndex;
-						foundFirstPending = true;
-					}
+				// If needed, update the first valid slot for searching for actions to run
+				if (!foundFirstPending)
+				{
+					_firstPendingAction = actionIndex;
+					foundFirstPending = true;
+				}
 
-					// Based on the ready state, use this action or mark as an error
-					switch (GetActionReadyState(Actions[actionIndex].Action))
-					{
-						case ActionReadyState.NotReady:
-							break;
+				// Based on the ready state, use this action or mark as an error
+				switch (GetActionReadyState(Actions[actionIndex].Action))
+				{
+					case ActionReadyState.NotReady:
+						break;
 
-						case ActionReadyState.Error:
-							Actions[actionIndex].Status = ActionStatus.Error;
-							completedActions++;
-							break;
+					case ActionReadyState.Error:
+						Actions[actionIndex].Status = ActionStatus.Error;
+						completedActions++;
+						break;
 
-						case ActionReadyState.Ready:
-							ImmediateActionQueueRunner? selectedRunner = null;
-							System.Action? action = null;
-							if (runner != null)
+					case ActionReadyState.Ready:
+						ImmediateActionQueueRunner? selectedRunner = null;
+						Func<Task>? runAction = null;
+						if (runner != null)
+						{
+							if (runner.IsUnderLimits && runner.ActionPhase == Actions[actionIndex].Phase)
 							{
-								if (runner.IsUnderLimits && runner.ActionPhase == Actions[actionIndex].Phase)
+								runAction = runner.RunAction(Actions[actionIndex].Action);
+								if (runAction != null)
 								{
-									action = runner.RunAction(Actions[actionIndex].Action);
-									if (action != null)
-									{
-										selectedRunner = runner;
-									}
+									selectedRunner = runner;
 								}
 							}
-							else
+						}
+						else
+						{
+							foreach (ImmediateActionQueueRunner tryRunner in _runners)
 							{
-								foreach (ImmediateActionQueueRunner tryRunner in _runners)
+								if (tryRunner.Type == ImmediateActionQueueRunnerType.Automatic &&
+									tryRunner.IsUnderLimits && tryRunner.ActionPhase == Actions[actionIndex].Phase)
 								{
-									if (tryRunner.Type == ImmediateActionQueueRunnerType.Automatic &&
-										tryRunner.IsUnderLimits && tryRunner.ActionPhase == Actions[actionIndex].Phase)
+									runAction = tryRunner.RunAction(Actions[actionIndex].Action);
+									if (runAction != null)
 									{
-										action = tryRunner.RunAction(Actions[actionIndex].Action);
-										if (action != null)
-										{
-											selectedRunner = tryRunner;
-											break;
-										}
+										selectedRunner = tryRunner;
+										break;
 									}
 								}
 							}
+						}
 
-							if (action != null && selectedRunner != null)
-							{
-								Actions[actionIndex].Status = ActionStatus.Running;
-								Actions[actionIndex].Runner = selectedRunner;
-								selectedRunner.ActiveActions++;
-								selectedRunner.ActiveActionWeight += Actions[actionIndex].Action.Weight;
-								return (action, completedActions);
-							}
-							break;
+						if (runAction != null && selectedRunner != null)
+						{
+							Actions[actionIndex].Status = ActionStatus.Running;
+							Actions[actionIndex].Runner = selectedRunner;
+							selectedRunner.ActiveActions++;
+							selectedRunner.ActiveActionWeight += Actions[actionIndex].Action.Weight;
+							return (runAction, Actions[actionIndex].Action, completedActions);
+						}
+						break;
 
-						default:
-							throw new ApplicationException("Unexpected action ready state");
-					}
+					default:
+						throw new ApplicationException("Unexpected action ready state");
 				}
 			}
-			return (null, completedActions);
+			return (null, null, completedActions);
 		}
 
 		/// <summary>
@@ -861,16 +930,12 @@ namespace UnrealBuildTool
 
 					// prevent overwriting of error text
 					s_previousLineLength = -1;
-				}
-			}
 
-			if (exitCode != 0 && exitCode != int.MaxValue)
-			{
-				// Cancel all other pending tasks
-				if (StopCompilationAfterErrors)
-				{
-					CancellationTokenSource.Cancel();
-					AddCompletedActions(int.MaxValue);
+					// Cancel all other pending tasks
+					if (StopCompilationAfterErrors)
+					{
+						CancellationTokenSource.Cancel();
+					}
 				}
 			}
 		}
diff --git a/Engine/Source/Programs/UnrealBuildTool/Executors/ParallelExecutor.cs b/Engine/Source/Programs/UnrealBuildTool/Executors/ParallelExecutor.cs
--- a/Engine/Source/Programs/UnrealBuildTool/Executors/ParallelExecutor.cs
+++ b/Engine/Source/Programs/UnrealBuildTool/Executors/ParallelExecutor.cs
@@ -204,16 +204,12 @@ namespace UnrealBuildTool
 			}
 		}
 
-		private static System.Action? RunAction(ImmediateActionQueue queue, LinkedAction action)
+		private static Func<Task>? RunAction(ImmediateActionQueue queue, LinkedAction action)
 		{
 			return async () =>
 			{
-				Task<ExecuteResults> task = RunAction(action, queue.ProcessGroup, queue.CancellationToken);
-				await task;
-
-				ExecuteResults? results = task.Status == TaskStatus.RanToCompletion ? task.Result : null;
-				bool success = results?.ExitCode == 0;
-				queue.OnActionCompleted(action, success, results);
+				ExecuteResults results = await RunAction(action, queue.ProcessGroup, queue.CancellationToken);
+				queue.OnActionCompleted(action, results.ExitCode == 0, results);
 			};
 		}
 
