Modern Xcode: using copies of info.plist instead of using the original

diff --git a/Engine/Config/BaseEngine.ini b/Engine/Config/BaseEngine.ini
--- a/Engine/Config/BaseEngine.ini
+++ b/Engine/Config/BaseEngine.ini
@@ -3032,8 +3032,7 @@ ModernSigningTeam=
 ModernSigningPrefix=com.YourCompany
 ModernBundleIdentifier=$(UE_SIGNING_PREFIX).$(UE_PRODUCT_NAME_STRIPPED)
 AppCategory=public.app-category.games
-TemplateMacPlist=(FilePath="/Engine/Build/Mac/Resources/Info.Template.plist")
-TemplateMacEditorPlist=(FilePath="/Engine/Build/Mac/Resources/Info-Editor.Template.plist")
+TemplateMacPlist=(FilePath="/Game/Build/Mac/Resources/Info.Template.plist")
 TemplateIOSPlist=(FilePath="/Game/Build/IOS/UBTGenerated/Info.Template.plist")
 PremadeMacEntitlements=(FilePath="/Engine/Build/Mac/Resources/Sandbox.Server.entitlements")
 ShippingSpecificMacEntitlements=(FilePath="/Engine/Build/Mac/Resources/Sandbox.NoNet.entitlements")
diff --git a/Engine/Source/Programs/UnrealBuildTool/ProjectFiles/Xcode/XcodeProject.cs b/Engine/Source/Programs/UnrealBuildTool/ProjectFiles/Xcode/XcodeProject.cs
--- a/Engine/Source/Programs/UnrealBuildTool/ProjectFiles/Xcode/XcodeProject.cs
+++ b/Engine/Source/Programs/UnrealBuildTool/ProjectFiles/Xcode/XcodeProject.cs
@@ -143,13 +143,14 @@
 
 	enum MetadataMode
 	{
+		Unset = -1,
 		UsePremade,
 		UpdateTemplate,
 	}
 
 	class MetadataItem
 	{
-		public MetadataMode Mode;
+		public MetadataMode Mode = MetadataMode.Unset;
 		public FileReference? File = null;
 		public string? XcodeProjectRelative = null;
 
@@ -160,74 +161,84 @@
 		//	XcodeProjectRelative = File.MakeRelativeTo(XcodeProjectFile.ParentDirectory!);
 		//}
 
-		public MetadataItem(DirectoryReference ProductDirectory, DirectoryReference XcodeProject, ConfigHierarchy Ini, params string[] Locations)
+		
+		// Location: Can be key of setting entry in .ini, or the full file path
+		// CopyFromFolderIfNotFound: If the file at "Location" does not exist, try to copy the same named file from this folder
+		public MetadataItem(DirectoryReference ProductDirectory, DirectoryReference XcodeProject, ConfigHierarchy Ini, string Location, MetadataMode InMode, DirectoryReference CopyFromFolderIfNotFound)
 		{
-			foreach (string Loc in Locations)
+			// no extension means it's a .ini entry
+			if (Path.GetExtension(Location) == "")
 			{
-				string[] Tokens = Loc.Split(':');
-				Mode = (Tokens[0] == "premade") ? MetadataMode.UsePremade : MetadataMode.UpdateTemplate;
+				string? FileLocation;
+				if (Ini.TryGetValue("/Script/MacTargetPlatform.XcodeProjectSettings", Location, out FileLocation) && FileLocation != "")
+				{
+					if (FileLocation.StartsWith("(FilePath=", StringComparison.OrdinalIgnoreCase))
+					{
+						FileLocation = ConfigHierarchy.GetStructEntry(FileLocation, "FilePath", false)!;
+						if (FileLocation.StartsWith("/Engine/"))
+						{
+							File = FileReference.Combine(Unreal.EngineDirectory, FileLocation.Substring(8));
+						}
+						else if (FileLocation.StartsWith("/Game/"))
+						{
+							File = FileReference.Combine(ProductDirectory, FileLocation.Substring(6));
+						}
+						else
+						{
+							File = new FileReference(FileLocation);
+						}
+					}
+					else
+					{
+						File = new FileReference(FileLocation);
+					}
+				}
+			}
+			else
+			{
+				File = new FileReference(Location);
+			}
 
-				FileReference? TestFile = null;
-
-				// no extension means it's a .ini entry
-				if (Path.GetExtension(Tokens[1]) == "")
+			if (File != null)
+			{
+				// Copy from source location if no such file exist
+				// Except UBTGenerated, they will be generated during UBT runs
+				if (!FileReference.Exists(File) && !File.ContainsName("UBTGenerated", 0))
 				{
-					string? PremadeLocation;
-					if (Ini.TryGetValue("/Script/MacTargetPlatform.XcodeProjectSettings", Tokens[1], out PremadeLocation) && PremadeLocation != "")
+					FileReference SourceFile = FileReference.Combine(CopyFromFolderIfNotFound, File.GetFileName());
+					if (FileReference.Exists(SourceFile))
 					{
-						if (PremadeLocation.StartsWith("(FilePath=", StringComparison.OrdinalIgnoreCase))
+						DirectoryReference.CreateDirectory(File.Directory);
+						try
 						{
-							PremadeLocation = ConfigHierarchy.GetStructEntry(PremadeLocation, "FilePath", false)!;
-							if (PremadeLocation.StartsWith("/Engine/"))
+							FileReference.Copy(SourceFile, File, false);
+						}
+						catch (System.IO.IOException ex)
+						{
+							if (ex.Message.Contains("already exists"))
 							{
-								TestFile = FileReference.Combine(Unreal.EngineDirectory, PremadeLocation.Substring(8));
+								// Some other thread is probably copying the same file, ignore
 							}
-							else if (PremadeLocation.StartsWith("/Game/"))
-							{
-								TestFile = FileReference.Combine(ProductDirectory, PremadeLocation.Substring(6));
-							}
 							else
 							{
-								TestFile = new FileReference(PremadeLocation);
+								throw ex;
 							}
-						}
-						else
-						{
-							TestFile = new FileReference(PremadeLocation);
 						}
 					}
 				}
-				else
-				{
-					TestFile = new FileReference(Tokens[1]);
-				}
-
-				// make sure the file (if it's a generated file, then it may not exist yet, so allow it through)
-				if (TestFile != null)
-				{
-					if (TestFile.ContainsName("UBTGenerated", 0) || FileReference.Exists(TestFile))
-					{
-						File = TestFile;
-						break;
-					}
-
-					throw new BuildException($"Metadata file {Tokens[1]} (resolved to {TestFile}) was specified in project settings, but does not exist. Ignoring...");
-				}
 			}
 
-			if (File == null)
+			if (File != null && (File.ContainsName("UBTGenerated", 0) || FileReference.Exists(File)))
 			{
-				// if null, we assume UpdateTemplate, since xcode can update "nothing" to something useful (for plist), or we ignore 
-				// UpdateTemplate for entitlements
-				Mode = MetadataMode.UpdateTemplate;
-			}
-			else
-			{
-				// make it relative if it exists
+				// We found a valid metadata file
+				Mode = InMode;
 				XcodeProjectRelative = File.MakeRelativeTo(XcodeProject.ParentDirectory!);
 			}
-			
-			// Log.TraceInformation($"   Metadata {Locations[0]} found {Mode} File {File}");
+			else
+			{
+				// Either key is missing, or file is missing
+				Mode = MetadataMode.Unset;
+			}
 		}
 	}
 
@@ -241,20 +252,84 @@
 		public Metadata(DirectoryReference ProductDirectory, DirectoryReference XcodeProject, ConfigHierarchy Ini, bool bSupportsMac, bool bSupportsIOSOrTVOS, ILogger Logger)
 		{
 			Logger.LogInformation("making metadata for {ProductDirectory} / {XcodeProject}", ProductDirectory, XcodeProject);
+			DirectoryReference ResourceFolder = DirectoryReference.Combine(Unreal.EngineDirectory, "Build/Mac/Resources/");
 			if (bSupportsMac)
 			{
-				PlistFiles[MetadataPlatform.MacEditor] = new MetadataItem(ProductDirectory, XcodeProject, Ini, "premade:PremadeMacEditorPlist", "template:TemplateMacEditorPlist");
-				PlistFiles[MetadataPlatform.Mac] = new MetadataItem(ProductDirectory, XcodeProject, Ini, "premade:PremadeMacPlist", "template:TemplateMacPlist");
-				EntitlementsFiles[MetadataPlatform.MacEditor] = new MetadataItem(ProductDirectory, XcodeProject, Ini, "premade:PremadeMacEditorEntitlements");
-				EntitlementsFiles[MetadataPlatform.Mac] = new MetadataItem(ProductDirectory, XcodeProject, Ini, "premade:PremadeMacEntitlements");
-				ShippingEntitlementsFiles[MetadataPlatform.MacEditor] = new MetadataItem(ProductDirectory, XcodeProject, Ini, "premade:ShippingSpecificMacEditorEntitlements");
-				ShippingEntitlementsFiles[MetadataPlatform.Mac] = new MetadataItem(ProductDirectory, XcodeProject, Ini, "premade:ShippingSpecificMacEntitlements");
+				// All editor use template plist which is just a copy of default info.plist, user should not need to modify or change this
+				// TEMP: currently cooked editor (e.g. QAGCookedEditor) uses this hardcoded file too, we should find a way to expose this in Game/Config/DefaultEngine.ini
+				PlistFiles[MetadataPlatform.MacEditor] = new MetadataItem(ProductDirectory, XcodeProject, Ini,
+					Unreal.EngineDirectory.FullName + "/Intermediate/Build/Mac/Resources/Info-Editor.Template.plist",
+					MetadataMode.UpdateTemplate,
+					ResourceFolder);
+
+				if (ProductDirectory == Unreal.EngineDirectory)
+				{
+					// Engine dir is the same as Product dir, meaning no uproject (UnrealGame.app)
+					// Don't use ini value in this case, since ini value points at /Game/, use this hardcoded location instead
+					PlistFiles[MetadataPlatform.Mac] = new MetadataItem(ProductDirectory, XcodeProject, Ini,
+						Unreal.EngineDirectory.FullName + "/Intermediate/Build/Mac/Resources/Info.Template.plist",
+						MetadataMode.UpdateTemplate,
+						ResourceFolder);
+				}
+				else
+				{
+					// E.g. QAGame.app
+					// Try Premade first
+					PlistFiles[MetadataPlatform.Mac] = new MetadataItem(ProductDirectory, XcodeProject, Ini,
+						"PremadeMacPlist",
+						MetadataMode.UsePremade,
+						ResourceFolder);
+					if (PlistFiles[MetadataPlatform.Mac].Mode == MetadataMode.Unset)
+					{
+						// Premade not found, try Template
+						PlistFiles[MetadataPlatform.Mac] = new MetadataItem(ProductDirectory, XcodeProject, Ini,
+						"TemplateMacPlist",
+						MetadataMode.UpdateTemplate,
+						ResourceFolder);
+					}
+				}
+
+				EntitlementsFiles[MetadataPlatform.MacEditor] = new MetadataItem(ProductDirectory, XcodeProject, Ini,
+					"PremadeMacEditorEntitlements",
+					MetadataMode.UsePremade,
+					ResourceFolder);
+				EntitlementsFiles[MetadataPlatform.Mac] = new MetadataItem(ProductDirectory, XcodeProject, Ini,
+					"PremadeMacEntitlements",
+					MetadataMode.UsePremade,
+					ResourceFolder);
+				ShippingEntitlementsFiles[MetadataPlatform.MacEditor] = new MetadataItem(ProductDirectory, XcodeProject, Ini,
+					"ShippingSpecificMacEditorEntitlements",
+					MetadataMode.UsePremade,
+					ResourceFolder);
+				ShippingEntitlementsFiles[MetadataPlatform.Mac] = new MetadataItem(ProductDirectory, XcodeProject, Ini,
+					"ShippingSpecificMacEntitlements",
+					MetadataMode.UsePremade,
+					ResourceFolder);
 			}
 			if (bSupportsIOSOrTVOS)
 			{
-				PlistFiles[MetadataPlatform.IOS] = new MetadataItem(ProductDirectory, XcodeProject, Ini, "premade:PremadeIOSPlist", "template:TemplateIOSPlist");
-				EntitlementsFiles[MetadataPlatform.IOS] = new MetadataItem(ProductDirectory, XcodeProject, Ini, "premade:PremadeIOSEntitlements");
-				ShippingEntitlementsFiles[MetadataPlatform.IOS] = new MetadataItem(ProductDirectory, XcodeProject, Ini, "premade:ShippingSpecificIOSEntitlements");
+				// Try Premade first
+				PlistFiles[MetadataPlatform.IOS] = new MetadataItem(ProductDirectory, XcodeProject, Ini,
+					"PremadeIOSPlist",
+					MetadataMode.UsePremade,
+					ResourceFolder);
+				if (PlistFiles[MetadataPlatform.IOS].Mode == MetadataMode.Unset)
+				{
+					// Premade not found, try Template
+					PlistFiles[MetadataPlatform.IOS] = new MetadataItem(ProductDirectory, XcodeProject, Ini,
+						"TemplateIOSPlist",
+						MetadataMode.UpdateTemplate,
+						ResourceFolder);
+				}
+
+				EntitlementsFiles[MetadataPlatform.IOS] = new MetadataItem(ProductDirectory, XcodeProject, Ini,
+					"PremadeIOSEntitlements",
+					MetadataMode.UsePremade,
+					ResourceFolder);
+				ShippingEntitlementsFiles[MetadataPlatform.IOS] = new MetadataItem(ProductDirectory, XcodeProject, Ini,
+					"ShippingSpecificIOSEntitlements",
+					MetadataMode.UsePremade,
+					ResourceFolder);
 			}
 		}
 	}
@@ -1253,7 +1328,7 @@
 
 			CopyScript.AddRange(new string[]
 			{
-				"# We run this often during ApplePostBuildSync, which looks at the PkgInfo file in the .app to determine if this is up to date,",
+				"# We run this often during ApplePostBuildSync, which looks at the PkgInfo file in the .app to determine if this is up to date,",
 				"# so delete it in case we have an error, we will re-run PostBuildSync next build",
 				"echo rm ${CONFIGURATION_BUILD_DIR}/${CONTENTS_FOLDER_PATH}/PkgInfo",
 				"rm ${CONFIGURATION_BUILD_DIR}/${CONTENTS_FOLDER_PATH}/PkgInfo",
@@ -1294,7 +1369,7 @@
 					$"STAGED_DIR={DefaultStageDir}",
 					"if [[ -z ${STAGED_DIR} ]]; then exit 0; fi",
 					"# Make sure the staged directory exists and has files in it",
-					"if [[ ! -e ${STAGED_DIR} || `ls ${STAGED_DIR} | wc -l` -eq 0 ]]; then ",
+					"if [[ ! -e ${STAGED_DIR} || `ls ${STAGED_DIR} | wc -l` -eq 0 ]]; then ",
 					"  echo =========================================================================================",
 					"  echo \\\"ERROR: You must have a valid staged build directory\\\"",
 					"  echo \\\"Use the editor's Platforms menu, or run:\\\"",
