Fixed an issue where the public dependencies weren't being added to the include path list.

diff --git a/Engine/Plugins/Experimental/Mutable/Source/CustomizableObjectPopulation/CustomizableObjectPopulation.Build.cs b/Engine/Plugins/Experimental/Mutable/Source/CustomizableObjectPopulation/CustomizableObjectPopulation.Build.cs
--- a/Engine/Plugins/Experimental/Mutable/Source/CustomizableObjectPopulation/CustomizableObjectPopulation.Build.cs
+++ b/Engine/Plugins/Experimental/Mutable/Source/CustomizableObjectPopulation/CustomizableObjectPopulation.Build.cs
@@ -13,10 +13,7 @@ public class CustomizableObjectPopulation : ModuleRules
         PrivateDependencyModuleNames.AddRange(new string[] {
             "InputCore",
             "SlateCore",
-            "Slate",
-            "Core",
             "CoreUObject",
-            "Engine",
             "RenderCore",
             "RHI",
             "AppFramework",
@@ -35,7 +32,6 @@ public class CustomizableObjectPopulation : ModuleRules
 
         PrivateIncludePathModuleNames.AddRange(
         new string[] {
-				"CustomizableObjectEditor",
 				"TargetPlatform"
 		});
 
@@ -45,6 +41,7 @@ public class CustomizableObjectPopulation : ModuleRules
             PublicDependencyModuleNames.Add("DerivedDataCache");
             PublicDependencyModuleNames.Add("EditorStyle");
             PublicDependencyModuleNames.Add("MessageLog");
-        }
+			PrivateIncludePathModuleNames.Add("CustomizableObjectEditor");
+		}
     }
 }
diff --git a/Engine/Source/Programs/UnrealBuildTool/Configuration/UEBuildModule.cs b/Engine/Source/Programs/UnrealBuildTool/Configuration/UEBuildModule.cs
--- a/Engine/Source/Programs/UnrealBuildTool/Configuration/UEBuildModule.cs
+++ b/Engine/Source/Programs/UnrealBuildTool/Configuration/UEBuildModule.cs
@@ -1114,15 +1114,18 @@ namespace UnrealBuildTool
 				{
 					Logger.LogDebug("Found circular reference to {ThisRefName}, but {GuiltyModule} declares a cycle on {VictimModule} which breaks the chain", ThisRefName, GuiltyModule, VictimModule);
 				}
-			}		
+			}
 
-			// Recursively create all the public include path modules. These modules may not be added to the target (and we don't process their referenced 
+			// Recursively create all the public include path and dependency modules. These modules may not be added to the target (and we don't process their referenced 
 			// dependencies), but they need to be created to set up their include paths.
-			RecursivelyCreateIncludePathModulesByName(Rules.PublicIncludePathModuleNames, ref PublicIncludePathModules, ref bDependsOnVerse, CreateModule, NextReferenceChain);
+			HashSet<string> PublicIncludePathModuleNames = new();
+			PublicIncludePathModuleNames.UnionWith(Rules.PublicIncludePathModuleNames!);
+			PublicIncludePathModuleNames.UnionWith(Rules.PublicDependencyModuleNames!);
+			RecursivelyCreateIncludePathModulesByName(PublicIncludePathModuleNames, ref PublicIncludePathModules, ref bDependsOnVerse, CreateModule, NextReferenceChain);
 
 			// Create all the referenced modules. This path can be recursive, so we check against PrivateIncludePathModules to ensure we don't recurse through the 
 			// same module twice (it produces better errors if something fails).
-			if(PrivateIncludePathModules == null)
+			if (PrivateIncludePathModules == null)
 			{
 				// Log dependencies if required
 				if (Logger.IsEnabled((LogLevel)LogEventType.VeryVerbose))
@@ -1192,7 +1195,7 @@ namespace UnrealBuildTool
 			}
 		}
 
-		private static void RecursivelyCreateIncludePathModulesByName(List<string> ModuleNames, ref List<UEBuildModule>? Modules, ref bool bDependsOnVerse, CreateModuleDelegate CreateModule, string ReferenceChain)
+		private static void RecursivelyCreateIncludePathModulesByName(IEnumerable<string> ModuleNames, ref List<UEBuildModule>? Modules, ref bool bDependsOnVerse, CreateModuleDelegate CreateModule, string ReferenceChain)
 		{
 			// Check whether the module list is already set. We set this immediately (via the ref) to avoid infinite recursion.
 			if (Modules == null)
@@ -1201,7 +1204,10 @@ namespace UnrealBuildTool
 				foreach (string ModuleName in ModuleNames)
 				{
 					UEBuildModule Module = CreateModule(ModuleName, ReferenceChain);
-					RecursivelyCreateIncludePathModulesByName(Module.Rules.PublicIncludePathModuleNames, ref Module.PublicIncludePathModules, ref Module.bDependsOnVerse, CreateModule, ReferenceChain);
+					HashSet<string> PublicIncludePathModuleNames = new();
+					PublicIncludePathModuleNames.UnionWith(Module.Rules.PublicIncludePathModuleNames);
+					PublicIncludePathModuleNames.UnionWith(Module.Rules.PublicDependencyModuleNames);
+					RecursivelyCreateIncludePathModulesByName(PublicIncludePathModuleNames, ref Module.PublicIncludePathModules, ref Module.bDependsOnVerse, CreateModule, ReferenceChain);
 					Modules.Add(Module);
 					bDependsOnVerse |= Module.bDependsOnVerse;
 				}
diff --git a/Engine/Source/Runtime/Engine/Engine.Build.cs b/Engine/Source/Runtime/Engine/Engine.Build.cs
--- a/Engine/Source/Runtime/Engine/Engine.Build.cs
+++ b/Engine/Source/Runtime/Engine/Engine.Build.cs
@@ -21,7 +21,6 @@ public class Engine : ModuleRules
 				"AudioMixer", 
 				"AudioMixerCore",
 				"MovieSceneCapture", 
-				"PacketHandler", 
 				"Renderer"
 			}
 		);
@@ -43,9 +42,7 @@ public class Engine : ModuleRules
 				"EyeTracker",
 				"MRMesh",
 				"Advertising",
-				"MovieSceneCapture",
 				"AutomationWorker",
-				"MovieSceneCapture",
 				"DesktopPlatform"
 			}
 		);
@@ -79,7 +76,6 @@ public class Engine : ModuleRules
 				"CoreUObject",
 				"FieldNotification",
 				"NetCore",
-				"ApplicationCore",
 				"ImageCore",
 				"Json",
 				"JsonUtilities",
@@ -113,20 +109,16 @@ public class Engine : ModuleRules
 			}
 		);
 
+		if (Target.bCompileAgainstApplicationCore)
+		{
+			PublicDependencyModuleNames.Add("ApplicationCore");
+		}
+
 		PublicIncludePathModuleNames.AddRange(
 			new string[] {
 				"TypedElementFramework",
 				"TypedElementRuntime",
-				"NetCore",
-				"RenderCore",
-				"CoreUObject",
-				"CoreOnline",
-				"PhysicsCore",
-				"ChaosCore",
-				"DeveloperSettings",
 				"NetCommon",
-				"Sockets",
-				"MeshDescription"
 			}
 		);
 
@@ -147,13 +139,11 @@ public class Engine : ModuleRules
 				"Analytics",
 				"AudioMixer",
 				"AudioMixerCore",
-				"SignalProcessing",
 				"IntelISPC",
 				"TraceLog",
 				"ColorManagement",
 				"Icmp",
 				"XmlParser",
-				"AudioExtensions"
 			}
 		);
 
diff --git a/Engine/Source/Runtime/Slate/Slate.Build.cs b/Engine/Source/Runtime/Slate/Slate.Build.cs
--- a/Engine/Source/Runtime/Slate/Slate.Build.cs
+++ b/Engine/Source/Runtime/Slate/Slate.Build.cs
@@ -13,13 +13,17 @@ public class Slate : ModuleRules
 			new string[] {
 				"Core",
 				"CoreUObject",
-				"ApplicationCore",
 				"InputCore",
 				"Json",
 				"SlateCore",
 				"ImageWrapper"
 			});
 
+		if (Target.bCompileAgainstApplicationCore)
+		{
+			PublicDependencyModuleNames.Add("ApplicationCore");
+		}
+
 		if (Target.Platform == UnrealTargetPlatform.Win64)
 		{
 			AddEngineThirdPartyPrivateStaticDependencies(Target, "XInput");
diff --git a/Engine/Source/Runtime/SlateCore/SlateCore.Build.cs b/Engine/Source/Runtime/SlateCore/SlateCore.Build.cs
--- a/Engine/Source/Runtime/SlateCore/SlateCore.Build.cs
+++ b/Engine/Source/Runtime/SlateCore/SlateCore.Build.cs
@@ -13,12 +13,16 @@ public class SlateCore : ModuleRules
 				"Core",
 				"CoreUObject",
 				"InputCore",
-				"ApplicationCore",
 				"Json",
 				"TraceLog",
 			});
 
-        if (Target.Type != TargetType.Server)
+		if (Target.bCompileAgainstApplicationCore)
+		{
+			PublicDependencyModuleNames.Add("ApplicationCore");
+		}
+
+		if (Target.Type != TargetType.Server)
 		{
 			if (Target.bCompileFreeType)
 			{
