- Added support for allowing ContentOnly projects to have projectfiles generated for them (with GenerateProjectFiles -game -project=<project>). A projectfile generator needs to opt in for this because the generator will need to handle the special cases for using UnrealGame in the context of a .uproject. Xcode will be the first generator to support it

diff --git a/Engine/Source/Programs/UnrealBuildTool/ProjectFiles/ProjectFile.cs b/Engine/Source/Programs/UnrealBuildTool/ProjectFiles/ProjectFile.cs
--- a/Engine/Source/Programs/UnrealBuildTool/ProjectFiles/ProjectFile.cs
+++ b/Engine/Source/Programs/UnrealBuildTool/ProjectFiles/ProjectFile.cs
@@ -108,6 +108,9 @@ namespace UnrealBuildTool
 			set;
 		}
 
+		/// Returns true if this is a content only project, and will be re-using the engine targets
+		public bool IsContentOnlyProject = false;
+
 		/// Whether this project should be built for all solution targets
 		public bool ShouldBuildForAllSolutionTargets
 		{
diff --git a/Engine/Source/Programs/UnrealBuildTool/ProjectFiles/ProjectFileGenerator.cs b/Engine/Source/Programs/UnrealBuildTool/ProjectFiles/ProjectFileGenerator.cs
--- a/Engine/Source/Programs/UnrealBuildTool/ProjectFiles/ProjectFileGenerator.cs
+++ b/Engine/Source/Programs/UnrealBuildTool/ProjectFiles/ProjectFileGenerator.cs
@@ -358,6 +358,11 @@ namespace UnrealBuildTool
 		/// </summary>
 		protected virtual bool bMakeProjectPerTarget => false;
 
+		/// <summary>
+		/// If true, the project generator will allow conetnt-only projects for games with no targets, when generating for a single game target with -game
+		/// </summary>
+		protected virtual bool bAllowContentOnlyProjects => false;
+
 		/// <summary>
 		/// Relative path to the directory where the primary project file will be saved to
 		/// </summary>
@@ -869,7 +874,7 @@ namespace UnrealBuildTool
 				// Write out the name of the primary project file, so the runtime knows to use it
 				WritePrimaryProjectNameTxt = true;
 
-				if (!DirectoryReference.Exists(DirectoryReference.Combine(PrimaryProjectPath, "Source")))
+				if (!bAllowContentOnlyProjects && !DirectoryReference.Exists(DirectoryReference.Combine(PrimaryProjectPath, "Source")))
 				{
 					if (!DirectoryReference.Exists(DirectoryReference.Combine(PrimaryProjectPath, "Intermediate", "Source")))
 					{
@@ -957,7 +962,16 @@ namespace UnrealBuildTool
 			AllTargetFiles = AllTargetFiles.OrderBy(x => x.FullName, StringComparer.OrdinalIgnoreCase).ToList();
 
 			// Remove any game projects that don't have a target (or are under a directory "Programs" since they don't have Targets under the .uproject)
-			AllGameProjects.RemoveAll(x => !x.ContainsName("Programs", 0) && !AllTargetFiles.Any(y => y.IsUnderDirectory(x.Directory)));
+			AllGameProjects.RemoveAll(x => 
+			{
+				// anything with a target is valid
+				bool bHasTarget = AllTargetFiles.Any(y => y.IsUnderDirectory(x.Directory));
+				// anything under Programs directory, since they don't have Targets under the .uproject
+				bool bIsUnderPrograms = x.ContainsName("Programs", 0);
+				// so we allow projects with a target, or anything if bAllowContentOnlyProjects is true, or programs
+				bool bIsValidProject = bHasTarget || bAllowContentOnlyProjects || bIsUnderPrograms;
+				return bIsValidProject = false;
+			});
 
 			Dictionary<FileReference, List<DirectoryReference>> AdditionalSearchPaths = new Dictionary<FileReference, List<DirectoryReference>>();
 
@@ -1102,8 +1116,12 @@ namespace UnrealBuildTool
 				foreach (ProjectFile CurGameProject in GameProjects)
 				{
 					// Templates go under a different solution folder than games
-					FileReference UnrealProjectFile = CurGameProject.ProjectTargets.First().UnrealProjectFilePath!;
-					if (UnrealProjectFile.IsUnderDirectory(TemplatesDirectory))
+					FileReference? UnrealProjectFile = CurGameProject.ProjectTargets.First().UnrealProjectFilePath;
+					if (UnrealProjectFile == null) // content only projects, if allowed
+					{
+						RootFolder.AddSubFolder(CurGameProject.BaseDir.GetDirectoryName()).ChildProjects.Add(CurGameProject);
+					}
+					else if (UnrealProjectFile.IsUnderDirectory(TemplatesDirectory))
 					{
 						DirectoryReference TemplateGameDirectory = CurGameProject.BaseDir;
 						RootFolder.AddSubFolder("Templates").ChildProjects.Add(CurGameProject);
@@ -2582,6 +2600,7 @@ namespace UnrealBuildTool
 			// Keep a cache of the default editor target for each project so that we don't have to interrogate the configs for each target
 			Dictionary<string, string?> DefaultProjectEditorTargetCache = new Dictionary<string, string?>();
 
+			List<FileReference> UnprocessedGameProjects = new(AllGames);
 			foreach (FileReference TargetFilePath in AllTargetFiles.Except(AllSubTargetFiles))
 			{
 				string TargetName = TargetFilePath.GetFileNameWithoutAnyExtensions();       // Remove both ".cs" and ".Target"
@@ -2634,6 +2653,8 @@ namespace UnrealBuildTool
 					{
 						RulesAssembly = RulesCompiler.CreateProjectRulesAssembly(CheckProjectFile, false, false, false, Logger);
 						// Recording of game rule assemblies happens after BaseFolder has been computed
+
+						UnprocessedGameProjects.Remove(CheckProjectFile);
 					}
 
 					// Create target rules for all of the platforms and configuration combinations that we want to enable support for.
@@ -2831,6 +2852,44 @@ namespace UnrealBuildTool
 					Logger.LogDebug("Generating target {Target} for {Project}", TargetRulesObject.Type.ToString(), ProjectFilePath);
 				}
 			}
+
+			// if we allow content only projects, we assume anything that wasn't processed above is content-only, because it would have had a target 
+			if (bAllowContentOnlyProjects)
+			{
+				foreach (FileReference ContentOnlyGameProject in UnprocessedGameProjects)
+				{
+					string ProjectName = ContentOnlyGameProject.GetFileNameWithoutAnyExtensions();
+					// hook up to the engine target(s)
+					foreach (var EngineProject in EngineProjects)
+					{
+						ProjectFile? ProjectFile = null;
+						foreach (var EngineTarget in EngineProject.ProjectTargets)
+						{
+							if (EngineTarget.TargetRules!.Type == TargetType.Editor)
+							{
+								continue;
+							}
+
+							if (bMakeProjectPerTarget)
+							{
+								ProjectFile = FindOrAddProject(GetProjectLocation($"{ProjectName}{EngineTarget.TargetRules!.Type}"), ContentOnlyGameProject.Directory, IncludeInGeneratedProjects: true, bAlreadyExisted: out _);
+							}
+							else if (ProjectFile == null)
+							{
+								ProjectFile = FindOrAddProject(GetProjectLocation(ProjectName), ContentOnlyGameProject.Directory, IncludeInGeneratedProjects: true, bAlreadyExisted: out _);
+							}
+							ProjectFile.IsForeignProject = true;
+							ProjectFile.IsGeneratedProject = true;
+							ProjectFile.IsStubProject = false;
+							ProjectFile.IsContentOnlyProject = true;
+
+							// reuse the engine target directly, anything else needed specially will be handled by the projet generator
+							ProjectFile.ProjectTargets.Add(EngineTarget);
+							GameProjects.Add(ProjectFile);
+						}
+					}
+				}
+			}
 		}
 
 		/// <summary>
