UnrealBuildTool: Add AdditionalData propety to TargetRules, allowing a user to pass arbitary objects from a Target to any ModuleRules without needing to recompile UnrealBuildTool (request via UDN)

#rnx

diff --git a/Engine/Source/Programs/UnrealBuildTool/Configuration/TargetRules.cs b/Engine/Source/Programs/UnrealBuildTool/Configuration/TargetRules.cs
--- a/Engine/Source/Programs/UnrealBuildTool/Configuration/TargetRules.cs
+++ b/Engine/Source/Programs/UnrealBuildTool/Configuration/TargetRules.cs
@@ -503,6 +503,14 @@ namespace UnrealBuildTool
 		/// </summary>
 		public ILogger Logger { get; internal set; }
 
+		/// <summary>
+		/// Generic nullable object so a user can set additional data in a project's TargetRule and access in a project's ModuleRule without needing to add new properties post-release.
+		/// For example:
+		/// * in .Target.cs:	AdditionalData = "data";
+		/// * in .Build.cs:		if ((Target.AdditionalData as string) == "data") { ... }
+		/// </summary>
+		public object? AdditionalData { get; set; }
+
 		/// <summary>
 		/// Platform that this target is being built for.
 		/// </summary>
@@ -2818,6 +2826,11 @@ namespace UnrealBuildTool
 			get { return Inner.TargetFiles!; }
 		}
 
+		public object? AdditionalData
+		{
+			get { return Inner.AdditionalData; }
+		}
+
 		public UnrealTargetPlatform Platform
 		{
 			get { return Inner.Platform; }
