Fix UBT project generation skipping remove on invalid targets when buliding the collection of project file targets.

diff --git a/Engine/Source/Programs/UnrealBuildTool/ProjectFiles/ProjectFileGenerator.cs b/Engine/Source/Programs/UnrealBuildTool/ProjectFiles/ProjectFileGenerator.cs
--- a/Engine/Source/Programs/UnrealBuildTool/ProjectFiles/ProjectFileGenerator.cs
+++ b/Engine/Source/Programs/UnrealBuildTool/ProjectFiles/ProjectFileGenerator.cs
@@ -970,7 +970,7 @@ namespace UnrealBuildTool
 				bool bIsUnderPrograms = x.ContainsName("Programs", 0);
 				// so we allow projects with a target, or anything if bAllowContentOnlyProjects is true, or programs
 				bool bIsValidProject = bHasTarget || bAllowContentOnlyProjects || bIsUnderPrograms;
-				return bIsValidProject = false;
+				return bIsValidProject == false;
 			});
 
 			Dictionary<FileReference, List<DirectoryReference>> AdditionalSearchPaths = new Dictionary<FileReference, List<DirectoryReference>>();
