Remove unnecessary setting from generated workspace file.

#jra none

diff --git a/Engine/Source/Programs/UnrealBuildTool/ProjectFiles/VisualStudioCode/VSCodeProjectFileGenerator.cs b/Engine/Source/Programs/UnrealBuildTool/ProjectFiles/VisualStudioCode/VSCodeProjectFileGenerator.cs
--- a/Engine/Source/Programs/UnrealBuildTool/ProjectFiles/VisualStudioCode/VSCodeProjectFileGenerator.cs
+++ b/Engine/Source/Programs/UnrealBuildTool/ProjectFiles/VisualStudioCode/VSCodeProjectFileGenerator.cs
@@ -1647,7 +1647,6 @@ namespace UnrealBuildTool
 
 				if (bUseVSCodeExtension)
 				{
-					WorkspaceFile.AddField("UE.loggingLevel", "Warning");
 					if (HostPlatform == UnrealTargetPlatform.Win64)
 					{
 						WorkspaceFile.AddField("UE.UBTScriptPath", MakePathString(FileReference.Combine(ProjectRoot, "Engine", "Build", "BatchFiles", "RunUBT.bat"), true));
