Fixed issue with artifacts where the last group of writes were not getting completed.

diff --git a/Engine/Source/Programs/UnrealBuildTool/Artifacts/HordeStorageArtifactCache.cs b/Engine/Source/Programs/UnrealBuildTool/Artifacts/HordeStorageArtifactCache.cs
--- a/Engine/Source/Programs/UnrealBuildTool/Artifacts/HordeStorageArtifactCache.cs
+++ b/Engine/Source/Programs/UnrealBuildTool/Artifacts/HordeStorageArtifactCache.cs
@@ -5,6 +5,7 @@ using EpicGames.Horde.Storage;
 using EpicGames.Horde.Storage.Backends;
 using EpicGames.Horde.Storage.Nodes;
 using Microsoft.Extensions.Logging;
+using Microsoft.Extensions.Logging.Abstractions;
 using System;
 using System.Collections.Generic;
 using System.IO;
@@ -254,7 +255,7 @@ namespace UnrealBuildTool.Artifacts
 		public static IArtifactCache CreateFileCache(DirectoryReference directory, ILogger logger, bool cleanDirectory)
 		{
 			HordeStorageArtifactCache cache = new(null, logger);
-			cache._readyTask = Task.Run(() => cache.InitFileCache(directory, logger, cleanDirectory));
+			cache._readyTask = Task.Run(() => cache.InitFileCache(directory, NullLogger.Instance, cleanDirectory));
 			return cache;
 		}
 
@@ -443,7 +444,7 @@ namespace UnrealBuildTool.Artifacts
 					ArtifactMapping[] mappingsToFlush = _pendingMappings.ToArray();
 					Task? priorTask = _pendingFlushTask;
 					_pendingMappings.Clear();
-					pendingFlushTask = _pendingFlushTask = new(async () =>
+					async Task action()
 					{
 
 						// When forcing, we might have a prior flush task in progress.  Wait for it to complete
@@ -463,7 +464,8 @@ namespace UnrealBuildTool.Artifacts
 						{
 							_semaphore.Release();
 						}
-					}, cancellationToken);
+					}
+					pendingFlushTask = _pendingFlushTask = new(() => action().Wait(), cancellationToken);
 				}
 			}
 
@@ -568,7 +570,7 @@ namespace UnrealBuildTool.Artifacts
 		/// <returns>TreeReader</returns>
 		private TreeReader CreateTreeReader()
 		{
-			return new(_store!, null, _treeReaderOptions, _logger);
+			return new(_store!, null, _treeReaderOptions, NullLogger.Instance);
 		}
 
 		/// <summary>
@@ -577,7 +579,7 @@ namespace UnrealBuildTool.Artifacts
 		/// <returns>TreeWriter</returns>
 		private TreeWriter CreateTreeWriter()
 		{
-			return new(_store!, null, default, null, _logger);
+			return new(_store!, null, default, null, NullLogger.Instance);
 		}
 	}
 }
