Fixed issue where actions completed from the cache would be resaved.

diff --git a/Engine/Source/Programs/UnrealBuildTool/Executors/ImmediateActionQueue.cs b/Engine/Source/Programs/UnrealBuildTool/Executors/ImmediateActionQueue.cs
--- a/Engine/Source/Programs/UnrealBuildTool/Executors/ImmediateActionQueue.cs
+++ b/Engine/Source/Programs/UnrealBuildTool/Executors/ImmediateActionQueue.cs
@@ -683,8 +683,10 @@ namespace UnrealBuildTool
 				runner.ActiveActionWeight -= Actions[actionIndex].Action.Weight;
 
 				// If we are doing an artifact check, then move to compile phase
+				bool wasArtifactCheck = false;
 				if (Actions[actionIndex].Phase == ActionPhase.ArtifactCheck)
 				{
+					wasArtifactCheck = true;
 					Actions[actionIndex].Phase = ActionPhase.Compile;
 					if (status != ActionStatus.Finished)
 					{
@@ -724,7 +726,10 @@ namespace UnrealBuildTool
 					case ActionStatus.Finished:
 						// Notify the artifact handler of the action completing.  We don't wait on the resulting task.  The
 						// cache is required to remember any pending saves and a final call to Flush will wait for everything to complete.
-						_actionArtifactCache?.ActionCompleteAsync(action, CancellationToken);
+						if (!wasArtifactCheck)
+						{
+							_actionArtifactCache?.ActionCompleteAsync(action, CancellationToken);
+						}
 						completedActions = 1;
 						break;
 
