Write tasks and launch configurations to workspace file instead of .vscode directory so users can persist their own files in that directory.

diff --git a/Engine/Source/Programs/UnrealBuildTool/ProjectFiles/VisualStudioCode/VSCodeProjectFileGenerator.cs b/Engine/Source/Programs/UnrealBuildTool/ProjectFiles/VisualStudioCode/VSCodeProjectFileGenerator.cs
--- a/Engine/Source/Programs/UnrealBuildTool/ProjectFiles/VisualStudioCode/VSCodeProjectFileGenerator.cs
+++ b/Engine/Source/Programs/UnrealBuildTool/ProjectFiles/VisualStudioCode/VSCodeProjectFileGenerator.cs
@@ -290,11 +290,9 @@ namespace UnrealBuildTool
 
 			ProjectData ProjectData = GatherProjectData(Projects, Logger);
 
-			WriteTasksFile(ProjectData, VSCodeDir);
-			WriteLaunchFile(ProjectData, VSCodeDir, Logger);
 			WriteWorkspaceIgnoreFile(Projects);
 			WriteCppPropertiesFile(VSCodeDir, ProjectData);
-			WriteWorkspaceFile();
+			WriteWorkspaceFile(ProjectData, Logger);
 
 			if (bForeignProject && bIncludeEngineSource)
 			{
@@ -1061,35 +1059,27 @@ namespace UnrealBuildTool
 			}
 		}
 
-		private void WriteTasksFile(ProjectData ProjectData, DirectoryReference VSCodeDir)
+		private void WriteTasks(JsonFile OutFile, ProjectData ProjectData)
 		{
-			JsonFile OutFile = new JsonFile();
+			OutFile.AddField("version", "2.0.0");
 
-			OutFile.BeginRootObject();
+			OutFile.BeginArray("tasks");
 			{
-				OutFile.AddField("version", "2.0.0");
-
-				OutFile.BeginArray("tasks");
+				if (!bUseVSCodeExtension)
 				{
-					if (!bUseVSCodeExtension)
-					{
-						foreach (ProjectData.Project NativeProject in ProjectData.NativeProjects)
-						{
-							WriteNativeTask(NativeProject, OutFile);
-						}
-					}
-
-					foreach (ProjectData.Project CSharpProject in ProjectData.CSharpProjects)
+					foreach (ProjectData.Project NativeProject in ProjectData.NativeProjects)
 					{
-						WriteCSharpTask(CSharpProject, OutFile);
+						WriteNativeTask(NativeProject, OutFile);
 					}
+				}
 
-					OutFile.EndArray();
+				foreach (ProjectData.Project CSharpProject in ProjectData.CSharpProjects)
+				{
+					WriteCSharpTask(CSharpProject, OutFile);
 				}
-			}
-			OutFile.EndRootObject();
 
-			OutFile.Write(FileReference.Combine(VSCodeDir, "tasks.json"));
+				OutFile.EndArray();
+			}
 		}
 		
 		private FileReference GetExecutableFilename(ProjectFile Project, ProjectTarget Target, UnrealTargetPlatform Platform, UnrealTargetConfiguration Configuration)
@@ -1456,84 +1446,76 @@ namespace UnrealBuildTool
 			}
 		}
 
-		private void WriteLaunchFile(ProjectData ProjectData, DirectoryReference VSCodeDir, ILogger Logger)
+		private void WriteLaunch(JsonFile OutFile, ProjectData ProjectData, ILogger Logger)
 		{
-			JsonFile OutFile = new JsonFile();
-
-			OutFile.BeginRootObject();
+			OutFile.AddField("version", "0.2.0");
+			if(bAddDebugCoreConfig)
 			{
-				OutFile.AddField("version", "0.2.0");
-				if(bAddDebugCoreConfig)
-				{
-					OutFile.BeginArray("inputs");
-					OutFile.BeginObject();
-					OutFile.AddField("id", "coreFileName");
-					OutFile.AddField("type", "command");
-					OutFile.AddField("command", "filePicker.pick");
-					OutFile.BeginObject("args");
-					OutFile.AddField("masks", "core*");
-					OutFile.BeginObject("display");
-					OutFile.AddField("type", "fileRelativePath");
-					OutFile.AddField("detail", "filePath");
-					OutFile.EndObject();
-					OutFile.AddField("output", "filePath");
-					OutFile.EndObject();
-					OutFile.EndObject();
-					OutFile.EndArray();
-				}
+				OutFile.BeginArray("inputs");
+				OutFile.BeginObject();
+				OutFile.AddField("id", "coreFileName");
+				OutFile.AddField("type", "command");
+				OutFile.AddField("command", "filePicker.pick");
+				OutFile.BeginObject("args");
+				OutFile.AddField("masks", "core*");
+				OutFile.BeginObject("display");
+				OutFile.AddField("type", "fileRelativePath");
+				OutFile.AddField("detail", "filePath");
+				OutFile.EndObject();
+				OutFile.AddField("output", "filePath");
+				OutFile.EndObject();
+				OutFile.EndObject();
+				OutFile.EndArray();
+			}
 
-				OutFile.BeginArray("configurations");
+			OutFile.BeginArray("configurations");
+			{
+				if (!bUseVSCodeExtension)
 				{
-					if (!bUseVSCodeExtension)
-					{
-						foreach (ProjectData.Project Project in ProjectData.NativeProjects)
-						{
-							WriteNativeLaunchConfig(Project, OutFile, Logger);
-						}
-					}
-
-					foreach (ProjectData.Project Project in ProjectData.CSharpProjects)
+					foreach (ProjectData.Project Project in ProjectData.NativeProjects)
 					{
-						WriteCSharpLaunchConfig(Project, OutFile);
+						WriteNativeLaunchConfig(Project, OutFile, Logger);
 					}
 				}
 
-				// Add in a special task for regenerating project files
-				string PreLaunchTask = "";
-				List<string> Args = new List<string>();
-				Args.Add("-projectfiles");
-				Args.Add("-vscode");
-
-				if (bGeneratingGameProjectFiles)
-				{
-					Args.Add("-project=" + MakeUnquotedPathString(OnlyGameProject!, EPathType.Absolute));
-					Args.Add("-game");
-				}
-				if (bIncludeEngineSource)
-				{
-					Args.Add("-engine");
-				}
-				
-				if (bIncludeDotNetPrograms)
+				foreach (ProjectData.Project Project in ProjectData.CSharpProjects)
 				{
-					Args.Add("-dotnet");
-					PreLaunchTask = "UnrealBuildTool " + HostPlatform.ToString() + " Development Build";
+					WriteCSharpLaunchConfig(Project, OutFile);
 				}
+			}
 
-				FileReference RunUbtPath = FileReference.Combine(ProjectRoot, "Engine", "Build", "BatchFiles", "RunUBT.bat");
-				WriteSingleCSharpLaunchConfig(
-					OutFile,
-					"Generate Project Files",
-					PreLaunchTask,
-					RunUbtPath,
-					Args.ToArray()
-				);
+			// Add in a special task for regenerating project files
+			string PreLaunchTask = "";
+			List<string> Args = new List<string>();
+			Args.Add("-projectfiles");
+			Args.Add("-vscode");
 
-				OutFile.EndArray();
+			if (bGeneratingGameProjectFiles)
+			{
+				Args.Add("-project=" + MakeUnquotedPathString(OnlyGameProject!, EPathType.Absolute));
+				Args.Add("-game");
+			}
+			if (bIncludeEngineSource)
+			{
+				Args.Add("-engine");
+			}
+			
+			if (bIncludeDotNetPrograms)
+			{
+				Args.Add("-dotnet");
+				PreLaunchTask = "UnrealBuildTool " + HostPlatform.ToString() + " Development Build";
 			}
-			OutFile.EndRootObject();
 
-			OutFile.Write(FileReference.Combine(VSCodeDir, "launch.json"));
+			FileReference RunUbtPath = FileReference.Combine(ProjectRoot, "Engine", "Build", "BatchFiles", "RunUBT.bat");
+			WriteSingleCSharpLaunchConfig(
+				OutFile,
+				"Generate Project Files",
+				PreLaunchTask,
+				RunUbtPath,
+				Args.ToArray()
+			);
+
+			OutFile.EndArray();
 		}
 
 		private void WriteWorkspaceIgnoreFile(List<ProjectFile> Projects)
@@ -1606,7 +1588,7 @@ namespace UnrealBuildTool
 			}
 		}
 		
-		private void WriteWorkspaceFile()
+		private void WriteWorkspaceFile(ProjectData ProjectData, ILogger Logger)
 		{
 			JsonFile WorkspaceFile = new JsonFile();
 
@@ -1703,6 +1685,14 @@ namespace UnrealBuildTool
 			}
 			WorkspaceFile.EndObject();
 
+			WorkspaceFile.BeginObject("tasks");
+			WriteTasks(WorkspaceFile, ProjectData);
+			WorkspaceFile.EndObject();
+
+			WorkspaceFile.BeginObject("launch");
+			WriteLaunch(WorkspaceFile, ProjectData, Logger);
+			WorkspaceFile.EndObject();
+
 			WorkspaceFile.EndRootObject();
 
 			string? WorkspaceName = bForeignProject ? GameProjectName : PrimaryProjectName;
