Make sure the high refreshrate on iPhone is available with both old and new XCode flow.

diff --git a/Engine/Source/Programs/UnrealBuildTool/Platform/IOS/IOSExports.cs b/Engine/Source/Programs/UnrealBuildTool/Platform/IOS/IOSExports.cs
--- a/Engine/Source/Programs/UnrealBuildTool/Platform/IOS/IOSExports.cs
+++ b/Engine/Source/Programs/UnrealBuildTool/Platform/IOS/IOSExports.cs
@@ -527,18 +527,7 @@ namespace UnrealBuildTool
 					Text.AppendLine("\t<key>com.apple.developer.user-management</key>");
 					Text.AppendLine("\t<array><string>runs-as-current-user</string></array>");
 				}
-
-				// Support high refresh rates (iPhone only)
-				// https://developer.apple.com/documentation/quartzcore/optimizing_promotion_refresh_rates_for_iphone_13_pro_and_ipad_pro
-				bool bSupportHighRefreshRates = false;
-				PlatformGameConfig.GetBool("/Script/IOSRuntimeSettings.IOSRuntimeSettings", "bSupportHighRefreshRates", out bSupportHighRefreshRates);
-
-				if (bSupportHighRefreshRates && Platform == UnrealTargetPlatform.IOS)
-				{
-					Text.AppendLine("\t<key>CADisableMinimumFrameDurationOnPhone</key><true/>");
-				}
-
-
+				
 				// End of entitlements
 				Text.AppendLine("</dict>");
 				Text.AppendLine("</plist>");
diff --git a/Engine/Source/Programs/UnrealBuildTool/Platform/IOS/UEDeployIOS.cs b/Engine/Source/Programs/UnrealBuildTool/Platform/IOS/UEDeployIOS.cs
--- a/Engine/Source/Programs/UnrealBuildTool/Platform/IOS/UEDeployIOS.cs
+++ b/Engine/Source/Programs/UnrealBuildTool/Platform/IOS/UEDeployIOS.cs
@@ -376,6 +376,15 @@ namespace UnrealBuildTool
 			Text.AppendLine("\t<key>UILaunchStoryboardName</key>");
 			Text.AppendLine("\t<string>LaunchScreen</string>");
 
+			// Support high refresh rates (iPhone only)
+			// https://developer.apple.com/documentation/quartzcore/optimizing_promotion_refresh_rates_for_iphone_13_pro_and_ipad_pro
+			bool bSupportHighRefreshRates = false;
+			Ini.GetBool("/Script/IOSRuntimeSettings.IOSRuntimeSettings", "bSupportHighRefreshRates", out bSupportHighRefreshRates);
+			if (bSupportHighRefreshRates)
+			{
+				Text.AppendLine("\t<key>CADisableMinimumFrameDurationOnPhone</key><true/>");
+			}
+
 			// disable exempt encryption
 			Text.AppendLine("\t<key>ITSAppUsesNonExemptEncryption</key>");
 			Text.AppendLine("\t<false/>");
@@ -777,6 +786,15 @@ namespace UnrealBuildTool
 				}
 			}
 
+			// Support high refresh rates (iPhone only)
+			// https://developer.apple.com/documentation/quartzcore/optimizing_promotion_refresh_rates_for_iphone_13_pro_and_ipad_pro
+			bool bSupportHighRefreshRates = false;
+			Ini.GetBool("/Script/IOSRuntimeSettings.IOSRuntimeSettings", "bSupportHighRefreshRates", out bSupportHighRefreshRates);
+			if (bSupportHighRefreshRates)
+			{
+				Text.AppendLine("\t<key>CADisableMinimumFrameDurationOnPhone</key><true/>");
+			}
+
 			Text.AppendLine("\t<key>CFBundleSupportedPlatforms</key>");
 			Text.AppendLine("\t<array>");
 			Text.AppendLine("\t\t<string>iPhoneOS</string>");
