Fixed issue where UBT would lockup when the artifact system was enabled.

diff --git a/Engine/Source/Programs/UnrealBuildTool/Executors/ImmediateActionQueue.cs b/Engine/Source/Programs/UnrealBuildTool/Executors/ImmediateActionQueue.cs
--- a/Engine/Source/Programs/UnrealBuildTool/Executors/ImmediateActionQueue.cs
+++ b/Engine/Source/Programs/UnrealBuildTool/Executors/ImmediateActionQueue.cs
@@ -570,6 +570,7 @@ namespace UnrealBuildTool
 
 					case ActionReadyState.Error:
 						Actions[actionIndex].Status = ActionStatus.Error;
+						Actions[actionIndex].Phase = ActionPhase.Compile;
 						completedActions++;
 						break;
 
