- Brought the idea of hybrid/temp targets for content only projects that need to be treated as code, to UBT instead of UAT. Generate temp targets when building or generating projects
- Only active when the project generator says it supports ContentOnly projects, although that is just for now until Xcode can test it out, in case it causes Visual Studio issues

diff --git a/Engine/Source/Programs/AutomationTool/AutomationUtils/ProjectUtils.cs b/Engine/Source/Programs/AutomationTool/AutomationUtils/ProjectUtils.cs
--- a/Engine/Source/Programs/AutomationTool/AutomationUtils/ProjectUtils.cs
+++ b/Engine/Source/Programs/AutomationTool/AutomationUtils/ProjectUtils.cs
@@ -213,366 +213,6 @@ namespace AutomationTool
 			return ProjectClientBinariesPath;
 		}
 
-		private static bool ProjectHasCode(FileReference RawProjectPath)
-		{
-			// check to see if we already have a Target.cs file
-			if (File.Exists(Path.Combine(Path.GetDirectoryName(RawProjectPath.FullName), "Source", RawProjectPath.GetFileNameWithoutExtension() + ".Target.cs")))
-			{
-				return true;
-			}
-			else if (Directory.Exists(Path.Combine(Path.GetDirectoryName(RawProjectPath.FullName), "Source")))
-			{
-				// wasn't one in the main Source directory, let's check all sub-directories
-				//@todo: may want to read each target.cs to see if it has a target corresponding to the project name as a final check
-				FileInfo[] Files = (new DirectoryInfo(Path.Combine(Path.GetDirectoryName(RawProjectPath.FullName), "Source")).GetFiles("*.Target.cs", SearchOption.AllDirectories));
-				if (Files.Length > 0)
-				{
-					return true;
-				}
-			}
-			return false;
-		}
-
-		private static bool RequiresTempTarget(FileReference RawProjectPath, List<UnrealTargetPlatform> Platforms, List<UnrealTargetConfiguration> Configurations, bool AssetNativizationRequested)
-		{
-			bool bHasCode = ProjectHasCode(RawProjectPath);
-			foreach (UnrealTargetPlatform Platform in Platforms)
-			{
-				foreach(UnrealTargetConfiguration Configuration in Configurations)
-				{
-					string Reason;
-					if(RequiresTempTarget(RawProjectPath, bHasCode, Platform, Configuration, TargetType.Game, AssetNativizationRequested, true, out Reason))
-					{
-						Logger.LogInformation("{Project} requires a temporary target.cs to be generated ({Reason})", RawProjectPath.GetFileName(), Reason);
-						return true;
-					}
-				}
-			}
-			return false;
-		}
-
-		/// <summary>
-		/// NOTE: This function must mirror the functionality of TargetPlatformBase::RequiresTempTarget
-		/// </summary>
-		private static bool RequiresTempTarget(FileReference RawProjectPath, bool bProjectHasCode, UnrealTargetPlatform Platform, UnrealTargetConfiguration Configuration, TargetType TargetType, bool bRequiresAssetNativization, bool bRequiresCookedData, out string OutReason)
-		{
-			// check to see if we already have a Target.cs file
-			if (bProjectHasCode)
-			{
-				OutReason = null;
-				return false;
-			}
-
-			// check if asset nativization is enabled
-			if (bRequiresAssetNativization)
-            {
-				OutReason = "asset nativization is enabled";
-                return true;
-            }
-
-			// Check if encryption or signing is enabled
-			EncryptionAndSigning.CryptoSettings Settings = EncryptionAndSigning.ParseCryptoSettings(RawProjectPath.Directory, Platform, Log.Logger);
-			if (Settings.IsAnyEncryptionEnabled() || Settings.IsPakSigningEnabled())
-			{
-				OutReason = "encryption/signing is enabled";
-				return true;
-			}
-
-			// check the target platforms for any differences in build settings or additional plugins
-			if(!Unreal.IsEngineInstalled() && !PlatformExports.HasDefaultBuildConfig(RawProjectPath, Platform))
-			{
-				OutReason = "project has non-default build configuration";
-				return true;
-			}
-			if(PlatformExports.RequiresBuild(RawProjectPath, Platform))
-			{
-				OutReason = "overriden by target platform";
-				return true;
-			}
-
-			// Read the project descriptor, and find all the plugins available to this project
-			ProjectDescriptor Project = ProjectDescriptor.FromFile(RawProjectPath);
-
-			// Enumerate all the available plugins
-			Dictionary<string, PluginInfo> AllPlugins = Plugins.ReadAvailablePlugins(Unreal.EngineDirectory, DirectoryReference.FromFile(RawProjectPath), new List<DirectoryReference>()).ToDictionary(x => x.Name, x => x, StringComparer.OrdinalIgnoreCase);
-
-			// find if there are any plugins enabled or disabled which differ from the default
-			string Reason;
-			if (RequiresTempTargetForCodePlugin(Project, Platform, Configuration, TargetType, AllPlugins, out Reason))
-			{
-				OutReason = Reason;
-				return true;
-			}
-
-			OutReason = null;
-			return false;
-		}
-
-		/// <summary>
-		/// NOTE: This function must mirror FPluginManager::RequiresTempTargetForCodePlugin
-		/// </summary>
-		static bool RequiresTempTargetForCodePlugin(ProjectDescriptor ProjectDescriptor, UnrealTargetPlatform Platform, UnrealTargetConfiguration Configuration, TargetType TargetType, Dictionary<string, PluginInfo> AllPlugins, out string OutReason)
-		{
-			PluginReferenceDescriptor MissingPlugin;
-
-			HashSet<string> ProjectCodePlugins = new HashSet<string>(StringComparer.OrdinalIgnoreCase);
-			if (!GetCodePluginsForTarget(ProjectDescriptor, Platform, Configuration, TargetType, ProjectCodePlugins, AllPlugins, out MissingPlugin))
-			{
-				OutReason = String.Format("{0} plugin is referenced by target but not found", MissingPlugin.Name);
-				return true;
-			}
-
-			HashSet<string> DefaultCodePlugins = new HashSet<string>(StringComparer.OrdinalIgnoreCase);
-			if (!GetCodePluginsForTarget(null, Platform, Configuration, TargetType, DefaultCodePlugins, AllPlugins, out MissingPlugin))
-			{
-				OutReason = String.Format("{0} plugin is referenced by the default target but not found", MissingPlugin.Name);
-				return true;
-			}
-
-			foreach (string ProjectCodePlugin in ProjectCodePlugins)
-			{
-				if (!DefaultCodePlugins.Contains(ProjectCodePlugin))
-				{
-					OutReason = String.Format("{0} plugin is enabled", ProjectCodePlugin);
-					return true;
-				}
-			}
-
-			foreach (string DefaultCodePlugin in DefaultCodePlugins)
-			{
-				if (!ProjectCodePlugins.Contains(DefaultCodePlugin))
-				{
-					OutReason = String.Format("{0} plugin is disabled", DefaultCodePlugin);
-					return true;
-				}
-			}
-
-			OutReason = null;
-			return false;
-		}
-
-		/// <summary>
-		/// NOTE: This function must mirror FPluginManager::GetCodePluginsForTarget
-		/// </summary>
-		static bool GetCodePluginsForTarget(ProjectDescriptor ProjectDescriptor, UnrealTargetPlatform Platform, UnrealTargetConfiguration Configuration, TargetType TargetType, HashSet<string> CodePluginNames, Dictionary<string, PluginInfo> AllPlugins, out PluginReferenceDescriptor OutMissingPlugin)
-		{
-			bool bLoadPluginsForTargetPlatforms = (TargetType == TargetType.Editor);
-
-			// Map of all enabled plugins
-			Dictionary<string, PluginInfo> EnabledPlugins = new Dictionary<string, PluginInfo>(StringComparer.OrdinalIgnoreCase);
-					   			 
-			// Keep a set of all the plugin names that have been configured. We read configuration data from different places, but only configure a plugin from the first place that it's referenced.
-			HashSet<string> ConfiguredPluginNames = new HashSet<string>(StringComparer.OrdinalIgnoreCase);
-
-			bool bAllowEnginePluginsEnabledByDefault = true;
-
-			// Find all the plugin references in the project file
-			if (ProjectDescriptor != null)
-			{
-				bAllowEnginePluginsEnabledByDefault = !ProjectDescriptor.DisableEnginePluginsByDefault;
-				if (ProjectDescriptor.Plugins != null)
-				{
-					// Copy the plugin references, since we may modify the project if any plugins are missing
-					foreach (PluginReferenceDescriptor PluginReference in ProjectDescriptor.Plugins)
-					{
-						if (!ConfiguredPluginNames.Contains(PluginReference.Name))
-						{
-							PluginReferenceDescriptor MissingPlugin;
-							if (!ConfigureEnabledPluginForTarget(PluginReference, ProjectDescriptor, null, Platform, Configuration, TargetType, bLoadPluginsForTargetPlatforms, AllPlugins, EnabledPlugins, out MissingPlugin))
-							{
-								OutMissingPlugin = MissingPlugin;
-								return false;
-							}
-							ConfiguredPluginNames.Add(PluginReference.Name);
-						}
-					}
-				}
-			}
-
-			// Add the plugins which are enabled by default
-			foreach (KeyValuePair<string, PluginInfo> PluginPair in AllPlugins)
-			{
-				if (PluginPair.Value.IsEnabledByDefault(bAllowEnginePluginsEnabledByDefault) && !ConfiguredPluginNames.Contains(PluginPair.Key))
-				{
-					PluginReferenceDescriptor MissingPlugin;
-					if (!ConfigureEnabledPluginForTarget(new PluginReferenceDescriptor(PluginPair.Key, null, true), ProjectDescriptor, null, Platform, Configuration, TargetType, bLoadPluginsForTargetPlatforms, AllPlugins, EnabledPlugins, out MissingPlugin))
-					{
-						OutMissingPlugin = MissingPlugin;
-						return false;
-					}
-					ConfiguredPluginNames.Add(PluginPair.Key);
-				}
-			}
-
-			// Figure out which plugins have code 
-			bool bBuildDeveloperTools = (TargetType == TargetType.Editor || TargetType == TargetType.Program || (Configuration != UnrealTargetConfiguration.Test && Configuration != UnrealTargetConfiguration.Shipping));
-			bool bRequiresCookedData = (TargetType != TargetType.Editor);
-			foreach (KeyValuePair<string, PluginInfo> Pair in EnabledPlugins)
-			{
-				if (Pair.Value.Descriptor.Modules != null)
-				{
-					foreach (ModuleDescriptor Module in Pair.Value.Descriptor.Modules)
-					{
-						if (Module.IsCompiledInConfiguration(Platform, Configuration, null, TargetType, bBuildDeveloperTools, bRequiresCookedData))
-						{
-							CodePluginNames.Add(Pair.Key);
-							break;
-						}
-					}
-				}
-			}
-
-			OutMissingPlugin = null;
-			return true;
-		}
-
-		/// <summary>
-		/// NOTE: This function should mirror FPluginManager::ConfigureEnabledPluginForTarget
-		/// </summary>
-		static bool ConfigureEnabledPluginForTarget(PluginReferenceDescriptor FirstReference, ProjectDescriptor ProjectDescriptor, string TargetName, UnrealTargetPlatform Platform, UnrealTargetConfiguration Configuration, TargetType TargetType, bool bLoadPluginsForTargetPlatforms, Dictionary<string, PluginInfo> AllPlugins, Dictionary<string, PluginInfo> EnabledPlugins, out PluginReferenceDescriptor OutMissingPlugin)
-		{
-			if (!EnabledPlugins.ContainsKey(FirstReference.Name))
-			{
-				// Set of plugin names we've added to the queue for processing
-				HashSet<string> NewPluginNames = new HashSet<string>(StringComparer.OrdinalIgnoreCase);
-				NewPluginNames.Add(FirstReference.Name);
-
-				// Queue of plugin references to consider
-				List<PluginReferenceDescriptor> NewPluginReferences = new List<PluginReferenceDescriptor>();
-				NewPluginReferences.Add(FirstReference);
-
-				// Loop through the queue of plugin references that need to be enabled, queuing more items as we go
-				for (int Idx = 0; Idx < NewPluginReferences.Count; Idx++)
-				{
-					PluginReferenceDescriptor Reference = NewPluginReferences[Idx];
-
-					// Check if the plugin is required for this platform
-					if(!Reference.IsEnabledForPlatform(Platform) || !Reference.IsEnabledForTargetConfiguration(Configuration) || !Reference.IsEnabledForTarget(TargetType))
-					{
-						Logger.LogDebug("Ignoring plugin '{Arg0}' for platform/configuration", Reference.Name);
-						continue;
-					}
-
-					// Check if the plugin is required for this platform
-					if(!bLoadPluginsForTargetPlatforms && !Reference.IsSupportedTargetPlatform(Platform))
-					{
-						Logger.LogDebug("Ignoring plugin '{Arg0}' due to unsupported platform", Reference.Name);
-						continue;
-					}
-
-					// Find the plugin being enabled
-					PluginInfo Plugin;
-					if (!AllPlugins.TryGetValue(Reference.Name, out Plugin))
-					{
-						// Ignore any optional plugins
-						if (Reference.bOptional)
-						{
-							Logger.LogDebug("Ignored optional reference to '%s' plugin; plugin was not found.", Reference.Name);
-							continue;
-						}
-
-						// Add it to the missing list
-						OutMissingPlugin = Reference;
-						return false;
-					}
-
-					// Check the plugin supports this platform
-					if(!bLoadPluginsForTargetPlatforms && !Plugin.Descriptor.SupportsTargetPlatform(Platform))
-					{
-						Logger.LogDebug("Ignoring plugin '{Arg0}' due to unsupported platform in plugin descriptor", Reference.Name);
-						continue;
-					}
-
-					// Check that this plugin supports the current program
-					if (TargetType == TargetType.Program && !Plugin.Descriptor.SupportedPrograms.Contains(TargetName))
-					{
-						Logger.LogDebug("Ignoring plugin '{Arg0}' due to absence from the supported programs list", Reference.Name);
-						continue;
-					}
-
-					// Add references to all its dependencies
-					if (Plugin.Descriptor.Plugins != null)
-					{
-						foreach (PluginReferenceDescriptor NextReference in Plugin.Descriptor.Plugins)
-						{
-							if (!EnabledPlugins.ContainsKey(NextReference.Name) && !NewPluginNames.Contains(NextReference.Name))
-							{
-								NewPluginNames.Add(NextReference.Name);
-								NewPluginReferences.Add(NextReference);
-							}
-						}
-					}
-
-					// Add the plugin
-					EnabledPlugins.Add(Plugin.Name, Plugin);
-				}
-			}
-
-			OutMissingPlugin = null;
-			return true;
-		}
-
-		private static void GenerateTempTarget(FileReference RawProjectPath)
-		{
-			DirectoryReference TempDir = DirectoryReference.Combine(RawProjectPath.Directory, "Intermediate", "Source");
-			DirectoryReference.CreateDirectory(TempDir);
-
-			// Get the project name for use in temporary files
-			string ProjectName = RawProjectPath.GetFileNameWithoutExtension();
-
-			// Create a target.cs file
-			MemoryStream TargetStream = new MemoryStream();
-			using (StreamWriter Writer = new StreamWriter(TargetStream))
-			{
-				Writer.WriteLine("using UnrealBuildTool;");
-				Writer.WriteLine();
-				Writer.WriteLine("public class {0}Target : TargetRules", ProjectName);
-				Writer.WriteLine("{");
-				Writer.WriteLine("\tpublic {0}Target(TargetInfo Target) : base(Target)", ProjectName);
-				Writer.WriteLine("\t{");
-				Writer.WriteLine("\t\tDefaultBuildSettings = BuildSettingsVersion.V3;");
-				Writer.WriteLine("\t\tIncludeOrderVersion = EngineIncludeOrderVersion.Latest;");
-				Writer.WriteLine("\t\tType = TargetType.Game;");
-				Writer.WriteLine("\t\tExtraModuleNames.Add(\"{0}\");", ProjectName);
-				Writer.WriteLine("\t}");
-				Writer.WriteLine("}");
-			}
-			FileReference TargetLocation = FileReference.Combine(TempDir, ProjectName + ".Target.cs");
-			FileReference.WriteAllBytesIfDifferent(TargetLocation, TargetStream.ToArray());
-
-			// Create a build.cs file
-			MemoryStream ModuleStream = new MemoryStream();
-			using (StreamWriter Writer = new StreamWriter(ModuleStream))
-			{
-				Writer.WriteLine("using UnrealBuildTool;");
-				Writer.WriteLine();
-				Writer.WriteLine("public class {0} : ModuleRules", ProjectName);
-				Writer.WriteLine("{");
-				Writer.WriteLine("\tpublic {0}(ReadOnlyTargetRules Target) : base(Target)", ProjectName);
-				Writer.WriteLine("\t{");
-				Writer.WriteLine("\t\tPCHUsage = PCHUsageMode.UseExplicitOrSharedPCHs;");
-				Writer.WriteLine();
-				Writer.WriteLine("\t\tPrivateDependencyModuleNames.Add(\"Core\");");
-				Writer.WriteLine("\t\tPrivateDependencyModuleNames.Add(\"Core\");");
-				Writer.WriteLine("\t}");
-				Writer.WriteLine("}");
-			}
-			FileReference ModuleLocation = FileReference.Combine(TempDir, ProjectName + ".Build.cs");
-			FileReference.WriteAllBytesIfDifferent(ModuleLocation, ModuleStream.ToArray());
-
-			// Create a main module cpp file
-			MemoryStream SourceFileStream = new MemoryStream();
-			using (StreamWriter Writer = new StreamWriter(SourceFileStream))
-			{
-				Writer.WriteLine("#include \"CoreTypes.h\"");
-				Writer.WriteLine("#include \"Modules/ModuleManager.h\"");
-				Writer.WriteLine();
-				Writer.WriteLine("IMPLEMENT_PRIMARY_GAME_MODULE(FDefaultModuleImpl, {0}, \"{0}\");", ProjectName);
-			}
-			FileReference SourceFileLocation = FileReference.Combine(TempDir, ProjectName + ".cpp");
-			FileReference.WriteAllBytesIfDifferent(SourceFileLocation, SourceFileStream.ToArray());
-		}
-
 		/// <summary>
 		/// Attempts to autodetect project properties.
 		/// </summary>
@@ -613,23 +253,12 @@ namespace AutomationTool
 					}
 				}
 
-				string TempTargetDir = CommandUtils.CombinePaths(Path.GetDirectoryName(RawProjectPath.FullName), "Intermediate", "Source");
-                if (RequiresTempTarget(RawProjectPath, TargetPlatforms, TargetConfigurations, AssetNativizationRequested))
+				if (NativeProjects.ConditionalMakeTempTargetForHybridProject(RawProjectPath, Logger))
 				{
-					GenerateTempTarget(RawProjectPath);
 					Properties.bWasGenerated = true;
-                    ExtraSearchPaths.Add(TempTargetDir);
-				}
-				else if (File.Exists(Path.Combine(Path.GetDirectoryName(RawProjectPath.FullName), "Intermediate", "Source", Path.GetFileNameWithoutExtension(RawProjectPath.FullName) + ".Target.cs")))
-				{
-					File.Delete(Path.Combine(Path.GetDirectoryName(RawProjectPath.FullName), "Intermediate", "Source", Path.GetFileNameWithoutExtension(RawProjectPath.FullName) + ".Target.cs"));
+					string TempTargetDir = CommandUtils.CombinePaths(Path.GetDirectoryName(RawProjectPath.FullName), "Intermediate", "Source");
+					ExtraSearchPaths.Add(TempTargetDir);
 				}
-
-                // in case the RulesCompiler (what we use to find all the 
-                // Target.cs files) has already cached the contents of this 
-                // directory, then we need to invalidate that cache (so 
-                // it'll find/use the new Target.cs file)
-                Rules.InvalidateRulesFileCache(TempTargetDir);
             }
 
 			if (CommandUtils.CmdEnv.HasCapabilityToCompile)
diff --git a/Engine/Source/Programs/UnrealBuildTool/Configuration/TargetDescriptor.cs b/Engine/Source/Programs/UnrealBuildTool/Configuration/TargetDescriptor.cs
--- a/Engine/Source/Programs/UnrealBuildTool/Configuration/TargetDescriptor.cs
+++ b/Engine/Source/Programs/UnrealBuildTool/Configuration/TargetDescriptor.cs
@@ -421,11 +421,17 @@ namespace UnrealBuildTool
 						if (TargetProjectFile == null)
 						{
 							// find one with a matching name
-							TargetProjectFile = NativeProjects.EnumerateProjectFiles(Log.Logger)
+							TargetProjectFile = NativeProjects.EnumerateProjectFiles(Logger)
 								.Where(x => x.GetFileNameWithoutAnyExtensions().Equals(TargetName, StringComparison.InvariantCultureIgnoreCase))
 								.FirstOrDefault();
 						}
 
+						// make a temp target for hybrid content-as-code projects
+						if (TargetProjectFile != null)
+						{
+							NativeProjects.ConditionalMakeTempTargetForHybridProject(TargetProjectFile, Logger);
+						}
+
 						if (ParamArchitectures == null)
 						{
 							// ask the platform what achitectures it wants for this project
diff --git a/Engine/Source/Programs/UnrealBuildTool/ProjectFiles/ProjectFile.cs b/Engine/Source/Programs/UnrealBuildTool/ProjectFiles/ProjectFile.cs
--- a/Engine/Source/Programs/UnrealBuildTool/ProjectFiles/ProjectFile.cs
+++ b/Engine/Source/Programs/UnrealBuildTool/ProjectFiles/ProjectFile.cs
@@ -109,7 +109,20 @@ namespace UnrealBuildTool
 		}
 
 		/// Returns true if this is a content only project, and will be re-using the engine targets
-		public bool IsContentOnlyProject = false;
+		public bool IsContentOnlyProject
+		{
+			get;
+			set;
+		}
+
+		/// Returns true if this is a project without source code, but some build setting requires it to be compiled as code, which will make
+		/// a unique executable, instead of using UnrealGame, etc
+		/// @note: Currently unused by any generator, but a future generator may make use of it
+		public bool IsHybridContentOnlyProject
+		{
+			get;
+			set;
+		}
 
 		/// Whether this project should be built for all solution targets
 		public bool ShouldBuildForAllSolutionTargets
diff --git a/Engine/Source/Programs/UnrealBuildTool/ProjectFiles/ProjectFileGenerator.cs b/Engine/Source/Programs/UnrealBuildTool/ProjectFiles/ProjectFileGenerator.cs
--- a/Engine/Source/Programs/UnrealBuildTool/ProjectFiles/ProjectFileGenerator.cs
+++ b/Engine/Source/Programs/UnrealBuildTool/ProjectFiles/ProjectFileGenerator.cs
@@ -954,6 +954,18 @@ namespace UnrealBuildTool
 			// Build the list of games to generate projects for
 			List<FileReference> AllGameProjects = FindGameProjects(Logger);
 
+			// before we detect targets, if we allow hybrid content only projects, look for hybrid projects and create .Target.cs files as needed
+			if (bAllowContentOnlyProjects)
+			{
+				foreach (FileReference GameUProjectFile in AllGameProjects)
+				{
+					NativeProjects.ConditionalMakeTempTargetForHybridProject(GameUProjectFile, Logger);
+				}
+
+				// they are created in a temp location, which we need to scan
+				bIncludeTempTargets = true;
+			}
+
 			// Find all of the target files.  This will filter out any modules or targets that don't
 			// belong to platforms we're generating project files for.
 			List<FileReference> AllTargetFiles = DiscoverTargets(AllGameProjects, Logger, OnlyGameProject, SupportedPlatforms, bIncludeEngineSource, bIncludeTempTargets);
@@ -962,20 +974,21 @@ namespace UnrealBuildTool
 			AllTargetFiles = AllTargetFiles.OrderBy(x => x.FullName, StringComparer.OrdinalIgnoreCase).ToList();
 
 			// Remove any game projects that don't have a target (or are under a directory "Programs" since they don't have Targets under the .uproject)
-			AllGameProjects.RemoveAll(x => 
+			List<FileReference> CodeBasedGameProjects = new(AllGameProjects);
+			CodeBasedGameProjects.RemoveAll(x => 
 			{
 				// anything with a target is valid
 				bool bHasTarget = AllTargetFiles.Any(y => y.IsUnderDirectory(x.Directory));
 				// anything under Programs directory, since they don't have Targets under the .uproject
 				bool bIsUnderPrograms = x.ContainsName("Programs", 0);
-				// so we allow projects with a target, or anything if bAllowContentOnlyProjects is true, or programs
-				bool bIsValidProject = bHasTarget || bAllowContentOnlyProjects || bIsUnderPrograms;
+				// so we allow projects with a target or are programs
+				bool bIsValidProject = bHasTarget || bIsUnderPrograms;
 				return bIsValidProject == false;
 			});
 
 			Dictionary<FileReference, List<DirectoryReference>> AdditionalSearchPaths = new Dictionary<FileReference, List<DirectoryReference>>();
 
-			foreach (FileReference GameProject in AllGameProjects)
+			foreach (FileReference GameProject in CodeBasedGameProjects)
 			{
 				ProjectDescriptor Project = ProjectDescriptor.FromFile(GameProject);
 				if (Project.AdditionalPluginDirectories.Count > 0)
@@ -986,7 +999,7 @@ namespace UnrealBuildTool
 
 			// Find all of the module files.  This will filter out any modules or targets that don't belong to platforms
 			// we're generating project files for.
-			List<FileReference> AllModuleFiles = DiscoverModules(AllGameProjects, AdditionalSearchPaths.Values.SelectMany(x => x).ToList());
+			List<FileReference> AllModuleFiles = DiscoverModules(CodeBasedGameProjects, AdditionalSearchPaths.Values.SelectMany(x => x).ToList());
 
 			List<ProjectFile> EngineProjects = new List<ProjectFile>();
 			List<ProjectFile> GameProjects = new List<ProjectFile>();
@@ -1544,14 +1557,18 @@ namespace UnrealBuildTool
 			foreach (ProjectFile GameProject in GameProjects)
 			{
 				DirectoryReference GameProjectDirectory = GameProject.BaseDir;
+
+				// add the uproject to all projects using this game project - some project generators requires at least one source file to exist, but if
+				//this is a codeonly project, then there is not the usual Target.cs file to be that one file
+				foreach (DirectoryReference ExtensionDir in Unreal.GetExtensionDirs(GameProjectDirectory))
+				{
+					GameProject.AddFilesToProject(SourceFileSearch.FindFiles(ExtensionDir, SearchSubdirectories: false), GameProjectDirectory);
+				}
+
 				if (UniqueGameProjectDirectories.Add(GameProjectDirectory))
 				{
 					// @todo projectfiles: We have engine localization files, but should we also add GAME localization files?
 
-					foreach (DirectoryReference ExtensionDir in Unreal.GetExtensionDirs(GameProjectDirectory))
-					{
-						GameProject.AddFilesToProject(SourceFileSearch.FindFiles(ExtensionDir, SearchSubdirectories: false), GameProjectDirectory);
-					}
 
 					// Game restricted source files, since they won't be added via a module FileReference
 					foreach (DirectoryReference GameRestrictedSourceDirectory in Unreal.GetExtensionDirs(GameProjectDirectory, "Source", bIncludePlatformDirectories: false, bIncludeBaseDirectory: false))
@@ -2762,6 +2779,7 @@ namespace UnrealBuildTool
 					ProjectFile.IsForeignProject = CheckProjectFile != null && !NativeProjects.IsNativeProject(CheckProjectFile, Logger);
 					ProjectFile.IsGeneratedProject = true;
 					ProjectFile.IsStubProject = UnrealBuildTool.IsProjectInstalled();
+					ProjectFile.IsHybridContentOnlyProject = CheckProjectFile != null && NativeProjects.IsHybridContentOnlyProject(CheckProjectFile, Logger);
 					if (TargetRulesObject.bBuildInSolutionByDefault.HasValue)
 					{
 						ProjectFile.ShouldBuildByDefaultForSolutionTargets = TargetRulesObject.bBuildInSolutionByDefault.Value;
diff --git a/Engine/Source/Programs/UnrealBuildTool/System/NativeProjects.cs b/Engine/Source/Programs/UnrealBuildTool/System/NativeProjects.cs
--- a/Engine/Source/Programs/UnrealBuildTool/System/NativeProjects.cs
+++ b/Engine/Source/Programs/UnrealBuildTool/System/NativeProjects.cs
@@ -78,5 +78,464 @@ namespace UnrealBuildTool
 			}
 			return CachedTargetNameToProjectFile.TryGetValue(InTargetName, out OutProjectFileName);
 		}
+
+
+
+		#region Hybrid content-only/code-based project
+
+		/// <summary>
+		/// Returns true if this project is a Hybrid content only project that requires it to be built as code
+		/// </summary>
+		/// <param name="UProjectFile"></param>
+		/// <param name="Logger"></param>
+		/// <returns></returns>
+		public static bool IsHybridContentOnlyProject(FileReference UProjectFile, ILogger Logger)
+		{
+			return RequiresTempTarget(
+				UProjectFile, 
+				new List<UnrealTargetPlatform>() { BuildHostPlatform.Current.Platform },
+				new List<UnrealTargetConfiguration>() { UnrealTargetConfiguration.Development, UnrealTargetConfiguration.Shipping }, 
+				out _,
+				Logger);
+		}
+
+		/// <summary>
+		/// Returns true if this project is a Hybrid content only project that requires it to be built as code
+		/// </summary>
+		/// <param name="UProjectFile"></param>
+		/// <param name="Reason">Contains a description of the reason the project is hybrid</param>
+		/// <param name="Logger"></param>
+		/// <returns></returns>
+		public static bool IsHybridContentOnlyProject(FileReference UProjectFile, [NotNullWhen(true)] out string? Reason, ILogger Logger)
+		{
+			return RequiresTempTarget(
+				UProjectFile,
+				new List<UnrealTargetPlatform>() { BuildHostPlatform.Current.Platform },
+				new List<UnrealTargetConfiguration>() { UnrealTargetConfiguration.Development, UnrealTargetConfiguration.Shipping },
+				out Reason,
+				Logger);
+		}
+
+		/// <summary>
+		/// Creates temporary target files, if needed, for a hybrid content only project
+		/// </summary>
+		/// <param name="UProjectFile"></param>
+		/// <param name="Logger"></param>
+		/// <returns>True if the project is hybrid</returns>
+		public static bool ConditionalMakeTempTargetForHybridProject(FileReference UProjectFile, ILogger Logger)
+		{
+			string? Reason;
+			bool bIsHybrid = IsHybridContentOnlyProject(UProjectFile, out Reason, Logger);
+
+			DirectoryReference TempDir = DirectoryReference.Combine(UProjectFile.Directory, "Intermediate", "Source");
+
+			// Get the project name for use in temporary files
+			string ProjectName = UProjectFile.GetFileNameWithoutExtension();
+			Dictionary<TargetType, FileReference> TargetFiles = new()
+			{
+				{ TargetType.Editor, FileReference.Combine(TempDir, $"{ProjectName}Editor.Target.cs") },
+				{ TargetType.Game, FileReference.Combine(TempDir, $"{ProjectName}.Target.cs") },
+				{ TargetType.Client, FileReference.Combine(TempDir, $"{ProjectName}Client.Target.cs") },
+				{ TargetType.Server, FileReference.Combine(TempDir, $"{ProjectName}Server.Target.cs") },
+			};
+			FileReference ModuleLocation = FileReference.Combine(TempDir, ProjectName + ".Build.cs");
+			FileReference SourceFileLocation = FileReference.Combine(TempDir, ProjectName + ".cpp");
+
+			// if all files exist, early out
+			bool bWasHybrid = TargetFiles.Values.All(x => FileReference.Exists(x)) && FileReference.Exists(ModuleLocation) && FileReference.Exists(SourceFileLocation);
+			if (!bIsHybrid)
+			{
+				// clean up if needed
+				if (bWasHybrid)
+				{
+					Logger.LogWarning("Cleaning old temporary Target files for {Project} because it no longer being treated as a code-based project.", ProjectName);
+					DirectoryReference.Delete(TempDir, bRecursive: true);
+				}
+				return false;
+			}
+			
+			// if the files existed, just leave them be
+			if (bIsHybrid && bWasHybrid)
+			{
+				return true;
+			}
+
+			Logger.LogInformation($"{Reason} Creating temporary .Target.cs files.", Reason);
+
+			// make sure directory exists
+			DirectoryReference.CreateDirectory(TempDir);
+
+			// Create a target.cs file
+			foreach (TargetType TargetType in TargetFiles.Keys)
+			{
+				string TargetTypeString = TargetType == TargetType.Game ? "" : TargetType.ToString();
+
+				MemoryStream TargetStream = new MemoryStream();
+				using (StreamWriter Writer = new StreamWriter(TargetStream))
+				{
+					Writer.WriteLine("using UnrealBuildTool;");
+					Writer.WriteLine();
+					Writer.WriteLine("public class {0}{1}Target : TargetRules", ProjectName, TargetTypeString);
+					Writer.WriteLine("{");
+					Writer.WriteLine("\tpublic {0}{1}Target(TargetInfo Target) : base(Target)", ProjectName, TargetTypeString);
+					Writer.WriteLine("\t{");
+					Writer.WriteLine("\t\tDefaultBuildSettings = BuildSettingsVersion.V3;");
+					Writer.WriteLine("\t\tIncludeOrderVersion = EngineIncludeOrderVersion.Latest;");
+					Writer.WriteLine("\t\tType = TargetType.{0};", TargetType);
+					Writer.WriteLine("\t\tExtraModuleNames.Add(\"{0}\");", ProjectName);
+					Writer.WriteLine("\t}");
+					Writer.WriteLine("}");
+				}
+				FileReference.WriteAllBytesIfDifferent(TargetFiles[TargetType], TargetStream.ToArray());
+			}
+
+			// Create a build.cs file
+			MemoryStream ModuleStream = new MemoryStream();
+			using (StreamWriter Writer = new StreamWriter(ModuleStream))
+			{
+				Writer.WriteLine("using UnrealBuildTool;");
+				Writer.WriteLine();
+				Writer.WriteLine("public class {0} : ModuleRules", ProjectName);
+				Writer.WriteLine("{");
+				Writer.WriteLine("\tpublic {0}(ReadOnlyTargetRules Target) : base(Target)", ProjectName);
+				Writer.WriteLine("\t{");
+				Writer.WriteLine("\t\tPCHUsage = PCHUsageMode.UseExplicitOrSharedPCHs;");
+				Writer.WriteLine();
+				Writer.WriteLine("\t\tPrivateDependencyModuleNames.Add(\"Core\");");
+				Writer.WriteLine("\t\tPrivateDependencyModuleNames.Add(\"Core\");");
+				Writer.WriteLine("\t}");
+				Writer.WriteLine("}");
+			}
+			FileReference.WriteAllBytesIfDifferent(ModuleLocation, ModuleStream.ToArray());
+
+			// Create a main module cpp file
+			MemoryStream SourceFileStream = new MemoryStream();
+			using (StreamWriter Writer = new StreamWriter(SourceFileStream))
+			{
+				Writer.WriteLine("#include \"CoreTypes.h\"");
+				Writer.WriteLine("#include \"Modules/ModuleManager.h\"");
+				Writer.WriteLine();
+				Writer.WriteLine("IMPLEMENT_PRIMARY_GAME_MODULE(FDefaultModuleImpl, {0}, \"{0}\");", ProjectName);
+			}
+			FileReference.WriteAllBytesIfDifferent(SourceFileLocation, SourceFileStream.ToArray());
+
+			// need to clear out some caches now that we've added a Target
+			ClearCache();
+			Rules.InvalidateRulesFileCache(UProjectFile.Directory.FullName);
+			DirectoryItem.ResetCachedInfo(UProjectFile.Directory.FullName);
+			DirectoryItem.ResetCachedInfo(TempDir.FullName);
+
+			return true;
+		}
+
+
+
+
+
+		/// <summary>
+		/// Determines if a project (given a .uproject file) has source code. This is determined by finding at least one .Target.cs file
+		/// </summary>
+		/// <param name="UProjectFile">Path to a .uproject file</param>
+		/// <param name="bCheckForTempTargets">If true, search Intermediate/Source for target files</param>
+		/// <returns>True if this is a source-based project</returns>
+		public static bool ProjectHasCode(FileReference UProjectFile, bool bCheckForTempTargets)
+		{
+			DirectoryReference SourceDir = DirectoryReference.Combine(UProjectFile.Directory, "Source");
+			DirectoryReference TempSourceDir = DirectoryReference.Combine(UProjectFile.Directory, "Intermediate", "Source");
+
+			// check to see if we have a Target.cs file in Source or Intermediate/Source
+			if (DirectoryReference.Exists(SourceDir) && DirectoryReference.EnumerateFiles(SourceDir, "*.Target.cs", SearchOption.TopDirectoryOnly).Count() > 0)
+			{
+				return true;
+			}
+			if (bCheckForTempTargets && DirectoryReference.Exists(TempSourceDir) && DirectoryReference.EnumerateFiles(TempSourceDir, "*.Target.cs", SearchOption.TopDirectoryOnly).Count() > 0)
+			{
+				return true;
+			}
+			return false;
+		}
+
+
+		private static bool RequiresTempTarget(FileReference UProjectFile, List<UnrealTargetPlatform> Platforms, List<UnrealTargetConfiguration> Configurations, [NotNullWhen(true)] out string? Reason, ILogger Logger)
+		{
+			// no reason by default
+			Reason = null;
+
+			// never create temp targets for the Template projects
+			if (UProjectFile.ContainsName("Templates", 0))
+			{
+				return false;
+			}
+
+			bool bHasCode = ProjectHasCode(UProjectFile, bCheckForTempTargets: false);
+			foreach (UnrealTargetPlatform Platform in Platforms)
+			{
+				foreach (UnrealTargetConfiguration Configuration in Configurations)
+				{
+					string? InnerReason;
+					if (RequiresTempTarget(UProjectFile, bHasCode, Platform, Configuration, TargetType.Game, out InnerReason, Logger))
+					{
+						Reason = $"{UProjectFile.GetFileName()} is has no code, but is being treated as a code-based project because: {InnerReason}.";
+						return true;
+					}
+				}
+			}
+			return false;
+		}
+
+		/// <summary>
+		/// NOTE: This function must mirror the functionality of TargetPlatformBase::RequiresTempTarget
+		/// </summary>
+		public static bool RequiresTempTarget(FileReference RawProjectPath, bool bProjectHasCode, UnrealTargetPlatform Platform, UnrealTargetConfiguration Configuration, TargetType TargetType, out string? OutReason, ILogger Logger)
+		{
+			// check to see if we already have a Target.cs file
+			if (bProjectHasCode)
+			{
+				OutReason = null;
+				return false;
+			}
+
+			// Check if encryption or signing is enabled
+			EncryptionAndSigning.CryptoSettings Settings = EncryptionAndSigning.ParseCryptoSettings(RawProjectPath.Directory, Platform, Log.Logger);
+			if (Settings.IsAnyEncryptionEnabled() || Settings.IsPakSigningEnabled())
+			{
+				OutReason = "encryption/signing is enabled";
+				return true;
+			}
+
+			// check the target platforms for any differences in build settings or additional plugins
+			if (!Unreal.IsEngineInstalled() && !PlatformExports.HasDefaultBuildConfig(RawProjectPath, Platform))
+			{
+				OutReason = "project has non-default build configuration";
+				return true;
+			}
+			if (PlatformExports.RequiresBuild(RawProjectPath, Platform))
+			{
+				OutReason = "overriden by target platform";
+				return true;
+			}
+
+			// Read the project descriptor, and find all the plugins available to this project
+			ProjectDescriptor Project = ProjectDescriptor.FromFile(RawProjectPath);
+
+			// Enumerate all the available plugins
+			Dictionary<string, PluginInfo> AllPlugins = Plugins.ReadAvailablePlugins(Unreal.EngineDirectory, DirectoryReference.FromFile(RawProjectPath), new List<DirectoryReference>()).ToDictionary(x => x.Name, x => x, StringComparer.OrdinalIgnoreCase);
+
+			// find if there are any plugins enabled or disabled which differ from the default
+			string? Reason;
+			if (RequiresTempTargetForCodePlugin(Project, Platform, Configuration, TargetType, AllPlugins, out Reason, Logger))
+			{
+				OutReason = Reason;
+				return true;
+			}
+
+			OutReason = null;
+			return false;
+		}
+
+		/// <summary>
+		/// NOTE: This function must mirror FPluginManager::RequiresTempTargetForCodePlugin
+		/// </summary>
+		static bool RequiresTempTargetForCodePlugin(ProjectDescriptor ProjectDescriptor, UnrealTargetPlatform Platform, UnrealTargetConfiguration Configuration, TargetType TargetType, Dictionary<string, PluginInfo> AllPlugins, out string? OutReason, ILogger Logger)
+		{
+			PluginReferenceDescriptor? MissingPlugin;
+
+			HashSet<string> ProjectCodePlugins = new HashSet<string>(StringComparer.OrdinalIgnoreCase);
+			if (!GetCodePluginsForTarget(ProjectDescriptor, Platform, Configuration, TargetType, ProjectCodePlugins, AllPlugins, out MissingPlugin, Logger))
+			{
+				OutReason = String.Format("{0} plugin is referenced by target but not found", MissingPlugin!.Name);
+				return true;
+			}
+
+			HashSet<string> DefaultCodePlugins = new HashSet<string>(StringComparer.OrdinalIgnoreCase);
+			if (!GetCodePluginsForTarget(null, Platform, Configuration, TargetType, DefaultCodePlugins, AllPlugins, out MissingPlugin, Logger))
+			{
+				OutReason = String.Format("{0} plugin is referenced by the default target but not found", MissingPlugin!.Name);
+				return true;
+			}
+
+			foreach (string ProjectCodePlugin in ProjectCodePlugins)
+			{
+				if (!DefaultCodePlugins.Contains(ProjectCodePlugin))
+				{
+					OutReason = String.Format("{0} plugin is enabled", ProjectCodePlugin);
+					return true;
+				}
+			}
+
+			foreach (string DefaultCodePlugin in DefaultCodePlugins)
+			{
+				if (!ProjectCodePlugins.Contains(DefaultCodePlugin))
+				{
+					OutReason = String.Format("{0} plugin is disabled", DefaultCodePlugin);
+					return true;
+				}
+			}
+
+			OutReason = null;
+			return false;
+		}
+
+		/// <summary>
+		/// NOTE: This function must mirror FPluginManager::GetCodePluginsForTarget
+		/// </summary>
+		static bool GetCodePluginsForTarget(ProjectDescriptor? ProjectDescriptor, UnrealTargetPlatform Platform, UnrealTargetConfiguration Configuration, TargetType TargetType, HashSet<string> CodePluginNames, Dictionary<string, PluginInfo> AllPlugins, out PluginReferenceDescriptor? OutMissingPlugin, ILogger Logger)
+		{
+			bool bLoadPluginsForTargetPlatforms = (TargetType == TargetType.Editor);
+
+			// Map of all enabled plugins
+			Dictionary<string, PluginInfo> EnabledPlugins = new Dictionary<string, PluginInfo>(StringComparer.OrdinalIgnoreCase);
+
+			// Keep a set of all the plugin names that have been configured. We read configuration data from different places, but only configure a plugin from the first place that it's referenced.
+			HashSet<string> ConfiguredPluginNames = new HashSet<string>(StringComparer.OrdinalIgnoreCase);
+
+			bool bAllowEnginePluginsEnabledByDefault = true;
+
+			// Find all the plugin references in the project file
+			if (ProjectDescriptor != null)
+			{
+				bAllowEnginePluginsEnabledByDefault = !ProjectDescriptor.DisableEnginePluginsByDefault;
+				if (ProjectDescriptor.Plugins != null)
+				{
+					// Copy the plugin references, since we may modify the project if any plugins are missing
+					foreach (PluginReferenceDescriptor PluginReference in ProjectDescriptor.Plugins)
+					{
+						if (!ConfiguredPluginNames.Contains(PluginReference.Name))
+						{
+							PluginReferenceDescriptor? MissingPlugin;
+							if (!ConfigureEnabledPluginForTarget(PluginReference, ProjectDescriptor, null, Platform, Configuration, TargetType, bLoadPluginsForTargetPlatforms, AllPlugins, EnabledPlugins, out MissingPlugin, Logger))
+							{
+								OutMissingPlugin = MissingPlugin;
+								return false;
+							}
+							ConfiguredPluginNames.Add(PluginReference.Name);
+						}
+					}
+				}
+			}
+
+			// Add the plugins which are enabled by default
+			foreach (KeyValuePair<string, PluginInfo> PluginPair in AllPlugins)
+			{
+				if (PluginPair.Value.IsEnabledByDefault(bAllowEnginePluginsEnabledByDefault) && !ConfiguredPluginNames.Contains(PluginPair.Key))
+				{
+					PluginReferenceDescriptor? MissingPlugin;
+					if (!ConfigureEnabledPluginForTarget(new PluginReferenceDescriptor(PluginPair.Key, null, true), ProjectDescriptor, null, Platform, Configuration, TargetType, bLoadPluginsForTargetPlatforms, AllPlugins, EnabledPlugins, out MissingPlugin, Logger))
+					{
+						OutMissingPlugin = MissingPlugin;
+						return false;
+					}
+					ConfiguredPluginNames.Add(PluginPair.Key);
+				}
+			}
+
+			// Figure out which plugins have code 
+			bool bBuildDeveloperTools = (TargetType == TargetType.Editor || TargetType == TargetType.Program || (Configuration != UnrealTargetConfiguration.Test && Configuration != UnrealTargetConfiguration.Shipping));
+			bool bRequiresCookedData = (TargetType != TargetType.Editor);
+			foreach (KeyValuePair<string, PluginInfo> Pair in EnabledPlugins)
+			{
+				if (Pair.Value.Descriptor.Modules != null)
+				{
+					foreach (ModuleDescriptor Module in Pair.Value.Descriptor.Modules)
+					{
+						if (Module.IsCompiledInConfiguration(Platform, Configuration, "", TargetType, bBuildDeveloperTools, bRequiresCookedData))
+						{
+							CodePluginNames.Add(Pair.Key);
+							break;
+						}
+					}
+				}
+			}
+
+			OutMissingPlugin = null;
+			return true;
+		}
+
+		/// <summary>
+		/// NOTE: This function should mirror FPluginManager::ConfigureEnabledPluginForTarget
+		/// </summary>
+		static bool ConfigureEnabledPluginForTarget(PluginReferenceDescriptor FirstReference, ProjectDescriptor? ProjectDescriptor, string? TargetName, UnrealTargetPlatform Platform, UnrealTargetConfiguration Configuration, TargetType TargetType, bool bLoadPluginsForTargetPlatforms, Dictionary<string, PluginInfo> AllPlugins, Dictionary<string, PluginInfo> EnabledPlugins, out PluginReferenceDescriptor? OutMissingPlugin, ILogger Logger)
+		{
+			if (!EnabledPlugins.ContainsKey(FirstReference.Name))
+			{
+				// Set of plugin names we've added to the queue for processing
+				HashSet<string> NewPluginNames = new HashSet<string>(StringComparer.OrdinalIgnoreCase);
+				NewPluginNames.Add(FirstReference.Name);
+
+				// Queue of plugin references to consider
+				List<PluginReferenceDescriptor> NewPluginReferences = new List<PluginReferenceDescriptor>();
+				NewPluginReferences.Add(FirstReference);
+
+				// Loop through the queue of plugin references that need to be enabled, queuing more items as we go
+				for (int Idx = 0; Idx < NewPluginReferences.Count; Idx++)
+				{
+					PluginReferenceDescriptor Reference = NewPluginReferences[Idx];
+
+					// Check if the plugin is required for this platform
+					if (!Reference.IsEnabledForPlatform(Platform) || !Reference.IsEnabledForTargetConfiguration(Configuration) || !Reference.IsEnabledForTarget(TargetType))
+					{
+						Logger.LogDebug("Ignoring plugin '{Arg0}' for platform/configuration", Reference.Name);
+						continue;
+					}
+
+					// Check if the plugin is required for this platform
+					if (!bLoadPluginsForTargetPlatforms && !Reference.IsSupportedTargetPlatform(Platform))
+					{
+						Logger.LogDebug("Ignoring plugin '{Arg0}' due to unsupported platform", Reference.Name);
+						continue;
+					}
+
+					// Find the plugin being enabled
+					PluginInfo? Plugin;
+					if (!AllPlugins.TryGetValue(Reference.Name, out Plugin))
+					{
+						// Ignore any optional plugins
+						if (Reference.bOptional)
+						{
+							Logger.LogDebug("Ignored optional reference to '%s' plugin; plugin was not found.", Reference.Name);
+							continue;
+						}
+
+						// Add it to the missing list
+						OutMissingPlugin = Reference;
+						return false;
+					}
+
+					// Check the plugin supports this platform
+					if (!bLoadPluginsForTargetPlatforms && !Plugin.Descriptor.SupportsTargetPlatform(Platform))
+					{
+						Logger.LogDebug("Ignoring plugin '{Arg0}' due to unsupported platform in plugin descriptor", Reference.Name);
+						continue;
+					}
+
+					// Check that this plugin supports the current program
+					if (TargetType == TargetType.Program && !Plugin.Descriptor.SupportedPrograms!.Contains(TargetName))
+					{
+						Logger.LogDebug("Ignoring plugin '{Arg0}' due to absence from the supported programs list", Reference.Name);
+						continue;
+					}
+
+					// Add references to all its dependencies
+					if (Plugin.Descriptor.Plugins != null)
+					{
+						foreach (PluginReferenceDescriptor NextReference in Plugin.Descriptor.Plugins)
+						{
+							if (!EnabledPlugins.ContainsKey(NextReference.Name) && !NewPluginNames.Contains(NextReference.Name))
+							{
+								NewPluginNames.Add(NextReference.Name);
+								NewPluginReferences.Add(NextReference);
+							}
+						}
+					}
+
+					// Add the plugin
+					EnabledPlugins.Add(Plugin.Name, Plugin);
+				}
+			}
+
+			OutMissingPlugin = null;
+			return true;
+		}
+		#endregion
 	}
 }
