AppXManifest localized strings are parsed during initialization, not read from the ini file each time.

 - TargetName is also cached from initialization, rather than being passed through everywhere.
 - some member variables renamed to clear up confusion over which culture type they refer to - UE or AppX

#rnx

diff --git a/Engine/Source/Programs/UnrealBuildTool/Platform/Windows/AppXManifestGeneratorBase.cs b/Engine/Source/Programs/UnrealBuildTool/Platform/Windows/AppXManifestGeneratorBase.cs
--- a/Engine/Source/Programs/UnrealBuildTool/Platform/Windows/AppXManifestGeneratorBase.cs
+++ b/Engine/Source/Programs/UnrealBuildTool/Platform/Windows/AppXManifestGeneratorBase.cs
@@ -49,15 +49,21 @@ namespace UnrealBuildTool
 		protected ConfigHierarchy? GameIni;
 
 		/// the default culture to use
-		protected string? DefaultCulture;
+		protected string? DefaultAppXCultureId;
 
 		/// all cultures that will be staged
-		protected List<string>? CulturesToStage;
+		protected List<string> SelectedAppXCultureIds = new();
+
+		/// localization dictionary - AppX CultureId to ConfigKey/Localized String pairs. Empty CultureId represents the default strings
+		protected Dictionary<string,Dictionary<string,string>>? PerAppXCultureIdStrings;
+
+		/// lookup table for UE's CultureToStage StageId to AppX CultureId
+		protected Dictionary<string,string> UEStageIdToAppXCultureId = new Dictionary<string, string>();
 
 		/// resource writer for the default culture
 		protected UEResXWriter? DefaultResourceWriter;
 
-		/// resource writers for each culture
+		/// resource writers for each AppX culture
 		protected Dictionary<string, UEResXWriter>? PerCultureResourceWriters;
 
 		/// the platform to generate the manifest for
@@ -66,6 +72,9 @@ namespace UnrealBuildTool
 		/// project file to use
 		protected FileReference? ProjectFile;
 
+		/// target name to use
+		protected string? TargetName;
+
 		/// output path - where the manifest will be created
 		protected string? OutputPath;
 
@@ -320,6 +329,70 @@ namespace UnrealBuildTool
 
 
 
+		/// <summary>
+		/// Create all the localization data. Returns whether there is any per-culture data set up
+		/// </summary>
+		protected virtual bool BuildLocalizationData()
+		{
+			bool bHasPerCultureResources = false;
+
+			// reset per-culture strings and make sure the default culture entry exists
+			PerAppXCultureIdStrings = new ();
+			PerAppXCultureIdStrings[""] = new();
+
+			// add all default strings
+			if (EngineIni!.GetString(IniSection_PlatformTargetSettings, "CultureStringResources", out string DefaultCultureScratchValue) && ConfigHierarchy.TryParse(DefaultCultureScratchValue, out Dictionary<string,string>? DefaultStrings))
+			{
+				PerAppXCultureIdStrings[""] = DefaultStrings;
+			}
+
+			// add per culture strings
+			if (EngineIni.GetArray(IniSection_PlatformTargetSettings, "PerCultureResources", out List<string>? PerCultureResources))
+			{
+				bHasPerCultureResources = true;
+
+				foreach (string CultureResources in PerCultureResources)
+				{
+					if (!ConfigHierarchy.TryParse(CultureResources, out Dictionary<string, string>? CultureProperties)
+						|| !CultureProperties.ContainsKey("CultureStringResources")
+						|| !CultureProperties.ContainsKey("CultureId")
+						|| !CultureProperties.ContainsKey("StageId"))
+					{
+						Logger.LogWarning("Invalid per-culture resource value: {Culture}", CultureResources);
+						continue;
+					}
+
+					string StageId = CultureProperties["StageId"];
+					if (string.IsNullOrEmpty(StageId))
+					{
+						Logger.LogWarning("Missing StageId value: {Culture}", CultureResources);
+						continue;
+					}
+
+					string CultureId = CultureProperties["CultureId"];
+					if (string.IsNullOrEmpty(CultureId))
+					{
+						Logger.LogWarning("Missing CultureId value: {Culture}", CultureResources);
+						continue;
+					}
+
+					// read culture strings
+					if (!ConfigHierarchy.TryParse(CultureProperties["CultureStringResources"], out Dictionary<string, string>? CultureStringResources))
+					{
+						Logger.LogError("Invalid culture string resources: \"{Culture}\". Unable to add resource entry.", CultureResources);
+						continue;
+					}
+					PerAppXCultureIdStrings[CultureId] = CultureStringResources;
+
+					// build culture lookup table
+					UEStageIdToAppXCultureId[StageId] = CultureId;
+				}
+			}
+
+
+			return bHasPerCultureResources;
+		}
+
 
 		private bool RemoveStaleResourceFiles()
 		{
@@ -427,7 +500,7 @@ namespace UnrealBuildTool
 			AddFileReference(SourceFilePath, Path.Combine(TargetPath, ResourceFileName), bIsGeneratedFile:false);
 
 			// Copy all per-culture resource files
-			foreach( string CultureId in CulturesToStage!)
+			foreach( string CultureId in SelectedAppXCultureIds!)
 			{
 				bFileExists = FindResourceBinaryFile(out SourceFilePath, ResourceFileName, CultureId, AllowEngineFallback);
 				if (bFileExists)
@@ -551,64 +624,32 @@ namespace UnrealBuildTool
 		/// <summary>
 		/// Adds the given string to the culture string writers
 		/// </summary>
-		protected string AddResourceEntry(string ResourceEntryName, string ConfigKey, string GenericINISection, string GenericINIKey, string DefaultValue, string ValueSuffix = "")
+		protected string AddResourceEntry(string ResourceEntryName, string ConfigKey, string DefaultValue, string ValueSuffix = "")
 		{
-			string? ConfigScratchValue = null;
-
 			// Get the default culture value
-			string DefaultCultureScratchValue;
-			if (EngineIni!.GetString(IniSection_PlatformTargetSettings, "CultureStringResources", out DefaultCultureScratchValue))
+			string? DefaultCultureString = null;
+			if (PerAppXCultureIdStrings!.TryGetValue("", out Dictionary<string,string>? DefaultCultureValues))
 			{
-				Dictionary<string, string>? Values;
-				if (!ConfigHierarchy.TryParse(DefaultCultureScratchValue, out Values))
+				if (DefaultCultureValues!.TryGetValue(ConfigKey, out string? Value) && Value != null)
 				{
-					Logger.LogError("Invalid default culture string resources: \"{Culture}\". Unable to add resource entry.", DefaultCultureScratchValue);
-					return "";
+					DefaultCultureString = Value;
 				}
-
-				ConfigScratchValue = Values[ConfigKey];
 			}
-
-			if (string.IsNullOrEmpty(ConfigScratchValue))
+			if (string.IsNullOrEmpty(DefaultCultureString))
 			{
-				// No platform specific value is provided. Use the generic config or default value
-				ConfigScratchValue = ReadIniString(GenericINIKey, GenericINISection, DefaultValue);
+				DefaultCultureString = DefaultValue;
 			}
+			DefaultResourceWriter!.AddResource(ResourceEntryName, DefaultCultureString + ValueSuffix);
 
-			DefaultResourceWriter!.AddResource(ResourceEntryName, ConfigScratchValue + ValueSuffix);
-
-			// Find the default value
-			List<string>? PerCultureValues;
-			if (EngineIni.GetArray(IniSection_PlatformTargetSettings, "PerCultureResources", out PerCultureValues))
+			// Get the localized culture values
+			foreach (string CultureId in SelectedAppXCultureIds!)
 			{
-				foreach (string CultureCombinedValues in PerCultureValues)
+				if (PerAppXCultureIdStrings.TryGetValue(CultureId, out Dictionary<string,string>? CultureValues))
 				{
-					Dictionary<string, string>? SeparatedCultureValues;
-					if (!ConfigHierarchy.TryParse(CultureCombinedValues, out SeparatedCultureValues)
-						|| !SeparatedCultureValues.ContainsKey("CultureStringResources")
-						|| !SeparatedCultureValues.ContainsKey("CultureId"))
+					if (CultureValues.TryGetValue(ConfigKey, out string? CultureString) && !string.IsNullOrEmpty(CultureString) )
 					{
-						Logger.LogError("Invalid per-culture resource: \"{Culture}\". Unable to add resource entry.", CultureCombinedValues);
-						continue;
-					}
-
-					var CultureId = SeparatedCultureValues["CultureId"];
-					if (CulturesToStage!.Contains(CultureId))
-					{
-						Dictionary<string, string>? CultureStringResources;
-						if (!ConfigHierarchy.TryParse(SeparatedCultureValues["CultureStringResources"], out CultureStringResources))
-						{
-							Logger.LogError("Invalid culture string resources: \"{Culture}\". Unable to add resource entry.", CultureCombinedValues);
-							continue;
-						}
-
-						var Value = CultureStringResources[ConfigKey];
-
-						if (CulturesToStage.Contains(CultureId) && !string.IsNullOrEmpty(Value))
-						{
-							var Writer = PerCultureResourceWriters![CultureId];
-							Writer.AddResource(ResourceEntryName, Value + ValueSuffix);
-						}
+						var Writer = PerCultureResourceWriters![CultureId];
+						Writer.AddResource(ResourceEntryName, CultureString + ValueSuffix);
 					}
 				}
 			}
@@ -616,37 +657,7 @@ namespace UnrealBuildTool
 			return "ms-resource:" + ResourceEntryName;
 		}
 
-		/// <summary>
-		/// Adds an additional string to all culture resource writers
-		/// </summary>
-		protected string AddExternalResourceEntry(string ResourceEntryName, string DefaultValue, Dictionary<string,string> CultureIdToCultureValues)
-		{
-			DefaultResourceWriter!.AddResource(ResourceEntryName, DefaultValue);
-
-			foreach( KeyValuePair<string,string> CultureIdToCultureValue in CultureIdToCultureValues)
-			{
-				var Writer = PerCultureResourceWriters![CultureIdToCultureValue.Key];
-				Writer.AddResource(ResourceEntryName, CultureIdToCultureValue.Value);
-			}
-
-			return "ms-resource:" + ResourceEntryName;
-		}
-
-		/// <summary>
-		/// Adds a debug-only string to all resource writers
-		/// </summary>
-		protected string AddDebugResourceString(string ResourceEntryName, string Value)
-		{
-			DefaultResourceWriter!.AddResource(ResourceEntryName, Value);
 
-			foreach (var CultureId in CulturesToStage!)
-			{
-				var Writer = PerCultureResourceWriters![CultureId];
-				Writer.AddResource(ResourceEntryName, Value);
-			}
-
-			return "ms-resource:" + ResourceEntryName;
-		}
 
 		/// <summary>
 		/// Get the XName from a given string and schema pair
@@ -661,13 +672,13 @@ namespace UnrealBuildTool
 		/// </summary>
 		protected XElement GetResources()
 		{
-			List<string> ResourceCulturesList = new List<string>(CulturesToStage!);
+			List<string> ResourceCulturesList = new List<string>(SelectedAppXCultureIds!);
 			// Move the default culture to the front of the list
-			ResourceCulturesList.Remove(DefaultCulture!);
-			ResourceCulturesList.Insert(0, DefaultCulture!);
+			ResourceCulturesList.Remove(DefaultAppXCultureId!);
+			ResourceCulturesList.Insert(0, DefaultAppXCultureId!);
 
 			// Check that we have a valid number of cultures
-			if (CulturesToStage!.Count < 1 || CulturesToStage.Count >= MaxResourceEntries)
+			if (SelectedAppXCultureIds!.Count < 1 || SelectedAppXCultureIds.Count >= MaxResourceEntries)
 			{
 				Logger.LogWarning("Incorrect number of cultures to stage. There must be between 1 and {MaxCultures} cultures selected.", MaxResourceEntries);
 			}
@@ -682,7 +693,7 @@ namespace UnrealBuildTool
 		/// <summary>
 		/// Get the package identity name string
 		/// </summary>
-		protected string GetIdentityPackageName(string? TargetName)
+		protected string GetIdentityPackageName()
 		{
             // Read the PackageName from config
 			var DefaultName = (ProjectFile != null) ? ProjectFile.GetFileNameWithoutAnyExtensions() : (TargetName ?? "DefaultUEProject");
@@ -734,9 +745,9 @@ namespace UnrealBuildTool
 		/// <summary>
 		/// Get the package identity element
 		/// </summary>
-		protected XElement GetIdentity(string? TargetName, out string IdentityName)
+		protected XElement GetIdentity(out string IdentityName)
         {
-            string PackageName = GetIdentityPackageName(TargetName);
+            string PackageName = GetIdentityPackageName();
             string PublisherName = GetIdentityPublisherName();
             string? VersionNumber = GetIdentityVersionNumber();
 
@@ -784,7 +795,7 @@ namespace UnrealBuildTool
 		/// <summary>
 		/// Return the entire manifest element
 		/// </summary>
-		protected abstract XElement GetManifest(Dictionary<UnrealTargetConfiguration,string> InExecutablePairs, string? TargetName, out string IdentityName);
+		protected abstract XElement GetManifest(Dictionary<UnrealTargetConfiguration,string> InExecutablePairs, out string IdentityName);
 
 		/// <summary>
 		/// Perform any platform-specific processing on the manifest before it is saved
@@ -818,6 +829,7 @@ namespace UnrealBuildTool
 
 			OutputPath = InOutputPath;
 			IntermediatePath = InIntermediatePath;
+			TargetName = InTargetName;
 			ProjectFile = InProjectFile;
 			ManifestFiles = new Dictionary<string, string>();
 
@@ -827,56 +839,60 @@ namespace UnrealBuildTool
 			EngineIni = ConfigCache.ReadHierarchy(ConfigHierarchyType.Engine, DirectoryReference.FromFile(InProjectFile), ConfigPlatform);
 
 			// Load and verify/clean culture list
+			List<string>? SelectedUECultureIds;
+			string DefaultUECultureId;
+			GameIni.GetArray("/Script/UnrealEd.ProjectPackagingSettings", "CulturesToStage", out SelectedUECultureIds);
+			GameIni.GetString("/Script/UnrealEd.ProjectPackagingSettings", "DefaultCulture", out DefaultUECultureId);
+			if (SelectedUECultureIds == null || SelectedUECultureIds.Count < 1)
 			{
-				List<string>? CulturesToStageWithDuplicates;
-				GameIni.GetArray("/Script/UnrealEd.ProjectPackagingSettings", "CulturesToStage", out CulturesToStageWithDuplicates);
-				GameIni.GetString("/Script/UnrealEd.ProjectPackagingSettings", "DefaultCulture", out DefaultCulture);
-				if (CulturesToStageWithDuplicates == null || CulturesToStageWithDuplicates.Count < 1)
-				{
-					Logger.LogError("At least one culture must be selected to stage.");
-					return null;
-				}
-
-				CulturesToStage = CulturesToStageWithDuplicates.Distinct().ToList();
+				Logger.LogError("At least one culture must be selected to stage.");
+				return null;
 			}
-			if (DefaultCulture == null || DefaultCulture.Length < 1)
+			SelectedUECultureIds = SelectedUECultureIds.Distinct().ToList();
+
+			if (DefaultUECultureId == null || DefaultUECultureId.Length < 1)
 			{
-				DefaultCulture = CulturesToStage[0];
-				Logger.LogWarning("A default culture must be selected to stage. Using {DefaultCulture}.", DefaultCulture);
+				DefaultUECultureId = SelectedUECultureIds[0];
+				Logger.LogWarning("A default culture must be selected to stage. Using {DefaultCulture}.", DefaultUECultureId);
 			}
-			if (!CulturesToStage.Contains(DefaultCulture))
+			if (!SelectedUECultureIds.Contains(DefaultUECultureId))
 			{
-				DefaultCulture = CulturesToStage[0];
-				Logger.LogWarning("The default culture must be one of the staged cultures. Using {DefaultCulture}.", DefaultCulture);
+				DefaultUECultureId = SelectedUECultureIds[0];
+				Logger.LogWarning("The default culture must be one of the staged cultures. Using {DefaultCulture}.", DefaultUECultureId);
 			}
 
-			List<string>? PerCultureValues;
-			if (EngineIni.GetArray(IniSection_PlatformTargetSettings, "PerCultureResources", out PerCultureValues))
-			{
-				foreach (string CultureCombinedValues in PerCultureValues)
+			bool bHasLocalizationData = BuildLocalizationData();
+			if (bHasLocalizationData)
+			{			
+				// generate the list of AppX cultures to stage
+				SelectedAppXCultureIds = new();
+				foreach (string UEStageId in SelectedUECultureIds)
 				{
-					Dictionary<string, string>? SeparatedCultureValues;
-					if (!ConfigHierarchy.TryParse(CultureCombinedValues, out SeparatedCultureValues))
+					if (UEStageIdToAppXCultureId.TryGetValue(UEStageId, out string? AppXCultureId) && !string.IsNullOrEmpty(AppXCultureId) )
 					{
-						Logger.LogWarning("Invalid per-culture resource value: {Culture}", CultureCombinedValues);
-						continue;
+						SelectedAppXCultureIds.Add(AppXCultureId);
 					}
-
-					string StageId = SeparatedCultureValues["StageId"];
-					int CultureIndex = CulturesToStage.FindIndex(x => x == StageId);
-					if (CultureIndex >= 0)
+					else
 					{
-						CulturesToStage[CultureIndex] = SeparatedCultureValues["CultureId"];
-						if (DefaultCulture == StageId)
-						{
-							DefaultCulture = SeparatedCultureValues["CultureId"];
-						}
+						// use the culture directly - no remapping required
+						UEStageIdToAppXCultureId[UEStageId] = UEStageId;
+						SelectedAppXCultureIds.Add(UEStageId);
 					}
 				}
+
+				// look up the default AppX culture
+				if (UEStageIdToAppXCultureId.TryGetValue(DefaultUECultureId, out string? NewDefaultCultureId))
+				{
+					DefaultAppXCultureId = NewDefaultCultureId;
+				}
+				if (string.IsNullOrEmpty(DefaultAppXCultureId))
+				{
+					throw new Exception($"unable to find default culture {DefaultUECultureId} in the staged culture list");
+				}
 			}
-			// Only warn if shipping, we can run without translated cultures they're just needed for cert
 			else if (InExecutablePairs.ContainsKey(UnrealTargetConfiguration.Shipping))
 			{
+				// Warn in shipping, we can run without translated cultures they're just needed for cert
 				Logger.LogInformation("Staged culture mappings not setup in the editor. See Per Culture Resources in the {Platform} Target Settings.", Platform.ToString() );
 			}
 
@@ -895,29 +911,23 @@ namespace UnrealBuildTool
 
 			// Construct the ResXWriters for each culture
 			PerCultureResourceWriters = new Dictionary<string, UEResXWriter>();
-			foreach (string Culture in CulturesToStage)
+			foreach (string AppXCultureId in SelectedAppXCultureIds)
 			{
-				string IntermediateStringResourcePath = Path.Combine(IntermediateResourceDirectory, Culture);
+				string IntermediateStringResourcePath = Path.Combine(IntermediateResourceDirectory, AppXCultureId);
 				string IntermediateStringResourceFile = Path.Combine(IntermediateStringResourcePath, "resources.resw");
 				if (!CreateCheckDirectory(IntermediateStringResourcePath, Logger))
 				{
-					Logger.LogWarning("Culture {Culture} resources not staged.", Culture);
-					CulturesToStage.Remove(Culture);
-					if (Culture == DefaultCulture)
-					{
-						DefaultCulture = CulturesToStage[0];
-						Logger.LogWarning("Default culture skipped. Using {DefaultCulture} as default culture.", DefaultCulture);
-					}
-					continue;
+					Logger.LogError("Could not create directory {IntDir}.", IntermediateStringResourcePath);
+					return null;
 				}
-				PerCultureResourceWriters.Add(Culture, new UEResXWriter(IntermediateStringResourceFile));
+				PerCultureResourceWriters.Add(AppXCultureId, new UEResXWriter(IntermediateStringResourceFile));
 			}
 
 
 
 			// Create the manifest document
 			string? IdentityName = null;
-			var ManifestXmlDocument = new XDocument(GetManifest(InExecutablePairs, InTargetName, out IdentityName));
+			var ManifestXmlDocument = new XDocument(GetManifest(InExecutablePairs, out IdentityName));
 
 			// Export manifest to the intermediate directory then compare the contents to any existing target manifest
 			// and replace if there are differences.
@@ -956,7 +966,7 @@ namespace UnrealBuildTool
 			{
 				// Create resource index configuration
 				string ResourceConfigFile = Path.Combine(IntermediatePath, "priconfig.xml");
-				RunMakePri("createconfig /cf \"" + ResourceConfigFile + "\" /dq " + DefaultCulture + " /o " + GetMakePriExtraCommandLine());
+				RunMakePri("createconfig /cf \"" + ResourceConfigFile + "\" /dq " + DefaultAppXCultureId + " /o " + GetMakePriExtraCommandLine());
 
 				// Load the new resource index configuration
 				XmlDocument PriConfig = new XmlDocument();
