Adding some additional hook points to AppXManifest generation for specializations.

#rnx

diff --git a/Engine/Source/Programs/UnrealBuildTool/Platform/Windows/AppXManifestGeneratorBase.cs b/Engine/Source/Programs/UnrealBuildTool/Platform/Windows/AppXManifestGeneratorBase.cs
--- a/Engine/Source/Programs/UnrealBuildTool/Platform/Windows/AppXManifestGeneratorBase.cs
+++ b/Engine/Source/Programs/UnrealBuildTool/Platform/Windows/AppXManifestGeneratorBase.cs
@@ -693,7 +693,7 @@ namespace UnrealBuildTool
 		/// <summary>
 		/// Get the package identity name string
 		/// </summary>
-		protected string GetIdentityPackageName()
+		protected virtual string GetIdentityPackageName()
 		{
             // Read the PackageName from config
 			var DefaultName = (ProjectFile != null) ? ProjectFile.GetFileNameWithoutAnyExtensions() : (TargetName ?? "DefaultUEProject");
@@ -718,7 +718,7 @@ namespace UnrealBuildTool
 		/// <summary>
 		/// Get the publisher name string
 		/// </summary>
-		protected string GetIdentityPublisherName()
+		protected virtual string GetIdentityPublisherName()
 		{
             var PublisherName = GetConfigString("PublisherName", "CompanyDistinguishedName", "CN=NoPublisher");
 			return PublisherName;
@@ -727,7 +727,7 @@ namespace UnrealBuildTool
 		/// <summary>
 		/// Get the package version string
 		/// </summary>
-		protected string? GetIdentityVersionNumber()
+		protected virtual string? GetIdentityVersionNumber()
 		{
             var VersionNumber = GetConfigString("PackageVersion", "ProjectVersion", "1.0.0.0");
             VersionNumber = ValidatePackageVersion(VersionNumber);
@@ -804,6 +804,13 @@ namespace UnrealBuildTool
         {
 		}
 
+		/// <summary>
+		/// Perform any additional initialization once all parameters and configuration are ready
+		/// </summary>
+		protected virtual void PostConfigurationInit()
+		{
+		}
+
 		/// <summary>
 		/// Create a manifest and return the list of modified files
 		/// </summary>
@@ -837,6 +844,7 @@ namespace UnrealBuildTool
 			// values in either game or engine settings, so we'll keep both.
 			GameIni = ConfigCache.ReadHierarchy(ConfigHierarchyType.Game, DirectoryReference.FromFile(InProjectFile), ConfigPlatform);
 			EngineIni = ConfigCache.ReadHierarchy(ConfigHierarchyType.Engine, DirectoryReference.FromFile(InProjectFile), ConfigPlatform);
+			PostConfigurationInit();
 
 			// Load and verify/clean culture list
 			List<string>? SelectedUECultureIds;
