UnrealBuildTool: Add .suppressed.lib to link action delete list

#rnx

diff --git a/Engine/Source/Programs/UnrealBuildTool/Platform/Windows/VCToolChain.cs b/Engine/Source/Programs/UnrealBuildTool/Platform/Windows/VCToolChain.cs
--- a/Engine/Source/Programs/UnrealBuildTool/Platform/Windows/VCToolChain.cs
+++ b/Engine/Source/Programs/UnrealBuildTool/Platform/Windows/VCToolChain.cs
@@ -2440,6 +2440,14 @@ namespace UnrealBuildTool
 				LinkAction.DeleteItems.UnionWith(LinkAction.ProducedItems.Where(x => x.Location.HasExtension(".pdb")));
 			}
 
+			// Delete any .suppressed.lib files before building, even if they're not tracked
+			if (ImportLibraryFilePath.GetFileName().EndsWith(".suppressed.lib"))
+			{
+				FileReference ExportFilePath = ImportLibraryFilePath.ChangeExtension(".exp");
+				LinkAction.DeleteItems.Add(FileItem.GetItemByFileReference(ImportLibraryFilePath));
+				LinkAction.DeleteItems.Add(FileItem.GetItemByFileReference(ExportFilePath));
+			}
+
 			// Tell the action that we're building an import library here and it should conditionally be
 			// ignored as a prerequisite for other actions
 			LinkAction.bProducesImportLibrary = bBuildImportLibraryOnly || LinkEnvironment.bIsBuildingDLL;
