fix for when there is no per-culture settings specified for the current project.

#rnx

diff --git a/Engine/Source/Programs/UnrealBuildTool/Platform/Windows/AppXManifestGeneratorBase.cs b/Engine/Source/Programs/UnrealBuildTool/Platform/Windows/AppXManifestGeneratorBase.cs
--- a/Engine/Source/Programs/UnrealBuildTool/Platform/Windows/AppXManifestGeneratorBase.cs
+++ b/Engine/Source/Programs/UnrealBuildTool/Platform/Windows/AppXManifestGeneratorBase.cs
@@ -870,38 +870,34 @@ namespace UnrealBuildTool
 			}
 
 			bool bHasLocalizationData = BuildLocalizationData();
-			if (bHasLocalizationData)
-			{			
-				// generate the list of AppX cultures to stage
-				SelectedAppXCultureIds = new();
-				foreach (string UEStageId in SelectedUECultureIds)
-				{
-					if (UEStageIdToAppXCultureId.TryGetValue(UEStageId, out string? AppXCultureId) && !string.IsNullOrEmpty(AppXCultureId) )
-					{
-						SelectedAppXCultureIds.Add(AppXCultureId);
-					}
-					else
-					{
-						// use the culture directly - no remapping required
-						UEStageIdToAppXCultureId[UEStageId] = UEStageId;
-						SelectedAppXCultureIds.Add(UEStageId);
-					}
-				}
 
-				// look up the default AppX culture
-				if (UEStageIdToAppXCultureId.TryGetValue(DefaultUECultureId, out string? NewDefaultCultureId))
-				{
-					DefaultAppXCultureId = NewDefaultCultureId;
-				}
-				if (string.IsNullOrEmpty(DefaultAppXCultureId))
+			// generate the list of AppX cultures to stage
+			SelectedAppXCultureIds = new();
+			foreach (string UEStageId in SelectedUECultureIds)
+			{
+				if (!UEStageIdToAppXCultureId.TryGetValue(UEStageId, out string? AppXCultureId) || string.IsNullOrEmpty(AppXCultureId) )
 				{
-					throw new Exception($"unable to find default culture {DefaultUECultureId} in the staged culture list");
+					// use the culture directly - no remapping required
+					AppXCultureId = UEStageId;
+					UEStageIdToAppXCultureId[UEStageId] = AppXCultureId;
 				}
+
+				SelectedAppXCultureIds.Add(AppXCultureId);
+			}
+
+			// look up the default AppX culture
+			if (!UEStageIdToAppXCultureId.TryGetValue(DefaultUECultureId, out DefaultAppXCultureId) || string.IsNullOrEmpty(DefaultAppXCultureId))
+			{
+				// use the default culture directly - no remapping required
+				DefaultAppXCultureId = DefaultUECultureId;
+				UEStageIdToAppXCultureId[DefaultUECultureId] = DefaultAppXCultureId;
 			}
-			else if (InExecutablePairs.ContainsKey(UnrealTargetConfiguration.Shipping))
+
+
+			// Warn in shipping, we can run without translated cultures they're just needed for cert
+			if (!bHasLocalizationData && InExecutablePairs.ContainsKey(UnrealTargetConfiguration.Shipping))
 			{
-				// Warn in shipping, we can run without translated cultures they're just needed for cert
-				Logger.LogInformation("Staged culture mappings not setup in the editor. See Per Culture Resources in the {Platform} Target Settings.", Platform.ToString() );
+				Logger.LogInformation("Staged culture mappings not setup in the editor. See Per Culture Resources in the {Platform} Target Settings.", Platform.ToString());
 			}
 
 			// Clean out the resources intermediate path so that we know there are no stale binary files.
