Fixing issue where artifacts were not being seen.

diff --git a/Engine/Source/Programs/UnrealBuildTool/Executors/ImmediateActionQueue.cs b/Engine/Source/Programs/UnrealBuildTool/Executors/ImmediateActionQueue.cs
--- a/Engine/Source/Programs/UnrealBuildTool/Executors/ImmediateActionQueue.cs
+++ b/Engine/Source/Programs/UnrealBuildTool/Executors/ImmediateActionQueue.cs
@@ -239,6 +239,11 @@ namespace UnrealBuildTool
 		/// </summary>
 		public bool IsDone => _doneTaskSource.Task.IsCompleted;
 
+		/// <summary>
+		///  Action that can be to notify when artifacts have been read for an action
+		/// </summary>
+		public Action<LinkedAction>? OnArtifactsRead = null;
+
 		/// <summary>
 		/// Collection of available runners
 		/// </summary>
@@ -333,6 +338,7 @@ namespace UnrealBuildTool
 						bool success = await _actionArtifactCache!.CompleteActionFromCacheAsync(action, CancellationToken);
 						if (success)
 						{
+							OnArtifactsRead?.Invoke(action);
 							OnActionCompleted(action, success, s_copiedFromCacheResults);
 						}
 						else
