Add missing NotNullWhen to config string array TryParse helper function.

#rnx

diff --git a/Engine/Source/Programs/UnrealBuildTool/System/ConfigHierarchy.cs b/Engine/Source/Programs/UnrealBuildTool/System/ConfigHierarchy.cs
--- a/Engine/Source/Programs/UnrealBuildTool/System/ConfigHierarchy.cs
+++ b/Engine/Source/Programs/UnrealBuildTool/System/ConfigHierarchy.cs
@@ -769,7 +769,7 @@ namespace UnrealBuildTool
 		/// <param name="Line">Line of text to parse</param>
 		/// <param name="Array">Receives array for the config array</param>
 		/// <returns>True if an array was parsed, false otherwise</returns>
-		public static bool TryParse(string Line, out string[]? Array)
+		public static bool TryParse(string Line, [NotNullWhen(true)] out string[]? Array)
 		{
 			// Convert the string to a zero-terminated array, to make parsing easier.
 			char[] Chars = new char[Line.Length + 1];
