Fixed issue where action outputs were not having their file information reset.  This resulted in old hashes being used to create keys for other actions.

diff --git a/Engine/Source/Programs/UnrealBuildTool/Executors/ImmediateActionQueue.cs b/Engine/Source/Programs/UnrealBuildTool/Executors/ImmediateActionQueue.cs
--- a/Engine/Source/Programs/UnrealBuildTool/Executors/ImmediateActionQueue.cs
+++ b/Engine/Source/Programs/UnrealBuildTool/Executors/ImmediateActionQueue.cs
@@ -10,6 +10,7 @@ using System.Collections.Generic;
 using System.Linq;
 using System.Threading;
 using System.Threading.Tasks;
+using UnrealBuildBase;
 using UnrealBuildTool.Artifacts;
 
 namespace UnrealBuildTool
@@ -682,6 +683,16 @@ namespace UnrealBuildTool
 			int actionIndex = action.SortIndex;
 			int completedActions = 0;
 
+			// If we are finished, then invalidate the file info for all the produced items.  This needs to be
+			// done so the artifact system sees any compiled files.
+			if (status == ActionStatus.Finished)
+			{
+				foreach (FileItem output in action.ProducedItems)
+				{
+					output.ResetCachedInfo();
+				}
+			}
+
 			// Update the actions data
 			lock (Actions)
 			{
