Added the ability to log artifact cache misses.

diff --git a/Engine/Source/Programs/UnrealBuildTool/Artifacts/ActionArtifactCache.cs b/Engine/Source/Programs/UnrealBuildTool/Artifacts/ActionArtifactCache.cs
--- a/Engine/Source/Programs/UnrealBuildTool/Artifacts/ActionArtifactCache.cs
+++ b/Engine/Source/Programs/UnrealBuildTool/Artifacts/ActionArtifactCache.cs
@@ -25,6 +25,9 @@ namespace UnrealBuildTool.Artifacts
 		/// <inheritdoc/>
 		public bool EnableWrites { get; set; } = true;
 
+		/// <inheritdoc/>
+		public bool LogCacheMisses { get; set; } = false;
+
 		/// <inheritdoc/>
 		public IArtifactCache ArtifactCache { get; init; }
 
@@ -94,8 +97,18 @@ namespace UnrealBuildTool.Artifacts
 
 			ArtifactMapping[] mappings = await ArtifactCache.QueryArtifactMappingsAsync(new IoHash[] { key }, cancellationToken);
 
+			string actionDescription = string.Empty;
+			if (LogCacheMisses)
+			{
+				actionDescription = $"{(action.CommandDescription ?? action.CommandPath.GetFileNameWithoutExtension())} {action.StatusDescription}".Trim();
+			}
+
 			if (mappings.Length == 0)
 			{
+				if (LogCacheMisses)
+				{
+					_logger.LogInformation("Artifact Cache Miss: No artifact mappings found for {ActionDescription}", actionDescription);
+				}
 				return false;
 			}
 
@@ -109,11 +122,19 @@ namespace UnrealBuildTool.Artifacts
 					FileItem item = FileItem.GetItemByPath(name);
 					if (!item.Exists)
 					{
+						if (LogCacheMisses)
+						{
+							_logger.LogInformation("Artifact Cache Miss: Input file missing {ActionDescription}/{File}", actionDescription, item.FullName);
+						}
 						match = false;
 						break;
 					}
 					if (input.ContentHash != await _fileHasher.GetDigestAsync(item, cancellationToken))
 					{
+						if (LogCacheMisses)
+						{
+							_logger.LogInformation("Artifact Cache Miss: Content hash different {actionDescription}/{File}", actionDescription, item.FullName);
+						}
 						match = false;
 						break;
 					}
diff --git a/Engine/Source/Programs/UnrealBuildTool/Artifacts/Artifacts.cs b/Engine/Source/Programs/UnrealBuildTool/Artifacts/Artifacts.cs
--- a/Engine/Source/Programs/UnrealBuildTool/Artifacts/Artifacts.cs
+++ b/Engine/Source/Programs/UnrealBuildTool/Artifacts/Artifacts.cs
@@ -320,6 +320,11 @@ namespace UnrealBuildTool.Artifacts
 		/// </summary>
 		public bool EnableWrites { get; set; }
 
+		/// <summary>
+		/// If true, log all cache misses.  Defaults to false.
+		/// </summary>
+		public bool LogCacheMisses { get; set; }
+
 		/// <summary>
 		/// Complete an action from existing cached data
 		/// </summary>
diff --git a/Engine/Source/Programs/UnrealBuildTool/Configuration/BuildConfiguration.cs b/Engine/Source/Programs/UnrealBuildTool/Configuration/BuildConfiguration.cs
--- a/Engine/Source/Programs/UnrealBuildTool/Configuration/BuildConfiguration.cs
+++ b/Engine/Source/Programs/UnrealBuildTool/Configuration/BuildConfiguration.cs
@@ -219,6 +219,13 @@ namespace UnrealBuildTool
 		[CommandLine("-NoArtifactWrites", Value = "False")]
 		public bool bArtifactWrites = true;
 
+		/// <summary>
+		/// If true, log all artifact cache misses as informational messages
+		/// </summary>
+		[XmlConfigFile(Category = "BuildConfiguration")]
+		[CommandLine("-LogArtifactCacheMisses", Value = "True")]
+		public bool bLogArtifactCacheMisses = false;
+
 		/// <summary>
 		/// Location to store the artifacts.
 		/// </summary>
diff --git a/Engine/Source/Programs/UnrealBuildTool/Modes/BuildMode.cs b/Engine/Source/Programs/UnrealBuildTool/Modes/BuildMode.cs
--- a/Engine/Source/Programs/UnrealBuildTool/Modes/BuildMode.cs
+++ b/Engine/Source/Programs/UnrealBuildTool/Modes/BuildMode.cs
@@ -552,6 +552,7 @@ namespace UnrealBuildTool
 					actionArtifactCache = ActionArtifactCache.CreateHordeFileCache(artifactDirectory, CppDependencies, Logger);
 					actionArtifactCache.EnableReads = BuildConfiguration.bArtifactRead;
 					actionArtifactCache.EnableWrites = BuildConfiguration.bArtifactWrites;
+					actionArtifactCache.LogCacheMisses = BuildConfiguration.bLogArtifactCacheMisses;
 				}
 
 				// Pre-process module interfaces to generate dependency files
