- Fixed an issue with checking to see if a project is a hybrid (content-only treated as code-based) when the project and engine contain plugins with the same name

diff --git a/Engine/Source/Programs/UnrealBuildTool/System/NativeProjects.cs b/Engine/Source/Programs/UnrealBuildTool/System/NativeProjects.cs
--- a/Engine/Source/Programs/UnrealBuildTool/System/NativeProjects.cs
+++ b/Engine/Source/Programs/UnrealBuildTool/System/NativeProjects.cs
@@ -319,7 +319,18 @@ namespace UnrealBuildTool
 			ProjectDescriptor Project = ProjectDescriptor.FromFile(RawProjectPath);
 
 			// Enumerate all the available plugins
-			Dictionary<string, PluginInfo> AllPlugins = Plugins.ReadAvailablePlugins(Unreal.EngineDirectory, DirectoryReference.FromFile(RawProjectPath), new List<DirectoryReference>()).ToDictionary(x => x.Name, x => x, StringComparer.OrdinalIgnoreCase);
+			List<PluginInfo> EnabledPlugins = Plugins.ReadAvailablePlugins(Unreal.EngineDirectory, DirectoryReference.FromFile(RawProjectPath), null);
+
+			// instead of using .ToDictionary, we do this in a loop so that the non-engine plugins replace engine plugins with the same name
+			Dictionary<string, PluginInfo> AllPlugins = new(StringComparer.OrdinalIgnoreCase);
+			foreach (PluginInfo Plugin in EnabledPlugins)
+			{
+				// if we don't already have it, or we do but this one is in the game directory, use it
+				if (!AllPlugins.ContainsKey(Plugin.Name) || Plugin.File.IsUnderDirectory(DirectoryReference.FromFile(RawProjectPath)))
+				{
+					AllPlugins[Plugin.Name] = Plugin;
+				}
+			}
 
 			// find if there are any plugins enabled or disabled which differ from the default
 			string? Reason;
