UnrealBuildTool: Fix missing properties in VCCompileAction copy constructor & Write

#rnx

diff --git a/Engine/Source/Programs/UnrealBuildTool/Platform/Windows/VCCompileAction.cs b/Engine/Source/Programs/UnrealBuildTool/Platform/Windows/VCCompileAction.cs
--- a/Engine/Source/Programs/UnrealBuildTool/Platform/Windows/VCCompileAction.cs
+++ b/Engine/Source/Programs/UnrealBuildTool/Platform/Windows/VCCompileAction.cs
@@ -331,6 +331,7 @@ namespace UnrealBuildTool
 		/// <param name="InAction">Action to copy from</param>
 		public VCCompileAction(VCCompileAction InAction)
 		{
+			ActionType = InAction.ActionType;
 			CompilerExe = InAction.CompilerExe;
 			CompilerType = InAction.CompilerType;
 			ToolChainVersion = InAction.ToolChainVersion;
@@ -396,7 +397,8 @@ namespace UnrealBuildTool
 			bCanExecuteRemotely = Reader.ReadBool();
 			bCanExecuteRemotelyWithSNDBS = Reader.ReadBool();
 			bCanExecuteRemotelyWithXGE = Reader.ReadBool();
-
+			Architecture = UnrealArch.Parse(Reader.ReadString()!);
+			Weight = Reader.ReadDouble();
 
 			AdditionalPrerequisiteItems = Reader.ReadList(() => Reader.ReadFileItem())!;
 			AdditionalProducedItems = Reader.ReadList(() => Reader.ReadFileItem())!;
@@ -432,6 +434,8 @@ namespace UnrealBuildTool
 			Writer.WriteBool(bCanExecuteRemotely);
 			Writer.WriteBool(bCanExecuteRemotelyWithSNDBS);
 			Writer.WriteBool(bCanExecuteRemotelyWithXGE);
+			Writer.WriteString(Architecture.ToString());
+			Writer.WriteDouble(Weight);
 
 			Writer.WriteList(AdditionalPrerequisiteItems, Item => Writer.WriteFileItem(Item));
 			Writer.WriteList(AdditionalProducedItems, Item => Writer.WriteFileItem(Item));
