Prevent fatal error when generating projects if there's an error creating the TargetRules that can be ignored

#rnx

diff --git a/Engine/Source/Programs/UnrealBuildTool/ProjectFiles/VisualStudio/VCProject.cs b/Engine/Source/Programs/UnrealBuildTool/ProjectFiles/VisualStudio/VCProject.cs
--- a/Engine/Source/Programs/UnrealBuildTool/ProjectFiles/VisualStudio/VCProject.cs
+++ b/Engine/Source/Programs/UnrealBuildTool/ProjectFiles/VisualStudio/VCProject.cs
@@ -2051,9 +2051,12 @@ namespace UnrealBuildTool
 
 					// Figure out if this is a monolithic build
 					bool bShouldCompileMonolithic = BuildPlatform.ShouldCompileMonolithicBinary(Platform);
-					if (!bShouldCompileMonolithic)
+					try
+					{
+						bShouldCompileMonolithic |= (Combination.ProjectTarget.CreateRulesDelegate(Platform, Configuration).LinkType == TargetLinkType.Monolithic);
+					}
+					catch (BuildException)
 					{
-						bShouldCompileMonolithic = (Combination.ProjectTarget.CreateRulesDelegate(Platform, Configuration).LinkType == TargetLinkType.Monolithic);
 					}
 
 					// Get the .uproject directory
diff --git a/Engine/Source/Programs/UnrealBuildTool/ProjectFiles/VisualStudioCode/VSCodeProjectFileGenerator.cs b/Engine/Source/Programs/UnrealBuildTool/ProjectFiles/VisualStudioCode/VSCodeProjectFileGenerator.cs
--- a/Engine/Source/Programs/UnrealBuildTool/ProjectFiles/VisualStudioCode/VSCodeProjectFileGenerator.cs
+++ b/Engine/Source/Programs/UnrealBuildTool/ProjectFiles/VisualStudioCode/VSCodeProjectFileGenerator.cs
@@ -1097,7 +1097,13 @@ namespace UnrealBuildTool
 
 			if (TargetRulesObject != null)
 			{
-				bShouldCompileMonolithic |= (Target.CreateRulesDelegate(Platform, Configuration).LinkType == TargetLinkType.Monolithic);
+				try
+				{
+					bShouldCompileMonolithic |= (Target.CreateRulesDelegate(Platform, Configuration).LinkType == TargetLinkType.Monolithic);
+				}
+				catch (BuildException)
+				{
+				}
 			}
 
 			TargetType TargetRulesType = Target.TargetRules == null ? TargetType.Program : Target.TargetRules.Type;
diff --git a/Engine/Source/Programs/UnrealBuildTool/ProjectFiles/Xcode/XcodeProject.cs b/Engine/Source/Programs/UnrealBuildTool/ProjectFiles/Xcode/XcodeProject.cs
--- a/Engine/Source/Programs/UnrealBuildTool/ProjectFiles/Xcode/XcodeProject.cs
+++ b/Engine/Source/Programs/UnrealBuildTool/ProjectFiles/Xcode/XcodeProject.cs
@@ -651,7 +651,13 @@ namespace UnrealBuildTool.XcodeProjectXcconfig
 									{
 										// Figure out if this is a monolithic build
 										bool bShouldCompileMonolithic = BuildPlatform.ShouldCompileMonolithicBinary(Platform);
-										bShouldCompileMonolithic |= (ProjectTarget.CreateRulesDelegate(Platform, Configuration).LinkType == TargetLinkType.Monolithic);
+										try
+										{
+											bShouldCompileMonolithic |= (ProjectTarget.CreateRulesDelegate(Platform, Configuration).LinkType == TargetLinkType.Monolithic);
+										}
+										catch (BuildException)
+										{
+										}
 
 										string ConfigName = Configuration.ToString();
 										if (!bMakeProjectPerTarget)
diff --git a/Engine/Source/Programs/UnrealBuildTool/ProjectFiles/Xcode/XcodeProjectLegacy.cs b/Engine/Source/Programs/UnrealBuildTool/ProjectFiles/Xcode/XcodeProjectLegacy.cs
--- a/Engine/Source/Programs/UnrealBuildTool/ProjectFiles/Xcode/XcodeProjectLegacy.cs
+++ b/Engine/Source/Programs/UnrealBuildTool/ProjectFiles/Xcode/XcodeProjectLegacy.cs
@@ -1457,7 +1457,13 @@ namespace UnrealBuildTool.XcodeProjectLegacy
 									{
 										// Figure out if this is a monolithic build
 										bool bShouldCompileMonolithic = BuildPlatform.ShouldCompileMonolithicBinary(Platform);
-										bShouldCompileMonolithic |= (ProjectTarget.CreateRulesDelegate(Platform, Configuration).LinkType == TargetLinkType.Monolithic);
+										try
+										{
+											bShouldCompileMonolithic |= (ProjectTarget.CreateRulesDelegate(Platform, Configuration).LinkType == TargetLinkType.Monolithic);
+										}
+										catch (BuildException)
+										{
+										}
 
 										string ConfigName = Configuration.ToString();
 										if (!bMakeProjectPerTarget && ProjectTarget.TargetRules!.Type != TargetType.Game && ProjectTarget.TargetRules.Type != TargetType.Program)
