UnrealBuildTool: Ensure that plugins do not have dependencies on any plugin that is sealed. Throw a build error if detected

diff --git a/Engine/Source/Programs/UnrealBuildTool/Configuration/UEBuildTarget.cs b/Engine/Source/Programs/UnrealBuildTool/Configuration/UEBuildTarget.cs
--- a/Engine/Source/Programs/UnrealBuildTool/Configuration/UEBuildTarget.cs
+++ b/Engine/Source/Programs/UnrealBuildTool/Configuration/UEBuildTarget.cs
@@ -2249,6 +2249,16 @@ namespace UnrealBuildTool
 				}
 			}
 
+			// Check that each plugin does not have a dependency on any sealed plugins
+			foreach (UEBuildPlugin Plugin in BuildPlugins)
+			{
+				UEBuildPlugin? DependencyPlugin = Plugin.Dependencies?.FirstOrDefault(X => X.Descriptor.bIsSealed);
+				if (DependencyPlugin != null)
+				{
+					throw new BuildException("Plugin '{0}' cannot depend on plugin '{1}' because it is sealed.", Plugin.Name, DependencyPlugin.Name);
+				}
+			}
+
 			// Check that each plugin declares its dependencies explicitly
 			foreach (UEBuildPlugin Plugin in BuildPlugins)
 			{
diff --git a/Engine/Source/Programs/UnrealBuildTool/System/PluginDescriptor.cs b/Engine/Source/Programs/UnrealBuildTool/System/PluginDescriptor.cs
--- a/Engine/Source/Programs/UnrealBuildTool/System/PluginDescriptor.cs
+++ b/Engine/Source/Programs/UnrealBuildTool/System/PluginDescriptor.cs
@@ -184,6 +184,11 @@ namespace UnrealBuildTool
 		/// </summary>
 		public bool bRequiresBuildPlatform;
 
+		/// <summary>
+		/// When true, prevents other plugins from depending on this plugin
+		/// </summary>
+		public bool bIsSealed;
+
 		/// <summary>
 		/// When true, this plugin's modules will not be loaded automatically nor will it's content be mounted automatically. It will load/mount when explicitly requested and LoadingPhases will be ignored
 		/// </summary>
@@ -326,6 +331,7 @@ namespace UnrealBuildTool
 			}
 
 			RawObject.TryGetBoolField("RequiresBuildPlatform", out bRequiresBuildPlatform);
+			RawObject.TryGetBoolField("Sealed", out bIsSealed);
 			RawObject.TryGetBoolField("ExplicitlyLoaded", out bExplicitlyLoaded);
 			RawObject.TryGetBoolField("HasExplicitPlatforms", out bHasExplicitPlatforms);
 
@@ -429,6 +435,11 @@ namespace UnrealBuildTool
 				Writer.WriteValue("RequiresBuildPlatform", bRequiresBuildPlatform);
 			}
 
+			if (bIsSealed)
+			{
+				Writer.WriteValue("Sealed", bIsSealed);
+			}
+
 			if (bExplicitlyLoaded)
 			{
 				Writer.WriteValue("ExplicitlyLoaded", bExplicitlyLoaded);
