UnrealBuildTool: Move several defines from VCToolchain to instead be added in the target, so they are placed in Definitions.h rather than passed via the command line response file

#rnx

diff --git a/Engine/Source/Programs/UnrealBuildTool/Platform/Windows/UEBuildWindows.cs b/Engine/Source/Programs/UnrealBuildTool/Platform/Windows/UEBuildWindows.cs
--- a/Engine/Source/Programs/UnrealBuildTool/Platform/Windows/UEBuildWindows.cs
+++ b/Engine/Source/Programs/UnrealBuildTool/Platform/Windows/UEBuildWindows.cs
@@ -1080,6 +1080,30 @@ namespace UnrealBuildTool
 				Target.bChainPCHs = false;
 			}
 
+			// Various global definitions alway set or only based on target settings
+			if (Target.WindowsPlatform.Compiler.IsMSVC() && Target.WindowsPlatform.bStrictConformanceMode)
+			{
+				// This define is needed to ensure that MSVC static analysis mode doesn't declare attributes that are incompatible with strict conformance mode
+				Target.GlobalDefinitions.Add("SAL_NO_ATTRIBUTE_DECLARATIONS=1");
+			}
+
+			// Ignore secure CRT warnings on Clang
+			if (Target.WindowsPlatform.Compiler.IsClang())
+			{
+				Target.GlobalDefinitions.Add("_CRT_SECURE_NO_WARNINGS");
+			}
+
+			// @todo: UE is non-compliant when it comes to use of %s and %S
+			// Previously %s meant "the current character set" and %S meant "the other one".
+			// Now %s means multibyte and %S means wide. %Ts means "natural width".
+			// Reverting this behaviour until the UE source catches up.
+			Target.GlobalDefinitions.Add("_CRT_STDIO_LEGACY_WIDE_SPECIFIERS=1");
+
+			// @todo: Silence the hash_map deprecation errors for now. This should be replaced with unordered_map for the real fix.
+			Target.GlobalDefinitions.Add("_SILENCE_STDEXT_HASH_DEPRECATION_WARNINGS=1");
+
+			// Maintain the old std::aligned_storage behavior from VS from v15.8 onwards, in case of prebuilt third party libraries are reliant on it
+			Target.GlobalDefinitions.Add("_DISABLE_EXTENDED_ALIGNED_STORAGE");
 
 			// E&C support.
 			if (Target.bSupportEditAndContinue || Target.bAdaptiveUnityEnablesEditAndContinue)
diff --git a/Engine/Source/Programs/UnrealBuildTool/Platform/Windows/VCToolChain.cs b/Engine/Source/Programs/UnrealBuildTool/Platform/Windows/VCToolChain.cs
--- a/Engine/Source/Programs/UnrealBuildTool/Platform/Windows/VCToolChain.cs
+++ b/Engine/Source/Programs/UnrealBuildTool/Platform/Windows/VCToolChain.cs
@@ -105,6 +105,12 @@ namespace UnrealBuildTool
 			ExternalDependencies.Add(FileItem.GetItemByFileReference(EnvVars.LinkerPath));
 		}
 
+		/// <summary>
+		/// Add a definition to the argument list to be passed on the command line.
+		/// If possible, defines should instead be added to Target.GlobalDefinitions in WindowsPlatform
+		/// </summary>
+		/// <param name="Arguments"></param>
+		/// <param name="Definition"></param>
 		static public void AddDefinition(List<string> Arguments, string Definition)
 		{
 			// Split the definition into name and value
@@ -119,6 +125,13 @@ namespace UnrealBuildTool
 			}
 		}
 
+		/// <summary>
+		/// Add a definition to the argument list to be passed on the command line.
+		/// If possible, defines should instead be added to Target.GlobalDefinitions in WindowsPlatform
+		/// </summary>
+		/// <param name="Arguments"></param>
+		/// <param name="Variable"></param>
+		/// <param name="Value"></param>
 		static public void AddDefinition(List<string> Arguments, string Variable, string? Value)
 		{
 			// If the value has a space in it and isn't wrapped in quotes, do that now
@@ -463,9 +476,6 @@ namespace UnrealBuildTool
 			{
 				if (Target.WindowsPlatform.bStrictConformanceMode)
 				{
-					// This define is needed to ensure that MSVC static analysis mode doesn't declare attributes that are incompatible with strict conformance mode
-					AddDefinition(Arguments, "SAL_NO_ATTRIBUTE_DECLARATIONS=1");
-
 					Arguments.Add("/permissive-");
 					Arguments.Add("/Zc:strictStrings-"); // Have to disable strict const char* semantics due to Windows headers not being compliant.
 					if (EnvVars.CompilerVersion >= new VersionNumber(14, 32) && EnvVars.CompilerVersion < new VersionNumber(14, 33, 31629))
@@ -499,30 +509,12 @@ namespace UnrealBuildTool
 				}
 			}
 
-			// @todo HoloLens: UE is non-compliant when it comes to use of %s and %S
-			// Previously %s meant "the current character set" and %S meant "the other one".
-			// Now %s means multibyte and %S means wide. %Ts means "natural width".
-			// Reverting this behaviour until the UE source catches up.
-			AddDefinition(Arguments, "_CRT_STDIO_LEGACY_WIDE_SPECIFIERS=1");
-
-			// @todo HoloLens: Silence the hash_map deprecation errors for now. This should be replaced with unordered_map for the real fix.
-			AddDefinition(Arguments, "_SILENCE_STDEXT_HASH_DEPRECATION_WARNINGS=1");
-
-			// Ignore secure CRT warnings on Clang
-			if(Target.WindowsPlatform.Compiler.IsClang())
-			{
-				AddDefinition(Arguments, "_CRT_SECURE_NO_WARNINGS");
-			}
-
 			// If compiling as a DLL, set the relevant defines
 			if (CompileEnvironment.bIsBuildingDLL)
 			{
 				AddDefinition(Arguments, "_WINDLL");
 			}
 
-			// Maintain the old std::aligned_storage behavior from VS from v15.8 onwards, in case of prebuilt third party libraries are reliant on it
-			AddDefinition(Arguments, "_DISABLE_EXTENDED_ALIGNED_STORAGE");
-
 			// Fix Incredibuild errors with helpers using heterogeneous character sets
 			if (Target.WindowsPlatform.Compiler.IsMSVC())
 			{
