UnrealBuildTool: Clean up some PGO logic to be shared, don't raise an exception if no project is set

#rnx

diff --git a/Engine/Source/Programs/UnrealBuildTool/Platform/Android/UEBuildAndroid.cs b/Engine/Source/Programs/UnrealBuildTool/Platform/Android/UEBuildAndroid.cs
--- a/Engine/Source/Programs/UnrealBuildTool/Platform/Android/UEBuildAndroid.cs
+++ b/Engine/Source/Programs/UnrealBuildTool/Platform/Android/UEBuildAndroid.cs
@@ -523,12 +523,12 @@ namespace UnrealBuildTool
 				CompileEnvironment.Definitions.Add("RHI_RAYTRACING=1");
 			}
 
-			if (Target.bPGOOptimize || Target.bPGOProfile)
+			if ((Target.bPGOOptimize || Target.bPGOProfile) && Target.ProjectFile != null)
 			{
 				Logger.LogInformation("PGO {PgoType} build", Target.bPGOOptimize ? "optimize" : "profile");
 				if (Target.bPGOOptimize)
 				{
-					CompileEnvironment.PGODirectory = Path.Combine(DirectoryReference.FromFile(Target.ProjectFile)!.FullName, "Platforms", "Android", "Build", "PGO");
+					CompileEnvironment.PGODirectory = DirectoryReference.Combine(Target.ProjectFile.Directory, "Platforms", "Android", "Build", "PGO").FullName;
 					CompileEnvironment.PGOFilenamePrefix = string.Format("{0}-Android", Target.Name);
 
 					LinkEnvironment.PGODirectory = CompileEnvironment.PGODirectory;
diff --git a/Engine/Source/Programs/UnrealBuildTool/Platform/Windows/UEBuildWindows.cs b/Engine/Source/Programs/UnrealBuildTool/Platform/Windows/UEBuildWindows.cs
--- a/Engine/Source/Programs/UnrealBuildTool/Platform/Windows/UEBuildWindows.cs
+++ b/Engine/Source/Programs/UnrealBuildTool/Platform/Windows/UEBuildWindows.cs
@@ -1633,11 +1633,10 @@ namespace UnrealBuildTool
 
 			if ((Target.bPGOOptimize || Target.bPGOProfile) && Target.ProjectFile != null)
 			{
-				// LTCG is required for PGO
-				//CompileEnvironment.bAllowLTCG = true;
-				//LinkEnvironment.bAllowLTCG = true;
+				// Win64 PGO folder is Windows, the rest match the platform name
+				string PGOPlatform = Target.Platform == UnrealTargetPlatform.Win64 ? "Windows" : Target.Platform.ToString();
 
-				CompileEnvironment.PGODirectory = Path.Combine(DirectoryReference.FromFile(Target.ProjectFile).FullName, "Platforms", "Windows", "Build", "PGO");
+				CompileEnvironment.PGODirectory = DirectoryReference.Combine(Target.ProjectFile.Directory, "Platforms", PGOPlatform, "Build", "PGO").FullName;
 				CompileEnvironment.PGOFilenamePrefix = string.Format("{0}-{1}-{2}", Target.Name, Target.Platform, Target.Configuration);
 
 				LinkEnvironment.PGODirectory = CompileEnvironment.PGODirectory;
