Make sure iOS15 is used when minimum version not specified in ini

diff --git a/Engine/Source/Programs/UnrealBuildTool/Platform/IOS/UEDeployIOS.cs b/Engine/Source/Programs/UnrealBuildTool/Platform/IOS/UEDeployIOS.cs
--- a/Engine/Source/Programs/UnrealBuildTool/Platform/IOS/UEDeployIOS.cs
+++ b/Engine/Source/Programs/UnrealBuildTool/Platform/IOS/UEDeployIOS.cs
@@ -203,6 +203,7 @@ namespace UnrealBuildTool
 			string MinVersionToReturn = "";
 			switch (MinVersion)
 			{
+				case "":
 				case "IOS_15":
 				case "IOS_Minimum":
 					MinVersionToReturn = "15.0";
