Support enabling AutoRTFM instrumentation on a Module level in UBT

diff --git a/Engine/Source/Programs/UnrealBuildTool/Configuration/ModuleRules.cs b/Engine/Source/Programs/UnrealBuildTool/Configuration/ModuleRules.cs
--- a/Engine/Source/Programs/UnrealBuildTool/Configuration/ModuleRules.cs
+++ b/Engine/Source/Programs/UnrealBuildTool/Configuration/ModuleRules.cs
@@ -1089,6 +1089,11 @@ namespace UnrealBuildTool
 		/// </summary>
 		public bool bAllowConfidentialPlatformDefines = false;
 
+		/// <summary>
+		/// Enables AutoRTFM instrumentation to this module only when AutoRTFMCompiler is enabled
+		/// </summary>
+		public bool bAllowAutoRTFMInstrumentation = false;
+
 		/// <summary>
 		/// List of modules names (no path needed) with header files that our module's public headers needs access to, but we don't need to "import" or link against.
 		/// </summary>
diff --git a/Engine/Source/Programs/UnrealBuildTool/Configuration/UEBuildModuleCPP.cs b/Engine/Source/Programs/UnrealBuildTool/Configuration/UEBuildModuleCPP.cs
--- a/Engine/Source/Programs/UnrealBuildTool/Configuration/UEBuildModuleCPP.cs
+++ b/Engine/Source/Programs/UnrealBuildTool/Configuration/UEBuildModuleCPP.cs
@@ -1743,6 +1743,12 @@ namespace UnrealBuildTool
 			Result.IncludeOrderVersion = Rules.IncludeOrderVersion;
 			Result.DeterministicWarningLevel = Rules.DeterministicWarningLevel;
 
+			// Only enable the AutoRTFM flag if we are using the AutoRTFM compiler
+			if (Target.bUseAutoRTFMCompiler)
+			{
+				Result.bAllowAutoRTFMInstrumentation = Rules.bAllowAutoRTFMInstrumentation;
+			}
+
 			if (Result.OptimizationLevel != Rules.OptimizationLevel)
 			{
 				Logger.LogInformation("Module {0} - Optimization level changed for module due to override. Old: {1} New: {2}", Name, Result.OptimizationLevel, Rules.OptimizationLevel);
diff --git a/Engine/Source/Programs/UnrealBuildTool/System/CppCompileEnvironment.cs b/Engine/Source/Programs/UnrealBuildTool/System/CppCompileEnvironment.cs
--- a/Engine/Source/Programs/UnrealBuildTool/System/CppCompileEnvironment.cs
+++ b/Engine/Source/Programs/UnrealBuildTool/System/CppCompileEnvironment.cs
@@ -270,6 +270,11 @@ namespace UnrealBuildTool
 		/// </summary>
 		public bool bEnableBufferSecurityChecks = true;
 
+		/// <summary>
+		/// Enables AutoRTFM instrumentation to this cpp file only when AutoRTFMCompiler is enabled
+		/// </summary>
+		public bool bAllowAutoRTFMInstrumentation = false;
+
 		/// <summary>
 		/// If unity builds are enabled this can be used to override if this specific module will build using Unity.
 		/// This is set using the per module configurations in BuildConfiguration.
@@ -721,6 +726,7 @@ namespace UnrealBuildTool
 			StaticAnalyzerDisabledCheckers = new HashSet<string>(Other.StaticAnalyzerDisabledCheckers);
 			StaticAnalyzerAdditionalCheckers = new HashSet<string>(Other.StaticAnalyzerAdditionalCheckers);
 			bOptimizeCode = Other.bOptimizeCode;
+			bAllowAutoRTFMInstrumentation = Other.bAllowAutoRTFMInstrumentation;
 			OptimizationLevel = Other.OptimizationLevel;
 			bCreateDebugInfo = Other.bCreateDebugInfo;
 			bIsBuildingLibrary = Other.bIsBuildingLibrary;
diff --git a/Engine/Source/Programs/UnrealBuildTool/ToolChain/ClangToolChain.cs b/Engine/Source/Programs/UnrealBuildTool/ToolChain/ClangToolChain.cs
--- a/Engine/Source/Programs/UnrealBuildTool/ToolChain/ClangToolChain.cs
+++ b/Engine/Source/Programs/UnrealBuildTool/ToolChain/ClangToolChain.cs
@@ -363,6 +363,11 @@ namespace UnrealBuildTool
 				// Validate PCH inputs by content if mtime check fails
 				Arguments.Add("-fpch-validate-input-files-content");
 			}
+
+			if (CompileEnvironment.bAllowAutoRTFMInstrumentation)
+			{
+				Arguments.Add("-fautortfm");
+			}
 		}
 
 		protected virtual void GetCStandardCompileArgument(CppCompileEnvironment CompileEnvironment, List<string> Arguments)
diff --git a/Engine/Source/Runtime/Core/Core.Build.cs b/Engine/Source/Runtime/Core/Core.Build.cs
--- a/Engine/Source/Runtime/Core/Core.Build.cs
+++ b/Engine/Source/Runtime/Core/Core.Build.cs
@@ -315,6 +315,8 @@ public class Core : ModuleRules
 		UnsafeTypeCastWarningLevel = WarningLevel.Error;
 
 		IWYUSupport = IWYUSupport.KeepAsIs;
+
+		bAllowAutoRTFMInstrumentation = true;
 	}
 
 	protected virtual bool SupportsBinaryConfig(ReadOnlyTargetRules Target)
diff --git a/Engine/Source/Runtime/CoreUObject/CoreUObject.Build.cs b/Engine/Source/Runtime/CoreUObject/CoreUObject.Build.cs
--- a/Engine/Source/Runtime/CoreUObject/CoreUObject.Build.cs
+++ b/Engine/Source/Runtime/CoreUObject/CoreUObject.Build.cs
@@ -32,5 +32,7 @@ public class CoreUObject : ModuleRules
 		}
 
 		PrivateDefinitions.Add("UE_DEFINE_LEGACY_MATH_CONSTANT_MACRO_NAMES=0");
+
+		bAllowAutoRTFMInstrumentation = true;
 	}
 }
