Manually make change here due to conflict being to hard to understand

diff --git a/Engine/Source/Programs/UnrealBuildTool/Platform/Windows/VCToolChain.cs b/Engine/Source/Programs/UnrealBuildTool/Platform/Windows/VCToolChain.cs
--- a/Engine/Source/Programs/UnrealBuildTool/Platform/Windows/VCToolChain.cs
+++ b/Engine/Source/Programs/UnrealBuildTool/Platform/Windows/VCToolChain.cs
@@ -1011,6 +1011,11 @@ namespace UnrealBuildTool
 
 				// Additional disabled warnings for msvc. Everything below should be checked if it is necessary
 				ClangWarnings.GetVCDisabledWarnings(Arguments);
+
+				if (CompileEnvironment.bAllowAutoRTFMInstrumentation)
+				{
+					Arguments.Add("-fautortfm");
+				}
 			}
 
 			if (Target.WindowsPlatform.Compiler == WindowsCompiler.Intel)
