Make sure iOS15 is used when minimum version not specified in ini

diff --git a/Engine/Source/Programs/UnrealBuildTool/Platform/IOS/UEDeployIOS.cs b/Engine/Source/Programs/UnrealBuildTool/Platform/IOS/UEDeployIOS.cs
--- a/Engine/Source/Programs/UnrealBuildTool/Platform/IOS/UEDeployIOS.cs
+++ b/Engine/Source/Programs/UnrealBuildTool/Platform/IOS/UEDeployIOS.cs
@@ -213,7 +213,7 @@ namespace UnrealBuildTool
 					break;
 				default:
 					MinVersionToReturn = "15.0";
-					Logger.LogWarning("MinimumiOSVersion {MinVersion} specified in ini file is no longer supported, defaulting to {MinVersionToReturn}", MinVersion, MinVersionToReturn);
+					Logger.LogInformation("MinimumiOSVersion {MinVersion} specified in ini file is no longer supported, defaulting to {MinVersionToReturn}", MinVersion, MinVersionToReturn);
 					break;
 			}
 			return MinVersionToReturn;
