- Added support for -architecture=Host to override other settings and use whatever architecture is running UBT
- Added UnrealArch.Host, which can be used to determine the current platform's  current architecture (Windows would have to add support if it supports running on Arm64)

diff --git a/Engine/Source/Programs/UnrealBuildTool/Configuration/UEBuildPlatform.cs b/Engine/Source/Programs/UnrealBuildTool/Configuration/UEBuildPlatform.cs
--- a/Engine/Source/Programs/UnrealBuildTool/Configuration/UEBuildPlatform.cs
+++ b/Engine/Source/Programs/UnrealBuildTool/Configuration/UEBuildPlatform.cs
@@ -154,6 +154,20 @@ namespace UnrealBuildTool
 			return Architectures.GetFolderNameForPlatform(this);
 		}
 
+		/// <summary>
+		/// Returns the architecture of the currently running OS (only used for desktop platforms, so will throw an exception in the general case)
+		/// </summary>
+		/// <returns></returns>
+		public virtual UnrealArch GetHostArchitecture()
+		{
+			if (AllSupportedArchitectures.bIsMultiArch)
+			{
+				throw new BuildException($"Asking for Host architecture from {GetType()} which supports multiple architectures, but did not override GetHostArchitecture()");
+			}
+
+			return AllSupportedArchitectures.SingleArchitecture;
+		}
+
 
 
 		/// <summary>
diff --git a/Engine/Source/Programs/UnrealBuildTool/Configuration/UEBuildTarget.cs b/Engine/Source/Programs/UnrealBuildTool/Configuration/UEBuildTarget.cs
--- a/Engine/Source/Programs/UnrealBuildTool/Configuration/UEBuildTarget.cs
+++ b/Engine/Source/Programs/UnrealBuildTool/Configuration/UEBuildTarget.cs
@@ -740,6 +740,13 @@ namespace UnrealBuildTool
 		/// <returns></returns>
 		static public bool TryParse(string Name, out UnrealArch Arch)
 		{
+			// handle the Host string
+			if (Name.Equals("Host", StringComparison.OrdinalIgnoreCase))
+			{
+				Arch.Id = UnrealArch.Host.Value.Id;
+				return true;
+			}
+
 			Name = FixName(Name);
 
 			if (GetUniqueStringRegistry().HasString(Name))
@@ -765,6 +772,12 @@ namespace UnrealBuildTool
 		/// <returns></returns>
 		static public UnrealArch Parse(string Name)
 		{
+			// handle the Host string
+			if (Name.Equals("Host", StringComparison.OrdinalIgnoreCase))
+			{
+				return UnrealArch.Host.Value;
+			}
+
 			Name = FixName(Name);
 
 			if (GetUniqueStringRegistry().HasString(Name))
@@ -825,6 +838,11 @@ namespace UnrealBuildTool
 		/// Used in place when needing to handle deprecated architectures that a licensee may still have. Do not use in normal logic
 		/// </summary>
 		public static UnrealArch Deprecated = FindOrAddByName("deprecated", bIsX64: false);
+
+		/// <summary>
+		/// Maps to the currently running architecture on the host platform (lazy because this needs other classes to be intialized)
+		/// </summary>
+		public static Lazy<UnrealArch> Host = new(() => UnrealArchitectureConfig.ForPlatform(BuildHostPlatform.Current.Platform).GetHostArchitecture());
 	}
 
 	internal class UnrealArchPlatformTypeConverter : TypeConverter
diff --git a/Engine/Source/Programs/UnrealBuildTool/Platform/Mac/UEBuildMac.cs b/Engine/Source/Programs/UnrealBuildTool/Platform/Mac/UEBuildMac.cs
--- a/Engine/Source/Programs/UnrealBuildTool/Platform/Mac/UEBuildMac.cs
+++ b/Engine/Source/Programs/UnrealBuildTool/Platform/Mac/UEBuildMac.cs
@@ -145,6 +145,11 @@ namespace UnrealBuildTool
 		{
 		}
 
+		public override UnrealArch GetHostArchitecture()
+		{
+			return MacExports.IsRunningOnAppleArchitecture ? UnrealArch.Arm64 : UnrealArch.X64;
+		}
+
 		public override string ConvertToReadableArchitecture(UnrealArch Architecture)
 		{
 			if (Architecture == UnrealArch.X64)
