- Fixes NFL + Programs + Platform directories to look in the removed restricted program path, brings structure closer to Engine standard within Restricted Platform folders.
Before:
 "Engine\Restricted\NoRedist\Platforms\{Platform}\Restricted\NotForLicensees\Programs\{Project}Config\{Platform}Engine.ini"
After:
 "Engine\Restricted\NoRedist\Platforms\{Platform}\Programs\{Project}\Config\{Platform}Engine.ini"
- Fixed config paths for Restricted, non-Platform, program ini files (last change had put OPT_SUBDIR in too many places, it's only needed for Restricted Programs Platform files)

diff --git a/Engine/Source/Programs/UnrealBuildTool/System/ConfigHierarchy.cs b/Engine/Source/Programs/UnrealBuildTool/System/ConfigHierarchy.cs
--- a/Engine/Source/Programs/UnrealBuildTool/System/ConfigHierarchy.cs
+++ b/Engine/Source/Programs/UnrealBuildTool/System/ConfigHierarchy.cs
@@ -956,8 +956,8 @@ namespace UnrealBuildTool
 			// The base expansion (ie, no expansion)
 			new ConfigLayerExpansion { }, 
 			// Restricted Locations
-			new ConfigLayerExpansion { Before1 = "{ENGINE}/", After1 = "{ENGINE}/Restricted/NotForLicensees/",	Before2 = "{PROJECT}/Config/", After2 = "{RESTRICTEDPROJECT_NFL}/{OPT_SUBDIR}Config/" },
-			new ConfigLayerExpansion { Before1 = "{ENGINE}/", After1 = "{ENGINE}/Restricted/NoRedist/",			Before2 = "{PROJECT}/Config/", After2 = "{RESTRICTEDPROJECT_NR}/{OPT_SUBDIR}Config/" },
+			new ConfigLayerExpansion { Before1 = "{ENGINE}/", After1 = "{ENGINE}/Restricted/NotForLicensees/",	Before2 = "{PROJECT}/Config/", After2 = "{RESTRICTEDPROJECT_NFL}/Config/" },
+			new ConfigLayerExpansion { Before1 = "{ENGINE}/", After1 = "{ENGINE}/Restricted/NoRedist/",			Before2 = "{PROJECT}/Config/", After2 = "{RESTRICTEDPROJECT_NR}/Config/" },
 			// Platform Extensions
 			new ConfigLayerExpansion { Before1 = "{ENGINE}/Config/{PLATFORM}/", After1 = "{EXTENGINE}/Config/",	Before2 = "{PROJECT}/Config/{PLATFORM}/", After2 = "{EXTPROJECT}/Config/" },
 			// Platform Extensions in Restricted Locations
@@ -1023,6 +1023,7 @@ namespace UnrealBuildTool
 				DirectoryReference NFLDir;
 				DirectoryReference NRDir;
 				string OptionalSubDir = "";
+
 				if (ProjectDir.IsUnderDirectory(Unreal.EngineDirectory))
 				{
 					OptionalSubDir = ProjectDir.MakeRelativeTo(Unreal.EngineDirectory) + "/";
@@ -1034,6 +1035,16 @@ namespace UnrealBuildTool
 					NFLDir = DirectoryReference.Combine(ProjectDir, "Restricted/NotForLicensees");
 					NRDir = DirectoryReference.Combine(ProjectDir, "Restricted/NoRedist");
 				}
+
+				if (ProjectDir.IsUnderDirectory(NFLDir))
+				{
+					OptionalSubDir = ProjectDir.MakeRelativeTo(NFLDir) + "/";
+				}
+				else if (ProjectDir.IsUnderDirectory(NRDir))
+				{
+					OptionalSubDir = ProjectDir.MakeRelativeTo(NRDir) + "/";
+				}
+
 				string PlatformExtensionProjectConfigDir = DirectoryReference.Combine(ProjectDir, "Platforms", PlatformName).FullName;
 
 				OutString = OutString.Replace("{PROJECT}", ProjectDir.FullName);
diff --git a/Engine/Source/Runtime/Core/Private/Misc/ConfigContext.cpp b/Engine/Source/Runtime/Core/Private/Misc/ConfigContext.cpp
--- a/Engine/Source/Runtime/Core/Private/Misc/ConfigContext.cpp
+++ b/Engine/Source/Runtime/Core/Private/Misc/ConfigContext.cpp
@@ -455,12 +455,26 @@ FString FConfigContext::PerformFinalExpansions(const FString& InString, const FS
 	OutString = OutString.Replace(TEXT("{RESTRICTEDPROJECT_NFL}"), *ProjectNotForLicenseesDir);
 	OutString = OutString.Replace(TEXT("{RESTRICTEDPROJECT_NR}"), *ProjectNoRedistDir);
 
-	if (FPaths::IsUnderDirectory(ProjectRootDir, EngineRootDir))
+	if (FPaths::IsUnderDirectory(ProjectRootDir, ProjectNotForLicenseesDir))
 	{
 		FString RelativeDir = ProjectRootDir;
-		FPaths::MakePathRelativeTo(RelativeDir, *EngineRootDir);
+		FPaths::MakePathRelativeTo(RelativeDir, *(ProjectNotForLicenseesDir + TEXT("/")));
+
+		OutString = OutString.Replace(TEXT("{OPT_SUBDIR}"), *(RelativeDir + TEXT("/")));
+	}
+	else if (FPaths::IsUnderDirectory(ProjectRootDir, ProjectNoRedistDir))
+	{
+		FString RelativeDir = ProjectRootDir;
+		FPaths::MakePathRelativeTo(RelativeDir, *(ProjectNoRedistDir + TEXT("/")));
+
+		OutString = OutString.Replace(TEXT("{OPT_SUBDIR}"), *(RelativeDir + TEXT("/")));
+	}
+	else if (FPaths::IsUnderDirectory(ProjectRootDir, EngineRootDir))
+	{
+		FString RelativeDir = ProjectRootDir;
+		FPaths::MakePathRelativeTo(RelativeDir, *(EngineRootDir + TEXT("/")));
 		
-		OutString = OutString.Replace(TEXT("{OPT_SUBDIR}"), *RelativeDir);
+		OutString = OutString.Replace(TEXT("{OPT_SUBDIR}"), *(RelativeDir + TEXT("/")));
 	}
 	else
 	{
diff --git a/Engine/Source/Runtime/Core/Public/Misc/ConfigHierarchy.h b/Engine/Source/Runtime/Core/Public/Misc/ConfigHierarchy.h
--- a/Engine/Source/Runtime/Core/Public/Misc/ConfigHierarchy.h
+++ b/Engine/Source/Runtime/Core/Public/Misc/ConfigHierarchy.h
@@ -78,12 +78,12 @@ inline FConfigLayerExpansion GConfigExpansions[] =
 	// Restricted Locations
 	{ 
 		TEXT("{ENGINE}/"),						TEXT("{ENGINE}/Restricted/NotForLicensees/"),	
-		TEXT("{PROJECT}/Config/"),				TEXT("{RESTRICTEDPROJECT_NFL}/{OPT_SUBDIR}Config/"),
+		TEXT("{PROJECT}/Config/"),				TEXT("{RESTRICTEDPROJECT_NFL}/Config/"),
 		EConfigExpansionFlags::ForUncooked | EConfigExpansionFlags::ForCooked
 	},
 	{ 
 		TEXT("{ENGINE}/"),						TEXT("{ENGINE}/Restricted/NoRedist/"),			
-		TEXT("{PROJECT}/Config/"),				TEXT("{RESTRICTEDPROJECT_NR}/{OPT_SUBDIR}Config/"),
+		TEXT("{PROJECT}/Config/"),				TEXT("{RESTRICTEDPROJECT_NR}/Config/"),
 		EConfigExpansionFlags::ForUncooked 
 	},
 
