UBT: Add debugging flag -dumpiniloads to show which .ini files are being loaded into hierarchy with ConfigCache.cs

diff --git a/Engine/Source/Programs/UnrealBuildTool/System/ConfigCache.cs b/Engine/Source/Programs/UnrealBuildTool/System/ConfigCache.cs
--- a/Engine/Source/Programs/UnrealBuildTool/System/ConfigCache.cs
+++ b/Engine/Source/Programs/UnrealBuildTool/System/ConfigCache.cs
@@ -6,12 +6,11 @@ using System.Diagnostics.CodeAnalysis;
 using System.Linq;
 using System.Reflection;
 using System.IO;
-using System.Text;
-using System.Threading.Tasks;
 using EpicGames.Core;
-using UnrealBuildTool;
 using Microsoft.Extensions.Logging;
+using Microsoft.Extensions.Logging.Abstractions;
 using System.Collections.Concurrent;
+using Microsoft.Build.Logging;
 
 namespace UnrealBuildTool
 {
@@ -312,6 +311,7 @@ namespace UnrealBuildTool
 
 			string IniConfigArgPrefix = "-ini:" + Enum.GetName(typeof(ConfigHierarchyType), Type) + ":";
 			string CustomConfigPrefix = "-CustomConfig=";
+			bool bDumpIniLoads = false;
 			foreach (string CmdLineArg in CmdLineArgs)
 			{
 				if (CmdLineArg.StartsWith(IniConfigArgPrefix, StringComparison.InvariantCultureIgnoreCase))
@@ -322,6 +322,10 @@ namespace UnrealBuildTool
 				{
 					CustomConfig = CmdLineArg.Substring(CustomConfigPrefix.Length);
 				}
+				if (CmdLineArg.Equals("-dumpiniloads", StringComparison.InvariantCultureIgnoreCase))
+				{
+					bDumpIniLoads = true;
+				}
 			}
 
 			if (CustomConfig == null)
@@ -332,7 +336,15 @@ namespace UnrealBuildTool
 			// Get the key to use for the cache. It cannot be null, so we use the engine directory if a project directory is not given.
 			ConfigHierarchyKey Key = new ConfigHierarchyKey(Type, ProjectDir, Platform, CustomConfig, OverrideStrings);
 
+			ILogger Logger = NullLogger.Instance;
+			
+			if(bDumpIniLoads)
+			{
+				Logger = Log.Logger;
+			};
+
 			// Try to get the cached hierarchy with this key
+			Logger.LogInformation($"Requested Hierarchy: {Type},{ProjectDir},{Platform},{CustomConfig}");
 			return HierarchyKeyToHierarchy.GetOrAdd(Key, _ =>
 			{
 				// Find all the input files
@@ -340,8 +352,10 @@ namespace UnrealBuildTool
 				foreach (FileReference IniFileName in ConfigHierarchy.EnumerateConfigFileLocations(Type, ProjectDir, Platform, CustomConfig))
 				{
 					ConfigFile? File;
+					Logger.LogInformation($"Trying to load: {IniFileName}");
 					if (TryReadFile(IniFileName, out File))
 					{
+						Logger.LogInformation($"    Found!");
 						Files.Add(File);
 					}
 				}
@@ -350,6 +364,7 @@ namespace UnrealBuildTool
 				{
 					ConfigFile OverrideFile = new ConfigFile(OverrideString);
 					Files.Add(OverrideFile);
+					Logger.LogInformation($"Trying to load override: {OverrideFile}");
 				}
 
 				// Create the hierarchy
