Identify 30Hz device groups. Only run AiBotSoak for half as many frames as 60Hz device. Shortens tests which eases pressure on device pool and fixes an issue where csv wasn't written due to long test run.

diff --git a/Engine/Source/Programs/UnrealBuildTool/Configuration/UEBuildTarget.cs b/Engine/Source/Programs/UnrealBuildTool/Configuration/UEBuildTarget.cs
--- a/Engine/Source/Programs/UnrealBuildTool/Configuration/UEBuildTarget.cs
+++ b/Engine/Source/Programs/UnrealBuildTool/Configuration/UEBuildTarget.cs
@@ -563,6 +563,12 @@ namespace UnrealBuildTool
 		/// SDLPlatform is for platforms that use SDL for windows, cursors, etc
 		/// </summary>
 		public static UnrealPlatformGroup SDLPlatform = FindOrAddByName("SDLPlatform");
+
+		/// <summary>
+		/// 30Hz is for platforms that typically run at 30Hz (Android, Switch) vs. the baseline 60Hz 
+		/// </summary>
+		public static UnrealPlatformGroup ThirtyHz = FindOrAddByName("30Hz");
+
 	}
 
 	/// <summary>
diff --git a/Engine/Source/Programs/UnrealBuildTool/Platform/Android/UEBuildAndroid.cs b/Engine/Source/Programs/UnrealBuildTool/Platform/Android/UEBuildAndroid.cs
--- a/Engine/Source/Programs/UnrealBuildTool/Platform/Android/UEBuildAndroid.cs
+++ b/Engine/Source/Programs/UnrealBuildTool/Platform/Android/UEBuildAndroid.cs
@@ -638,6 +638,7 @@ namespace UnrealBuildTool
 			// Register this build platform
 			UEBuildPlatform.RegisterBuildPlatform(new AndroidPlatform(SDK, Logger), Logger);
 			UEBuildPlatform.RegisterPlatformWithGroup(UnrealTargetPlatform.Android, UnrealPlatformGroup.Android);
+			UEBuildPlatform.RegisterPlatformWithGroup(UnrealTargetPlatform.Android, UnrealPlatformGroup.ThirtyHz);
 		}
 	}
 }
