UnrealBuildTool: Bump msvc preferred version to [REDACTED_PHONE]. Notably this version improves /LTCG link duration by 25%-50% depending on the project

diff --git a/Engine/Source/Programs/UnrealBuildTool/Platform/Windows/MicrosoftPlatformSDK.Versions.cs b/Engine/Source/Programs/UnrealBuildTool/Platform/Windows/MicrosoftPlatformSDK.Versions.cs
--- a/Engine/Source/Programs/UnrealBuildTool/Platform/Windows/MicrosoftPlatformSDK.Versions.cs
+++ b/Engine/Source/Programs/UnrealBuildTool/Platform/Windows/MicrosoftPlatformSDK.Versions.cs
@@ -44,6 +44,7 @@ namespace UnrealBuildTool
 		/// </summary>
 		static readonly VersionNumberRange[] PreferredVisualCppVersions = new VersionNumberRange[]
 		{
+			VersionNumberRange.Parse("14.36.32532", "14.36.99999"), // VS2022 17.6.x
 			VersionNumberRange.Parse("14.35.32215", "14.35.99999"), // VS2022 17.5.x
 			VersionNumberRange.Parse("14.34.31933", "14.34.99999"), // VS2022 17.4.x
 			VersionNumberRange.Parse("14.29.30133", "14.29.99999"), // VS2019 16.11.x
diff --git a/Engine/Source/Programs/UnrealBuildTool/Platform/Windows/MicrosoftPlatformSDK.VisualStudio.cs b/Engine/Source/Programs/UnrealBuildTool/Platform/Windows/MicrosoftPlatformSDK.VisualStudio.cs
--- a/Engine/Source/Programs/UnrealBuildTool/Platform/Windows/MicrosoftPlatformSDK.VisualStudio.cs
+++ b/Engine/Source/Programs/UnrealBuildTool/Platform/Windows/MicrosoftPlatformSDK.VisualStudio.cs
@@ -46,7 +46,7 @@ namespace UnrealBuildTool
 		static readonly string[] VisualStudio2022SuggestedComponents = new string[]
 		{
 			"Microsoft.VisualStudio.Workload.ManagedDesktop",
-			"Microsoft.VisualStudio.Component.VC.14.35.17.5.x86.x64",
+			"Microsoft.VisualStudio.Component.VC.14.36.17.6.x86.x64",
 			"Microsoft.Net.Component.4.6.2.TargetingPack",
 		};
 
