Split resources.pri file generation from xml appx manifest generation.

 - also removed several unused members & functions.
 - use File/DirectoryReference rather than strings in the api.

#rnx

diff --git a/Engine/Source/Programs/UnrealBuildTool/Platform/Windows/AppXManifestGeneratorBase.cs b/Engine/Source/Programs/UnrealBuildTool/Platform/Windows/AppXManifestGeneratorBase.cs
--- a/Engine/Source/Programs/UnrealBuildTool/Platform/Windows/AppXManifestGeneratorBase.cs
+++ b/Engine/Source/Programs/UnrealBuildTool/Platform/Windows/AppXManifestGeneratorBase.cs
@@ -21,9 +21,6 @@ namespace UnrealBuildTool
 	/// </summary>
 	abstract public class AppXManifestGeneratorBase
 	{
-		/// default schema namespace
-        protected virtual string Schema2010NS { get { return "http://schemas.microsoft.com/appx/2010/manifest"; } }
-
 		/// config section for platform-specific target settings
 		protected virtual string IniSection_PlatformTargetSettings { get { return string.Format( "/Script/{0}PlatformEditor.{0}TargetSettings", Platform.ToString() ); } }
 		
@@ -36,9 +33,6 @@ namespace UnrealBuildTool
 		/// default subdirectory for engine resources
 		protected const string EngineResourceSubPath = "DefaultImages";
 
-		/// platform to use for reading configuration
-		protected virtual UnrealTargetPlatform ConfigPlatform { get { return Platform; } }
-
         /// Manifest compliance values
         protected const int MaxResourceEntries = 200;
 
@@ -48,24 +42,15 @@ namespace UnrealBuildTool
 		/// cached game ini
 		protected ConfigHierarchy? GameIni;
 
+		/// AppX package resource generator
+		protected UEAppXResources? AppXResources;
+
 		/// the default culture to use
 		protected string? DefaultAppXCultureId;
 
-		/// all cultures that will be staged
-		protected List<string> SelectedAppXCultureIds = new();
-
-		/// localization dictionary - AppX CultureId to ConfigKey/Localized String pairs. Empty CultureId represents the default strings
-		protected Dictionary<string,Dictionary<string,string>>? PerAppXCultureIdStrings;
-
 		/// lookup table for UE's CultureToStage StageId to AppX CultureId
 		protected Dictionary<string,string> UEStageIdToAppXCultureId = new Dictionary<string, string>();
 
-		/// resource writer for the default culture
-		protected UEResXWriter? DefaultResourceWriter;
-
-		/// resource writers for each AppX culture
-		protected Dictionary<string, UEResXWriter>? PerCultureResourceWriters;
-
 		/// the platform to generate the manifest for
 		protected UnrealTargetPlatform Platform;
 
@@ -75,18 +60,7 @@ namespace UnrealBuildTool
 		/// target name to use
 		protected string? TargetName;
 
-		/// output path - where the manifest will be created
-		protected string? OutputPath;
-
-		/// intermediate path - used for temporary files
-		protected string? IntermediatePath;
-
-		/// files that should be included in the manifest output folder and where to source them from
-		protected Dictionary<string,string>? ManifestFiles; // Dst, Src
-
-		/// <summary>
 		/// Logger for output
-		/// </summary>
 		protected readonly ILogger Logger;
 
 		/// <summary>
@@ -98,95 +72,6 @@ namespace UnrealBuildTool
 			Logger = InLogger;
 		}
 
-		/// <summary>
-		/// Lookup a switch in a dictionary
-		/// </summary>
-        protected static bool SafeGetBool(IDictionary<string, string> InDictionary, string Key, bool DefaultValue = false)
-		{
-			if (InDictionary.ContainsKey(Key))
-			{
-				var Value = InDictionary[Key].Trim().ToLower();
-				return Value == "true" || Value == "1" || Value == "yes";
-			}
-
-			return DefaultValue;
-		}
-
-		/// <summary>
-		/// Attempts to create the given directory
-		/// </summary>
-		protected static bool CreateCheckDirectory(string TargetDirectory, ILogger Logger)
-		{
-			if (!Directory.Exists(TargetDirectory))
-			{
-				try
-				{
-					Directory.CreateDirectory(TargetDirectory);
-				}
-				catch (Exception)
-				{
-					Logger.LogError("Could not create directory {TargetDir}.", TargetDirectory);
-					return false;
-				}
-				if (!Directory.Exists(TargetDirectory))
-				{
-					Logger.LogError("Path {TargetDir} does not exist or is not a directory.", TargetDirectory);
-					return false;
-				}
-			}
-			return true;
-		}
-
-
-        /// <summary>
-        /// Runs Makepri. Blocking
-        /// </summary>
-        /// <param name="CommandLine">Commandline</param>
-        /// <returns>bool    Application ran successfully</returns>
-        protected bool RunMakePri(string CommandLine)
-        {
-			string PriExecutable = GetMakePriBinaryPath();
-
-			if (File.Exists(PriExecutable) == false)
-			{
-				throw new BuildException("BUILD FAILED: Couldn't find the makepri executable: {0}", PriExecutable);
-			}
-
-			StringBuilder ProcessOutput = new StringBuilder();
-			void LocalProcessOutput(DataReceivedEventArgs Args)
-			{
-				if (Args != null && Args.Data != null)
-				{
-					ProcessOutput.AppendLine(Args.Data.TrimEnd());
-				}
-			}
-
-			ProcessStartInfo StartInfo = new ProcessStartInfo(PriExecutable, CommandLine);			
-			StartInfo.UseShellExecute = false;
-			StartInfo.CreateNoWindow = true;
-			StartInfo.StandardOutputEncoding = Encoding.Unicode;
-			StartInfo.StandardErrorEncoding = Encoding.Unicode;
-
-			Process LocalProcess = new Process();
-			LocalProcess.StartInfo = StartInfo;
-			LocalProcess.OutputDataReceived += (Sender, Args) => { LocalProcessOutput(Args); };
-			LocalProcess.ErrorDataReceived += (Sender, Args) => { LocalProcessOutput(Args); };
-			int ExitCode = Utils.RunLocalProcess(LocalProcess);
-
-			if (ExitCode == 0)
-			{
-				Logger.LogDebug("Output", ProcessOutput.ToString());
-				return true;
-			}
-			else
-			{
-				Logger.LogInformation("Output", ProcessOutput.ToString());
-				Logger.LogError("{File} returned an error.", Path.GetFileName(PriExecutable));
-				Logger.LogError("Exit code: {Code}", ExitCode);
-				return false;
-			}
-        }
-
 
 		/// <summary>
 		/// Returns a valid version of the given package version string
@@ -337,13 +222,12 @@ namespace UnrealBuildTool
 			bool bHasPerCultureResources = false;
 
 			// reset per-culture strings and make sure the default culture entry exists
-			PerAppXCultureIdStrings = new ();
-			PerAppXCultureIdStrings[""] = new();
+			AppXResources!.ClearStrings();
 
 			// add all default strings
 			if (EngineIni!.GetString(IniSection_PlatformTargetSettings, "CultureStringResources", out string DefaultCultureScratchValue) && ConfigHierarchy.TryParse(DefaultCultureScratchValue, out Dictionary<string,string>? DefaultStrings))
 			{
-				PerAppXCultureIdStrings[""] = DefaultStrings;
+				AppXResources!.AddDefaultStrings(DefaultStrings);
 			}
 
 			// add per culture strings
@@ -375,6 +259,7 @@ namespace UnrealBuildTool
 						Logger.LogWarning("Missing CultureId value: {Culture}", CultureResources);
 						continue;
 					}
+					AppXResources.AddCulture(CultureId);
 
 					// read culture strings
 					if (!ConfigHierarchy.TryParse(CultureProperties["CultureStringResources"], out Dictionary<string, string>? CultureStringResources))
@@ -382,7 +267,7 @@ namespace UnrealBuildTool
 						Logger.LogError("Invalid culture string resources: \"{Culture}\". Unable to add resource entry.", CultureResources);
 						continue;
 					}
-					PerAppXCultureIdStrings[CultureId] = CultureStringResources;
+					AppXResources!.AddCultureStrings(CultureId, CultureStringResources);
 
 					// build culture lookup table
 					UEStageIdToAppXCultureId[StageId] = CultureId;
@@ -394,300 +279,44 @@ namespace UnrealBuildTool
 		}
 
 
-		private bool RemoveStaleResourceFiles()
-		{
-			string TargetResourceDir = Path.Combine(OutputPath!, BuildResourceSubPath);
-			if (!Directory.Exists(TargetResourceDir))
-			{
-				return false;
-			}
-
-			// remove all resource files that should not be included
-			var TargetResourceInstances = Directory.EnumerateFiles(TargetResourceDir, "*.*", SearchOption.AllDirectories);
-
-			var StaleResourceFiles = TargetResourceInstances.Where(X => !ManifestFiles!.ContainsKey(X)).ToList();
-			if (StaleResourceFiles.Any())
-			{
-				Logger.LogDebug("Removing stale manifest resource files...");
-				foreach (string StaleResourceFile in StaleResourceFiles)
-				{
-					// try to delete the file & the directory that contains it
-					try
-					{
-						Logger.LogDebug("    removing {Path}", Utils.MakePathRelativeTo(StaleResourceFile, OutputPath!));
-						FileUtils.ForceDeleteFile(StaleResourceFile);
-						if (!Directory.EnumerateFileSystemEntries(Path.GetDirectoryName(StaleResourceFile)!).Any())
-						{
-							Directory.Delete(Path.GetDirectoryName(StaleResourceFile)!, false);
-						}
-					}
-					catch (Exception E)
-					{
-						Logger.LogError("    Could not remove {StaleResourceFile} - {Message}.", StaleResourceFile, E.Message);
-					}
-				}
-			}
-
-			return StaleResourceFiles.Any();
-		}
-
-
 		/// <summary>
-		/// Attempts to locate the given resource binary file in several known folder locations
+		/// Register the locations where resource binary files can be found
 		/// </summary>
-		protected virtual bool FindResourceBinaryFile( out string SourceFilePath, string ResourceFileName, string CultureId = "", bool AllowEngineFallback = true)
+		protected virtual void PrepareResourceBinaryPaths()
 		{
-			bool bFileExists = false;
-			SourceFilePath = "";
-
 			if (ProjectFile != null)
 			{
-				// look in project normal Build location
-				SourceFilePath = Path.Combine(ProjectFile.Directory.FullName!, "Build", Platform.ToString(), BuildResourceSubPath, CultureId, ResourceFileName);
-				bFileExists = File.Exists(SourceFilePath);
-
-				// look in Platform Extensions next
-				if (!bFileExists)
-				{
-					SourceFilePath = Path.Combine(ProjectFile.Directory.FullName!, "Platforms", Platform.ToString(), "Build", BuildResourceSubPath, CultureId, ResourceFileName);
-					bFileExists = File.Exists(SourceFilePath);
-				}
-			}
-
-			// look in Engine, if allowed
-			if (!bFileExists && AllowEngineFallback)
-			{
-				SourceFilePath = Path.Combine(Unreal.EngineDirectory.FullName, "Build", Platform.ToString(), EngineResourceSubPath, CultureId, ResourceFileName);
-				bFileExists = File.Exists(SourceFilePath);
-
-				// look in Platform extensions too
-				if (!bFileExists)
-				{
-					SourceFilePath = Path.Combine(Unreal.EngineDirectory.FullName, "Platforms", Platform.ToString(), "Build", EngineResourceSubPath, CultureId, ResourceFileName);
-					bFileExists = File.Exists(SourceFilePath);
-				}
-			}
-
-			return bFileExists;
-		}
-
-		/// <summary>
-		/// Determines whether the given resource binary file can be found in one of the known folder locations
-		/// </summary>
-		protected bool DoesResourceBinaryFileExist(string ResourceFileName, string CultureId = "", bool AllowEngineFallback = true)
-		{
-			string SourcePath;
-			return FindResourceBinaryFile( out SourcePath, ResourceFileName, CultureId, AllowEngineFallback );
-		}
-
-
-		/// <summary>
-		/// Adds the given resource binary file(s) to the manifest files
-		/// </summary>
-		protected bool AddResourceBinaryFileReference(string ResourceFileName, bool AllowEngineFallback = true)
-		{
-			string TargetPath = Path.Combine(OutputPath!, BuildResourceSubPath);
-			string SourceFilePath;
-			bool bFileExists = FindResourceBinaryFile(out SourceFilePath, ResourceFileName, CultureId:"", AllowEngineFallback);
-
-			// At least the default culture entry for any resource binary must always exist
-			if (!bFileExists)
-			{
-				return false;
-			}
-
-			// Copy the default resource file
-			AddFileReference(SourceFilePath, Path.Combine(TargetPath, ResourceFileName), bIsGeneratedFile:false);
-
-			// Copy all per-culture resource files
-			foreach( string CultureId in SelectedAppXCultureIds!)
-			{
-				bFileExists = FindResourceBinaryFile(out SourceFilePath, ResourceFileName, CultureId, AllowEngineFallback);
-				if (bFileExists)
-				{
-					AddFileReference(SourceFilePath, Path.Combine(TargetPath, CultureId, ResourceFileName), bIsGeneratedFile:false);
-				}
-			}
-
-			return true;
-		}
-
-
-		/// <summary>
-		/// Adds the given file to the manifest files
-		/// </summary>
-		protected void AddFileReference(string SourcePath, string TargetPath, bool bIsGeneratedFile)
-		{
-			// check if the file is unchanged
-			if (File.Exists(TargetPath) && !string.IsNullOrEmpty(SourcePath) )
-			{
-				FileInfo SrcFileInfo = new FileInfo(SourcePath);
-				FileInfo DstFileInfo = new FileInfo(TargetPath);
-
-				bool bFileIsUnchanged = (SrcFileInfo.Length == DstFileInfo.Length);
-				if (bFileIsUnchanged)
-				{
-					if (bIsGeneratedFile)
-					{
-						// this file is auto-generated - we need to compare the contents
-						byte[] OriginalContents = File.ReadAllBytes(TargetPath);
-						byte[] NewContents = File.ReadAllBytes(SourcePath);
-						bFileIsUnchanged = Enumerable.SequenceEqual(OriginalContents, NewContents);
-					}
-					else
-					{
-						// this file is not generated, just copied from somewhere else - can check the time to confirm if they're different
-						bFileIsUnchanged = (SrcFileInfo.CreationTime == DstFileInfo.CreationTime);
-					}
-				}
-
-				if (bFileIsUnchanged)
-				{
-					// use an empty source string if the file doesn't need changing
-					SourcePath = "";
-				}
-			}
-
-			ManifestFiles!.Add(TargetPath, SourcePath);
-		}
-
-
-		/// <summary>
-		/// Copies all of the generated files to the output folder
-		/// </summary>
-		private List<string> CopyFilesToOutput()
-		{
-			List<string> UpdatedFiles = new List<string>();
-
-			// early out if there's nothing to copy
-			if (ManifestFiles == null || !ManifestFiles.Any(X => !string.IsNullOrEmpty(X.Value)))
-			{
-				return UpdatedFiles;
-			}
-
-			Logger.LogDebug("Updating manifest resource files...");
-
-			// copy over any new or updated files
-			foreach ( var ManifestFilePair in ManifestFiles! )
-			{
-				string TargetPath = ManifestFilePair.Key;
-				string SourcePath = ManifestFilePair.Value;
-				if (string.IsNullOrEmpty(SourcePath))
-				{
-					// if source is an empty string then the file is up-to-date
-					continue;
-				}
-				
-				// remove old version, if any
-				bool bFileExists = File.Exists(TargetPath);
-				if (bFileExists)
-				{
-					try
-					{
-						FileUtils.ForceDeleteFile(TargetPath);
-					}
-					catch (Exception E)
-					{
-						Logger.LogError("    Could not replace file {TargetPath} - {Message}", TargetPath, E.Message);
-					}
-				}
-
-				// copy new version
-				try
-				{
-					if (bFileExists)
-					{
-						Logger.LogDebug("    updating {Path}", Utils.MakePathRelativeTo(TargetPath, OutputPath!));
-					}
-					else
-					{
-						Logger.LogDebug("    adding {Path}", Utils.MakePathRelativeTo(TargetPath, OutputPath!));
-					}
-
-					Directory.CreateDirectory(Path.GetDirectoryName(TargetPath)!);
-					File.Copy(SourcePath, TargetPath);
-					File.SetAttributes(TargetPath, FileAttributes.Normal);
-					File.SetCreationTime(TargetPath, File.GetCreationTime(SourcePath));
-
-					UpdatedFiles.Add(TargetPath);
-				}
-				catch (Exception E)
-				{
-					Logger.LogError("    Unable to copy file {TargetPath} - {Message}", TargetPath, E.Message);
-				}
-			}
-
-			return UpdatedFiles;
-		}
-
-
-		/// <summary>
-		/// Adds the given string to the culture string writers
-		/// </summary>
-		protected string AddResourceEntry(string ResourceEntryName, string ConfigKey, string DefaultValue, string ValueSuffix = "")
-		{
-			// Get the default culture value
-			string? DefaultCultureString = null;
-			if (PerAppXCultureIdStrings!.TryGetValue("", out Dictionary<string,string>? DefaultCultureValues))
-			{
-				if (DefaultCultureValues!.TryGetValue(ConfigKey, out string? Value) && Value != null)
-				{
-					DefaultCultureString = Value;
-				}
-			}
-			if (string.IsNullOrEmpty(DefaultCultureString))
-			{
-				DefaultCultureString = DefaultValue;
-			}
-			DefaultResourceWriter!.AddResource(ResourceEntryName, DefaultCultureString + ValueSuffix);
-
-			// Get the localized culture values
-			foreach (string CultureId in SelectedAppXCultureIds!)
-			{
-				if (PerAppXCultureIdStrings.TryGetValue(CultureId, out Dictionary<string,string>? CultureValues))
-				{
-					if (CultureValues.TryGetValue(ConfigKey, out string? CultureString) && !string.IsNullOrEmpty(CultureString) )
-					{
-						var Writer = PerCultureResourceWriters![CultureId];
-						Writer.AddResource(ResourceEntryName, CultureString + ValueSuffix);
-					}
-				}
+				AppXResources!.ProjectBinaryResourceDirectories.Add( DirectoryReference.Combine(ProjectFile.Directory, "Build", Platform.ToString(), BuildResourceSubPath) );
+				AppXResources!.ProjectBinaryResourceDirectories.Add( DirectoryReference.Combine(ProjectFile.Directory, "Platforms", Platform.ToString(), "Build", BuildResourceSubPath) );
 			}
 
-			return "ms-resource:" + ResourceEntryName;
+			AppXResources!.EngineFallbackBinaryResourceDirectories.Add( DirectoryReference.Combine(Unreal.EngineDirectory, "Build", Platform.ToString(), EngineResourceSubPath) );
+			AppXResources!.EngineFallbackBinaryResourceDirectories.Add( DirectoryReference.Combine(Unreal.EngineDirectory, "Platforms", Platform.ToString(), "Build", EngineResourceSubPath) );
 		}
 
 
-
-		/// <summary>
-		/// Get the XName from a given string and schema pair
-		/// </summary>
-		protected virtual XName GetName( string BaseName, string SchemaName )
-		{
-			return XName.Get(BaseName);
-		}
-
 		/// <summary>
 		/// Get the resources element
 		/// </summary>
 		protected XElement GetResources()
 		{
-			List<string> ResourceCulturesList = new List<string>(SelectedAppXCultureIds!);
+			List<string> ResourceCulturesList = AppXResources!.GetAllCultureIds().ToList();
+
 			// Move the default culture to the front of the list
 			ResourceCulturesList.Remove(DefaultAppXCultureId!);
 			ResourceCulturesList.Insert(0, DefaultAppXCultureId!);
 
 			// Check that we have a valid number of cultures
-			if (SelectedAppXCultureIds!.Count < 1 || SelectedAppXCultureIds.Count >= MaxResourceEntries)
+			if (ResourceCulturesList!.Count < 1 || ResourceCulturesList.Count >= MaxResourceEntries)
 			{
 				Logger.LogWarning("Incorrect number of cultures to stage. There must be between 1 and {MaxCultures} cultures selected.", MaxResourceEntries);
 			}
 
 			// Create the culture list. This list is unordered except that the default language must be first which we already took care of above.
 			var CultureElements = ResourceCulturesList.Select(c =>
-				new XElement(GetName("Resource", Schema2010NS), new XAttribute("Language", c)));
+				new XElement("Resource", new XAttribute("Language", c)));
 
-			return new XElement(GetName("Resources", Schema2010NS), CultureElements);
+			return new XElement("Resources", CultureElements);
 		}
 
 		/// <summary>
@@ -753,7 +382,7 @@ namespace UnrealBuildTool
 
             IdentityName = PackageName;
 
-            return new XElement(GetName("Identity", Schema2010NS),
+            return new XElement("Identity",
                 new XAttribute("Name", PackageName),
                 new XAttribute("Publisher", PublisherName),
                 new XAttribute("Version", VersionNumber!));
@@ -782,7 +411,7 @@ namespace UnrealBuildTool
 		/// <summary>
 		/// Get the path to the makepri.exe tool
 		/// </summary>
-		protected abstract string GetMakePriBinaryPath();
+		protected abstract FileReference GetMakePriBinaryPath();
 
 		/// <summary>
 		/// Get any additional platform-specific parameters for makepri.exe
@@ -814,7 +443,7 @@ namespace UnrealBuildTool
 		/// <summary>
 		/// Create a manifest and return the list of modified files
 		/// </summary>
-		public List<string>? CreateManifest(string InManifestName, string InOutputPath, string InIntermediatePath, string? InTargetName, FileReference? InProjectFile, Dictionary<UnrealTargetConfiguration,string> InExecutablePairs)
+		public List<string>? CreateManifest(string InManifestName, DirectoryReference OutputDirectory, string? InTargetName, FileReference? InProjectFile, Dictionary<UnrealTargetConfiguration,string> InExecutablePairs)
 		{
 			// Check parameter values are valid.
 			if (InExecutablePairs.Count < 1)
@@ -823,27 +452,22 @@ namespace UnrealBuildTool
 				return null;
 			}
 
-			if (!CreateCheckDirectory(InOutputPath, Logger))
-			{
-				Logger.LogError("Failed to create output directory \"{OutputDir}\".", InOutputPath);
-				return null;
-			}
-			if (!CreateCheckDirectory(InIntermediatePath, Logger))
-			{
-				Logger.LogError("Failed to create intermediate directory \"{IntDir}\".", InIntermediatePath);
-				return null;
-			}
+			FileUtils.CreateDirectoryTree(OutputDirectory);
+
+			DirectoryReference ProjectRoot = (InProjectFile != null) ? InProjectFile.Directory : Unreal.EngineDirectory;
+			DirectoryReference IntermediateDirectory = DirectoryReference.Combine(ProjectRoot, "Intermediate", "Manifest", Platform.ToString());
+			FileUtils.ForceDeleteDirectory(IntermediateDirectory);
+			FileUtils.CreateDirectoryTree(IntermediateDirectory);
 
-			OutputPath = InOutputPath;
-			IntermediatePath = InIntermediatePath;
 			TargetName = InTargetName;
 			ProjectFile = InProjectFile;
-			ManifestFiles = new Dictionary<string, string>();
+			AppXResources = new(Logger, GetMakePriBinaryPath());
+			PrepareResourceBinaryPaths();
 
 			// Load up INI settings. We'll use engine settings to retrieve the manifest configuration, but these may reference
 			// values in either game or engine settings, so we'll keep both.
-			GameIni = ConfigCache.ReadHierarchy(ConfigHierarchyType.Game, DirectoryReference.FromFile(InProjectFile), ConfigPlatform);
-			EngineIni = ConfigCache.ReadHierarchy(ConfigHierarchyType.Engine, DirectoryReference.FromFile(InProjectFile), ConfigPlatform);
+			GameIni = ConfigCache.ReadHierarchy(ConfigHierarchyType.Game, DirectoryReference.FromFile(InProjectFile), Platform);
+			EngineIni = ConfigCache.ReadHierarchy(ConfigHierarchyType.Engine, DirectoryReference.FromFile(InProjectFile), Platform);
 			PostConfigurationInit();
 
 			// Load and verify/clean culture list
@@ -872,7 +496,6 @@ namespace UnrealBuildTool
 			bool bHasLocalizationData = BuildLocalizationData();
 
 			// generate the list of AppX cultures to stage
-			SelectedAppXCultureIds = new();
 			foreach (string UEStageId in SelectedUECultureIds)
 			{
 				if (!UEStageIdToAppXCultureId.TryGetValue(UEStageId, out string? AppXCultureId) || string.IsNullOrEmpty(AppXCultureId) )
@@ -880,9 +503,9 @@ namespace UnrealBuildTool
 					// use the culture directly - no remapping required
 					AppXCultureId = UEStageId;
 					UEStageIdToAppXCultureId[UEStageId] = AppXCultureId;
+					AppXResources.AddCulture(UEStageId);
 				}
 
-				SelectedAppXCultureIds.Add(AppXCultureId);
 			}
 
 			// look up the default AppX culture
@@ -900,139 +523,24 @@ namespace UnrealBuildTool
 				Logger.LogInformation("Staged culture mappings not setup in the editor. See Per Culture Resources in the {Platform} Target Settings.", Platform.ToString());
 			}
 
-			// Clean out the resources intermediate path so that we know there are no stale binary files.
-			string IntermediateResourceDirectory = Path.Combine(IntermediatePath, BuildResourceSubPath);
-			FileUtils.ForceDeleteDirectory(IntermediateResourceDirectory);
-			if (!CreateCheckDirectory(IntermediateResourceDirectory, Logger))
-			{
-				Logger.LogError("Could not create directory {IntDir}.", IntermediateResourceDirectory);
-				return null;
-			}
-
-			// Construct a single resource writer for the default (no-culture) values
-			string DefaultResourceIntermediatePath = Path.Combine(IntermediateResourceDirectory, "resources.resw");
-			DefaultResourceWriter = new UEResXWriter(DefaultResourceIntermediatePath);
-
-			// Construct the ResXWriters for each culture
-			PerCultureResourceWriters = new Dictionary<string, UEResXWriter>();
-			foreach (string AppXCultureId in SelectedAppXCultureIds)
-			{
-				string IntermediateStringResourcePath = Path.Combine(IntermediateResourceDirectory, AppXCultureId);
-				string IntermediateStringResourceFile = Path.Combine(IntermediateStringResourcePath, "resources.resw");
-				if (!CreateCheckDirectory(IntermediateStringResourcePath, Logger))
-				{
-					Logger.LogError("Could not create directory {IntDir}.", IntermediateStringResourcePath);
-					return null;
-				}
-				PerCultureResourceWriters.Add(AppXCultureId, new UEResXWriter(IntermediateStringResourceFile));
-			}
-
-
-
 			// Create the manifest document
 			string? IdentityName = null;
 			var ManifestXmlDocument = new XDocument(GetManifest(InExecutablePairs, out IdentityName));
 
-			// Export manifest to the intermediate directory then compare the contents to any existing target manifest
-			// and replace if there are differences.
-			string ManifestIntermediatePath = Path.Combine(IntermediatePath, InManifestName);
-			string ManifestTargetPath = Path.Combine(OutputPath, InManifestName);
-			ManifestXmlDocument.Save(ManifestIntermediatePath);
-			AddFileReference(ManifestIntermediatePath, ManifestTargetPath, bIsGeneratedFile: true);
-			ProcessManifest(InExecutablePairs, InManifestName, ManifestTargetPath, ManifestIntermediatePath);
-
-			// Export the resource tables starting with the default culture
-			string DefaultResourceTargetPath = Path.Combine(OutputPath, BuildResourceSubPath, "resources.resw");
-			DefaultResourceWriter.Close();
-			AddFileReference(DefaultResourceIntermediatePath, DefaultResourceTargetPath, bIsGeneratedFile: true);
-
-			foreach (var Writer in PerCultureResourceWriters)
-			{
-				Writer.Value.Close();
-
-				string IntermediateStringResourceFile = Path.Combine(IntermediateResourceDirectory, Writer.Key, "resources.resw");
-				string TargetStringResourceFile = Path.Combine(OutputPath, BuildResourceSubPath, Writer.Key, "resources.resw");
-
-				AddFileReference(IntermediateStringResourceFile, TargetStringResourceFile, bIsGeneratedFile: true);
-			}
-
-
-			// include a reference to the Package Resource Index so it isn't removed by RemoveStaleResourceFiles
-			string TargetResourceIndexFile = Path.Combine(OutputPath, "resources.pri");
-			AddFileReference("", TargetResourceIndexFile, bIsGeneratedFile: true);
-
-			// The resource database is dependent on everything else calculated here (manifest, resource string tables, binary resources).
-			// So if any file has been updated we'll need to run the config.
-			bool bHadStaleResources = RemoveStaleResourceFiles();
-			List<string> UpdatedFilePaths = CopyFilesToOutput();
-
-			if (bHadStaleResources || UpdatedFilePaths.Any())
-			{
-				// Create resource index configuration
-				string ResourceConfigFile = Path.Combine(IntermediatePath, "priconfig.xml");
-				RunMakePri("createconfig /cf \"" + ResourceConfigFile + "\" /dq " + DefaultAppXCultureId + " /o " + GetMakePriExtraCommandLine());
-
-				// Load the new resource index configuration
-				XmlDocument PriConfig = new XmlDocument();
-				PriConfig.Load(ResourceConfigFile);
-
-				// remove the packaging node - we do not want to split the pri & only want one .pri file
-				XmlNode PackagingNode = PriConfig.SelectSingleNode("/resources/packaging")!;
-				PackagingNode.ParentNode!.RemoveChild(PackagingNode);
-
-				// all required resources are explicitly listed in resources.resfiles, rather than relying on makepri to discover them
-				string ResourcesResFile = Path.Combine(IntermediatePath, "resources.resfiles");
-				XmlNode PriIndexNode = PriConfig.SelectSingleNode("/resources/index")!;
-				XmlAttribute PriStartIndex = PriIndexNode.Attributes!["startIndexAt"]!;
-				PriStartIndex.Value = ResourcesResFile;
-
-				// swap the folder indexer-config to a RESFILES indexer-config.
-				XmlElement FolderIndexerConfigNode = (XmlElement)PriConfig.SelectSingleNode("/resources/index/indexer-config[@type='folder']")!;
-				FolderIndexerConfigNode.SetAttribute("type", "RESFILES");
-				FolderIndexerConfigNode.RemoveAttribute("foldernameAsQualifier");
-				FolderIndexerConfigNode.RemoveAttribute("filenameAsQualifier");
-
-				PriConfig.Save(ResourceConfigFile);
-
-				// generate resources.resfiles
-				IEnumerable<string> Resources = Directory.EnumerateFiles(Path.Combine(OutputPath, BuildResourceSubPath), "*.*", SearchOption.AllDirectories);
-				System.Text.StringBuilder ResourcesList = new System.Text.StringBuilder();
-				foreach (string Resource in Resources)
-				{
-					ResourcesList.AppendLine( Utils.MakePathRelativeTo( Resource, OutputPath ) );
-				}
-				File.WriteAllText(ResourcesResFile, ResourcesList.ToString());
-
-				// remove old Package Resource Index
-				try
-				{
-					FileUtils.ForceDeleteFile(TargetResourceIndexFile);
-				}
-				catch (Exception E)
-				{
-					Logger.LogError("cannot remove old pri file: {TargetResourceIndexFile} - {Message}", TargetResourceIndexFile, E.Message);
-				}
-
-				// Generate the Package Resource Index
-				string ResourceLogFile = Path.Combine(IntermediatePath, "ResIndexLog.xml");
-				string MakePriCommandLine = "new /pr \"" + OutputPath + "\" /cf \"" + ResourceConfigFile + "\" /mn \"" + ManifestTargetPath + "\" /il \"" + ResourceLogFile + "\" /of \"" + TargetResourceIndexFile + "\" /o";
-				if (IdentityName != null)
-				{
-					MakePriCommandLine += " /indexName \"" + IdentityName + "\"";
-				}
+			// Export manifest to the intermediate directory and add it to the manifest resources files for copying
+			FileReference ManifestIntermediateFile = FileReference.Combine(IntermediateDirectory, InManifestName);
+			FileReference ManifestTargetFile = FileReference.Combine(OutputDirectory, InManifestName);
+			ManifestXmlDocument.Save(ManifestIntermediateFile.FullName);
+			AppXResources.AddFileReference(ManifestIntermediateFile, InManifestName);
+			ProcessManifest(InExecutablePairs, InManifestName, ManifestTargetFile.FullName, ManifestIntermediateFile.FullName);
 
-				Logger.LogDebug("    generating {Path}", Utils.MakePathRelativeTo(TargetResourceIndexFile, OutputPath!));
-				RunMakePri(MakePriCommandLine);
-				UpdatedFilePaths.Add(TargetResourceIndexFile);
-			}
-
-			// Report if nothing was changed
-			if (!bHadStaleResources && !UpdatedFilePaths.Any())
-			{
-				Logger.LogDebug($"Manifest resource files are up to date");
-			}
+			// Generate the package resource index and copy all resource files to the output
+			FileReference ManifestTargetPath = FileReference.Combine( OutputDirectory, InManifestName);
+			List<FileReference> UpdatedFiles = AppXResources.GenerateAppXResources( OutputDirectory, IntermediateDirectory, ManifestTargetFile, DefaultAppXCultureId, IdentityName );
 
-			return UpdatedFilePaths;
+			// Clean up and reutrn the list of updated files
+			FileUtils.ForceDeleteDirectory(IntermediateDirectory);
+			return UpdatedFiles.ConvertAll( X => X.FullName );
 		}
 	}
 }
diff --git a/Engine/Source/Programs/UnrealBuildTool/Platform/Windows/ResXResource/UEAppXResources.cs b/Engine/Source/Programs/UnrealBuildTool/Platform/Windows/ResXResource/UEAppXResources.cs
new file mode 100644
--- /dev/null
+++ b/Engine/Source/Programs/UnrealBuildTool/Platform/Windows/ResXResource/UEAppXResources.cs
@@ -0,0 +1,578 @@
+// Copyright Epic Games, Inc. All Rights Reserved.
+
+using System;
+using System.Collections.Generic;
+using System.Linq;
+using System.IO;
+using System.Text;
+using EpicGames.Core;
+using System.Diagnostics;
+using System.Xml;
+using System.Diagnostics.CodeAnalysis;
+using Microsoft.Extensions.Logging;
+
+namespace UnrealBuildTool
+{
+	/// <summary>
+	/// AppX Package Resource Index generator
+	/// </summary>
+	public class UEAppXResources
+	{
+		const string BuildResourceSubPath = "Resources";
+		const string DefaultResources = "";
+
+		class Resources
+		{
+			public Resources() { }
+
+			public Dictionary<string, string> SourceStrings = new();    // raw source strings, as registered by AddCultureString etc.
+			public Dictionary<string, string> StringResources = new();  // string resources that will be used, as defined by AddResourceString
+		};
+
+		readonly Dictionary<string, Resources> PerCultureResources = new();
+		readonly Dictionary<FileReference,string> FilesToCopy = new(); // source file -> target file fragment (relative to OutputDir)
+
+		readonly ILogger Logger;
+		readonly FileReference MakePriExe; // path to makepri.exe
+		
+		/// <summary>
+		/// Folders where binary resource files are located, in decreasing priority order
+		/// </summary>
+		public List<DirectoryReference> ProjectBinaryResourceDirectories { get; }
+
+		/// <summary>
+		/// Folders where fallback binary resource files are located, in decreasing priority order
+		/// </summary>
+		public List<DirectoryReference> EngineFallbackBinaryResourceDirectories { get; }
+
+		/// <summary>
+		/// Target OS version
+		/// </summary>
+		public string TargetOSVersion { get; set; }
+
+		/// <summary>
+		/// Create a new instance
+		/// </summary>
+		public UEAppXResources(ILogger InLogger, FileReference InMakePriExe)
+		{
+			Logger = InLogger;
+			MakePriExe = InMakePriExe;
+			ProjectBinaryResourceDirectories = new();
+			EngineFallbackBinaryResourceDirectories = new();
+			TargetOSVersion = "10.0.0";
+
+			if (FileReference.Exists(MakePriExe) == false)
+			{
+				throw new BuildException($"Couldn't find the makepri executable: {MakePriExe}");
+			}
+
+			AddCulture(DefaultResources);
+		}
+
+
+		/// <summary>
+		/// Adds a reference to the given culture
+		/// </summary>
+		public void AddCulture( string AppXCultureId )
+		{
+			if (!PerCultureResources.TryGetValue(AppXCultureId, out Resources? Result) || Result == null)
+			{
+				PerCultureResources.Add( AppXCultureId, new());
+			}
+		}
+
+		/// <summary>
+		/// Adds a reference to all of the given cultures
+		/// </summary>
+		/// <param name="AppXCultureIds"></param>
+		public void AddCultures( IEnumerable<string> AppXCultureIds )
+		{
+			foreach (string AppXCultureId in AppXCultureIds )
+			{
+				AddCulture(AppXCultureId);
+			}
+		}
+
+
+		/// <summary>
+		/// Returns a collection of all registered cultures
+		/// </summary>
+		/// <returns></returns>
+		public IEnumerable<string> GetAllCultureIds()
+		{
+			return PerCultureResources.Keys.Where( X => X != DefaultResources );
+		}
+
+
+		/// <summary>
+		/// Add the given string to the default culture
+		/// </summary>
+		public void AddDefaultString( string ConfigKey, string Value )
+		{
+			AddCultureString(DefaultResources, ConfigKey, Value);
+		}
+
+		/// <summary>
+		/// Adds all of the given strings to the default culture
+		/// </summary>
+		public void AddDefaultStrings( Dictionary<string, string> Strings )
+		{
+			AddCultureStrings(DefaultResources, Strings );
+		}
+
+		/// <summary>
+		/// Add the given string to the given culture
+		/// </summary>
+		public void AddCultureString( string AppXCultureId, string ConfigKey, string Value )
+		{
+			GetCultureResources(AppXCultureId).SourceStrings[ConfigKey] = Value;
+		}
+
+		/// <summary>
+		/// Adds all of the given strings to the given culture
+		/// </summary>
+		public void AddCultureStrings( string AppXCultureId, Dictionary<string, string> Strings )
+		{
+			Resources ThisCultureResources = GetCultureResources(AppXCultureId);
+			foreach (KeyValuePair<string,string> String in Strings)
+			{
+				ThisCultureResources.SourceStrings[String.Key] = String.Value;
+			}
+		}
+
+		/// <summary>
+		/// Returns whether the given string is known
+		/// </summary>
+		public bool HasString( string Key )
+		{
+			return GetDefaultResources().SourceStrings.ContainsKey(Key);
+		}
+
+
+		/// <summary>
+		/// Clear all stored strings for all cultures
+		/// </summary>
+		public void ClearStrings()
+		{
+			foreach (KeyValuePair<string, Resources> Itr in PerCultureResources)
+			{
+				Itr.Value.SourceStrings.Clear();
+			}
+		}
+
+
+		Resources GetCultureResources( string AppXCultureId )
+		{
+			if (PerCultureResources.TryGetValue(AppXCultureId, out Resources? Result) && Result != null)
+			{
+				return Result;
+			}
+			else
+			{
+				throw new Exception($"Culture {AppXCultureId} is not known. Missing call to AddCulture?");
+			}
+		}
+
+		Resources GetDefaultResources()
+		{
+			return GetCultureResources(DefaultResources);
+		}
+
+
+
+
+        bool RunMakePri(string CommandLine)
+        {
+			StringBuilder ProcessOutput = new StringBuilder();
+			void LocalProcessOutput(DataReceivedEventArgs Args)
+			{
+				if (Args != null && Args.Data != null)
+				{
+					ProcessOutput.AppendLine(Args.Data.TrimEnd());
+				}
+			}
+
+			ProcessStartInfo StartInfo = new ProcessStartInfo(MakePriExe.FullName, CommandLine);			
+			StartInfo.UseShellExecute = false;
+			StartInfo.CreateNoWindow = true;
+			StartInfo.StandardOutputEncoding = Encoding.Unicode;
+			StartInfo.StandardErrorEncoding = Encoding.Unicode;
+
+			Process LocalProcess = new Process();
+			LocalProcess.StartInfo = StartInfo;
+			LocalProcess.OutputDataReceived += (Sender, Args) => { LocalProcessOutput(Args); };
+			LocalProcess.ErrorDataReceived += (Sender, Args) => { LocalProcessOutput(Args); };
+			int ExitCode = Utils.RunLocalProcess(LocalProcess);
+
+			if (ExitCode == 0)
+			{
+				Logger.LogDebug("{Output}", ProcessOutput.ToString());
+				return true;
+			}
+			else
+			{
+				Logger.LogInformation("{Output}", ProcessOutput.ToString());
+				Logger.LogError("{File} returned an error.", MakePriExe.GetFileName());
+				Logger.LogError("Exit code: {Code}", ExitCode);
+				return false;
+			}
+        }
+
+
+
+
+		private bool RemoveStaleResourceFiles( IEnumerable<string> RequiredFileFragments, DirectoryReference OutputDirectory )
+		{
+			DirectoryReference ResourceDirectory = DirectoryReference.Combine( OutputDirectory, BuildResourceSubPath );
+			if (!DirectoryReference.Exists(ResourceDirectory))
+			{
+				return false;
+			}
+
+			// remove all files in the Resources/ subfolder that should not be included
+			IEnumerable<FileReference> ExistingFiles = DirectoryReference.EnumerateFiles(ResourceDirectory, "*.*", SearchOption.AllDirectories);
+			IEnumerable<FileReference> RequiredFiles = RequiredFileFragments
+				.Select( X => FileReference.Combine( OutputDirectory, X ))
+				.Where( X => X.IsUnderDirectory(ResourceDirectory))
+				;
+			IEnumerable<FileReference> StaleResourceFiles = ExistingFiles.Except(RequiredFiles);
+			if (!StaleResourceFiles.Any())
+			{
+				return false;
+			}
+
+			Logger.LogDebug("Removing stale manifest resource files...");
+			foreach (FileReference StaleResourceFile in StaleResourceFiles)
+			{
+				// try to delete the file & the directory that contains it
+				try
+				{
+					Logger.LogDebug("    removing {Path}", StaleResourceFile.MakeRelativeTo(OutputDirectory));
+					FileUtils.ForceDeleteFile(StaleResourceFile);
+					if (!Directory.EnumerateFileSystemEntries(StaleResourceFile.Directory.FullName).Any())
+					{
+						DirectoryReference.Delete(StaleResourceFile.Directory, false);
+					}
+				}
+				catch (Exception E)
+				{
+					Logger.LogError("    Could not remove {StaleResourceFile} - {Message}.", StaleResourceFile, E.Message);
+				}
+			}
+
+			return true;
+		}
+
+
+		/// <summary>
+		/// Attempts to locate the given resource binary file in the previously-specified resource folders
+		/// </summary>
+		bool FindResourceBinaryFile( [NotNullWhen(true)] out FileReference? SourceFilePath, string ResourceFileName, string CultureId = "", bool AllowEngineFallback = true)
+		{
+			// look in project binary resource directories
+			foreach (DirectoryReference BinaryResourceDirectory in ProjectBinaryResourceDirectories)
+			{
+				FileReference BinaryResourceFile = FileReference.Combine(BinaryResourceDirectory, CultureId, ResourceFileName );
+				if (FileReference.Exists(BinaryResourceFile))
+				{
+					SourceFilePath = BinaryResourceFile;
+					return true;
+				}
+			}
+
+			// look in Engine, if allowed
+			if (AllowEngineFallback )
+			{
+				foreach (DirectoryReference BinaryResourceDirectory in EngineFallbackBinaryResourceDirectories)
+				{
+					FileReference BinaryResourceFile = FileReference.Combine(BinaryResourceDirectory, CultureId, ResourceFileName );
+					if (FileReference.Exists(BinaryResourceFile))
+					{
+						SourceFilePath = BinaryResourceFile;
+						return true;
+					}
+				}
+			}
+
+			// not found
+			SourceFilePath = null;
+			return false;
+		}
+
+		/// <summary>
+		/// Determines whether the given resource binary file can be found in one of the known folder locations for the given culture
+		/// </summary>
+		public bool DoesCultureResourceBinaryFileExist(string ResourceFileName, string CultureId, bool AllowEngineFallback = true)
+		{
+			return FindResourceBinaryFile( out FileReference? _, ResourceFileName, CultureId, AllowEngineFallback );
+		}
+
+		/// <summary>
+		/// Determines whether the given default resource binary file can be found in one of the known folder locations
+		/// </summary>
+		public bool DoesDefaultResourceBinaryFileExist(string ResourceFileName, bool AllowEngineFallback = true)
+		{
+			return DoesCultureResourceBinaryFileExist( ResourceFileName, DefaultResources, AllowEngineFallback );
+		}
+
+
+		/// <summary>
+		/// Adds the given resource binary file(s) to the manifest files
+		/// </summary>
+		public bool AddResourceBinaryFileReference(string ResourceFileName, bool AllowEngineFallback = true)
+		{
+			// At least the default culture entry for any resource binary must always exist
+			if (!FindResourceBinaryFile( out FileReference? SourceFilePath, ResourceFileName, DefaultResources, AllowEngineFallback))
+			{
+				return false;
+			}
+			AddFileReference(SourceFilePath, Path.Combine(BuildResourceSubPath, ResourceFileName));
+
+			// Copy all per-culture resource files
+			foreach( string CultureId in GetAllCultureIds() )
+			{
+				if (FindResourceBinaryFile(out SourceFilePath, ResourceFileName, CultureId, AllowEngineFallback))
+				{
+					AddFileReference(SourceFilePath, Path.Combine(BuildResourceSubPath, CultureId, ResourceFileName));
+				}
+			}
+
+			return true;
+		}
+
+
+		/// <summary>
+		/// Adds the given file to the manifest files
+		/// </summary>
+		public void AddFileReference(FileReference SourcePath, string TargetPathFragment)
+		{
+			FilesToCopy.Add(SourcePath, TargetPathFragment);
+		}
+
+		private static bool AreFilesDifferent( FileReference File1, FileReference File2)
+		{
+			FileInfo FileInfo1 = File1.ToFileInfo();
+			FileInfo FileInfo2 = File2.ToFileInfo();
+			if (FileInfo1.Length != FileInfo2.Length)
+			{
+				return true;
+			}
+
+			byte[] FileContents1 = FileReference.ReadAllBytes(File1);
+			byte[] FileContents2 = FileReference.ReadAllBytes(File2);
+			return !Enumerable.SequenceEqual(FileContents1, FileContents2);
+		}
+
+
+
+		/// <summary>
+		/// Copies all of the generated files to the output folder
+		/// </summary>
+		private List<FileReference> CopyFilesToOutput(DirectoryReference OutputDirectory)
+		{
+			List<FileReference> UpdatedFiles = new();
+			if (!FilesToCopy.Any())
+			{
+				return UpdatedFiles;
+			}
+
+			Logger.LogDebug("Updating manifest resource files...");
+			foreach ( var FileToCopy in FilesToCopy )
+			{
+				FileReference SourceFile = FileToCopy.Key;
+				FileReference TargetFile = FileReference.Combine(OutputDirectory, FileToCopy.Value);
+				if (!FileReference.Exists(SourceFile))
+				{
+					Logger.LogError("    Source file not found {TargetFile}", TargetFile);
+					continue;
+				}
+				if (FileReference.Exists(TargetFile) && !AreFilesDifferent(TargetFile, SourceFile))
+				{
+					continue;
+				}
+
+				// remove old version, if any
+				bool bFileExists = FileReference.Exists(TargetFile);
+				if (bFileExists)
+				{
+					try
+					{
+						FileUtils.ForceDeleteFile(TargetFile);
+					}
+					catch (Exception E)
+					{
+						Logger.LogError("    Could not replace file {TargetFile} - {Message}", TargetFile, E.Message);
+					}
+				}
+
+				// copy new version
+				try
+				{
+					if (bFileExists)
+					{
+						Logger.LogDebug("    updating {Path}", Utils.MakePathRelativeTo(TargetFile.FullName, OutputDirectory.FullName!));
+					}
+					else
+					{
+						Logger.LogDebug("    adding {Path}", Utils.MakePathRelativeTo(TargetFile.FullName, OutputDirectory.FullName!));
+					}
+
+					Directory.CreateDirectory(Path.GetDirectoryName(TargetFile.FullName)!);
+					FileReference.Copy(SourceFile, TargetFile);
+					FileReference.SetAttributes(TargetFile, FileAttributes.Normal);
+					File.SetCreationTime(TargetFile.FullName, File.GetCreationTime(SourceFile.FullName));
+
+					UpdatedFiles.Add(TargetFile);
+				}
+				catch (Exception E)
+				{
+					Logger.LogError("    Unable to copy file {TargetFile} - {Message}", TargetFile, E.Message);
+				}
+			}
+
+			return UpdatedFiles;
+		}
+
+
+
+		/// <summary>
+		/// Adds the given string to the culture string writers
+		/// </summary>
+		public string AddResourceString(string ResourceEntryName, string ConfigKey, string DefaultValue, string ValueSuffix = "")
+		{
+			// Get the default culture value
+			Resources DefaultCultureResources = GetDefaultResources();
+			string? DefaultCultureString = null;
+			if (DefaultCultureResources.SourceStrings.TryGetValue(ConfigKey, out string? Value) && Value != null)
+			{
+				DefaultCultureString = Value;
+			}
+			if (string.IsNullOrEmpty(DefaultCultureString))
+			{
+				DefaultCultureString = DefaultValue;
+			}
+			DefaultCultureResources.StringResources.Add(ResourceEntryName, DefaultCultureString + ValueSuffix);
+
+
+			// Get the localized culture values
+			foreach (string CultureId in GetAllCultureIds())
+			{
+				Resources ThisCultureResources = GetCultureResources(CultureId);
+				if (ThisCultureResources.SourceStrings.TryGetValue(ConfigKey, out string? CultureString) && !string.IsNullOrEmpty(CultureString) )
+				{
+					ThisCultureResources.StringResources.Add(ResourceEntryName, CultureString + ValueSuffix);
+				}
+			}
+
+			return "ms-resource:" + ResourceEntryName;
+		}
+
+
+
+
+		/// <summary>
+		/// Generate the package resource index and copies all resources files to the output
+		/// </summary>
+		public List<FileReference> GenerateAppXResources( DirectoryReference OutputDirectory, DirectoryReference IntermediateDirectory, FileReference ManifestFile, string DefaultAppXCultureId, string? PackageIdentityName )
+		{
+			// Clean out the resources intermediate path so that we know there are no stale binary files.
+			FileUtils.ForceDeleteDirectory(DirectoryReference.Combine(IntermediateDirectory, BuildResourceSubPath));
+
+			// Finalize all cultures
+			foreach (KeyValuePair<string, Resources> Itr in PerCultureResources)
+			{
+				string AppXCultureId = Itr.Key;
+
+				// Create the culture folders
+				DirectoryReference IntermediateCultureResourceDirectory = DirectoryReference.Combine(IntermediateDirectory, BuildResourceSubPath, AppXCultureId);
+				DirectoryReference OutputCultureResourceDirectory = DirectoryReference.Combine(OutputDirectory, BuildResourceSubPath, AppXCultureId);
+				FileUtils.CreateDirectoryTree(IntermediateCultureResourceDirectory);
+				FileUtils.CreateDirectoryTree(OutputCultureResourceDirectory);
+
+				// Export the resource tables & add them to the list of files to copy
+				FileReference IntermediateResWIndexerFile = FileReference.Combine( IntermediateDirectory, BuildResourceSubPath, AppXCultureId, "resources.resw" );
+				UEResXWriter ResourceWriter = new(IntermediateResWIndexerFile.FullName);
+				foreach (KeyValuePair<string, string> StringResourcePair in Itr.Value.StringResources)
+				{
+					ResourceWriter.AddResource( StringResourcePair.Key, StringResourcePair.Value );
+				}
+				ResourceWriter.Close();
+			
+				AddFileReference( IntermediateResWIndexerFile, Path.Combine(BuildResourceSubPath, AppXCultureId, "resources.resw"));
+			}
+
+			// The resource database is dependent on everything else calculated here (manifest, resource string tables, binary resources).
+			// So if any file has been updated we'll need to run the config.
+			bool bHadStaleResources = RemoveStaleResourceFiles(FilesToCopy.Values, OutputDirectory);
+			List<FileReference> UpdatedFilePaths = CopyFilesToOutput(OutputDirectory);
+			FileReference TargetResourceIndexFile = FileReference.Combine( OutputDirectory, "resources.pri" );
+
+			if (bHadStaleResources || UpdatedFilePaths.Any() || !FileReference.Exists(TargetResourceIndexFile) )
+			{
+				// Create resource index configuration
+				FileReference ResourceConfigFile = FileReference.Combine(IntermediateDirectory, "priconfig.xml");
+				RunMakePri($"createconfig /cf \"{ResourceConfigFile}\" /dq {DefaultAppXCultureId} /o /pv {TargetOSVersion}");
+
+				// Load the new resource index configuration
+				XmlDocument PriConfig = new XmlDocument();
+				PriConfig.Load(ResourceConfigFile.FullName);
+
+				// remove the packaging node - we do not want to split the pri & only want one .pri file
+				XmlNode PackagingNode = PriConfig.SelectSingleNode("/resources/packaging")!;
+				PackagingNode.ParentNode!.RemoveChild(PackagingNode);
+
+				// all required resources are explicitly listed in resources.resfiles, rather than relying on makepri to discover them
+				FileReference ResourcesResFile = FileReference.Combine(IntermediateDirectory, "resources.resfiles");
+				XmlNode PriIndexNode = PriConfig.SelectSingleNode("/resources/index")!;
+				XmlAttribute PriStartIndex = PriIndexNode.Attributes!["startIndexAt"]!;
+				PriStartIndex.Value = ResourcesResFile.FullName;
+
+				// swap the folder indexer-config to a RESFILES indexer-config.
+				XmlElement FolderIndexerConfigNode = (XmlElement)PriConfig.SelectSingleNode("/resources/index/indexer-config[@type='folder']")!;
+				FolderIndexerConfigNode.SetAttribute("type", "RESFILES");
+				FolderIndexerConfigNode.RemoveAttribute("foldernameAsQualifier");
+				FolderIndexerConfigNode.RemoveAttribute("filenameAsQualifier");
+
+				PriConfig.Save(ResourceConfigFile.FullName);
+
+				// generate resources.resfiles
+				IEnumerable<FileReference> Resources = DirectoryReference.EnumerateFiles( DirectoryReference.Combine(OutputDirectory, BuildResourceSubPath), "*.*", SearchOption.AllDirectories);
+				System.Text.StringBuilder ResourcesList = new System.Text.StringBuilder();
+				foreach (FileReference Resource in Resources)
+				{
+					ResourcesList.AppendLine( Resource.MakeRelativeTo(OutputDirectory) );
+				}
+				File.WriteAllText(ResourcesResFile.FullName, ResourcesList.ToString());
+
+				// remove old Package Resource Index
+				FileUtils.ForceDeleteFile(TargetResourceIndexFile);
+
+				// generate new Package Resource Index
+				FileReference ResourceLogFile = FileReference.Combine(IntermediateDirectory, "ResIndexLog.xml");
+				string MakePriCommandLine = $"new /pr \"{OutputDirectory}\" /cf \"{ResourceConfigFile}\" /mn \"{ManifestFile}\" /il \"{ResourceLogFile}\" /of \"{TargetResourceIndexFile}\" /o";
+				if (PackageIdentityName != null)
+				{
+					MakePriCommandLine += $" /indexName \"{PackageIdentityName}\"";
+				}
+
+				Logger.LogDebug("    generating {Path}", TargetResourceIndexFile.MakeRelativeTo(OutputDirectory));
+				RunMakePri(MakePriCommandLine);
+				UpdatedFilePaths.Add(TargetResourceIndexFile);
+
+				// .resw files are not needed - the data is embedded in the resources.pri
+				foreach (FileReference ResW in DirectoryReference.EnumerateFiles( OutputDirectory, "*.resw", SearchOption.AllDirectories))
+				{
+					FileUtils.ForceDeleteFile(ResW);
+				}
+			}
+
+			// Report if nothing was changed
+			if (!bHadStaleResources && !UpdatedFilePaths.Any())
+			{
+				Logger.LogDebug($"Manifest resource files are up to date");
+			}
+
+			return UpdatedFilePaths;
+		}
+	}
+}
