Fix issue with UBT where the UE_TARGET_NAME isn't the actual target name.  This resulted in HotReload and Live Coding to no longer work when used with a non-unity build.

diff --git a/Engine/Source/Programs/UnrealBuildTool/Configuration/TargetRules.cs b/Engine/Source/Programs/UnrealBuildTool/Configuration/TargetRules.cs
--- a/Engine/Source/Programs/UnrealBuildTool/Configuration/TargetRules.cs
+++ b/Engine/Source/Programs/UnrealBuildTool/Configuration/TargetRules.cs
@@ -438,6 +438,11 @@ namespace UnrealBuildTool
 		/// </summary>
 		public bool IsNameOverriden() { return !String.IsNullOrEmpty(NameOverride); }
 
+		/// <summary>
+		/// Return the original target name without overrides or adornments
+		/// </summary>
+		public string OriginalName => DefaultName;
+
 		/// <summary>
 		/// Override the name used for this target
 		/// </summary>
@@ -2803,6 +2808,8 @@ namespace UnrealBuildTool
 			get { return Inner.Name; }
 		}
 
+		public string OriginalName => Inner.OriginalName;
+
 		internal FileReference File
 		{
 			get { return Inner.File!; }
diff --git a/Engine/Source/Programs/UnrealBuildTool/Configuration/UEBuildModuleCPP.cs b/Engine/Source/Programs/UnrealBuildTool/Configuration/UEBuildModuleCPP.cs
--- a/Engine/Source/Programs/UnrealBuildTool/Configuration/UEBuildModuleCPP.cs
+++ b/Engine/Source/Programs/UnrealBuildTool/Configuration/UEBuildModuleCPP.cs
@@ -1820,7 +1820,7 @@ namespace UnrealBuildTool
 			{
 				string ProjectName = Target.ProjectFile.GetFileNameWithoutExtension();
 				Result.Definitions.Add(String.Format("UE_PROJECT_NAME={0}", ProjectName));
-				Result.Definitions.Add(String.Format("UE_TARGET_NAME={0}", Target.Name));
+				Result.Definitions.Add(String.Format("UE_TARGET_NAME={0}", Target.OriginalName));
 			}
 
 			// Add the module's public and private definitions.
