Fix issue with UBT where the UE_TARGET_NAME isn't the actual target name.  This resulted in HotReload and Live Coding to no longer work when used with a non-unity build.

#rnx

diff --git a/Engine/Source/Programs/UnrealBuildTool/ProjectFiles/Xcode/XcodeProject.cs b/Engine/Source/Programs/UnrealBuildTool/ProjectFiles/Xcode/XcodeProject.cs
--- a/Engine/Source/Programs/UnrealBuildTool/ProjectFiles/Xcode/XcodeProject.cs
+++ b/Engine/Source/Programs/UnrealBuildTool/ProjectFiles/Xcode/XcodeProject.cs
@@ -1692,7 +1692,7 @@ namespace UnrealBuildTool.XcodeProjectXcconfig
 			Xcconfig.AppendLine($"UE_DISPLAY_NAME = {UnrealData.DisplayName}");
 			Xcconfig.AppendLine($"UE_SIGNING_PREFIX = {SigningPrefix}");
 			Xcconfig.AppendLine($"UE_PLATFORM_NAME = {Platform}");
-			Xcconfig.AppendLine($"UE_TARGET_NAME = {UnrealData.TargetRules.Name}");
+			Xcconfig.AppendLine($"UE_TARGET_NAME = {UnrealData.TargetRules.OriginalName}");
 			Xcconfig.AppendLine($"UE_TARGET_PLATFORM_NAME = {TargetPlatformName}");
 			Xcconfig.AppendLine($"UE_ENGINE_DIR = {Unreal.EngineDirectory}");
 			Xcconfig.AppendLine($"UE_BINARIES_DIR = {DirectoryReference.Combine(ProjectOrEngineDir, "Binaries", Platform.ToString())}");
