Correct placement of low level tests solution folder for projects

diff --git a/Engine/Source/Programs/UnrealBuildTool/Configuration/TestTargetRules.cs b/Engine/Source/Programs/UnrealBuildTool/Configuration/TestTargetRules.cs
--- a/Engine/Source/Programs/UnrealBuildTool/Configuration/TestTargetRules.cs
+++ b/Engine/Source/Programs/UnrealBuildTool/Configuration/TestTargetRules.cs
@@ -149,7 +149,25 @@ namespace UnrealBuildTool
 			if (ExplicitTestsTarget)
 			{
 				bBuildInSolutionByDefault = true;
-				SolutionDirectory = "Programs/LowLevelTests";
+				if (ProjectFile != null)
+				{
+					DirectoryReference SamplesDirectory = DirectoryReference.Combine(Unreal.RootDirectory, "Samples");
+					if (!ProjectFile.IsUnderDirectory(Unreal.EngineDirectory))
+					{
+						if (ProjectFile.IsUnderDirectory(SamplesDirectory))
+						{
+							SolutionDirectory = Path.Combine("Samples", ProjectFile.Directory.GetDirectoryName(), "LowLevelTests");
+						}
+						else
+						{
+							SolutionDirectory = Path.Combine("Games", ProjectFile.Directory.GetDirectoryName(), "LowLevelTests");
+						}
+					}
+				}
+				else
+				{
+					SolutionDirectory = "Programs/LowLevelTests";
+				}
 
 				// Default to true for explicit targets to reduce compilation times.
 				// Selective module compilation will automatically detect if Engine is required based on Engine include files in tests.
@@ -166,7 +184,15 @@ namespace UnrealBuildTool
 				if (File != null)
 				{
 					string OutputDirectory = UEBuildTarget.GetOutputDirectoryForExecutable(Unreal.EngineDirectory, File).FullName;
-					string OutputFileName = $"{OutputDirectory}\\Binaries\\{Target.Platform}\\{ExeBinariesSubFolder}\\{OutputName}.exe.is_unreal_test";
+					string OutputFileName;
+					if (ProjectFile == null)
+					{ 
+						OutputFileName = $"{OutputDirectory}\\Binaries\\{Target.Platform}\\{ExeBinariesSubFolder}\\{OutputName}.exe.is_unreal_test";
+					}
+					else
+					{
+						OutputFileName = $"{ProjectFile!.Directory.FullName}\\Binaries\\{Target.Platform}\\{ExeBinariesSubFolder}\\{OutputName}.exe.is_unreal_test";
+					}
 					PostBuildSteps.Add("echo > " + OutputFileName);
 				}
 			}
