- Allow a project file generator to create temporary primary projects (.sln, .xcworkspace, etc) in the Intermediate/ProjectFIles directory where the rest of the projects are created (the generator would need to not hardcode relative paths from parimary to sub projects)

diff --git a/Engine/Source/Programs/UnrealBuildTool/ProjectFiles/ProjectFileGenerator.cs b/Engine/Source/Programs/UnrealBuildTool/ProjectFiles/ProjectFileGenerator.cs
--- a/Engine/Source/Programs/UnrealBuildTool/ProjectFiles/ProjectFileGenerator.cs
+++ b/Engine/Source/Programs/UnrealBuildTool/ProjectFiles/ProjectFileGenerator.cs
@@ -216,6 +216,11 @@ namespace UnrealBuildTool
 		/// </summary>
 		public bool bGeneratingGameProjectFiles = false;
 
+		/// <summary>
+		/// True if we're generating temporary solution/workspace that we want to put under Intermediate/ProjectFiles along with the projects to not dirty up the root dir
+		/// </summary>
+		public bool bGeneratingTemporaryProjects = false;
+
 		/// <summary>
 		/// Optional list of platforms to generate projects for
 		/// </summary>
@@ -911,6 +916,12 @@ namespace UnrealBuildTool
 				WritePrimaryProjectNameTxt = true;
 			}
 
+			// if making a temp primary project, put it in with the rest of the project files
+			if (bGeneratingTemporaryProjects)
+			{
+				PrimaryProjectPath = IntermediateProjectFilesPath;
+			}
+
 			// Modify the name if specific platforms were given
 			if (ProjectPlatforms.Count > 0 && bAppendPlatformSuffix)
 			{
diff --git a/Engine/Source/Programs/UnrealBuildTool/ProjectFiles/Xcode/XcodeProjectFileGenerator.cs b/Engine/Source/Programs/UnrealBuildTool/ProjectFiles/Xcode/XcodeProjectFileGenerator.cs
--- a/Engine/Source/Programs/UnrealBuildTool/ProjectFiles/Xcode/XcodeProjectFileGenerator.cs
+++ b/Engine/Source/Programs/UnrealBuildTool/ProjectFiles/Xcode/XcodeProjectFileGenerator.cs
@@ -474,6 +474,7 @@ namespace UnrealBuildTool
 				bIncludeDocumentation = false;
 				bIncludeShaderSource = false;
 				bIncludeEngineSource = true;
+				bGeneratingTemporaryProjects = true;
 
 				// generate just the engine project
 				if (OnlyGameProject != null)
