UnrealBuildTool: Print reference chain correctly for include path modules. Process references before include path modules

#rnx

diff --git a/Engine/Source/Programs/UnrealBuildTool/Configuration/UEBuildModule.cs b/Engine/Source/Programs/UnrealBuildTool/Configuration/UEBuildModule.cs
--- a/Engine/Source/Programs/UnrealBuildTool/Configuration/UEBuildModule.cs
+++ b/Engine/Source/Programs/UnrealBuildTool/Configuration/UEBuildModule.cs
@@ -1114,18 +1114,11 @@ namespace UnrealBuildTool
 				{
 					Logger.LogDebug("Found circular reference to {ThisRefName}, but {GuiltyModule} declares a cycle on {VictimModule} which breaks the chain", ThisRefName, GuiltyModule, VictimModule);
 				}
-			}
-
-			// Recursively create all the public include path and dependency modules. These modules may not be added to the target (and we don't process their referenced 
-			// dependencies), but they need to be created to set up their include paths.
-			HashSet<string> PublicIncludePathModuleNames = new();
-			PublicIncludePathModuleNames.UnionWith(Rules.PublicIncludePathModuleNames!);
-			PublicIncludePathModuleNames.UnionWith(Rules.PublicDependencyModuleNames!);
-			RecursivelyCreateIncludePathModulesByName(PublicIncludePathModuleNames, ref PublicIncludePathModules, ref bDependsOnVerse, CreateModule, NextReferenceChain);
+			}		
 
-			// Create all the referenced modules. This path can be recursive, so we check against PrivateIncludePathModules to ensure we don't recurse through the 
+			// Create all the referenced modules. This path can be recursive, so we check against PublicDependencyModules to ensure we don't recurse through the 
 			// same module twice (it produces better errors if something fails).
-			if (PrivateIncludePathModules == null)
+			if(PublicDependencyModules == null)
 			{
 				// Log dependencies if required
 				if (Logger.IsEnabled((LogLevel)LogEventType.VeryVerbose))
@@ -1138,9 +1131,6 @@ namespace UnrealBuildTool
 					LogDependencyNameList("Private Include Paths:", Rules.PrivateIncludePathModuleNames, Logger);
 				}
 
-				// Create the private include path modules
-				RecursivelyCreateIncludePathModulesByName(Rules.PrivateIncludePathModuleNames, ref PrivateIncludePathModules, ref bDependsOnVerse, CreateModule, NextReferenceChain);
-
 				// Create all the dependency modules - pass through the reference stack so we can check for cycles
 				RecursivelyCreateModulesByName(Rules.PublicDependencyModuleNames, ref PublicDependencyModules, ref bDependsOnVerse, CreateModule, NextReferenceChain, ReferenceStack, Logger);
 				if (Rules.Target.IsTestTarget)
@@ -1149,6 +1139,16 @@ namespace UnrealBuildTool
 					MoveTestsRunnerDependencyToLastPosition();
 				}
 				RecursivelyCreateModulesByName(Rules.PrivateDependencyModuleNames, ref PrivateDependencyModules, ref bDependsOnVerse, CreateModule, NextReferenceChain, ReferenceStack, Logger);
+
+				// Recursively create all the public include path modules
+				HashSet<string> PublicIncludePathModuleNames = new();
+				PublicIncludePathModuleNames.UnionWith(Rules.PublicIncludePathModuleNames!);
+				PublicIncludePathModuleNames.UnionWith(Rules.PublicDependencyModuleNames!);
+				RecursivelyCreateIncludePathModulesByName(PublicIncludePathModuleNames, ref PublicIncludePathModules, ref bDependsOnVerse, CreateModule, NextReferenceChain);
+
+				// Create the private include path modules
+				RecursivelyCreateIncludePathModulesByName(Rules.PrivateIncludePathModuleNames, ref PrivateIncludePathModules, ref bDependsOnVerse, CreateModule, NextReferenceChain);
+
 				// Dynamic loads aren't considered a reference chain so start with an empty stack
 				RecursivelyCreateModulesByName(Rules.DynamicallyLoadedModuleNames, ref DynamicallyLoadedModules, ref bDependsOnVerse, CreateModule, NextReferenceChain, new List<UEBuildModule>(), Logger);
 			}
@@ -1204,12 +1204,20 @@ namespace UnrealBuildTool
 				foreach (string ModuleName in ModuleNames)
 				{
 					UEBuildModule Module = CreateModule(ModuleName, ReferenceChain);
-					HashSet<string> PublicIncludePathModuleNames = new();
-					PublicIncludePathModuleNames.UnionWith(Module.Rules.PublicIncludePathModuleNames);
-					PublicIncludePathModuleNames.UnionWith(Module.Rules.PublicDependencyModuleNames);
-					RecursivelyCreateIncludePathModulesByName(PublicIncludePathModuleNames, ref Module.PublicIncludePathModules, ref Module.bDependsOnVerse, CreateModule, ReferenceChain);
-					Modules.Add(Module);
-					bDependsOnVerse |= Module.bDependsOnVerse;
+					if (!Modules.Contains(Module))
+					{
+						// Name of this reference
+						string ModuleRefName = (Module.RulesFile == null) ? Module.Name : Module.RulesFile.GetFileName();
+						// Set the reference chain for anything referenced by this module
+						string NextReferenceChain = String.Format("{0} -> {1} (public include)", ReferenceChain, ModuleRefName);
+
+						HashSet<string> PublicIncludePathModuleNames = new();
+						PublicIncludePathModuleNames.UnionWith(Module.Rules.PublicIncludePathModuleNames);
+						PublicIncludePathModuleNames.UnionWith(Module.Rules.PublicDependencyModuleNames);
+						RecursivelyCreateIncludePathModulesByName(PublicIncludePathModuleNames, ref Module.PublicIncludePathModules, ref Module.bDependsOnVerse, CreateModule, NextReferenceChain);
+						Modules.Add(Module);
+						bDependsOnVerse |= Module.bDependsOnVerse;
+					}
 				}
 			}
 		}
