Make sure to include autortfm for C files when enabled for the module

diff --git a/Engine/Source/Programs/UnrealBuildTool/Platform/Windows/VCToolChain.cs b/Engine/Source/Programs/UnrealBuildTool/Platform/Windows/VCToolChain.cs
--- a/Engine/Source/Programs/UnrealBuildTool/Platform/Windows/VCToolChain.cs
+++ b/Engine/Source/Programs/UnrealBuildTool/Platform/Windows/VCToolChain.cs
@@ -1526,6 +1526,15 @@ namespace UnrealBuildTool
 				if (bIsPlainCFile)
 				{
 					AppendCLArguments_C(CompileEnvironment, CompileAction.Arguments);
+
+					// AppendCLArguments_C is a static and we cannot use the Target to check if we are using clang, so shoving here for now
+					if (Target.WindowsPlatform.Compiler.IsClang())
+					{
+						if (CompileEnvironment.bAllowAutoRTFMInstrumentation)
+						{
+							CompileAction.Arguments.Add("-fautortfm");
+						}
+					}
 				}
 				else
 				{
diff --git a/Engine/Source/Programs/UnrealBuildTool/ToolChain/ClangToolChain.cs b/Engine/Source/Programs/UnrealBuildTool/ToolChain/ClangToolChain.cs
--- a/Engine/Source/Programs/UnrealBuildTool/ToolChain/ClangToolChain.cs
+++ b/Engine/Source/Programs/UnrealBuildTool/ToolChain/ClangToolChain.cs
@@ -394,6 +394,11 @@ namespace UnrealBuildTool
 				default:
 					throw new BuildException($"Unsupported C standard type set: {CompileEnvironment.CStandard}");
 			}
+
+			if (CompileEnvironment.bAllowAutoRTFMInstrumentation)
+			{
+				Arguments.Add("-fautortfm");
+			}
 		}
 
 		protected virtual void GetCompileArguments_CPP(CppCompileEnvironment CompileEnvironment, List<string> Arguments)
