UBT: Recursively gather the PublicDependencies and link the entire chain so the writer of a module does not need to understand the linking requirements of every Plugin they use.

diff --git a/Engine/Source/Programs/UnrealBuildTool/Configuration/ModuleRules.cs b/Engine/Source/Programs/UnrealBuildTool/Configuration/ModuleRules.cs
--- a/Engine/Source/Programs/UnrealBuildTool/Configuration/ModuleRules.cs
+++ b/Engine/Source/Programs/UnrealBuildTool/Configuration/ModuleRules.cs
@@ -841,6 +841,17 @@ namespace UnrealBuildTool
 		/// </summary>
 		public bool bEnableObjCAutomaticReferenceCounting = false;
 
+		/// <summary>
+		/// Whether to allow PublicDependencies to be gathered and linked recursively (so it links against all dependents).
+		/// Consider PublicDepChain of PluginA -> PluginB -> PluginC.  If false, PluginA links PluginB. If true, PluginA links both PluginB and PluginC.
+		/// </summary>
+		public bool bLinkPublicDependencyChain
+		{
+			set { bLinkPublicDependencyChainOverride = value; }
+			get { return bLinkPublicDependencyChainOverride ?? Target.bLinkPublicDependencyChain; }
+		}
+		private bool? bLinkPublicDependencyChainOverride;
+
 		/// <summary>
 		/// How to treat deterministic warnings (experimental).
 		/// </summary>
diff --git a/Engine/Source/Programs/UnrealBuildTool/Configuration/TargetRules.cs b/Engine/Source/Programs/UnrealBuildTool/Configuration/TargetRules.cs
--- a/Engine/Source/Programs/UnrealBuildTool/Configuration/TargetRules.cs
+++ b/Engine/Source/Programs/UnrealBuildTool/Configuration/TargetRules.cs
@@ -846,6 +846,16 @@ namespace UnrealBuildTool
 		[ConfigFile(ConfigHierarchyType.Engine, "/Script/BuildSettings.BuildSettings", "bCompileCEF3")]
 		public bool bCompileCEF3 = true;
 
+		/// <summary>
+		/// Whether to allow PublicDependencies to be gathered and linked recursively (so it links against all dependents).
+		/// Consider PublicDepChain of PluginA -> PluginB -> PluginC.  If false, PluginA links PluginB. If true, PluginA links both PluginB and PluginC.
+		/// This sets the default, but it can be overridden per module by the ModuleRules.
+		/// </summary>
+		[XmlConfigFile(Category = "BuildConfiguration")]
+		[CommandLine("-NoLinkDependencyChain", Value = "false")]
+		[CommandLine("-LinkDependencyChain", Value = "true")]
+		public bool bLinkPublicDependencyChain = true;
+
 		/// <summary>
 		/// Whether to compile using ISPC.
 		/// </summary>
@@ -3072,6 +3082,11 @@ namespace UnrealBuildTool
 			get { return Inner.bCompileCEF3; }
 		}
 
+		public bool bLinkPublicDependencyChain
+		{
+			get { return Inner.bLinkPublicDependencyChain; }
+		}
+
 		public bool bCompileISPC
 		{
 			get { return Inner.bCompileISPC; }
diff --git a/Engine/Source/Programs/UnrealBuildTool/Configuration/UEBuildModule.cs b/Engine/Source/Programs/UnrealBuildTool/Configuration/UEBuildModule.cs
--- a/Engine/Source/Programs/UnrealBuildTool/Configuration/UEBuildModule.cs
+++ b/Engine/Source/Programs/UnrealBuildTool/Configuration/UEBuildModule.cs
@@ -926,7 +926,22 @@ namespace UnrealBuildTool
 			// Also allow the module's public and private dependencies to modify the link environment.
 			List<UEBuildModule> AllDependencyModules = new List<UEBuildModule>();
 			AllDependencyModules.AddRange(PrivateDependencyModules!);
-			AllDependencyModules.AddRange(PublicDependencyModules!);
+
+			// Build up a list of all the public dependencies
+			if (Rules.bLinkPublicDependencyChain)
+			{
+				List<UEBuildModule> AllPublicDependencies = new List<UEBuildModule>();
+				RecursivelyAddPublicDependencies(this, AllPublicDependencies);
+				foreach (var PrivateModule in PrivateDependencyModules!)
+				{
+					RecursivelyAddPublicDependencies(PrivateModule, AllPublicDependencies);
+				}
+				AllDependencyModules.AddRange(AllPublicDependencies);
+			}
+			else
+			{
+				AllDependencyModules.AddRange(PublicDependencyModules!);
+			}
 
 			foreach (UEBuildModule DependencyModule in AllDependencyModules)
 			{
@@ -1222,6 +1237,23 @@ namespace UnrealBuildTool
 			}
 		}
 
+		/// <summary>
+		/// Recursively build up the Public Dependencies. Module is the current root and it will add all children to AllPublicDependencies.
+		/// </summary>
+		/// <param name="Module">The current 'root' of the recursive operation</param>
+		/// <param name="AllPublicDependencies">Where the public dependencies get stored</param>
+		private static void RecursivelyAddPublicDependencies(UEBuildModule Module, List<UEBuildModule> AllPublicDependencies)
+		{
+			foreach (UEBuildModule PublicDependency in Module.PublicDependencyModules!)
+			{
+				if (!AllPublicDependencies.Contains(PublicDependency))
+				{
+					AllPublicDependencies.Add(PublicDependency);
+					RecursivelyAddPublicDependencies(PublicDependency, AllPublicDependencies);
+				}
+			}
+		}
+
 		/// <summary>
 		/// Returns valueless API defines (like MODULE_API)
 		/// </summary>
