1) Replaced bCanCache with new artifact mode flags
2) Added support in artifacts for relative paths
3) Modified ISPC header generation to be cached

At this time, only MSVC compiles are cached and -deterministic must be specified when building.

diff --git a/Engine/Source/Programs/UnrealBuildTool.Tests/ArtifactsTest.cs b/Engine/Source/Programs/UnrealBuildTool.Tests/ArtifactsTest.cs
--- a/Engine/Source/Programs/UnrealBuildTool.Tests/ArtifactsTest.cs
+++ b/Engine/Source/Programs/UnrealBuildTool.Tests/ArtifactsTest.cs
@@ -20,17 +20,17 @@ namespace UnrealBuildToolTests
 
 		private static Artifact MakeInput1()
 		{
-			return new Artifact(ArtifactDirectoryTree.Absolute, ArtifactType.Source, "Input1", IoHash.Compute(s_input1Data.Span));
+			return new Artifact(ArtifactDirectoryTree.Absolute, "Input1", IoHash.Compute(s_input1Data.Span));
 		}
 
 		private static Artifact MakeInput2()
 		{
-			return new Artifact(ArtifactDirectoryTree.Absolute, ArtifactType.Source, "Input2", IoHash.Compute(s_input2Data.Span));
+			return new Artifact(ArtifactDirectoryTree.Absolute, "Input2", IoHash.Compute(s_input2Data.Span));
 		}
 
 		private static Artifact MakeOutput1()
 		{
-			return new Artifact(ArtifactDirectoryTree.Absolute, ArtifactType.Object, "Output1", IoHash.Compute(s_output1Data.Span));
+			return new Artifact(ArtifactDirectoryTree.Absolute, "Output1", IoHash.Compute(s_output1Data.Span));
 		}
 
 		public static ArtifactMapping MakeBundle1()
diff --git a/Engine/Source/Programs/UnrealBuildTool/Artifacts/ActionArtifactCache.cs b/Engine/Source/Programs/UnrealBuildTool/Artifacts/ActionArtifactCache.cs
--- a/Engine/Source/Programs/UnrealBuildTool/Artifacts/ActionArtifactCache.cs
+++ b/Engine/Source/Programs/UnrealBuildTool/Artifacts/ActionArtifactCache.cs
@@ -3,8 +3,9 @@
 using EpicGames.Core;
 using Microsoft.Extensions.Logging;
 using Microsoft.Extensions.Logging.Abstractions;
+using System;
+using System.Collections.Concurrent;
 using System.Collections.Generic;
-using System.Linq;
 using System.Text;
 using System.Threading;
 using System.Threading.Tasks;
@@ -13,6 +14,87 @@ using UnrealBuildBase;
 namespace UnrealBuildTool.Artifacts
 {
 
+	/// <summary>
+	/// Helper class for directory mapping
+	/// </summary>
+	internal class ArtifactDirectoryMapping : IArtifactDirectoryMapping
+	{
+
+		/// <summary>
+		/// Creating cache
+		/// </summary>
+		public IActionArtifactCache? Cache;
+
+		/// <summary>
+		/// Root of the project
+		/// </summary>
+		public DirectoryReference? ProjectRoot;
+
+		/// <inheritdoc/>
+		public string GetDirectory(ArtifactDirectoryTree tree)
+		{
+			switch (tree)
+			{
+				case ArtifactDirectoryTree.Absolute:
+					return string.Empty;
+				case ArtifactDirectoryTree.Engine:
+					if (Cache == null || Cache.EngineRoot == null)
+					{
+						throw new ApplicationException("Attempt to get engine root when not set");
+					}
+					return Cache.EngineRoot.FullName;
+				case ArtifactDirectoryTree.Project:
+					if (ProjectRoot == null)
+					{
+						throw new ApplicationException("Attempt to get project root when not set");
+					}
+					return ProjectRoot.FullName;
+				default:
+					throw new NotImplementedException("Unexpected directory tree value");
+			}
+		}
+
+		/// <summary>
+		/// Given a file, return the artifact structure for it.  This routine tests to see if the 
+		/// file is under any of the well known directories.
+		/// </summary>
+		/// <param name="action">Action requesting the artifact</param>
+		/// <param name="file">File in question</param>
+		/// <param name="hash">Hash value used to populate artifact</param>
+		/// <returns>Created artifact</returns>
+		public Artifact GetArtifact(LinkedAction action, FileItem file, IoHash hash)
+		{
+			if (action.ArtifactMode.HasFlag(ArtifactMode.AbsolutePath))
+			{
+				return CreateArtifact(ArtifactDirectoryTree.Absolute, file.AbsolutePath, hash);
+			}
+			if (Cache != null && Cache.EngineRoot != null && file.Location.IsUnderDirectory(Cache.EngineRoot))
+			{
+				return CreateArtifact(ArtifactDirectoryTree.Engine, file.Location.MakeRelativeTo(Cache.EngineRoot), hash);
+			}
+			else if (ProjectRoot != null && file.Location.IsUnderDirectory(ProjectRoot))
+			{
+				return CreateArtifact(ArtifactDirectoryTree.Project, file.Location.MakeRelativeTo(ProjectRoot), hash);
+			}
+			else
+			{
+				return CreateArtifact(ArtifactDirectoryTree.Absolute, file.AbsolutePath, hash);
+			}
+		}
+
+		/// <summary>
+		/// Create an artifact with the given settings
+		/// </summary>
+		/// <param name="tree">Directory tree</param>
+		/// <param name="path">Path to artifact</param>
+		/// <param name="hash">Hash of the artifact</param>
+		/// <returns>The artifact</returns>
+		private Artifact CreateArtifact(ArtifactDirectoryTree tree, string path, IoHash hash)
+		{
+			return new(tree, new Utf8String(path), hash) { DirectoryMapping = this };
+		}
+	}
+
 	/// <summary>
 	/// Generic handler for artifacts
 	/// </summary>
@@ -31,6 +113,12 @@ namespace UnrealBuildTool.Artifacts
 		/// <inheritdoc/>
 		public IArtifactCache ArtifactCache { get; init; }
 
+		/// <inheritdoc/>
+		public DirectoryReference? EngineRoot { get; set; } = null;
+
+		/// <inheritdoc/>
+		public DirectoryReference[]? DirectoryRoots { get; set; } = null;
+
 		/// <summary>
 		/// Logging device
 		/// </summary>
@@ -46,6 +134,16 @@ namespace UnrealBuildTool.Artifacts
 		/// </summary>
 		private readonly FileHasher _fileHasher;
 
+		/// <summary>
+		/// Directory mapper to be used for targets without projects
+		/// </summary>
+		private readonly ArtifactDirectoryMapping _projectlessMapper;
+
+		/// <summary>
+		/// Artifact mappers for all targets
+		/// </summary>
+		private readonly ConcurrentDictionary<DirectoryReference, ArtifactDirectoryMapping> _mappings = new();
+
 		/// <summary>
 		/// Construct a new artifact handler object
 		/// </summary>
@@ -58,6 +156,7 @@ namespace UnrealBuildTool.Artifacts
 			_cppDependencyCache = cppDependencyCache;
 			ArtifactCache = artifactCache;
 			_fileHasher = new(NullLogger.Instance);
+			_projectlessMapper = new() { Cache = this };
 		}
 
 		/// <summary>
@@ -93,7 +192,14 @@ namespace UnrealBuildTool.Artifacts
 				return false;
 			}
 
-			IoHash key = await GetKeyAsync(action);
+			if (!action.ArtifactMode.HasFlag(ArtifactMode.Enabled))
+			{
+				return false;
+			}
+
+			ArtifactDirectoryMapping directoryMapping = GetDirectoryMapping(action);
+			(List<FileItem> inputs, _) = CollectInputs(action, false);
+			IoHash key = await GetKeyAsync(directoryMapping, action, inputs);
 
 			ArtifactMapping[] mappings = await ArtifactCache.QueryArtifactMappingsAsync(new IoHash[] { key }, cancellationToken);
 
@@ -118,7 +224,7 @@ namespace UnrealBuildTool.Artifacts
 
 				foreach (Artifact input in mapping.Inputs)
 				{
-					string name = input.Name.ToString();
+					string name = input.GetFullPath(directoryMapping);
 					FileItem item = FileItem.GetItemByPath(name);
 					if (!item.Exists)
 					{
@@ -142,7 +248,9 @@ namespace UnrealBuildTool.Artifacts
 
 				if (match)
 				{
-					bool[]? readResults = await ArtifactCache.QueryArtifactOutputsAsync(new[] { mapping }, cancellationToken);
+					ArtifactMapping mappingCopy = mapping;
+					mappingCopy.DirectoryMapping = directoryMapping;
+					bool[]? readResults = await ArtifactCache.QueryArtifactOutputsAsync(new[] { mappingCopy }, cancellationToken);
 
 					if (readResults == null || readResults.Length == 0 || !readResults[0])
 					{
@@ -152,7 +260,7 @@ namespace UnrealBuildTool.Artifacts
 					{
 						foreach (Artifact output in mapping.Outputs)
 						{
-							string outputName = output.Name.ToString();
+							string outputName = output.GetFullPath(directoryMapping);
 							FileItem item = FileItem.GetItemByPath(outputName);
 							item.ResetCachedInfo(); // newly created outputs need refreshing
 							_fileHasher.SetDigest(item, output.ContentHash);
@@ -172,8 +280,17 @@ namespace UnrealBuildTool.Artifacts
 				return;
 			}
 
-			ArtifactMapping[] mappings = new ArtifactMapping[] { await CreateArtifactMappingAsync(action, cancellationToken) };
-			await ArtifactCache.SaveArtifactMappingsAsync(mappings, cancellationToken);
+			if (!action.ArtifactMode.HasFlag(ArtifactMode.Enabled))
+			{
+				return;
+			}
+
+			ArtifactMapping mapping = await CreateArtifactMappingAsync(action, cancellationToken);
+			if (mapping.Key != IoHash.Zero)
+			{
+				await ArtifactCache.SaveArtifactMappingsAsync(new ArtifactMapping[] { (ArtifactMapping)mapping }, cancellationToken);
+			}
+			return;
 		}
 
 		/// <inheritdoc/>
@@ -190,52 +307,57 @@ namespace UnrealBuildTool.Artifacts
 		/// <returns>Artifact mapping</returns>
 		private async Task<ArtifactMapping> CreateArtifactMappingAsync(LinkedAction action, CancellationToken cancellationToken)
 		{
-			(IoHash key, IoHash mappingKey) = await GetKeyAndMappingKeyAsync(action);
+			ArtifactDirectoryMapping directoryMapping = GetDirectoryMapping(action);
+			(List<FileItem> inputs, List<FileItem>? dependencies) = CollectInputs(action, false);
+			(IoHash key, IoHash mappingKey) = await GetKeyAndMappingKeyAsync(directoryMapping, action, inputs, dependencies);
 
 			// We gather the output files first to make sure that all the generated files (including the dependency file) gets
 			// their cached FileInfo reset.
 			List<Artifact> outputs = new();
 			foreach (FileItem output in action.ProducedItems)
 			{
-				output.ResetCachedInfo();
+				// Outputs can not be a directory.
+				if (output.Attributes.HasFlag(System.IO.FileAttributes.Directory))
+				{
+					return new(IoHash.Zero, IoHash.Zero, Array.Empty<Artifact>(), Array.Empty<Artifact>());
+				}
 				IoHash hash = await _fileHasher.GetDigestAsync(output, cancellationToken);
-				Artifact artifact = new(ArtifactDirectoryTree.Absolute, ArtifactType.Source, new Utf8String(output.AbsolutePath), hash);
+				Artifact artifact = directoryMapping.GetArtifact(action, output, hash);
 				outputs.Add(artifact);
 			}
 
-			List<Artifact> inputs = new();
-			foreach (FileItem input in action.PrerequisiteItems)
+			List<Artifact> inputArtifacts = new();
+			foreach (FileItem input in inputs)
 			{
 				IoHash hash = await _fileHasher.GetDigestAsync(input, cancellationToken);
-				Artifact artifact = new(ArtifactDirectoryTree.Absolute, ArtifactType.Source, new Utf8String(input.AbsolutePath), hash);
-				inputs.Add(artifact);
+				Artifact artifact = directoryMapping.GetArtifact(action, input, hash);
+				inputArtifacts.Add(artifact);
 			}
 
-			if (action.DependencyListFile != null)
+			if (dependencies != null)
 			{
-				if (_cppDependencyCache.TryGetDependencies(action.DependencyListFile, _logger, out List<FileItem>? depedencies))
+				foreach (FileItem dependency in dependencies)
 				{
-					foreach (FileItem dep in depedencies)
-					{
-						IoHash hash = await _fileHasher.GetDigestAsync(dep, cancellationToken);
-						Artifact artifact = new(ArtifactDirectoryTree.Absolute, ArtifactType.Source, new Utf8String(dep.AbsolutePath), hash);
-						inputs.Add(artifact);
-					}
+					IoHash hash = await _fileHasher.GetDigestAsync(dependency, cancellationToken);
+					Artifact artifact = directoryMapping.GetArtifact(action, dependency, hash);
+					inputArtifacts.Add(artifact);
 				}
 			}
 
-			return new(key, mappingKey, inputs.ToArray(), outputs.ToArray());
+			return new(key, mappingKey, inputArtifacts.ToArray(), outputs.ToArray());
 		}
 
 		/// <summary>
 		/// Get the key has for the action
 		/// </summary>
+		/// <param name="directoryMapping">Directory mapping object</param>
 		/// <param name="action">Source action</param>
+		/// <param name="inputs">Inputs used to construct the key</param>
 		/// <returns>Task returning the key</returns>
-		private async Task<IoHash> GetKeyAsync(LinkedAction action)
+		private async Task<IoHash> GetKeyAsync(ArtifactDirectoryMapping directoryMapping, LinkedAction action, List<FileItem> inputs)
 		{
 			StringBuilder builder = new();
-			await AppendKeyAsync(builder, action);
+			await AppendKeyAsync(builder, directoryMapping, action, inputs);
 			IoHash key = IoHash.Compute(new Utf8String(builder.ToString()));
 			return key;
 		}
@@ -243,14 +365,17 @@ namespace UnrealBuildTool.Artifacts
 		/// <summary>
 		/// Generate the key and mapping key hashes for the action
 		/// </summary>
+		/// <param name="directoryMapping">Directory mapping object</param>
 		/// <param name="action">Source action</param>
+		/// <param name="inputs">Inputs used to construct the key</param>
+		/// <param name="dependencies">Dependencies used to construct the key</param>
 		/// <returns>Task object with the key and mapping key</returns>
-		private async Task<(IoHash, IoHash)> GetKeyAndMappingKeyAsync(LinkedAction action)
+		private async Task<(IoHash, IoHash)> GetKeyAndMappingKeyAsync(ArtifactDirectoryMapping directoryMapping, LinkedAction action, List<FileItem> inputs, List<FileItem>? dependencies)
 		{
 			StringBuilder builder = new();
-			await AppendKeyAsync(builder, action);
+			await AppendKeyAsync(builder, directoryMapping, action, inputs);
 			IoHash key = IoHash.Compute(new Utf8String(builder.ToString()));
-			await AppendMappingKeyAsync(builder, action);
+			await AppendMappingKeyAsync(builder, directoryMapping, action, dependencies);
 			IoHash mappingKey = IoHash.Compute(new Utf8String(builder.ToString()));
 			return (key, mappingKey);
 		}
@@ -259,52 +384,182 @@ namespace UnrealBuildTool.Artifacts
 		/// Generate the lookup key.  This key is generated from the action's inputs.
 		/// </summary>
 		/// <param name="builder">Destination builder</param>
+		/// <param name="directoryMapping">Directory mapping object</param>
 		/// <param name="action">Source action</param>
+		/// <param name="inputs">Inputs used to construct the key</param>
 		/// <returns>Task object</returns>
-		private async Task AppendKeyAsync(StringBuilder builder, LinkedAction action)
+		private async Task AppendKeyAsync(StringBuilder builder, ArtifactDirectoryMapping directoryMapping, LinkedAction action, List<FileItem> inputs)
 		{
 			builder.AppendLine(action.CommandVersion);
 			builder.AppendLine(action.CommandArguments);
+			await AppendFiles(builder, directoryMapping, action, inputs);
+		}
 
-			FileItem[] inputs = action.PrerequisiteItems.ToArray();
-			Task<IoHash>[] waits = new Task<IoHash>[inputs.Length];
-			for (int index = 0; index < inputs.Length; index++)
+		/// <summary>
+		/// Generate the full mapping key.  This contains the hashes for the action's inputs and the dependent files.
+		/// </summary>
+		/// <param name="builder">Destination builder</param>
+		/// <param name="directoryMapping">Directory mapping object</param>
+		/// <param name="action">Source action</param>
+		/// <param name="dependencies">Dependencies used to construct the key</param>
+		/// <returns>Task object</returns>
+		private async Task AppendMappingKeyAsync(StringBuilder builder, ArtifactDirectoryMapping directoryMapping, LinkedAction action, List<FileItem>? dependencies)
+		{
+			await AppendFiles(builder, directoryMapping, action, dependencies);
+		}
+
+		/// <summary>
+		/// Append the file information for a given list of files
+		/// </summary>
+		/// <param name="builder">Destination builder</param>
+		/// <param name="directoryMapping">Directory mapping object</param>
+		/// <param name="action">Source action</param>
+		/// <param name="files">Collection of files</param>
+		/// <returns>Task object</returns>
+		private async Task AppendFiles(StringBuilder builder, ArtifactDirectoryMapping directoryMapping, LinkedAction action, List<FileItem>? files)
+		{
+			if (files != null)
+			{
+				Task<IoHash>[] waits = new Task<IoHash>[files.Count];
+				for (int index = 0; index < files.Count; index++)
+				{
+					waits[index] = _fileHasher.GetDigestAsync(files[index]);
+				}
+				await Task.WhenAll(waits);
+				string[] lines = new string[files.Count];
+				for (int index = 0; index < files.Count; index++)
+				{
+					Artifact artifact = directoryMapping.GetArtifact(action, files[index], waits[index].Result);
+					lines[index] = $"{GetArtifactTreeName(artifact)} {artifact.Tree} {artifact.Name} {artifact.ContentHash}";
+				}
+				Array.Sort(lines, StringComparer.Ordinal);
+				foreach (string line in lines)
+				{
+					builder.AppendLine(line);
+				}
+			}
+		}
+
+		/// <summary>
+		/// Given an action, collect the list of inputs and dependencies.  This includes any processing required
+		/// for disabled actions that propagate their inputs to dependent actions.
+		/// </summary>
+		/// <param name="action">Action in question</param>
+		/// <param name="collectDependencies">If true, collect the dependencies too</param>
+		/// <returns>Collection of inputs and dependencies</returns>
+		private (List<FileItem>, List<FileItem>?) CollectInputs(LinkedAction action, bool collectDependencies)
+		{
+			// Search for any prerequisite action that is set to propagate inputs (i.e. PCH).  Collect those
+			// inputs and we will be inserting those inputs in our list
+			Dictionary<FileItem, List<FileItem>> substitutions = new();
+			foreach (LinkedAction prereq in action.PrerequisiteActions)
 			{
-				builder.AppendLine(inputs[index].FullName);
-				waits[index] = _fileHasher.GetDigestAsync(inputs[index]);
+				if (prereq.ArtifactMode.HasFlag(ArtifactMode.PropagateInputs))
+				{
+					(List<FileItem> prereqInputs, List<FileItem>? prereqDependencies) = CollectInputs(prereq, collectDependencies);
+					if (prereqDependencies != null)
+					{
+						prereqInputs.AddRange(prereqDependencies);
+					}
+					foreach (FileItem output in prereq.ProducedItems)
+					{
+						substitutions.TryAdd(output, prereqInputs);
+					}
+				}
 			}
-			await Task.WhenAll(waits);
-			foreach (Task<IoHash> wait in waits)
+
+			HashSet<FileItem> uniques = new();
+			List<FileItem> inputs = new();
+			AddFileItems(uniques, substitutions, inputs, action.PrerequisiteItems);
+
+			List<FileItem>? dependencies = null;
+			if (collectDependencies && action.DependencyListFile != null)
 			{
-				builder.AppendLine(wait.Result.ToString());
+				if (_cppDependencyCache.TryGetDependencies(action.DependencyListFile, _logger, out List<FileItem>? cppDeps))
+				{
+					dependencies = new();
+					AddFileItems(uniques, substitutions, dependencies, cppDeps);
+				}
 			}
+
+			return (inputs, dependencies);
 		}
 
 		/// <summary>
-		/// Generate the full mapping key.  This contains the hashes for the action's inputs and the dependent files.
+		/// Add a list of file items to the collection
 		/// </summary>
-		/// <param name="builder">Destination builder</param>
-		/// <param name="action">Source action</param>
-		/// <returns>Task object</returns>
-		private async Task AppendMappingKeyAsync(StringBuilder builder, LinkedAction action)
+		/// <param name="uniques">Hash set used to detect already included file items</param>
+		/// <param name="substitutions">Substitutions when a given input is found</param>
+		/// <param name="outputs">Destination list</param>
+		/// <param name="inputs">Source inputs</param>
+		private void AddFileItems(HashSet<FileItem> uniques, Dictionary<FileItem, List<FileItem>>? substitutions, List<FileItem> outputs, IEnumerable<FileItem> inputs)
 		{
-			if (action.DependencyListFile != null)
+			if (substitutions != null)
 			{
-				if (_cppDependencyCache.TryGetDependencies(action.DependencyListFile, _logger, out List<FileItem>? dependencies))
+				foreach (FileItem input in inputs)
 				{
-					Task<IoHash>[] waits = new Task<IoHash>[dependencies.Count];
-					for (int index = 0; index < dependencies.Count; index++)
+					if (substitutions.TryGetValue(input, out List<FileItem>? substituteFileItems))
+					{
+						AddFileItems(uniques, null, outputs, substituteFileItems);
+					}
+					else if (uniques.Add(input))
 					{
-						builder.AppendLine(dependencies[index].FullName);
-						waits[index] = _fileHasher.GetDigestAsync(dependencies[index]);
+						outputs.Add(input);
 					}
-					await Task.WhenAll(waits);
-					foreach (Task<IoHash> wait in waits)
+				}
+			}
+			else
+			{
+				foreach (FileItem input in inputs)
+				{
+					if (uniques.Add(input))
 					{
-						builder.AppendLine(wait.Result.ToString());
+						outputs.Add(input);
 					}
 				}
 			}
 		}
+
+		/// <summary>
+		/// Return the tree name of the artifact
+		/// </summary>
+		/// <param name="artifact">Artifact in question</param>
+		/// <returns>Tree name</returns>
+		private static string GetArtifactTreeName(Artifact artifact)
+		{
+			switch (artifact.Tree)
+			{
+				case ArtifactDirectoryTree.Absolute:
+					return "Absolute";
+				case ArtifactDirectoryTree.Engine:
+					return "Engine";
+				case ArtifactDirectoryTree.Project:
+					return "Project";
+				default:
+					throw new NotImplementedException("Unexpected artifact directory tree type");
+			}
+		}
+
+		/// <summary>
+		/// Return an artifact mapper for the given action
+		/// </summary>
+		/// <param name="action">Action in question</param>
+		/// <returns>Artifact mapper specific to the target's project directory</returns>
+		private ArtifactDirectoryMapping GetDirectoryMapping(LinkedAction action)
+		{
+			DirectoryReference? projectDirectory = action.Target != null && action.Target.ProjectFile != null ? action.Target.ProjectFile.Directory : null;
+			if (projectDirectory == null)
+			{
+				return _projectlessMapper;
+			}
+			return _mappings.GetOrAdd(projectDirectory, x =>
+			{
+				return new()
+				{
+					Cache = this,
+					ProjectRoot = x,
+				};
+			});
+		}
 	}
 }
diff --git a/Engine/Source/Programs/UnrealBuildTool/Artifacts/Artifacts.cs b/Engine/Source/Programs/UnrealBuildTool/Artifacts/Artifacts.cs
--- a/Engine/Source/Programs/UnrealBuildTool/Artifacts/Artifacts.cs
+++ b/Engine/Source/Programs/UnrealBuildTool/Artifacts/Artifacts.cs
@@ -4,6 +4,7 @@ using EpicGames.Core;
 using EpicGames.Serialization;
 using System;
 using System.Collections.Generic;
+using System.IO;
 using System.Linq;
 using System.Threading;
 using System.Threading.Tasks;
@@ -33,135 +34,64 @@ namespace UnrealBuildTool.Artifacts
 		Project,
 	}
 
-	/// <summary>
-	/// Type of the artifact
-	/// </summary>
-	public enum ArtifactType : byte
-	{
-
-		/// <summary>
-		/// Artifact is a source file
-		/// </summary>
-		Source,
-
-		/// <summary>
-		/// Artifact is an object file
-		/// </summary>
-		Object,
-	}
-
 	/// <summary>
 	/// Represents a single artifact
 	/// </summary>
-	public readonly struct Artifact
+	/// <param name="Tree">Directory tree containing the artifact</param>
+	/// <param name="Name">Name of the artifact</param>
+	/// <param name="ContentHash">Hash of the artifact contents</param>
+	public record struct Artifact(ArtifactDirectoryTree Tree, Utf8String Name, IoHash ContentHash)
 	{
-
-		/// <summary>
-		/// Directory tree containing the artifact
-		/// </summary>
-		public readonly ArtifactDirectoryTree Tree;
-
-		/// <summary>
-		/// Type of the artifact
-		/// </summary>
-		public readonly ArtifactType Type;
-
 		/// <summary>
-		/// Name of the artifact
+		/// Directory mapping object
 		/// </summary>
-		public readonly Utf8String Name;
+		public IArtifactDirectoryMapping? DirectoryMapping { get; set; } = null;
 
 		/// <summary>
-		/// Hash of the artifact contents
+		/// The full path of the artifact
 		/// </summary>
-		public readonly IoHash ContentHash;
+		/// <returns>Full path of the artifact</returns>
+		public string GetFullPath() => GetFullPath(DirectoryMapping);
 
 		/// <summary>
-		/// Constructor
+		/// The full path of the artifact
 		/// </summary>
-		/// <param name="tree">Directory tree containing the artifact</param>
-		/// <param name="type">Type of the artifact</param>
-		/// <param name="name">Name of the artifact</param>
-		/// <param name="contentHash">Hash of the artifact contents</param>
-		public Artifact(ArtifactDirectoryTree tree, ArtifactType type, Utf8String name, IoHash contentHash)
+		/// <param name="mapping">Mapping object to use instead of embedded object</param>
+		/// <returns>Full path of the artifact</returns>
+		public string GetFullPath(IArtifactDirectoryMapping? mapping)
 		{
-			Tree = tree;
-			Type = type;
-			Name = name;
-			ContentHash = contentHash;
+			if (Tree == ArtifactDirectoryTree.Absolute)
+			{
+				return Name.ToString();
+			}
+			if (mapping == null)
+			{
+				throw new ApplicationException("Attempt to get the full path without mapping interface set");
+			}
+			string directory = mapping.GetDirectory(Tree);
+			return Path.Combine(directory, Name.ToString());
 		}
 	}
 
 	/// <summary>
 	/// Given a set of inputs, provide a list of outputs
 	/// </summary>
-	public readonly struct ArtifactMapping : IEquatable<ArtifactMapping>
+	/// <param name="Key">The hash of the primary input and the environment</param>
+	/// <param name="MappingKey">The unique hash for all inputs and the environment</param>
+	/// <param name="Inputs">Information about all inputs</param>
+	/// <param name="Outputs">Information about all outputs</param>
+	public record struct ArtifactMapping(IoHash Key, IoHash MappingKey, Artifact[] Inputs, Artifact[] Outputs)
 	{
-
-		/// <summary>
-		/// The hash of the primary input and the environment
-		/// </summary>
-		public readonly IoHash Key;
-
-		/// <summary>
-		/// The unique hash for all inputs and the environment
-		/// </summary>
-		public readonly IoHash MappingKey;
-
-		/// <summary>
-		/// Information about all inputs
-		/// </summary>
-		public readonly Artifact[] Inputs;
-
 		/// <summary>
-		/// Information about all outputs
+		/// Directory mapping object
 		/// </summary>
-		public readonly Artifact[] Outputs;
-
-		/// <summary>
-		/// Constructor
-		/// </summary>
-		/// <param name="key">The hash of the primary input and the environment</param>
-		/// <param name="mappingKey">The unique hash for all inputs and the environment</param>
-		/// <param name="inputs">Information about all inputs</param>
-		/// <param name="outputs">Information about all outputs</param>
-		public ArtifactMapping(IoHash key, IoHash mappingKey, Artifact[] inputs, Artifact[] outputs)
-		{
-			Key = key;
-			MappingKey = mappingKey;
-			Inputs = inputs;
-			Outputs = outputs;
-		}
-
-		/// <inheritdoc/>
-		public bool Equals(ArtifactMapping other)
-		{
-			return Key == other.Key && MappingKey == other.MappingKey && Inputs.SequenceEqual(other.Inputs) && Outputs.SequenceEqual(other.Outputs);
-		}
-
-		/// <inheritdoc/>
-		public override bool Equals(object? obj)
-		{
-			return obj is ArtifactMapping artifactMapping && Equals(artifactMapping);
-		}
+		public IArtifactDirectoryMapping? DirectoryMapping { get; set; } = null;
 
 		/// <inheritdoc/>
 		public override int GetHashCode()
 		{
 			return HashCode.Combine(Key.GetHashCode(), MappingKey.GetHashCode(), Inputs.GetHashCode(), Outputs.GetHashCode()); 
 		}
-
-		/// <inheritdoc/>
-		public static bool operator ==(ArtifactMapping left, ArtifactMapping right)
-		{
-			return left.Equals(right);
-		}
-
-		/// <inheritdoc/>
-		public static bool operator !=(ArtifactMapping left, ArtifactMapping right)
-		{
-			return !(left == right);
-		}
 	}
 
 	/// <summary>
@@ -178,10 +108,9 @@ namespace UnrealBuildTool.Artifacts
 		public static Artifact ReadArtifact(this IMemoryReader reader)
 		{
 			ArtifactDirectoryTree tree = (ArtifactDirectoryTree)reader.ReadUInt8();
-			ArtifactType type = (ArtifactType)reader.ReadUInt8();
 			Utf8String name = reader.ReadUtf8String();
 			IoHash contentHash = reader.ReadIoHash();
-			return new Artifact(tree, type, name, contentHash);
+			return new Artifact(tree, name, contentHash);
 		}
 
 		/// <summary>
@@ -192,7 +121,6 @@ namespace UnrealBuildTool.Artifacts
 		public static void WriteArtifact(this IMemoryWriter writer, Artifact artifact)
 		{
 			writer.WriteUInt8((byte)artifact.Tree);
-			writer.WriteUInt8((byte)artifact.Type);
 			writer.WriteUtf8String(artifact.Name);
 			writer.WriteIoHash(artifact.ContentHash);
 		}
@@ -249,6 +177,20 @@ namespace UnrealBuildTool.Artifacts
 		Available,
 	}
 
+	/// <summary>
+	/// Interface used to map a directory tree to an actual directory
+	/// </summary>
+	public interface IArtifactDirectoryMapping
+	{
+
+		/// <summary>
+		/// Return the directory root for the given tree
+		/// </summary>
+		/// <param name="tree">Tree in question</param>
+		/// <returns>Directory root</returns>
+		public string GetDirectory(ArtifactDirectoryTree tree);
+	}
+
 	/// <summary>
 	/// Interface for querying and adding artifacts
 	/// </summary>
@@ -325,6 +267,17 @@ namespace UnrealBuildTool.Artifacts
 		/// </summary>
 		public bool LogCacheMisses { get; set; }
 
+		/// <summary>
+		/// Root location of the engine
+		/// </summary>
+		public DirectoryReference? EngineRoot { get; set; }
+
+		/// <summary>
+		/// Array of extra directory roots when handling relative directories.
+		/// The array MUST be consistent between runs since only the index is saved.
+		/// </summary>
+		public DirectoryReference[]? DirectoryRoots { get; set; }
+
 		/// <summary>
 		/// Complete an action from existing cached data
 		/// </summary>
diff --git a/Engine/Source/Programs/UnrealBuildTool/Artifacts/HordeStorageArtifactCache.cs b/Engine/Source/Programs/UnrealBuildTool/Artifacts/HordeStorageArtifactCache.cs
--- a/Engine/Source/Programs/UnrealBuildTool/Artifacts/HordeStorageArtifactCache.cs
+++ b/Engine/Source/Programs/UnrealBuildTool/Artifacts/HordeStorageArtifactCache.cs
@@ -84,7 +84,7 @@ namespace UnrealBuildTool.Artifacts
 			int index = 0;
 			foreach (Artifact artifact in ArtifactMapping.Outputs)
 			{
-				string outputName = artifact.Name.ToString();
+				string outputName = artifact.GetFullPath();
 				using FileStream stream = new(outputName, FileMode.Open, FileAccess.Read, FileShare.Read);
 				OutputRefs[index++] = new NodeRef<FileNode>(await fileWriter.CreateAsync(stream, nodeTypeOptions.TargetSize, cancellationToken)); 
 			}
@@ -348,7 +348,7 @@ namespace UnrealBuildTool.Artifacts
 							}
 							try
 							{
-								string outputName = hordeArtifactMapping.ArtifactMapping.Outputs[refIndex++].Name.ToString();
+								string outputName = hordeArtifactMapping.ArtifactMapping.Outputs[refIndex++].GetFullPath(mapping.DirectoryMapping);
 								using FileStream stream = new(outputName, FileMode.Create, FileAccess.Write, FileShare.ReadWrite);
 								await FileNode.CopyToStreamAsync(reader, artifactRef.Handle.Locator, stream, cancellationToken);
 							}
@@ -363,7 +363,7 @@ namespace UnrealBuildTool.Artifacts
 						{
 							foreach (Artifact artifact in hordeArtifactMapping.ArtifactMapping.Outputs)
 							{
-								string outputName = artifact.Name.ToString();
+								string outputName = artifact.GetFullPath(mapping.DirectoryMapping);
 								if (File.Exists(outputName))
 								{
 									try
@@ -561,7 +561,7 @@ namespace UnrealBuildTool.Artifacts
 		/// <returns>The reference name</returns>
 		private static RefName GetRefName(IoHash key)
 		{
-			return new RefName($"action_artifact_v1_{key}");
+			return new RefName($"action_artifact_v2_{key}");
 		}
 
 		/// <summary>
diff --git a/Engine/Source/Programs/UnrealBuildTool/Executors/ImmediateActionQueue.cs b/Engine/Source/Programs/UnrealBuildTool/Executors/ImmediateActionQueue.cs
--- a/Engine/Source/Programs/UnrealBuildTool/Executors/ImmediateActionQueue.cs
+++ b/Engine/Source/Programs/UnrealBuildTool/Executors/ImmediateActionQueue.cs
@@ -327,7 +327,13 @@ namespace UnrealBuildTool
 			foreach (LinkedAction action in actions)
 			{
 				action.SortIndex = index;
-				Actions[index++] = new ActionState { Action = action, Status = ActionStatus.Queued, Phase = initialPhase, Results = null };
+				Actions[index++] = new ActionState 
+				{ 
+					Action = action, 
+					Status = ActionStatus.Queued, 
+					Phase = action.ArtifactMode.HasFlag(ArtifactMode.Enabled) ? initialPhase : ActionPhase.Compile,
+					Results = null,
+				};
 			}
 
 			if (readArtifacts)
diff --git a/Engine/Source/Programs/UnrealBuildTool/Modes/BuildMode.cs b/Engine/Source/Programs/UnrealBuildTool/Modes/BuildMode.cs
--- a/Engine/Source/Programs/UnrealBuildTool/Modes/BuildMode.cs
+++ b/Engine/Source/Programs/UnrealBuildTool/Modes/BuildMode.cs
@@ -553,6 +553,7 @@ namespace UnrealBuildTool
 					actionArtifactCache.EnableReads = BuildConfiguration.bArtifactRead;
 					actionArtifactCache.EnableWrites = BuildConfiguration.bArtifactWrites;
 					actionArtifactCache.LogCacheMisses = BuildConfiguration.bLogArtifactCacheMisses;
+					actionArtifactCache.EngineRoot = Unreal.EngineDirectory;
 				}
 
 				// Pre-process module interfaces to generate dependency files
diff --git a/Engine/Source/Programs/UnrealBuildTool/Platform/Windows/VCCompileAction.cs b/Engine/Source/Programs/UnrealBuildTool/Platform/Windows/VCCompileAction.cs
--- a/Engine/Source/Programs/UnrealBuildTool/Platform/Windows/VCCompileAction.cs
+++ b/Engine/Source/Programs/UnrealBuildTool/Platform/Windows/VCCompileAction.cs
@@ -29,6 +29,11 @@ namespace UnrealBuildTool
 		/// </summary>
 		public ActionType ActionType { get; set; } = ActionType.Compile;
 
+		/// <summary>
+		/// Artifact support for this step
+		/// </summary>
+		public ArtifactMode ArtifactMode { get; set; } = ArtifactMode.None;
+
 		/// <summary>
 		/// Path to the compiler
 		/// </summary>
@@ -166,9 +171,6 @@ namespace UnrealBuildTool
 		/// </summary>
 		public List<FileItem> DeleteItems { get; } = new List<FileItem>();
 
-		/// <inheritdoc/>
-		public bool bCanCache { get; set; }
-
 		/// <inheritdoc/>
 		public bool bCanExecuteRemotely { get; set; }
 
@@ -332,6 +334,7 @@ namespace UnrealBuildTool
 		public VCCompileAction(VCCompileAction InAction)
 		{
 			ActionType = InAction.ActionType;
+			ArtifactMode = InAction.ArtifactMode;
 			CompilerExe = InAction.CompilerExe;
 			CompilerType = InAction.CompilerType;
 			ToolChainVersion = InAction.ToolChainVersion;
@@ -372,6 +375,7 @@ namespace UnrealBuildTool
 		public VCCompileAction(BinaryArchiveReader Reader)
 		{
 			ActionType = (ActionType)Reader.ReadInt();
+			ArtifactMode = (ArtifactMode)Reader.ReadByte();
 			CompilerExe = Reader.ReadFileItem()!;
 			CompilerType = (WindowsCompiler)Reader.ReadInt();
 			ToolChainVersion = Reader.ReadString()!;
@@ -409,6 +413,7 @@ namespace UnrealBuildTool
 		public void Write(BinaryArchiveWriter Writer)
 		{
 			Writer.WriteInt((int)ActionType);
+			Writer.WriteByte((byte)ArtifactMode);
 			Writer.WriteFileItem(CompilerExe);
 			Writer.WriteInt((int)CompilerType);
 			Writer.WriteString(ToolChainVersion);
diff --git a/Engine/Source/Programs/UnrealBuildTool/Platform/Windows/VCToolChain.cs b/Engine/Source/Programs/UnrealBuildTool/Platform/Windows/VCToolChain.cs
--- a/Engine/Source/Programs/UnrealBuildTool/Platform/Windows/VCToolChain.cs
+++ b/Engine/Source/Programs/UnrealBuildTool/Platform/Windows/VCToolChain.cs
@@ -1475,7 +1475,17 @@ namespace UnrealBuildTool
 					}
 				}
 
-				CompileAction.bCanCache = CompileEnvironment.bDeterministic && !Target.WindowsPlatform.Compiler.IsClang();
+				if (CompileEnvironment.bDeterministic && !Target.WindowsPlatform.Compiler.IsClang())
+				{
+					if (CompileEnvironment.PrecompiledHeaderAction == PrecompiledHeaderAction.Create)
+					{
+						CompileAction.ArtifactMode |= ArtifactMode.PropagateInputs;
+					}
+					else
+					{
+						CompileAction.ArtifactMode |= ArtifactMode.Enabled | ArtifactMode.AbsolutePath; // deps output file contains absolute paths
+					}
+				}
 
 				// Don't farm out creation of precompiled headers as it is the critical path task.
 				CompileAction.bCanExecuteRemotely =
diff --git a/Engine/Source/Programs/UnrealBuildTool/System/Action.cs b/Engine/Source/Programs/UnrealBuildTool/System/Action.cs
--- a/Engine/Source/Programs/UnrealBuildTool/System/Action.cs
+++ b/Engine/Source/Programs/UnrealBuildTool/System/Action.cs
@@ -1,16 +1,12 @@
 // Copyright Epic Games, Inc. All Rights Reserved.
 
+using EpicGames.Core;
+using Microsoft.Extensions.Logging;
 using System;
 using System.Collections.Generic;
 using System.Diagnostics;
-using System.IO;
 using System.Linq;
-using System.Reflection;
-using System.Text;
 using System.Threading;
-using System.Threading.Tasks;
-using EpicGames.Core;
-using Microsoft.Extensions.Logging;
 using UnrealBuildBase;
 
 namespace UnrealBuildTool
@@ -41,6 +37,37 @@ namespace UnrealBuildTool
 		ParseTimingInfo,
 	}
 
+	/// <summary>
+	/// Defines the action's support for artifacts
+	/// </summary>
+	[Flags]
+	enum ArtifactMode : byte
+	{
+
+		/// <summary>
+		/// Cached artifacts aren't enabled
+		/// </summary>
+		None = 0,
+
+		/// <summary>
+		/// If set, the outputs should be cached but respect the flags below
+		/// </summary>
+		Enabled = 1 << 0,
+
+		/// <summary>
+		/// Absolute file paths must be used when recording inputs and outputs
+		/// </summary>
+		AbsolutePath = 1 << 1,
+
+		/// <summary>
+		/// For actions that can't be cached, by setting this flag, the inputs for the action in question
+		/// will be used as additional inputs for any action that references outputs of this action.  
+		/// For example, PCH files shouldn't be cached, but compiles that use the PCH should still be 
+		/// valid if consider the inputs used to create the PCH file.
+		/// </summary>
+		PropagateInputs = 1 << 2,
+	}
+
 	interface IExternalAction
 	{
 		/// <summary>
@@ -48,6 +75,11 @@ namespace UnrealBuildTool
 		/// </summary>
 		ActionType ActionType { get; }
 
+		/// <summary>
+		/// Artifact support for this step
+		/// </summary>
+		ArtifactMode ArtifactMode { get; }
+
 		/// <summary>
 		/// Every file this action depends on.  These files need to exist and be up to date in order for this action to even be considered
 		/// </summary>
@@ -103,11 +135,6 @@ namespace UnrealBuildTool
 		/// </summary>
 		bool bCanExecuteRemotely { get; }
 
-		/// <summary>
-		/// True if this action can be cached. Cache will store all touched inputs and store outputs
-		/// </summary>
-		bool bCanCache { get; }
-
 		/// <summary>
 		/// True if this action is allowed to be run on a remote machine with SNDBS. Files with #import directives must be compiled locally. Also requires bCanExecuteRemotely = true.
 		/// </summary>
@@ -165,6 +192,11 @@ namespace UnrealBuildTool
 		/// </summary>
 		public ActionType ActionType { get; set; }
 
+		/// <summary>
+		/// Artifact support for this step
+		/// </summary>
+		public ArtifactMode ArtifactMode { get; set; } = ArtifactMode.None;
+
 		/// <summary>
 		/// Every file this action depends on.  These files need to exist and be up to date in order for this action to even be considered
 		/// </summary>
@@ -215,11 +247,6 @@ namespace UnrealBuildTool
 		/// </summary>
 		public string StatusDescription { get; set; } = "...";
 
-		/// <summary>
-		/// True if this action's result can be cached based on its inputs
-		/// </summary>
-		public bool bCanCache { get; set; }
-
 		/// <summary>
 		/// True if this action is allowed to be run on a remote machine when a distributed build system is being used, such as XGE
 		/// </summary>
@@ -278,6 +305,7 @@ namespace UnrealBuildTool
 		public Action(IExternalAction InOther)
 		{
 			ActionType = InOther.ActionType;
+			ArtifactMode = InOther.ArtifactMode;
 			PrerequisiteItems = new SortedSet<FileItem>(InOther.PrerequisiteItems);
 			ProducedItems = new SortedSet<FileItem>(InOther.ProducedItems);
 			DeleteItems = new SortedSet<FileItem>(InOther.DeleteItems);
@@ -291,7 +319,6 @@ namespace UnrealBuildTool
 			bCanExecuteRemotely = InOther.bCanExecuteRemotely;
 			bCanExecuteRemotelyWithSNDBS = InOther.bCanExecuteRemotelyWithSNDBS;
 			bCanExecuteRemotelyWithXGE = InOther.bCanExecuteRemotelyWithXGE;
-			bCanCache = InOther.bCanCache;
 			bIsGCCCompiler = InOther.bIsGCCCompiler;
 			bShouldOutputStatusDescription = InOther.bShouldOutputStatusDescription;
 			bProducesImportLibrary = InOther.bProducesImportLibrary;
@@ -303,6 +330,7 @@ namespace UnrealBuildTool
 		public Action(BinaryArchiveReader Reader)
 		{
 			ActionType = (ActionType)Reader.ReadByte();
+			ArtifactMode = (ArtifactMode)Reader.ReadByte();
 			WorkingDirectory = Reader.ReadDirectoryReferenceNotNull();
 			CommandPath = Reader.ReadFileReference();
 			CommandArguments = Reader.ReadString()!;
@@ -312,7 +340,6 @@ namespace UnrealBuildTool
 			bCanExecuteRemotely = Reader.ReadBool();
 			bCanExecuteRemotelyWithSNDBS = Reader.ReadBool();
 			bCanExecuteRemotelyWithXGE = Reader.ReadBool();
-			bCanCache = Reader.ReadBool();
 			bIsGCCCompiler = Reader.ReadBool();
 			bShouldOutputStatusDescription = Reader.ReadBool();
 			bProducesImportLibrary = Reader.ReadBool();
@@ -332,6 +359,7 @@ namespace UnrealBuildTool
 		public void Write(BinaryArchiveWriter Writer)
 		{
 			Writer.WriteByte((byte)ActionType);
+			Writer.WriteByte((byte)ArtifactMode);
 			Writer.WriteDirectoryReference(WorkingDirectory);
 			Writer.WriteFileReference(CommandPath);
 			Writer.WriteString(CommandArguments);
@@ -341,7 +369,6 @@ namespace UnrealBuildTool
 			Writer.WriteBool(bCanExecuteRemotely);
 			Writer.WriteBool(bCanExecuteRemotelyWithSNDBS);
 			Writer.WriteBool(bCanExecuteRemotelyWithXGE);
-			Writer.WriteBool(bCanCache);
 			Writer.WriteBool(bIsGCCCompiler);
 			Writer.WriteBool(bShouldOutputStatusDescription);
 			Writer.WriteBool(bProducesImportLibrary);
@@ -362,6 +389,12 @@ namespace UnrealBuildTool
 		{
 			Action Action = new Action(Object.GetEnumField<ActionType>("Type"));
 
+			ArtifactMode ArtifactMode;
+			if (Object.TryGetEnumField("ArtifactMode", out ArtifactMode))
+			{
+				Action.ArtifactMode = ArtifactMode;
+			}
+
 			string? WorkingDirectory;
 			if(Object.TryGetStringField("WorkingDirectory", out WorkingDirectory))
 			{
@@ -416,12 +449,6 @@ namespace UnrealBuildTool
 				Action.bCanExecuteRemotelyWithXGE = bCanExecuteRemotelyWithXGE;
 			}
 
-			bool bCanCache;
-			if (Object.TryGetBoolField("bCanCache", out bCanCache))
-			{
-				Action.bCanCache = bCanCache;
-			}
-
 			bool bIsGCCCompiler;
 			if(Object.TryGetBoolField("bIsGCCCompiler", out bIsGCCCompiler))
 			{
@@ -520,6 +547,7 @@ namespace UnrealBuildTool
 		{
 			Writer.WriteValue("Id", LinkedActionToId[Action]);
 			Writer.WriteEnumValue("Type", Action.ActionType);
+			Writer.WriteEnumValue("ArtifactMode", Action.ArtifactMode);
 			Writer.WriteValue("WorkingDirectory", Action.WorkingDirectory.FullName);
 			Writer.WriteValue("CommandPath", Action.CommandPath.FullName);
 			Writer.WriteValue("CommandArguments", Action.CommandArguments);
@@ -529,7 +557,6 @@ namespace UnrealBuildTool
 			Writer.WriteValue("bCanExecuteRemotely", Action.bCanExecuteRemotely);
 			Writer.WriteValue("bCanExecuteRemotelyWithSNDBS", Action.bCanExecuteRemotelyWithSNDBS);
 			Writer.WriteValue("bCanExecuteRemotelyWithXGE", Action.bCanExecuteRemotelyWithXGE);
-			Writer.WriteValue("bCanCache", Action.bCanCache);
 			Writer.WriteValue("bIsGCCCompiler", Action.bIsGCCCompiler);
 			Writer.WriteValue("bShouldOutputStatusDescription", Action.bShouldOutputStatusDescription);
 			Writer.WriteValue("bProducesImportLibrary", Action.bProducesImportLibrary);
@@ -637,6 +664,7 @@ namespace UnrealBuildTool
 		#region Wrapper implementation of IAction
 
 		public ActionType ActionType => Inner.ActionType;
+		public ArtifactMode ArtifactMode => Inner.ArtifactMode;
 		public IEnumerable<FileItem> PrerequisiteItems => Inner.PrerequisiteItems;
 		public IEnumerable<FileItem> ProducedItems => Inner.ProducedItems;
 		public IEnumerable<FileItem> DeleteItems => Inner.DeleteItems;
@@ -647,7 +675,6 @@ namespace UnrealBuildTool
 		public string CommandVersion => Inner.CommandVersion;
 		public string CommandDescription => Inner.CommandDescription;
 		public string StatusDescription => Inner.StatusDescription;
-		public bool bCanCache => Inner.bCanCache;
 		public bool bCanExecuteRemotely => Inner.bCanExecuteRemotely;
 		public bool bCanExecuteRemotelyWithSNDBS => Inner.bCanExecuteRemotelyWithSNDBS;
 		public bool bCanExecuteRemotelyWithXGE => Inner.bCanExecuteRemotelyWithXGE;
diff --git a/Engine/Source/Programs/UnrealBuildTool/System/TargetMakefile.cs b/Engine/Source/Programs/UnrealBuildTool/System/TargetMakefile.cs
--- a/Engine/Source/Programs/UnrealBuildTool/System/TargetMakefile.cs
+++ b/Engine/Source/Programs/UnrealBuildTool/System/TargetMakefile.cs
@@ -179,7 +179,7 @@ namespace UnrealBuildTool
 		/// <summary>
 		/// The version number to write
 		/// </summary>
-		public const int CurrentVersion = 34;
+		public const int CurrentVersion = 35;
 
 		/// <summary>
 		/// The time at which the makefile was created
diff --git a/Engine/Source/Programs/UnrealBuildTool/ToolChain/ISPCToolChain.cs b/Engine/Source/Programs/UnrealBuildTool/ToolChain/ISPCToolChain.cs
--- a/Engine/Source/Programs/UnrealBuildTool/ToolChain/ISPCToolChain.cs
+++ b/Engine/Source/Programs/UnrealBuildTool/ToolChain/ISPCToolChain.cs
@@ -463,6 +463,7 @@ namespace UnrealBuildTool
 				CompileAction.WorkingDirectory = Unreal.EngineSourceDirectory;
 				CompileAction.CommandPath = new FileReference(GetISPCHostCompilerPath(BuildHostPlatform.Current.Platform));
 				CompileAction.StatusDescription = Path.GetFileName(ISPCFile.AbsolutePath);
+				CompileAction.ArtifactMode = ArtifactMode.Enabled;
 
 				CompileAction.bCanExecuteRemotely = true;
 
