[UnrealBuildTool] Add option to throw exception when plugin in DisablePlugins list is enabled by another plugin's dependency

diff --git a/Engine/Source/Programs/UnrealBuildTool/Configuration/TargetRules.cs b/Engine/Source/Programs/UnrealBuildTool/Configuration/TargetRules.cs
--- a/Engine/Source/Programs/UnrealBuildTool/Configuration/TargetRules.cs
+++ b/Engine/Source/Programs/UnrealBuildTool/Configuration/TargetRules.cs
@@ -695,6 +695,11 @@ namespace UnrealBuildTool
 		[CommandLine("-DisablePlugin=", ListSeparator = '+')]
 		public List<string> DisablePlugins = new List<string>();
 
+		/// <summary>
+		/// How to treat conflicts when a disabled plugin is being enabled by another plugin referencing it
+		/// </summary>
+		public WarningLevel DisablePluginsConflictWarningLevel = WarningLevel.Default;
+
 		/// <summary>
 		/// A list of Plugin names that are allowed to exist as dependencies without being defined in the uplugin descriptor
 		/// </summary>
@@ -2953,6 +2958,11 @@ namespace UnrealBuildTool
 			get { return Inner.DisablePlugins; }
 		}
 
+		public WarningLevel DisablePluginsConflictWarningLevel
+		{
+			get { return Inner.DisablePluginsConflictWarningLevel; }
+		}
+
 		public IEnumerable<string> BuildPlugins
 		{
 			get { return Inner.BuildPlugins; }
diff --git a/Engine/Source/Programs/UnrealBuildTool/Configuration/UEBuildTarget.cs b/Engine/Source/Programs/UnrealBuildTool/Configuration/UEBuildTarget.cs
--- a/Engine/Source/Programs/UnrealBuildTool/Configuration/UEBuildTarget.cs
+++ b/Engine/Source/Programs/UnrealBuildTool/Configuration/UEBuildTarget.cs
@@ -4300,6 +4300,24 @@ namespace UnrealBuildTool
 					return null;
 				}
 
+				// Detect when a plugin in DisablePlugins has been enabled by another plugin referencing it
+				if (Rules.DisablePlugins.Contains(Reference.Name))
+				{
+					if (Rules.DisablePluginsConflictWarningLevel == WarningLevel.Error)
+					{
+						LogValue PluginLogValue = LogValue.SourceFile(Info.File, Info.File.GetFileName());
+						throw new BuildLogEventException("Error: Plugin '{Plugin}' (referenced via {ReferenceChain}) is being enabled but found in DisablePlugins list. Suppress this message by setting DisablePluginsConflictWarningLevel = WarningLevel.Off in '{TargetRulesFile}'", PluginLogValue, ReferenceChain, TargetRulesFile.GetFileName());
+					}
+					else if (Rules.DisablePluginsConflictWarningLevel == WarningLevel.Warning)
+					{
+						Logger.LogWarning("Plugin '{ReferenceName}' (referenced via {ReferenceChain}) is being enabled but found in DisablePlugins list. Suppress this message by setting DisablePluginsConflictWarningLevel = WarningLevel.Off in '{TargetRulesFile}'", Reference.Name, ReferenceChain, TargetRulesFile.GetFileName());
+					}
+					else if (Rules.DisablePluginsConflictWarningLevel != WarningLevel.Off)
+					{
+						Logger.LogDebug("Plugin '{ReferenceName}' (referenced via {ReferenceChain}) is being enabled but found in DisablePlugins list. Suppress this message by setting DisablePluginsConflictWarningLevel = WarningLevel.Off in '{TargetRulesFile}'", Reference.Name, ReferenceChain, TargetRulesFile.GetFileName());
+					}
+				}
+
 				// Create the new instance and add it to the cache
 				Logger.LogTrace("Enabling plugin '{ReferenceName}' (referenced via {ReferenceChain})", Reference.Name, ReferenceChain);
 				Instance = new UEBuildPlugin(Info, ReferenceChain);
