LLT targets with project files - solution directory missing case

diff --git a/Engine/Source/Programs/UnrealBuildTool/Configuration/TestTargetRules.cs b/Engine/Source/Programs/UnrealBuildTool/Configuration/TestTargetRules.cs
--- a/Engine/Source/Programs/UnrealBuildTool/Configuration/TestTargetRules.cs
+++ b/Engine/Source/Programs/UnrealBuildTool/Configuration/TestTargetRules.cs
@@ -163,6 +163,10 @@ namespace UnrealBuildTool
 							SolutionDirectory = Path.Combine("Games", ProjectFile.Directory.GetDirectoryName(), "LowLevelTests");
 						}
 					}
+					else
+					{
+						SolutionDirectory = "Programs/LowLevelTests";
+					}
 				}
 				else
 				{
