- FIxed the --delete mode in rsync to not delete extra files when staging in Modern xcode

diff --git a/Engine/Source/Programs/UnrealBuildTool/ProjectFiles/Xcode/XcodeProject.cs b/Engine/Source/Programs/UnrealBuildTool/ProjectFiles/Xcode/XcodeProject.cs
--- a/Engine/Source/Programs/UnrealBuildTool/ProjectFiles/Xcode/XcodeProject.cs
+++ b/Engine/Source/Programs/UnrealBuildTool/ProjectFiles/Xcode/XcodeProject.cs
@@ -1349,7 +1349,8 @@ namespace UnrealBuildTool.XcodeProjectXcconfig
 			{
 				bool bIsEngineBuild = Project.UnrealData.UProjectFileLocation == null;
 				string DefaultStageDir = bIsEngineBuild ? "${UE_OVERRIDE_STAGE_DIR}" : "${UE_PROJECT_DIR}/Saved/StagedBuilds/${UE_TARGET_PLATFORM_NAME}";
-				string SyncSubdir = (Platform == UnrealTargetPlatform.Mac) ? "/UE" : "";
+				string SyncSourceSubdir = (Platform == UnrealTargetPlatform.Mac) ? "" : "/cookeddata";
+				string SyncDestSubdir = (Platform == UnrealTargetPlatform.Mac) ? "/UE" : "/cookeddata";
 				CopyScript.AddRange(new string[]
 				{
 					"# Skip syncing if desired",
@@ -1389,8 +1390,8 @@ namespace UnrealBuildTool.XcodeProjectXcconfig
 				CopyScript.AddRange(new string[]
 				{
 					"",
-					$"echo \\\"Syncing ${{STAGED_DIR}} to ${{CONFIGURATION_BUILD_DIR}}/${{CONTENTS_FOLDER_PATH}}{SyncSubdir}\\\"",
-					$"rsync -a --delete --exclude=Info.plist --exclude=/Manifest_* --exclude=*.app \\\"${{STAGED_DIR}}/\\\" \\\"${{CONFIGURATION_BUILD_DIR}}/${{CONTENTS_FOLDER_PATH}}{SyncSubdir}\\\"",
+					$"echo \\\"Syncing ${{STAGED_DIR}}{SyncSourceSubdir} to ${{CONFIGURATION_BUILD_DIR}}/${{CONTENTS_FOLDER_PATH}}{SyncDestSubdir}\\\"",
+					$"rsync -a --delete --exclude=Info.plist --exclude=/Manifest_* --exclude=*.app \\\"${{STAGED_DIR}}{SyncSourceSubdir}/\\\" \\\"${{CONFIGURATION_BUILD_DIR}}/${{CONTENTS_FOLDER_PATH}}{SyncDestSubdir}\\\"",
 				});
 
 				// copy any dylibs from the staged stub .app into the real one
@@ -1404,7 +1405,7 @@ namespace UnrealBuildTool.XcodeProjectXcconfig
 						"pushd \\\"${STAGED_DIR}/${UE_STAGED_BINARIES_DIR_BASE}/Binaries/Mac/\\\" > /dev/null",
 						"APPS=($(echo *.app))",
 						"echo Syncing $PWD/${APPS[0]}/${BUNDLE_CONTENTS_FOLDER_PATH} to ${CONFIGURATION_BUILD_DIR}/${CONTENTS_FOLDER_PATH}/\\\"",
-						"rsync -a --delete --include='*/' --include='*.dylib' --exclude='*' ${APPS[0]}/${BUNDLE_CONTENTS_FOLDER_PATH}MacOS/\\\" \\\"${CONFIGURATION_BUILD_DIR}/${CONTENTS_FOLDER_PATH}/MacOS\\\"",
+						"rsync -a --include='*/' --include='*.dylib' --exclude='*' ${APPS[0]}/${BUNDLE_CONTENTS_FOLDER_PATH}MacOS/\\\" \\\"${CONFIGURATION_BUILD_DIR}/${CONTENTS_FOLDER_PATH}/MacOS\\\"",
 						"popd > /dev/null",
 					});
 				}
