Build fix.

diff --git a/Engine/Source/Programs/UnrealBuildTool/Artifacts/HordeStorageArtifactCache.cs b/Engine/Source/Programs/UnrealBuildTool/Artifacts/HordeStorageArtifactCache.cs
--- a/Engine/Source/Programs/UnrealBuildTool/Artifacts/HordeStorageArtifactCache.cs
+++ b/Engine/Source/Programs/UnrealBuildTool/Artifacts/HordeStorageArtifactCache.cs
@@ -30,7 +30,7 @@ namespace UnrealBuildTool.Artifacts
 		/// Collection of output file references.  There should be exactly the same number
 		/// of file references as outputs in the mapping
 		/// </summary>
-		public readonly NodeRef<FileNode>[] OutputRefs;
+		public readonly NodeRef<ChunkedDataNode>[] OutputRefs;
 
 		/// <summary>
 		/// Construct a new horde artifact number
@@ -40,7 +40,7 @@ namespace UnrealBuildTool.Artifacts
 		public HordeArtifactMapping(ArtifactMapping artifactMapping)
 		{
 			ArtifactMapping = artifactMapping;
-			OutputRefs = new NodeRef<FileNode>[ArtifactMapping.Outputs.Length];
+			OutputRefs = new NodeRef<ChunkedDataNode>[ArtifactMapping.Outputs.Length];
 		}
 
 		/// <summary>
@@ -50,7 +50,7 @@ namespace UnrealBuildTool.Artifacts
 		public HordeArtifactMapping(ITreeNodeReader reader)
 		{
 			ArtifactMapping = reader.ReadArtifactMapping();
-			OutputRefs = reader.ReadVariableLengthArray(() => new NodeRef<FileNode>(reader));
+			OutputRefs = reader.ReadVariableLengthArray(() => new NodeRef<ChunkedDataNode>(reader));
 		}
 
 		/// <summary>
@@ -80,13 +80,13 @@ namespace UnrealBuildTool.Artifacts
 			ChunkingOptionsForNodeType nodeTypeOptions = new(512 * 1024, 2 * 1024 * 1024, 1 * 1024 * 1024);
 			ChunkingOptions options = new(){ LeafOptions = nodeTypeOptions, InteriorOptions = nodeTypeOptions };
 
-			FileNodeWriter fileWriter = new(writer, options);
+			ChunkedDataWriter fileWriter = new(writer, options);
 			int index = 0;
 			foreach (Artifact artifact in ArtifactMapping.Outputs)
 			{
 				string outputName = artifact.GetFullPath();
 				using FileStream stream = new(outputName, FileMode.Open, FileAccess.Read, FileShare.Read);
-				OutputRefs[index++] = new NodeRef<FileNode>(await fileWriter.CreateAsync(stream, nodeTypeOptions.TargetSize, cancellationToken)); 
+				OutputRefs[index++] = new NodeRef<ChunkedDataNode>(await fileWriter.CreateAsync(stream, nodeTypeOptions.TargetSize, cancellationToken)); 
 			}
 		}
 	}
@@ -339,7 +339,7 @@ namespace UnrealBuildTool.Artifacts
 						output[index] = true;
 
 						int refIndex = 0;
-						foreach (NodeRef<FileNode> artifactRef in hordeArtifactMapping.OutputRefs)
+						foreach (NodeRef<ChunkedDataNode> artifactRef in hordeArtifactMapping.OutputRefs)
 						{
 							if (artifactRef.Handle == null)
 							{
@@ -350,7 +350,7 @@ namespace UnrealBuildTool.Artifacts
 							{
 								string outputName = hordeArtifactMapping.ArtifactMapping.Outputs[refIndex++].GetFullPath(mapping.DirectoryMapping);
 								using FileStream stream = new(outputName, FileMode.Create, FileAccess.Write, FileShare.ReadWrite);
-								await FileNode.CopyToStreamAsync(reader, artifactRef.Handle.Locator, stream, cancellationToken);
+								await ChunkedDataNode.CopyToStreamAsync(reader, artifactRef.Handle.Locator, stream, cancellationToken);
 							}
 							catch (Exception)
 							{
