Build fix.

diff --git a/Engine/Source/Programs/UnrealBuildTool/Artifacts/HordeStorageArtifactCache.cs b/Engine/Source/Programs/UnrealBuildTool/Artifacts/HordeStorageArtifactCache.cs
--- a/Engine/Source/Programs/UnrealBuildTool/Artifacts/HordeStorageArtifactCache.cs
+++ b/Engine/Source/Programs/UnrealBuildTool/Artifacts/HordeStorageArtifactCache.cs
@@ -47,7 +47,7 @@ namespace UnrealBuildTool.Artifacts
 		/// Construct a new artifact mapping from the reader
 		/// </summary>
 		/// <param name="reader">Source reader</param>
-		public HordeArtifactMapping(ITreeNodeReader reader)
+		public HordeArtifactMapping(NodeReader reader)
 		{
 			ArtifactMapping = reader.ReadArtifactMapping();
 			OutputRefs = reader.ReadVariableLengthArray(() => new NodeRef<ChunkedDataNode>(reader));
@@ -102,7 +102,7 @@ namespace UnrealBuildTool.Artifacts
 		/// </summary>
 		/// <param name="reader">Source reader</param>
 		/// <returns>Created artifact mapping</returns>
-		public static HordeArtifactMapping ReadHordeArtifactMapping(this ITreeNodeReader reader)
+		public static HordeArtifactMapping ReadHordeArtifactMapping(this NodeReader reader)
 		{
 			return new HordeArtifactMapping(reader);
 		}
@@ -141,7 +141,7 @@ namespace UnrealBuildTool.Artifacts
 		/// Construct a new mapping collection from the source reader
 		/// </summary>
 		/// <param name="reader">Source reader</param>
-		public ArtifactMappingCollectionNode(ITreeNodeReader reader)
+		public ArtifactMappingCollectionNode(NodeReader reader)
 		{
 			Mappings = reader.ReadDictionary<IoHash, HordeArtifactMapping>(() => reader.ReadIoHash(), () => reader.ReadHordeArtifactMapping());
 		}
@@ -203,11 +203,6 @@ namespace UnrealBuildTool.Artifacts
 		/// </summary>
 		private int _state = (int)ArtifactCacheState.Pending;
 
-		/// <summary>
-		/// Collection of tree reader options. This is needed for the extra node types.
-		/// </summary>
-		private readonly TreeReaderOptions _treeReaderOptions;
-
 		/// <summary>
 		/// Collection of mappings waiting to be written
 		/// </summary>
@@ -259,6 +254,11 @@ namespace UnrealBuildTool.Artifacts
 			return cache;
 		}
 
+		static HordeStorageArtifactCache()
+		{
+			Node.RegisterType<ArtifactMappingCollectionNode>();
+		}
+
 		/// <summary>
 		/// Constructor
 		/// </summary>
@@ -268,8 +268,6 @@ namespace UnrealBuildTool.Artifacts
 		{
 			_store = storage;
 			_logger = logger;
-			_treeReaderOptions = new TreeReaderOptions();
-			_treeReaderOptions.Types.Add(typeof(ArtifactMappingCollectionNode));
 			_pendingMappings = new(MaxPendingSize);
 		}
 
@@ -570,7 +568,7 @@ namespace UnrealBuildTool.Artifacts
 		/// <returns>TreeReader</returns>
 		private TreeReader CreateTreeReader()
 		{
-			return new(_store!, null, _treeReaderOptions, NullLogger.Instance);
+			return new(_store!, null, NullLogger.Instance);
 		}
 
 		/// <summary>
