LLT: the .is_unreal_test file path doesn't depend on having a project file or not

diff --git a/Engine/Source/Programs/UnrealBuildTool/Configuration/TestTargetRules.cs b/Engine/Source/Programs/UnrealBuildTool/Configuration/TestTargetRules.cs
--- a/Engine/Source/Programs/UnrealBuildTool/Configuration/TestTargetRules.cs
+++ b/Engine/Source/Programs/UnrealBuildTool/Configuration/TestTargetRules.cs
@@ -188,15 +188,7 @@ namespace UnrealBuildTool
 				if (File != null)
 				{
 					string OutputDirectory = UEBuildTarget.GetOutputDirectoryForExecutable(Unreal.EngineDirectory, File).FullName;
-					string OutputFileName;
-					if (ProjectFile == null)
-					{ 
-						OutputFileName = $"{OutputDirectory}\\Binaries\\{Target.Platform}\\{ExeBinariesSubFolder}\\{OutputName}.exe.is_unreal_test";
-					}
-					else
-					{
-						OutputFileName = $"{ProjectFile!.Directory.FullName}\\Binaries\\{Target.Platform}\\{ExeBinariesSubFolder}\\{OutputName}.exe.is_unreal_test";
-					}
+					string OutputFileName = $"{OutputDirectory}\\Binaries\\{Target.Platform}\\{ExeBinariesSubFolder}\\{OutputName}.exe.is_unreal_test";
 					PostBuildSteps.Add("echo > " + OutputFileName);
 				}
 			}
