UnrealBuildTool: Use expression body for simple properties and accessors

#rnx

diff --git a/Engine/Source/Programs/UnrealBuildTool/Configuration/ModuleRules.cs b/Engine/Source/Programs/UnrealBuildTool/Configuration/ModuleRules.cs
--- a/Engine/Source/Programs/UnrealBuildTool/Configuration/ModuleRules.cs
+++ b/Engine/Source/Programs/UnrealBuildTool/Configuration/ModuleRules.cs
@@ -579,13 +579,7 @@ namespace UnrealBuildTool
 		/// <summary>
 		/// True if a Plugin contains this module
 		/// </summary>
-		public bool IsPlugin
-		{
-			get
-			{
-				return Plugin != null;
-			}
-		}
+		public bool IsPlugin => Plugin != null;
 
 		/// <summary>
 		/// The rules context for this instance
@@ -614,7 +608,7 @@ namespace UnrealBuildTool
 		/// </summary>
 		public PackageOverrideType OverridePackageType
 		{
-			get { return overridePackageType ?? PackageOverrideType.None; }
+			get => overridePackageType ?? PackageOverrideType.None;
 			set
 			{
 				if (!IsPlugin)
@@ -633,13 +627,7 @@ namespace UnrealBuildTool
 		/// <summary>
 		/// Returns true if there has been an override type specified on this module
 		/// </summary>
-		public bool HasPackageOverride
-		{
-			get
-			{
-				return OverridePackageType != PackageOverrideType.None;
-			}
-		}
+		public bool HasPackageOverride => OverridePackageType != PackageOverrideType.None;
 
 		/// <summary>
 		/// Subfolder of Binaries/PLATFORM folder to put this module in when building DLLs. This should only be used by modules that are found via searching like the
@@ -670,7 +658,7 @@ namespace UnrealBuildTool
 
 				return ShouldOptimizeCode.Value ? CodeOptimization.Always : CodeOptimization.Never;
 			}
-			set { OptimizeCodeOverride = value; }
+			set => OptimizeCodeOverride = value;
 		}
 
 		private OptimizationMode? OptimizationLevelOverride = null;
@@ -699,7 +687,7 @@ namespace UnrealBuildTool
 					return Target.OptimizationLevel;
 				}
 			}
-			set { OptimizationLevelOverride = value; }
+			set => OptimizationLevelOverride = value;
 		}
 
 		/// <summary>
@@ -750,7 +738,7 @@ namespace UnrealBuildTool
 					return PCHUsageMode.NoSharedPCHs;
 				}
 			}
-			set { PCHUsagePrivate = value; }
+			set => PCHUsagePrivate = value;
 		}
 		private PCHUsageMode? PCHUsagePrivate;
 
@@ -765,8 +753,8 @@ namespace UnrealBuildTool
 		/// </summary>
 		public BuildSettingsVersion DefaultBuildSettings
 		{
-			get { return DefaultBuildSettingsPrivate ?? Target.DefaultBuildSettings; }
-			set { DefaultBuildSettingsPrivate = value; }
+			get => DefaultBuildSettingsPrivate ?? Target.DefaultBuildSettings;
+			set => DefaultBuildSettingsPrivate = value;
 		}
 		private BuildSettingsVersion? DefaultBuildSettingsPrivate;
 
@@ -787,7 +775,7 @@ namespace UnrealBuildTool
 				}
 				return IncludeOrderVersionPrivate ?? Target.IncludeOrderVersion;
 			}
-			set { IncludeOrderVersionPrivate = value; }
+			set => IncludeOrderVersionPrivate = value;
 		}
 		private EngineIncludeOrderVersion? IncludeOrderVersionPrivate;
 
@@ -803,14 +791,8 @@ namespace UnrealBuildTool
 		[Obsolete("bUseAVX is obsolete and will be removed in UE5.4, please replace with MinCpuArchX64")]
 		public bool bUseAVX
 		{
-			get
-			{
-				return MinCpuArchX64 >= MinimumCpuArchitectureX64.AVX;
-			}
-			set
-			{
-				MinCpuArchX64 = value ? MinimumCpuArchitectureX64.AVX : MinimumCpuArchitectureX64.None;
-			}
+			get => MinCpuArchX64 >= MinimumCpuArchitectureX64.AVX;
+			set => MinCpuArchX64 = value ? MinimumCpuArchitectureX64.AVX : MinimumCpuArchitectureX64.None;
 		}
 
 		/// <summary>
@@ -846,8 +828,8 @@ namespace UnrealBuildTool
 		/// </summary>
 		public WarningLevel DeterministicWarningLevel
 		{
-			get { return (DeterministicWarningLevelPrivate == WarningLevel.Default) ? (Target.bDeterministic ? WarningLevel.Warning : WarningLevel.Off) : DeterministicWarningLevelPrivate; }
-			set { DeterministicWarningLevelPrivate = value; }
+			get => (DeterministicWarningLevelPrivate == WarningLevel.Default) ? (Target.bDeterministic ? WarningLevel.Warning : WarningLevel.Off) : DeterministicWarningLevelPrivate;
+			set => DeterministicWarningLevelPrivate = value;
 		}
 
 		/// <inheritdoc cref="DeterministicWarningLevelPrivate"/>
@@ -858,8 +840,8 @@ namespace UnrealBuildTool
 		/// </summary>
 		public WarningLevel ShadowVariableWarningLevel
 		{
-			get { return (ShadowVariableWarningLevelPrivate == WarningLevel.Default)? ((DefaultBuildSettings >= BuildSettingsVersion.V2) ? WarningLevel.Error : Target.ShadowVariableWarningLevel) : ShadowVariableWarningLevelPrivate; }
-			set { ShadowVariableWarningLevelPrivate = value; }
+			get => (ShadowVariableWarningLevelPrivate == WarningLevel.Default) ? ((DefaultBuildSettings >= BuildSettingsVersion.V2) ? WarningLevel.Error : Target.ShadowVariableWarningLevel) : ShadowVariableWarningLevelPrivate;
+			set => ShadowVariableWarningLevelPrivate = value;
 		}
 
 		/// <inheritdoc cref="ShadowVariableWarningLevelPrivate"/>
@@ -875,8 +857,8 @@ namespace UnrealBuildTool
 		/// </summary>
 		public WarningLevel UnsafeTypeCastWarningLevel
 		{
-			get { return (UnsafeTypeCastWarningLevelPrivate == WarningLevel.Default)? Target.UnsafeTypeCastWarningLevel : UnsafeTypeCastWarningLevelPrivate; }
-			set { UnsafeTypeCastWarningLevelPrivate = value; }
+			get => (UnsafeTypeCastWarningLevelPrivate == WarningLevel.Default) ? Target.UnsafeTypeCastWarningLevel : UnsafeTypeCastWarningLevelPrivate;
+			set => UnsafeTypeCastWarningLevelPrivate = value;
 		}
 
 		/// <inheritdoc cref="UnsafeTypeCastWarningLevel"/>
@@ -892,8 +874,8 @@ namespace UnrealBuildTool
 		/// </summary>
 		public WarningLevel ModuleIncludePathWarningLevel
 		{
-			get { return (ModuleIncludePathWarningLevelPrivate == WarningLevel.Default) ? Target.ModuleIncludePathWarningLevel : ModuleIncludePathWarningLevelPrivate; }
-			set { ModuleIncludePathWarningLevelPrivate = value; }
+			get => (ModuleIncludePathWarningLevelPrivate == WarningLevel.Default) ? Target.ModuleIncludePathWarningLevel : ModuleIncludePathWarningLevelPrivate;
+			set => ModuleIncludePathWarningLevelPrivate = value;
 		}
 
 		/// <inheritdoc cref="ModuleIncludePathWarningLevel"/>
@@ -904,8 +886,8 @@ namespace UnrealBuildTool
 		/// </summary>
 		public WarningLevel ModuleIncludePrivateWarningLevel
 		{
-			get { return (ModuleIncludePrivateWarningLevelPrivate == WarningLevel.Default) ? Target.ModuleIncludePrivateWarningLevel : ModuleIncludePrivateWarningLevelPrivate; }
-			set { ModuleIncludePrivateWarningLevelPrivate = value; }
+			get => (ModuleIncludePrivateWarningLevelPrivate == WarningLevel.Default) ? Target.ModuleIncludePrivateWarningLevel : ModuleIncludePrivateWarningLevelPrivate;
+			set => ModuleIncludePrivateWarningLevelPrivate = value;
 		}
 
 		/// <inheritdoc cref="ModuleIncludePrivateWarningLevel"/>
@@ -916,8 +898,8 @@ namespace UnrealBuildTool
 		/// </summary>
 		public WarningLevel ModuleIncludeSubdirectoryWarningLevel
 		{
-			get { return (ModuleIncludeSubdirectoryWarningLevelPrivate == WarningLevel.Default) ? Target.ModuleIncludeSubdirectoryWarningLevel : ModuleIncludeSubdirectoryWarningLevelPrivate; }
-			set { ModuleIncludeSubdirectoryWarningLevelPrivate = value; }
+			get => (ModuleIncludeSubdirectoryWarningLevelPrivate == WarningLevel.Default) ? Target.ModuleIncludeSubdirectoryWarningLevel : ModuleIncludeSubdirectoryWarningLevelPrivate;
+			set => ModuleIncludeSubdirectoryWarningLevelPrivate = value;
 		}
 
 		/// <inheritdoc cref="ModuleIncludeSubdirectoryWarningLevel"/>
@@ -979,7 +961,7 @@ namespace UnrealBuildTool
 		/// </summary>
 		public bool bUseUnity
 		{
-			set { bUseUnityOverride = value; }
+			set => bUseUnityOverride = value;
 			get
 			{
 				bool UseUnity = true;
@@ -1304,8 +1286,8 @@ namespace UnrealBuildTool
 		/// </summary>
 		public bool bLegacyPublicIncludePaths
 		{
-			set { bLegacyPublicIncludePathsPrivate = value; }
-			get { return bLegacyPublicIncludePathsPrivate ?? ((DefaultBuildSettings < BuildSettingsVersion.V2) ? Target.bLegacyPublicIncludePaths : false); }
+			set => bLegacyPublicIncludePathsPrivate = value;
+			get => bLegacyPublicIncludePathsPrivate ?? ((DefaultBuildSettings < BuildSettingsVersion.V2) ? Target.bLegacyPublicIncludePaths : false);
 		}
 		private bool? bLegacyPublicIncludePathsPrivate;
 
@@ -1315,8 +1297,8 @@ namespace UnrealBuildTool
 		/// </summary>
 		public bool bLegacyParentIncludePaths
 		{
-			set { bLegacyParentIncludePathsPrivate = value; }
-			get { return bLegacyParentIncludePathsPrivate ?? ((DefaultBuildSettings < BuildSettingsVersion.V3) ? Target.bLegacyParentIncludePaths : false); }
+			set => bLegacyParentIncludePathsPrivate = value;
+			get => bLegacyParentIncludePathsPrivate ?? ((DefaultBuildSettings < BuildSettingsVersion.V3) ? Target.bLegacyParentIncludePaths : false);
 		}
 		private bool? bLegacyParentIncludePathsPrivate;
 
@@ -1364,13 +1346,7 @@ namespace UnrealBuildTool
 		/// <summary>
 		/// The current engine directory
 		/// </summary>
-		public string EngineDirectory
-		{
-			get
-			{
-				return Unreal.EngineDirectory.FullName;
-			}
-		}
+		public string EngineDirectory => Unreal.EngineDirectory.FullName;
 
 		/// <summary>
 		/// Property for the directory containing this plugin. Useful for adding paths to third party dependencies.
@@ -1393,13 +1369,7 @@ namespace UnrealBuildTool
 		/// <summary>
 		/// Property for the directory containing this module. Useful for adding paths to third party dependencies.
 		/// </summary>
-		public string ModuleDirectory
-		{
-			get
-			{
-				return Directory.FullName;
-			}
-		}
+		public string ModuleDirectory => Directory.FullName;
 
 		/// <summary>
 		/// Returns module's low level tests directory "Tests".
@@ -1763,10 +1733,7 @@ namespace UnrealBuildTool
 			}
 		}
 
-		internal bool IsTestModule
-		{
-			get { return bIsTestModuleOverride ?? false; }
-		}
+		internal bool IsTestModule => bIsTestModuleOverride ?? false;
 		/// <summary>
 		/// Whether this is a low level tests module.
 		/// </summary>
@@ -1829,16 +1796,10 @@ namespace UnrealBuildTool
 		/// Returns if VcPkg is supported for the build configuration.
 		/// </summary>
 		/// <returns>True if supported</returns>
-		public bool IsVcPackageSupported
-		{
-			get
-			{
-				return Target.Platform == UnrealTargetPlatform.Win64 ||
+		public bool IsVcPackageSupported => Target.Platform == UnrealTargetPlatform.Win64 ||
 					Target.Platform == UnrealTargetPlatform.Linux ||
 					Target.Platform == UnrealTargetPlatform.LinuxArm64 ||
 					Target.Platform == UnrealTargetPlatform.Mac;
-			}
-		}
 
 		/// <summary>
 		/// Returns the VcPkg root directory for the build configuration
diff --git a/Engine/Source/Programs/UnrealBuildTool/Configuration/TargetInfo.cs b/Engine/Source/Programs/UnrealBuildTool/Configuration/TargetInfo.cs
--- a/Engine/Source/Programs/UnrealBuildTool/Configuration/TargetInfo.cs
+++ b/Engine/Source/Programs/UnrealBuildTool/Configuration/TargetInfo.cs
@@ -44,10 +44,7 @@ namespace UnrealBuildTool
 		/// <summary>
 		/// The current build version
 		/// </summary>
-		public ReadOnlyBuildVersion Version
-		{
-			get { return ReadOnlyBuildVersion.Current; }
-		}
+		public ReadOnlyBuildVersion Version => ReadOnlyBuildVersion.Current;
 
 		/// <summary>
 		/// Additional command line arguments for this target
diff --git a/Engine/Source/Programs/UnrealBuildTool/Configuration/TargetRules.cs b/Engine/Source/Programs/UnrealBuildTool/Configuration/TargetRules.cs
--- a/Engine/Source/Programs/UnrealBuildTool/Configuration/TargetRules.cs
+++ b/Engine/Source/Programs/UnrealBuildTool/Configuration/TargetRules.cs
@@ -427,10 +427,7 @@ namespace UnrealBuildTool
 
 				return DefaultName;
 			}
-			set
-			{
-				NameOverride = value;
-			}
+			set => NameOverride = value;
 		}
 
 		/// <summary>
@@ -454,10 +451,7 @@ namespace UnrealBuildTool
 		/// <summary>
 		/// Whether this is a low level tests target.
 		/// </summary>
-		public bool IsTestTarget
-		{
-			get { return bIsTestTargetOverride; }
-		}
+		public bool IsTestTarget => bIsTestTargetOverride;
 		/// <summary>
 		/// Override this boolean flag in inheriting classes for low level test targets.
 		/// </summary>
@@ -468,21 +462,12 @@ namespace UnrealBuildTool
 		/// Explicitley defined test targets always inherit from TestTargetRules and define their own tests.
 		/// Implicit test targets are created from existing targets when building with -Mode=Test and they include tests from all dependencies.
 		/// </summary>
-		public bool ExplicitTestsTarget
-		{
-			get { return GetType().IsSubclassOf(typeof(TestTargetRules)); }
-		}
+		public bool ExplicitTestsTarget => GetType().IsSubclassOf(typeof(TestTargetRules));
 
 		/// <summary>
 		/// Controls the value of WITH_LOW_LEVEL_TESTS that dictates whether module-specific low level tests are compiled in or not.
 		/// </summary>
-		public bool WithLowLevelTests
-		{
-			get
-			{
-				return (IsTestTarget && !ExplicitTestsTarget) || bWithLowLevelTestsOverride;
-			}
-		}
+		public bool WithLowLevelTests => (IsTestTarget && !ExplicitTestsTarget) || bWithLowLevelTestsOverride;
 		/// <summary>
 		/// Override the value of WithLowLevelTests by setting this to true in inherited classes.
 		/// </summary>
@@ -558,8 +543,8 @@ namespace UnrealBuildTool
 		/// </summary>
 		public BuildSettingsVersion DefaultBuildSettings
 		{
-			get { return DefaultBuildSettingsPrivate ?? BuildSettingsVersion.V1; }
-			set { DefaultBuildSettingsPrivate = value; }
+			get => DefaultBuildSettingsPrivate ?? BuildSettingsVersion.V1;
+			set => DefaultBuildSettingsPrivate = value;
 		}
 		private BuildSettingsVersion? DefaultBuildSettingsPrivate; // Cannot be initialized inline; potentially overridden before the constructor is called.
 
@@ -582,7 +567,7 @@ namespace UnrealBuildTool
 				}
 				return IncludeOrderVersionPrivate ?? EngineIncludeOrderVersion.Oldest;
 			}
-			set { IncludeOrderVersionPrivate = value; }
+			set => IncludeOrderVersionPrivate = value;
 		}
 		private EngineIncludeOrderVersion? IncludeOrderVersionPrivate;
 
@@ -648,8 +633,8 @@ namespace UnrealBuildTool
 		/// </summary>
 		public bool bAllowHotReload
 		{
-			get { return bAllowHotReloadOverride ?? (Type == TargetType.Editor && LinkType == TargetLinkType.Modular); }
-			set { bAllowHotReloadOverride = value; }
+			get => bAllowHotReloadOverride ?? (Type == TargetType.Editor && LinkType == TargetLinkType.Modular);
+			set => bAllowHotReloadOverride = value;
 		}
 		private bool? bAllowHotReloadOverride;
 
@@ -905,7 +890,7 @@ namespace UnrealBuildTool
 		[RequiresUniqueBuildEnvironment]
 		public bool bUseGameplayDebuggerCore
 		{
-			get { return (UseGameplayDebuggerOverride == GameplayDebuggerOverrideState.Core) || bUseGameplayDebugger; }
+			get => (UseGameplayDebuggerOverride == GameplayDebuggerOverrideState.Core) || bUseGameplayDebugger;
 			set
 			{
 				if (UseGameplayDebuggerOverride != GameplayDebuggerOverrideState.Full)
@@ -929,8 +914,8 @@ namespace UnrealBuildTool
 		/// </summary>
 		public bool bBuildEditor
 		{
-			get { return (Type == TargetType.Editor || bCompileAgainstEditor); }
-			set { Logger.LogWarning("Setting {Type}.bBuildEditor is deprecated. Set {Type}.Type instead.", GetType().Name); }
+			get => (Type == TargetType.Editor || bCompileAgainstEditor);
+			set => Logger.LogWarning("Setting {Type}.bBuildEditor is deprecated. Set {Type}.Type instead.", GetType().Name);
 		}
 
 		/// <summary>
@@ -939,8 +924,8 @@ namespace UnrealBuildTool
 		[RequiresUniqueBuildEnvironment]
 		public bool bBuildRequiresCookedData
 		{
-			get { return bBuildRequiresCookedDataOverride ?? (Type == TargetType.Game || Type == TargetType.Client || Type == TargetType.Server); }
-			set { bBuildRequiresCookedDataOverride = value; }
+			get => bBuildRequiresCookedDataOverride ?? (Type == TargetType.Game || Type == TargetType.Client || Type == TargetType.Server);
+			set => bBuildRequiresCookedDataOverride = value;
 		}
 		bool? bBuildRequiresCookedDataOverride;
 
@@ -950,8 +935,8 @@ namespace UnrealBuildTool
 		[RequiresUniqueBuildEnvironment]
 		public bool bBuildWithEditorOnlyData
 		{
-			get { return bBuildWithEditorOnlyDataOverride ?? (Type == TargetType.Editor || Type == TargetType.Program); }
-			set { bBuildWithEditorOnlyDataOverride = value; }
+			get => bBuildWithEditorOnlyDataOverride ?? (Type == TargetType.Editor || Type == TargetType.Program);
+			set => bBuildWithEditorOnlyDataOverride = value;
 		}
 		private bool? bBuildWithEditorOnlyDataOverride;
 
@@ -966,8 +951,8 @@ namespace UnrealBuildTool
 		[RequiresUniqueBuildEnvironment]
 		public bool bBuildDeveloperTools
 		{
-			set { bBuildDeveloperToolsOverride = value; }
-			get { return bBuildDeveloperToolsOverride ?? (bCompileAgainstEngine && (Type == TargetType.Editor || Type == TargetType.Program || (Configuration != UnrealTargetConfiguration.Test && Configuration != UnrealTargetConfiguration.Shipping))); }
+			set => bBuildDeveloperToolsOverride = value;
+			get => bBuildDeveloperToolsOverride ?? (bCompileAgainstEngine && (Type == TargetType.Editor || Type == TargetType.Program || (Configuration != UnrealTargetConfiguration.Test && Configuration != UnrealTargetConfiguration.Shipping)));
 		}
 
 		/// <summary>
@@ -981,8 +966,8 @@ namespace UnrealBuildTool
 		[RequiresUniqueBuildEnvironment]
 		public bool bBuildTargetDeveloperTools
 		{
-			set { bBuildTargetDeveloperToolsOverride = value; }
-			get { return bBuildTargetDeveloperToolsOverride ?? bBuildDeveloperTools; }
+			set => bBuildTargetDeveloperToolsOverride = value;
+			get => bBuildTargetDeveloperToolsOverride ?? bBuildDeveloperTools;
 		}
 
 		/// <summary>
@@ -1005,8 +990,8 @@ namespace UnrealBuildTool
 		/// </summary>
 		public bool bNeedsExtraShaderFormats
 		{
-			set { bNeedsExtraShaderFormatsOverride = value; }
-			get { return bNeedsExtraShaderFormatsOverride ?? (bForceBuildShaderFormats || bBuildDeveloperTools) && (Type == TargetType.Editor || Type == TargetType.Program); }
+			set => bNeedsExtraShaderFormatsOverride = value;
+			get => bNeedsExtraShaderFormatsOverride ?? (bForceBuildShaderFormats || bBuildDeveloperTools) && (Type == TargetType.Editor || Type == TargetType.Program);
 		}
 
 		/// <summary>
@@ -1029,8 +1014,8 @@ namespace UnrealBuildTool
 		[RequiresUniqueBuildEnvironment]
 		public virtual bool bCompileAgainstEngine
 		{
-			get { return bCompileAgainstEnginePrivate; }
-			set { bCompileAgainstEnginePrivate = value; }
+			get => bCompileAgainstEnginePrivate;
+			set => bCompileAgainstEnginePrivate = value;
 		}
 		private bool bCompileAgainstEnginePrivate = true;
 
@@ -1040,8 +1025,8 @@ namespace UnrealBuildTool
 		[RequiresUniqueBuildEnvironment]
 		public virtual bool bCompileAgainstCoreUObject
 		{
-			get { return bCompileAgainstCoreUObjectPrivate; }
-			set { bCompileAgainstCoreUObjectPrivate = value; }
+			get => bCompileAgainstCoreUObjectPrivate;
+			set => bCompileAgainstCoreUObjectPrivate = value;
 		}
 		private bool bCompileAgainstCoreUObjectPrivate = true;
 
@@ -1052,8 +1037,8 @@ namespace UnrealBuildTool
 		[RequiresUniqueBuildEnvironment]
 		public virtual bool bCompileAgainstApplicationCore
 		{
-			get { return bCompileAgainstApplicationCorePrivate; }
-			set { bCompileAgainstApplicationCorePrivate = value; }
+			get => bCompileAgainstApplicationCorePrivate;
+			set => bCompileAgainstApplicationCorePrivate = value;
 		}
 		private bool bCompileAgainstApplicationCorePrivate = true;
 
@@ -1069,8 +1054,8 @@ namespace UnrealBuildTool
 		[RequiresUniqueBuildEnvironment]
 		public virtual bool bCompileAgainstEditor
 		{
-			set { bCompileAgainstEditorOverride = value; }
-			get { return bCompileAgainstEditorOverride ?? (Type == TargetType.Editor); }
+			set => bCompileAgainstEditorOverride = value;
+			get => bCompileAgainstEditorOverride ?? (Type == TargetType.Editor);
 		}
 
 		/// <summary>
@@ -1104,8 +1089,8 @@ namespace UnrealBuildTool
 		[RequiresUniqueBuildEnvironment]
 		public bool bCompileSpeedTree
 		{
-			set { bOverrideCompileSpeedTree = value; }
-			get { return bOverrideCompileSpeedTree ?? Type == TargetType.Editor; }
+			set => bOverrideCompileSpeedTree = value;
+			get => bOverrideCompileSpeedTree ?? Type == TargetType.Editor;
 		}
 
 		/// <summary>
@@ -1152,8 +1137,8 @@ namespace UnrealBuildTool
 		[RequiresUniqueBuildEnvironment]
 		public bool bWithServerCode
 		{
-			get { return bWithServerCodeOverride ?? (Type != TargetType.Client); }
-			set { bWithServerCodeOverride = value; }
+			get => bWithServerCodeOverride ?? (Type != TargetType.Client);
+			set => bWithServerCodeOverride = value;
 		}
 		private bool? bWithServerCodeOverride;
 
@@ -1165,8 +1150,8 @@ namespace UnrealBuildTool
 		[RequiresUniqueBuildEnvironment]
 		public bool bFNameOutlineNumber
 		{
-			get { return bFNameOutlineNumberOverride ?? false;  }
-			set { bFNameOutlineNumberOverride = value;  }
+			get => bFNameOutlineNumberOverride ?? false;
+			set => bFNameOutlineNumberOverride = value;
 		}
 		private bool? bFNameOutlineNumberOverride;
 
@@ -1178,14 +1163,8 @@ namespace UnrealBuildTool
 		[RequiresUniqueBuildEnvironment]
 		public bool bWithPushModel
 		{
-			get
-			{
-				return bWithPushModelOverride ?? (Type == TargetType.Editor);
-			}
-			set
-			{
-				bWithPushModelOverride = value;
-			}
+			get => bWithPushModelOverride ?? (Type == TargetType.Editor);
+			set => bWithPushModelOverride = value;
 		}
 		private bool? bWithPushModelOverride;
 
@@ -1208,8 +1187,8 @@ namespace UnrealBuildTool
 		[RequiresUniqueBuildEnvironment]
 		public bool bIncludePluginsForTargetPlatforms
 		{
-			get { return bIncludePluginsForTargetPlatformsOverride ?? (Type == TargetType.Editor); }
-			set { bIncludePluginsForTargetPlatformsOverride = value; }
+			get => bIncludePluginsForTargetPlatformsOverride ?? (Type == TargetType.Editor);
+			set => bIncludePluginsForTargetPlatformsOverride = value;
 		}
 		private bool? bIncludePluginsForTargetPlatformsOverride;
 
@@ -1225,8 +1204,8 @@ namespace UnrealBuildTool
 		[RequiresUniqueBuildEnvironment]
         public bool bWithPerfCounters
 		{
-			get { return bWithPerfCountersOverride ?? (Type == TargetType.Editor || Type == TargetType.Server); }
-			set { bWithPerfCountersOverride = value; }
+			get => bWithPerfCountersOverride ?? (Type == TargetType.Editor || Type == TargetType.Server);
+			set => bWithPerfCountersOverride = value;
 		}
 
 		[ConfigFile(ConfigHierarchyType.Engine, "/Script/BuildSettings.BuildSettings", "bWithPerfCounters")]
@@ -1238,8 +1217,8 @@ namespace UnrealBuildTool
 		[RequiresUniqueBuildEnvironment]
 		public bool bWithLiveCoding
 		{
-			get { return bWithLiveCodingPrivate ?? (Platform == UnrealTargetPlatform.Win64 && Architecture.bIsX64 && Configuration != UnrealTargetConfiguration.Shipping && Configuration != UnrealTargetConfiguration.Test && Type != TargetType.Program); }
-			set { bWithLiveCodingPrivate = value; }
+			get => bWithLiveCodingPrivate ?? (Platform == UnrealTargetPlatform.Win64 && Architecture.bIsX64 && Configuration != UnrealTargetConfiguration.Shipping && Configuration != UnrealTargetConfiguration.Test && Type != TargetType.Program);
+			set => bWithLiveCodingPrivate = value;
 		}
 		bool? bWithLiveCodingPrivate;
 
@@ -1387,8 +1366,8 @@ namespace UnrealBuildTool
 		/// </summary>
 		public bool bHasExports
 		{
-			get { return bHasExportsOverride ?? (LinkType == TargetLinkType.Modular); }
-			set { bHasExportsOverride = value; }
+			get => bHasExportsOverride ?? (LinkType == TargetLinkType.Modular);
+			set => bHasExportsOverride = value;
 		}
 		private bool? bHasExportsOverride;
 
@@ -1414,8 +1393,8 @@ namespace UnrealBuildTool
 		/// </summary>
 		public bool bBuildAdditionalConsoleApp
 		{
-			get { return bBuildAdditionalConsoleAppOverride ?? (Type == TargetType.Editor); }
-			set { bBuildAdditionalConsoleAppOverride = value; }
+			get => bBuildAdditionalConsoleAppOverride ?? (Type == TargetType.Editor);
+			set => bBuildAdditionalConsoleAppOverride = value;
 		}
 		private bool? bBuildAdditionalConsoleAppOverride;
 
@@ -1429,8 +1408,8 @@ namespace UnrealBuildTool
 		/// </summary>
 		public bool bUseUnityBuild
 		{
-			get { return bUseUnityBuildOverride ?? !bEnableCppModules; }
-			set { bUseUnityBuildOverride = value; }
+			get => bUseUnityBuildOverride ?? !bEnableCppModules;
+			set => bUseUnityBuildOverride = value;
 		}
 
 		/// <summary>
@@ -1493,8 +1472,8 @@ namespace UnrealBuildTool
 		/// </summary>
 		public bool bAdaptiveUnityDisablesPCHForProject
 		{
-			get { return bAdaptiveUnityDisablesProjectPCHForProjectPrivate ?? bAdaptiveUnityDisablesPCH; }
-			set { bAdaptiveUnityDisablesProjectPCHForProjectPrivate = value; }
+			get => bAdaptiveUnityDisablesProjectPCHForProjectPrivate ?? bAdaptiveUnityDisablesPCH;
+			set => bAdaptiveUnityDisablesProjectPCHForProjectPrivate = value;
 		}
 
 		/// <summary>
@@ -1612,12 +1591,10 @@ namespace UnrealBuildTool
 		/// </summary>
 		public bool bRetainFramePointers
 		{
-			get 
-			{
+			get =>
 				// Default to disabled on Linux to maintain legacy behavior
-				return bRetainFramePointersOverride ?? Platform.IsInGroup(UnrealPlatformGroup.Linux) == false;
-			}
-			set { bRetainFramePointersOverride = value; }
+				bRetainFramePointersOverride ?? Platform.IsInGroup(UnrealPlatformGroup.Linux) == false;
+			set => bRetainFramePointersOverride = value;
 		}
 
 		/// <summary>
@@ -2071,8 +2048,8 @@ namespace UnrealBuildTool
 		/// </summary>
 		public bool bLegacyPublicIncludePaths
 		{
-			get { return bLegacyPublicIncludePathsPrivate ?? (DefaultBuildSettings < BuildSettingsVersion.V2); }
-			set { bLegacyPublicIncludePathsPrivate = value; }
+			get => bLegacyPublicIncludePathsPrivate ?? (DefaultBuildSettings < BuildSettingsVersion.V2);
+			set => bLegacyPublicIncludePathsPrivate = value;
 		}
 		private bool? bLegacyPublicIncludePathsPrivate;
 
@@ -2081,8 +2058,8 @@ namespace UnrealBuildTool
 		/// </summary>
 		public bool bLegacyParentIncludePaths
 		{
-			get { return bLegacyParentIncludePathsPrivate ?? (DefaultBuildSettings < BuildSettingsVersion.V3); }
-			set { bLegacyParentIncludePathsPrivate = value; }
+			get => bLegacyParentIncludePathsPrivate ?? (DefaultBuildSettings < BuildSettingsVersion.V3);
+			set => bLegacyParentIncludePathsPrivate = value;
 		}
 		private bool? bLegacyParentIncludePathsPrivate;
 
@@ -2129,14 +2106,8 @@ namespace UnrealBuildTool
 		/// </summary>
 		public TargetLinkType LinkType
 		{
-			get
-			{
-				return (LinkTypePrivate != TargetLinkType.Default) ? LinkTypePrivate : ((Type == global::UnrealBuildTool.TargetType.Editor) ? TargetLinkType.Modular : TargetLinkType.Monolithic);
-			}
-			set
-			{
-				LinkTypePrivate = value;
-			}
+			get => (LinkTypePrivate != TargetLinkType.Default) ? LinkTypePrivate : ((Type == global::UnrealBuildTool.TargetType.Editor) ? TargetLinkType.Modular : TargetLinkType.Monolithic);
+			set => LinkTypePrivate = value;
 		}
 
 		/// <summary>
@@ -2165,14 +2136,8 @@ namespace UnrealBuildTool
 		/// </summary>
 		public string? LaunchModuleName
 		{
-			get
-			{
-				return (LaunchModuleNamePrivate == null && Type != global::UnrealBuildTool.TargetType.Program) ? "Launch" : LaunchModuleNamePrivate;
-			}
-			set
-			{
-				LaunchModuleNamePrivate = value;
-			}
+			get => (LaunchModuleNamePrivate == null && Type != global::UnrealBuildTool.TargetType.Program) ? "Launch" : LaunchModuleNamePrivate;
+			set => LaunchModuleNamePrivate = value;
 		}
 
 		/// <summary>
@@ -2216,7 +2181,7 @@ namespace UnrealBuildTool
 		{
 			get
 			{
-				if(BuildEnvironmentOverride.HasValue)
+				if (BuildEnvironmentOverride.HasValue)
 				{
 					return BuildEnvironmentOverride.Value;
 				}
@@ -2233,10 +2198,7 @@ namespace UnrealBuildTool
 					return TargetBuildEnvironment.Unique;
 				}
 			}
-			set
-			{
-				BuildEnvironmentOverride = value;
-			}
+			set => BuildEnvironmentOverride = value;
 		}
 
 		/// <summary>
@@ -2684,39 +2646,24 @@ namespace UnrealBuildTool
 		/// <summary>
 		/// Gets the host platform being built on
 		/// </summary>
-		public UnrealTargetPlatform HostPlatform
-		{
-			get { return BuildHostPlatform.Current.Platform; }
-		}
+		public UnrealTargetPlatform HostPlatform => BuildHostPlatform.Current.Platform;
 
 		/// <summary>
 		/// Expose the bGenerateProjectFiles flag to targets, so we can modify behavior as appropriate for better intellisense
 		/// </summary>
-		public bool bGenerateProjectFiles
-		{
-			get { return ProjectFileGenerator.bGenerateProjectFiles; }
-		}
+		public bool bGenerateProjectFiles => ProjectFileGenerator.bGenerateProjectFiles;
 
 		/// <summary>
 		/// Indicates whether target rules should be used to explicitly enable or disable plugins. Usually not needed for project generation unless project files indicate whether referenced plugins should be built or not.
 		/// </summary>
-		public bool bShouldTargetRulesTogglePlugins
-		{
-			get
-			{
-				return ((ProjectFileGenerator.Current != null) && ProjectFileGenerator.Current.ShouldTargetRulesTogglePlugins())
+		public bool bShouldTargetRulesTogglePlugins => ((ProjectFileGenerator.Current != null) && ProjectFileGenerator.Current.ShouldTargetRulesTogglePlugins())
 					|| ((ProjectFileGenerator.Current == null) && !ProjectFileGenerator.bGenerateProjectFiles);
-			}
-		}
 
 		/// <summary>
 		/// Expose a setting for whether or not the engine is installed
 		/// </summary>
 		/// <returns>Flag for whether the engine is installed</returns>
-		public bool bIsEngineInstalled
-		{
-			get { return Unreal.IsEngineInstalled(); }
-		}
+		public bool bIsEngineInstalled => Unreal.IsEngineInstalled();
 
 		/// <summary>
 		/// Gets diagnostic messages about default settings which have changed in newer versions of the engine
@@ -2806,1152 +2753,471 @@ namespace UnrealBuildTool
 		/// Accessors for fields on the inner TargetRules instance
 		/// </summary>
 		#region Read-only accessor properties
-		#pragma warning disable CS1591
+#pragma warning disable CS1591
 
-		public string Name
-		{
-			get { return Inner.Name; }
-		}
+		public string Name => Inner.Name;
 
 		public string OriginalName => Inner.OriginalName;
 
-		internal FileReference File
-		{
-			get { return Inner.File!; }
-		}
+		internal FileReference File => Inner.File!;
 
-		internal FileReference TargetSourceFile 
-		{
-			get { return Inner.TargetSourceFile!; }
-		}
+		internal FileReference TargetSourceFile => Inner.TargetSourceFile!;
 
-		internal IReadOnlySet<FileReference> TargetFiles
-		{
-			get { return Inner.TargetFiles!; }
-		}
+		internal IReadOnlySet<FileReference> TargetFiles => Inner.TargetFiles!;
 
-		public object? AdditionalData
-		{
-			get { return Inner.AdditionalData; }
-		}
+		public object? AdditionalData => Inner.AdditionalData;
 
-		public UnrealTargetPlatform Platform
-		{
-			get { return Inner.Platform; }
-		}
+		public UnrealTargetPlatform Platform => Inner.Platform;
 
-		public UnrealTargetConfiguration Configuration
-		{
-			get { return Inner.Configuration; }
-		}
+		public UnrealTargetConfiguration Configuration => Inner.Configuration;
 
-		public UnrealArchitectures Architectures
-		{
-			get { return Inner.Architectures; }
-		}
+		public UnrealArchitectures Architectures => Inner.Architectures;
 
-		public UnrealArch Architecture
-		{
-			get { return Inner.Architecture; }
-		}
+		public UnrealArch Architecture => Inner.Architecture;
 
-		public FileReference? ProjectFile
-		{
-			get { return Inner.ProjectFile; }
-		}
+		public FileReference? ProjectFile => Inner.ProjectFile;
 
-		public ReadOnlyBuildVersion Version
-		{
-			get { return Inner.Version; }
-		}
+		public ReadOnlyBuildVersion Version => Inner.Version;
 
-		public TargetType Type
-		{
-			get { return Inner.Type; }
-		}
+		public TargetType Type => Inner.Type;
 
 		/// <inheritdoc cref="TargetRules.Logger"/>
 		public ILogger Logger => Inner.Logger;
 
-		public BuildSettingsVersion DefaultBuildSettings
-		{
-			get { return Inner.DefaultBuildSettings; }
-		}
+		public BuildSettingsVersion DefaultBuildSettings => Inner.DefaultBuildSettings;
 
-		public EngineIncludeOrderVersion? ForcedIncludeOrder
-		{
-			get { return Inner.ForcedIncludeOrder; }
-		}
+		public EngineIncludeOrderVersion? ForcedIncludeOrder => Inner.ForcedIncludeOrder;
 
-		public EngineIncludeOrderVersion IncludeOrderVersion
-		{
-			get { return Inner.IncludeOrderVersion; }
-		}
+		public EngineIncludeOrderVersion IncludeOrderVersion => Inner.IncludeOrderVersion;
 
-		internal ConfigValueTracker ConfigValueTracker
-		{
-			get { return Inner.ConfigValueTracker; }
-		}
+		internal ConfigValueTracker ConfigValueTracker => Inner.ConfigValueTracker;
 
-		public string? OutputFile
-		{
-			get { return Inner.OutputFile; }
-		}
+		public string? OutputFile => Inner.OutputFile;
 
-		public bool bUsesSteam
-		{
-			get { return Inner.bUsesSteam; }
-		}
+		public bool bUsesSteam => Inner.bUsesSteam;
 
-		public bool bUsesCEF3
-		{
-			get { return Inner.bUsesCEF3; }
-		}
+		public bool bUsesCEF3 => Inner.bUsesCEF3;
 
-		public bool bUsesSlate
-		{
-			get { return Inner.bUsesSlate; }
-		}
+		public bool bUsesSlate => Inner.bUsesSlate;
 
-		public bool bUseStaticCRT
-		{
-			get { return Inner.bUseStaticCRT; }
-		}
+		public bool bUseStaticCRT => Inner.bUseStaticCRT;
 
-		public bool bDebugBuildsActuallyUseDebugCRT
-		{
-			get { return Inner.bDebugBuildsActuallyUseDebugCRT; }
-		}
+		public bool bDebugBuildsActuallyUseDebugCRT => Inner.bDebugBuildsActuallyUseDebugCRT;
 
-		public bool bLegalToDistributeBinary
-		{
-			get { return Inner.bLegalToDistributeBinary; }
-		}
+		public bool bLegalToDistributeBinary => Inner.bLegalToDistributeBinary;
 
-		public UnrealTargetConfiguration UndecoratedConfiguration
-		{
-			get { return Inner.UndecoratedConfiguration; }
-		}
+		public UnrealTargetConfiguration UndecoratedConfiguration => Inner.UndecoratedConfiguration;
 
-		public bool bAllowHotReload
-		{
-			get { return Inner.bAllowHotReload; }
-		}
+		public bool bAllowHotReload => Inner.bAllowHotReload;
 
-		public bool bBuildAllModules
-		{
-			get { return Inner.bBuildAllModules; }
-		}
+		public bool bBuildAllModules => Inner.bBuildAllModules;
 
-		public IEnumerable<string> AdditionalPlugins
-		{
-			get { return Inner.AdditionalPlugins; }
-		}
+		public IEnumerable<string> AdditionalPlugins => Inner.AdditionalPlugins;
 
-		public IEnumerable<string> EnablePlugins
-		{
-			get { return Inner.EnablePlugins; }
-		}
+		public IEnumerable<string> EnablePlugins => Inner.EnablePlugins;
 
-		public IEnumerable<string> DisablePlugins
-		{
-			get { return Inner.DisablePlugins; }
-		}
+		public IEnumerable<string> DisablePlugins => Inner.DisablePlugins;
 
-		public WarningLevel DisablePluginsConflictWarningLevel
-		{
-			get { return Inner.DisablePluginsConflictWarningLevel; }
-		}
+		public WarningLevel DisablePluginsConflictWarningLevel => Inner.DisablePluginsConflictWarningLevel;
 
-		public IEnumerable<string> BuildPlugins
-		{
-			get { return Inner.BuildPlugins; }
-		}
+		public IEnumerable<string> BuildPlugins => Inner.BuildPlugins;
 
-		public IEnumerable<string> InternalPluginDependencies
-		{
-			get { return Inner.InternalPluginDependencies; }
-		}
-		public bool bRuntimeDependenciesComeFromBuildPlugins
-		{
-			get { return Inner.bRuntimeDependenciesComeFromBuildPlugins; }
-		}
+		public IEnumerable<string> InternalPluginDependencies => Inner.InternalPluginDependencies;
+		public bool bRuntimeDependenciesComeFromBuildPlugins => Inner.bRuntimeDependenciesComeFromBuildPlugins;
 
-		public string PakSigningKeysFile
-		{
-			get { return Inner.PakSigningKeysFile; }
-		}
+		public string PakSigningKeysFile => Inner.PakSigningKeysFile;
 
-		public string SolutionDirectory
-		{
-			get { return Inner.SolutionDirectory; }
-		}
+		public string SolutionDirectory => Inner.SolutionDirectory;
 
-		public string CustomConfig
-		{
-			get { return Inner.CustomConfig; }
-		}
+		public string CustomConfig => Inner.CustomConfig;
 
-		public bool? bBuildInSolutionByDefault
-		{
-			get { return Inner.bBuildInSolutionByDefault; }
-		}
+		public bool? bBuildInSolutionByDefault => Inner.bBuildInSolutionByDefault;
 
-		public string ExeBinariesSubFolder
-		{
-			get { return Inner.ExeBinariesSubFolder; }
-		}
+		public string ExeBinariesSubFolder => Inner.ExeBinariesSubFolder;
 
-		public EGeneratedCodeVersion GeneratedCodeVersion
-		{
-			get { return Inner.GeneratedCodeVersion; }
-		}
-		public bool bEnableMeshEditor
-		{
-			get { return Inner.bEnableMeshEditor; }
-		}
+		public EGeneratedCodeVersion GeneratedCodeVersion => Inner.GeneratedCodeVersion;
+		public bool bEnableMeshEditor => Inner.bEnableMeshEditor;
 
-		public bool bUseVerse
-		{
-			get { return Inner.bUseVerse; }
-		}
+		public bool bUseVerse => Inner.bUseVerse;
 
-		public bool bUseAutoRTFMCompiler
-		{
-			get { return Inner.bUseAutoRTFMCompiler; }
-		}
+		public bool bUseAutoRTFMCompiler => Inner.bUseAutoRTFMCompiler;
 
 		[Obsolete("Deprecated in UE5.1 - No longer used as Chaos is always enabled.")]
-		public bool bCompileChaos
-		{
-			get { return Inner.bCompileChaos; }
-		}
+		public bool bCompileChaos => Inner.bCompileChaos;
 
 		[Obsolete("Deprecated in UE5.1 - No longer used as Chaos is always enabled.")]
-		public bool bUseChaos
-		{
-			get { return Inner.bUseChaos; }
-		}
+		public bool bUseChaos => Inner.bUseChaos;
 
-		public bool bUseChaosMemoryTracking
-		{
-			get { return Inner.bUseChaosMemoryTracking; }
-		}
+		public bool bUseChaosMemoryTracking => Inner.bUseChaosMemoryTracking;
 
-		public bool bCompileChaosVisualDebuggerSupport
-		{
-			get { return Inner.bCompileChaosVisualDebuggerSupport; }
-		}
+		public bool bCompileChaosVisualDebuggerSupport => Inner.bCompileChaosVisualDebuggerSupport;
 
-		public bool bUseChaosChecked
-		{
-			get { return Inner.bUseChaosChecked; }
-		}
+		public bool bUseChaosChecked => Inner.bUseChaosChecked;
 
 		[Obsolete("Deprecated in UE5.1 - No longer used in engine.")]
-		public bool bCustomSceneQueryStructure
-		{
-			get { return Inner.bCustomSceneQueryStructure; }
-		}
+		public bool bCustomSceneQueryStructure => Inner.bCustomSceneQueryStructure;
 
 		[Obsolete("Deprecated in UE5.1 - No longer used as Chaos is always enabled.")]
-		public bool bCompilePhysX
-		{
-			get { return Inner.bCompilePhysX; }
-		}
+		public bool bCompilePhysX => Inner.bCompilePhysX;
 
 		[Obsolete("Deprecated in UE5.1 - No longer used as Chaos is always enabled.")]
-		public bool bCompileAPEX
-		{
-			get { return Inner.bCompileAPEX; }
-		}
+		public bool bCompileAPEX => Inner.bCompileAPEX;
 
 		[Obsolete("Deprecated in UE5.1 - No longer used as Chaos is always enabled.")]
-		public bool bCompileNvCloth
-		{
-			get { return Inner.bCompileNvCloth; }
-		}
+		public bool bCompileNvCloth => Inner.bCompileNvCloth;
 
-		public bool bCompileICU
-		{
-			get { return Inner.bCompileICU; }
-		}
+		public bool bCompileICU => Inner.bCompileICU;
 
-		public bool bCompileCEF3
-		{
-			get { return Inner.bCompileCEF3; }
-		}
+		public bool bCompileCEF3 => Inner.bCompileCEF3;
 
-		public bool bCompileISPC
-		{
-			get { return Inner.bCompileISPC; }
-		}
+		public bool bCompileISPC => Inner.bCompileISPC;
 
-		public bool bUseGameplayDebugger
-		{
-			get { return Inner.bUseGameplayDebugger; }
-		}
+		public bool bUseGameplayDebugger => Inner.bUseGameplayDebugger;
 
-		public bool bUseGameplayDebuggerCore
-		{
-			get { return Inner.bUseGameplayDebuggerCore; }
-		}
+		public bool bUseGameplayDebuggerCore => Inner.bUseGameplayDebuggerCore;
 
-		public bool bUseIris
-		{
-			get { return Inner.bUseIris; }
-		}
+		public bool bUseIris => Inner.bUseIris;
 
-		public bool bCompileIntelMetricsDiscovery
-		{
-			get { return Inner.bCompileIntelMetricsDiscovery; }
-		}
-		
-		public bool bCompilePython
-		{
-			get { return Inner.bCompilePython; }
-		}
+		public bool bCompileIntelMetricsDiscovery => Inner.bCompileIntelMetricsDiscovery;
 
-		public bool bBuildEditor
-		{
-			get { return Inner.bBuildEditor; }
-		}
+		public bool bCompilePython => Inner.bCompilePython;
 
-		public bool bBuildRequiresCookedData
-		{
-			get { return Inner.bBuildRequiresCookedData; }
-		}
+		public bool bBuildEditor => Inner.bBuildEditor;
 
-		public bool bBuildWithEditorOnlyData
-		{
-			get { return Inner.bBuildWithEditorOnlyData; }
-		}
+		public bool bBuildRequiresCookedData => Inner.bBuildRequiresCookedData;
 
-		public bool bBuildDeveloperTools
-		{
-			get { return Inner.bBuildDeveloperTools; }
-		}
+		public bool bBuildWithEditorOnlyData => Inner.bBuildWithEditorOnlyData;
 
-		public bool bBuildTargetDeveloperTools
-		{
-			get { return Inner.bBuildTargetDeveloperTools; }
-		}
+		public bool bBuildDeveloperTools => Inner.bBuildDeveloperTools;
 
-		public bool bForceBuildTargetPlatforms
-		{
-			get { return Inner.bForceBuildTargetPlatforms; }
-		}
+		public bool bBuildTargetDeveloperTools => Inner.bBuildTargetDeveloperTools;
 
-		public bool bForceBuildShaderFormats
-		{
-			get { return Inner.bForceBuildShaderFormats; }
-		}
+		public bool bForceBuildTargetPlatforms => Inner.bForceBuildTargetPlatforms;
 
-		public bool bNeedsExtraShaderFormats
-		{
-			get { return Inner.bNeedsExtraShaderFormats; }
-		}
+		public bool bForceBuildShaderFormats => Inner.bForceBuildShaderFormats;
 
-		public bool bCompileCustomSQLitePlatform
-		{
-			get { return Inner.bCompileCustomSQLitePlatform; }
-		}
+		public bool bNeedsExtraShaderFormats => Inner.bNeedsExtraShaderFormats;
 
-		public bool bUseCacheFreedOSAllocs
-		{
-			get { return Inner.bUseCacheFreedOSAllocs; }
-		}
+		public bool bCompileCustomSQLitePlatform => Inner.bCompileCustomSQLitePlatform;
 
-		public bool bCompileAgainstEngine
-		{
-			get { return Inner.bCompileAgainstEngine; }
-		}
+		public bool bUseCacheFreedOSAllocs => Inner.bUseCacheFreedOSAllocs;
 
-		public bool bCompileAgainstCoreUObject
-		{
-			get { return Inner.bCompileAgainstCoreUObject; }
-		}
+		public bool bCompileAgainstEngine => Inner.bCompileAgainstEngine;
 
-		public bool bCompileAgainstApplicationCore
-		{
-			get { return Inner.bCompileAgainstApplicationCore; }
-		}
+		public bool bCompileAgainstCoreUObject => Inner.bCompileAgainstCoreUObject;
 
-		public bool bCompileAgainstEditor
-		{
-			get { return Inner.bCompileAgainstEditor; }
-		}
+		public bool bCompileAgainstApplicationCore => Inner.bCompileAgainstApplicationCore;
 
-		public bool bCompileRecast
-		{
-			get { return Inner.bCompileRecast; }
-		}
+		public bool bCompileAgainstEditor => Inner.bCompileAgainstEditor;
 
-		public bool bCompileNavmeshSegmentLinks
-		{
-			get { return Inner.bCompileNavmeshSegmentLinks; }
-		}
+		public bool bCompileRecast => Inner.bCompileRecast;
 
-		public bool bCompileNavmeshClusterLinks
-		{
-			get { return Inner.bCompileNavmeshClusterLinks; }
-		}
+		public bool bCompileNavmeshSegmentLinks => Inner.bCompileNavmeshSegmentLinks;
 
-		public bool bCompileSpeedTree
-		{
-			get { return Inner.bCompileSpeedTree; }
-		}
+		public bool bCompileNavmeshClusterLinks => Inner.bCompileNavmeshClusterLinks;
 
-		public bool bForceEnableExceptions
-		{
-			get { return Inner.bForceEnableExceptions; }
-		}
+		public bool bCompileSpeedTree => Inner.bCompileSpeedTree;
 
-		public bool bForceEnableObjCExceptions
-		{
-			get { return Inner.bForceEnableObjCExceptions; }
-		}
+		public bool bForceEnableExceptions => Inner.bForceEnableExceptions;
 
-		public bool bForceEnableRTTI
-		{
-			get { return Inner.bForceEnableRTTI; }
-		}
+		public bool bForceEnableObjCExceptions => Inner.bForceEnableObjCExceptions;
 
-		public bool bEnablePIE
-		{
-			get { return Inner.bEnablePIE; }
-		}
+		public bool bForceEnableRTTI => Inner.bForceEnableRTTI;
 
-		public bool bEnableStackProtection
-		{
-			get { return Inner.bEnableStackProtection; }
-		}
+		public bool bEnablePIE => Inner.bEnablePIE;
 
-		public bool bUseInlining
-		{
-			get { return Inner.bUseInlining; }
-		}
+		public bool bEnableStackProtection => Inner.bEnableStackProtection;
 
-		public bool bWithServerCode
-		{
-			get { return Inner.bWithServerCode; }
-		}
+		public bool bUseInlining => Inner.bUseInlining;
 
-		public bool bFNameOutlineNumber
-		{
-			get { return Inner.bFNameOutlineNumber;  }
-		}
+		public bool bWithServerCode => Inner.bWithServerCode;
 
-		public bool bWithPushModel
-		{
-			get { return Inner.bWithPushModel; }
-		}
+		public bool bFNameOutlineNumber => Inner.bFNameOutlineNumber;
 
-		public bool bCompileWithStatsWithoutEngine
-		{
-			get { return Inner.bCompileWithStatsWithoutEngine; }
-		}
+		public bool bWithPushModel => Inner.bWithPushModel;
 
-		public bool bCompileWithPluginSupport
-		{
-			get { return Inner.bCompileWithPluginSupport; }
-		}
+		public bool bCompileWithStatsWithoutEngine => Inner.bCompileWithStatsWithoutEngine;
 
-		public bool bIncludePluginsForTargetPlatforms
-		{
-			get { return Inner.bIncludePluginsForTargetPlatforms; }
-		}
+		public bool bCompileWithPluginSupport => Inner.bCompileWithPluginSupport;
 
-		public bool bCompileWithAccessibilitySupport
-		{
-			get { return Inner.bCompileWithAccessibilitySupport; }
-		}
+		public bool bIncludePluginsForTargetPlatforms => Inner.bIncludePluginsForTargetPlatforms;
 
-		public bool bWithPerfCounters
-		{
-			get { return Inner.bWithPerfCounters; }
-		}
+		public bool bCompileWithAccessibilitySupport => Inner.bCompileWithAccessibilitySupport;
 
-		public bool bWithLiveCoding
-		{
-			get { return Inner.bWithLiveCoding; }
-		}
+		public bool bWithPerfCounters => Inner.bWithPerfCounters;
 
-		public bool bUseDebugLiveCodingConsole
-		{
-			get { return Inner.bUseDebugLiveCodingConsole; }
-		}
+		public bool bWithLiveCoding => Inner.bWithLiveCoding;
 
-		public bool bWithDirectXMath
-		{
-			get { return Inner.bWithDirectXMath; }
-		}
+		public bool bUseDebugLiveCodingConsole => Inner.bUseDebugLiveCodingConsole;
 
-		public bool bWithFixedTimeStepSupport
-		{
-			get { return Inner.bWithFixedTimeStepSupport; }
-		}
+		public bool bWithDirectXMath => Inner.bWithDirectXMath;
 
-		public bool bUseLoggingInShipping
-		{
-			get { return Inner.bUseLoggingInShipping; }
-		}
+		public bool bWithFixedTimeStepSupport => Inner.bWithFixedTimeStepSupport;
 
-		public bool bLoggingToMemoryEnabled
-		{
-			get { return Inner.bLoggingToMemoryEnabled; }
-		}
+		public bool bUseLoggingInShipping => Inner.bUseLoggingInShipping;
 
-		public bool bUseLauncherChecks
-		{
-			get { return Inner.bUseLauncherChecks; }
-		}
+		public bool bLoggingToMemoryEnabled => Inner.bLoggingToMemoryEnabled;
 
-		public bool bUseChecksInShipping
-		{
-			get { return Inner.bUseChecksInShipping; }
-		}
+		public bool bUseLauncherChecks => Inner.bUseLauncherChecks;
 
-		public bool bTCHARIsUTF8
-		{
-			get { return Inner.bTCHARIsUTF8; }
-		}
+		public bool bUseChecksInShipping => Inner.bUseChecksInShipping;
 
-		public bool bUseEstimatedUtcNow
-		{
-			get { return Inner.bUseEstimatedUtcNow; }
-		}
+		public bool bTCHARIsUTF8 => Inner.bTCHARIsUTF8;
 
-		public bool bCompileFreeType
-		{
-			get { return Inner.bCompileFreeType; }
-		}
+		public bool bUseEstimatedUtcNow => Inner.bUseEstimatedUtcNow;
 
-		public bool bUseExecCommandsInShipping
-		{
-			get { return Inner.bUseExecCommandsInShipping; }
-		}
+		public bool bCompileFreeType => Inner.bCompileFreeType;
+
+		public bool bUseExecCommandsInShipping => Inner.bUseExecCommandsInShipping;
 
 		[Obsolete("Deprecated in UE5.1 - Please use OptimizationLevel instead.")]
-		public bool bCompileForSize
-		{
-			get { return Inner.bCompileForSize; }
-		}
+		public bool bCompileForSize => Inner.bCompileForSize;
 
-		public OptimizationMode OptimizationLevel
-		{
-			get { return Inner.OptimizationLevel; }
-		}
+		public OptimizationMode OptimizationLevel => Inner.OptimizationLevel;
 
-		public bool bRetainFramePointers
-		{
-			get { return Inner.bRetainFramePointers; }
-		}
+		public bool bRetainFramePointers => Inner.bRetainFramePointers;
 
-		public bool bForceCompileDevelopmentAutomationTests
-		{
-			get { return Inner.bForceCompileDevelopmentAutomationTests; }
-		}
+		public bool bForceCompileDevelopmentAutomationTests => Inner.bForceCompileDevelopmentAutomationTests;
 
-		public bool bForceCompilePerformanceAutomationTests
-		{
-			get { return Inner.bForceCompilePerformanceAutomationTests; }
-		}
+		public bool bForceCompilePerformanceAutomationTests => Inner.bForceCompilePerformanceAutomationTests;
 
-		public bool bForceDisableAutomationTests
-		{
-			get { return Inner.bForceDisableAutomationTests; }
-		}
+		public bool bForceDisableAutomationTests => Inner.bForceDisableAutomationTests;
 
-		public bool bUseXGEController
-		{
-			get { return Inner.bUseXGEController; }
-		}
+		public bool bUseXGEController => Inner.bUseXGEController;
 
-		public bool bEventDrivenLoader
-		{
-			get { return Inner.bEventDrivenLoader; }
-		}
+		public bool bEventDrivenLoader => Inner.bEventDrivenLoader;
 
-		public PointerMemberBehavior? NativePointerMemberBehaviorOverride
-		{
-			get { return Inner.NativePointerMemberBehaviorOverride; }
-		}
+		public PointerMemberBehavior? NativePointerMemberBehaviorOverride => Inner.NativePointerMemberBehaviorOverride;
 
-		public bool bIWYU
-		{
-			get { return Inner.bIWYU; }
-		}
+		public bool bIWYU => Inner.bIWYU;
 
-		public bool bIWYUHeadersOnly
-		{
-			get { return Inner.bIWYUHeadersOnly; }
-		}
+		public bool bIWYUHeadersOnly => Inner.bIWYUHeadersOnly;
 
-		public bool bEnforceIWYU
-		{
-			get { return Inner.bEnforceIWYU; }
-		}
+		public bool bEnforceIWYU => Inner.bEnforceIWYU;
 
-		public bool bWarnAboutMonolithicHeadersIncluded
-		{
-			get { return Inner.bWarnAboutMonolithicHeadersIncluded; }
-		}
+		public bool bWarnAboutMonolithicHeadersIncluded => Inner.bWarnAboutMonolithicHeadersIncluded;
 
-		public bool bHasExports
-		{
-			get { return Inner.bHasExports; }
-		}
+		public bool bHasExports => Inner.bHasExports;
 
-		public bool bPrecompile
-		{
-			get { return Inner.bPrecompile; }
-		}
+		public bool bPrecompile => Inner.bPrecompile;
 
-		public bool bEnableOSX109Support
-		{
-			get { return Inner.bEnableOSX109Support; }
-		}
+		public bool bEnableOSX109Support => Inner.bEnableOSX109Support;
 
-		public bool bIsBuildingConsoleApplication
-		{
-			get { return Inner.bIsBuildingConsoleApplication; }
-		}
+		public bool bIsBuildingConsoleApplication => Inner.bIsBuildingConsoleApplication;
 
-		public bool bBuildAdditionalConsoleApp
-		{
-			get { return Inner.bBuildAdditionalConsoleApp; }
-		}
+		public bool bBuildAdditionalConsoleApp => Inner.bBuildAdditionalConsoleApp;
 
-		public bool bDisableSymbolCache
-		{
-			get { return Inner.bDisableSymbolCache; }
-		}
+		public bool bDisableSymbolCache => Inner.bDisableSymbolCache;
 
-		public bool bUseUnityBuild
-		{
-			get { return Inner.bUseUnityBuild; }
-		}
+		public bool bUseUnityBuild => Inner.bUseUnityBuild;
 
-		public bool bForceUnityBuild
-		{
-			get { return Inner.bForceUnityBuild; }
-		}
+		public bool bForceUnityBuild => Inner.bForceUnityBuild;
 
-		public IReadOnlyList<string>? DisableMergingModuleAndGeneratedFilesInUnityFiles
-		{
-			get { return Inner.DisableMergingModuleAndGeneratedFilesInUnityFiles; }
-		}
+		public IReadOnlyList<string>? DisableMergingModuleAndGeneratedFilesInUnityFiles => Inner.DisableMergingModuleAndGeneratedFilesInUnityFiles;
 
-		public bool bAdaptiveUnityDisablesOptimizations
-		{
-			get { return Inner.bAdaptiveUnityDisablesOptimizations; }
-		}
+		public bool bAdaptiveUnityDisablesOptimizations => Inner.bAdaptiveUnityDisablesOptimizations;
 
-		public bool bAdaptiveUnityDisablesPCH
-		{
-			get { return Inner.bAdaptiveUnityDisablesPCH; }
-		}
+		public bool bAdaptiveUnityDisablesPCH => Inner.bAdaptiveUnityDisablesPCH;
 
-		public bool bAdaptiveUnityDisablesPCHForProject
-		{
-			get { return Inner.bAdaptiveUnityDisablesPCHForProject; }
-		}
+		public bool bAdaptiveUnityDisablesPCHForProject => Inner.bAdaptiveUnityDisablesPCHForProject;
 
-		public bool bAdaptiveUnityCreatesDedicatedPCH
-		{
-			get { return Inner.bAdaptiveUnityCreatesDedicatedPCH; }
-		}
+		public bool bAdaptiveUnityCreatesDedicatedPCH => Inner.bAdaptiveUnityCreatesDedicatedPCH;
 
-		public bool bAdaptiveUnityEnablesEditAndContinue
-		{
-			get { return Inner.bAdaptiveUnityEnablesEditAndContinue; }
-		}
+		public bool bAdaptiveUnityEnablesEditAndContinue => Inner.bAdaptiveUnityEnablesEditAndContinue;
 
-		public bool bAdaptiveUnityCompilesHeaderFiles
-		{
-			get { return Inner.bAdaptiveUnityCompilesHeaderFiles; }
-		}
+		public bool bAdaptiveUnityCompilesHeaderFiles => Inner.bAdaptiveUnityCompilesHeaderFiles;
 
-		public int MinGameModuleSourceFilesForUnityBuild
-		{
-			get { return Inner.MinGameModuleSourceFilesForUnityBuild; }
-		}
+		public int MinGameModuleSourceFilesForUnityBuild => Inner.MinGameModuleSourceFilesForUnityBuild;
 
-		public WarningLevel DefaultWarningLevel
-		{
-			get { return Inner.DefaultWarningLevel; }
-		}
+		public WarningLevel DefaultWarningLevel => Inner.DefaultWarningLevel;
 
-		public WarningLevel DeprecationWarningLevel
-		{
-			get { return Inner.DeprecationWarningLevel; }
-		}
+		public WarningLevel DeprecationWarningLevel => Inner.DeprecationWarningLevel;
 
-		public WarningLevel ShadowVariableWarningLevel
-		{
-			get { return Inner.ShadowVariableWarningLevel; }
-		}
+		public WarningLevel ShadowVariableWarningLevel => Inner.ShadowVariableWarningLevel;
 
-		public WarningLevel UnsafeTypeCastWarningLevel
-		{
-			get { return Inner.UnsafeTypeCastWarningLevel; }
-		}
+		public WarningLevel UnsafeTypeCastWarningLevel => Inner.UnsafeTypeCastWarningLevel;
 
-		public bool bUndefinedIdentifierErrors
-		{
-			get { return Inner.bUndefinedIdentifierErrors; }
-		}
+		public bool bUndefinedIdentifierErrors => Inner.bUndefinedIdentifierErrors;
 
-		public WarningLevel PCHPerformanceIssueWarningLevel
-		{
-			get { return Inner.PCHPerformanceIssueWarningLevel; }
-		}
+		public WarningLevel PCHPerformanceIssueWarningLevel => Inner.PCHPerformanceIssueWarningLevel;
 
-		public WarningLevel ModuleIncludePathWarningLevel
-		{
-			get { return Inner.ModuleIncludePathWarningLevel; }
-		}
+		public WarningLevel ModuleIncludePathWarningLevel => Inner.ModuleIncludePathWarningLevel;
 
-		public WarningLevel ModuleIncludePrivateWarningLevel
-		{
-			get { return Inner.ModuleIncludePrivateWarningLevel; }
-		}
+		public WarningLevel ModuleIncludePrivateWarningLevel => Inner.ModuleIncludePrivateWarningLevel;
 
-		public WarningLevel ModuleIncludeSubdirectoryWarningLevel
-		{
-			get { return Inner.ModuleIncludeSubdirectoryWarningLevel; }
-		}
+		public WarningLevel ModuleIncludeSubdirectoryWarningLevel => Inner.ModuleIncludeSubdirectoryWarningLevel;
 
-		public bool bWarningsAsErrors
-		{
-			get { return Inner.bWarningsAsErrors; }
-		}
+		public bool bWarningsAsErrors => Inner.bWarningsAsErrors;
 
-		public bool bUseFastMonoCalls
-		{
-			get { return Inner.bUseFastMonoCalls; }
-		}
+		public bool bUseFastMonoCalls => Inner.bUseFastMonoCalls;
 
-		public int NumIncludedBytesPerUnityCPP
-		{
-			get { return Inner.NumIncludedBytesPerUnityCPP; }
-		}
+		public int NumIncludedBytesPerUnityCPP => Inner.NumIncludedBytesPerUnityCPP;
 
-		public bool bDisableModuleNumIncludedBytesPerUnityCPPOverride
-		{
-			get { return Inner.bDisableModuleNumIncludedBytesPerUnityCPPOverride; }
-		}
+		public bool bDisableModuleNumIncludedBytesPerUnityCPPOverride => Inner.bDisableModuleNumIncludedBytesPerUnityCPPOverride;
 
-		public bool bStressTestUnity
-		{
-			get { return Inner.bStressTestUnity; }
-		}
+		public bool bStressTestUnity => Inner.bStressTestUnity;
 
-		public bool bDetailedUnityFiles
-		{
-			get { return Inner.bDetailedUnityFiles; }
-		}
+		public bool bDetailedUnityFiles => Inner.bDetailedUnityFiles;
 
-		public bool bDisableDebugInfo
-		{
-			get { return Inner.bDisableDebugInfo; }
-		}
+		public bool bDisableDebugInfo => Inner.bDisableDebugInfo;
 
-		public bool bDisableDebugInfoForGeneratedCode
-		{
-			get { return Inner.bDisableDebugInfoForGeneratedCode; }
-		}
+		public bool bDisableDebugInfoForGeneratedCode => Inner.bDisableDebugInfoForGeneratedCode;
 
-		public bool bOmitPCDebugInfoInDevelopment
-		{
-			get { return Inner.bOmitPCDebugInfoInDevelopment; }
-		}
+		public bool bOmitPCDebugInfoInDevelopment => Inner.bOmitPCDebugInfoInDevelopment;
 
-		public bool bUsePDBFiles
-		{
-			get { return Inner.bUsePDBFiles; }
-		}
+		public bool bUsePDBFiles => Inner.bUsePDBFiles;
 
-		public bool bUsePCHFiles
-		{
-			get { return Inner.bUsePCHFiles; }
-		}
+		public bool bUsePCHFiles => Inner.bUsePCHFiles;
 
-		public bool bDeterministic
-		{
-			get { return Inner.bDeterministic; }
-		}
-		public bool bChainPCHs
-		{
-			get { return Inner.bChainPCHs; }
-		}
+		public bool bDeterministic => Inner.bDeterministic;
+		public bool bChainPCHs => Inner.bChainPCHs;
 
-		public bool bForceIncludePCHHeadersForGenCppFilesWhenPCHIsDisabled
-		{
-			get { return Inner.bForceIncludePCHHeadersForGenCppFilesWhenPCHIsDisabled; }
-		}
+		public bool bForceIncludePCHHeadersForGenCppFilesWhenPCHIsDisabled => Inner.bForceIncludePCHHeadersForGenCppFilesWhenPCHIsDisabled;
 
-		public bool bPreprocessOnly
-		{
-			get { return Inner.bPreprocessOnly; }
-		}
+		public bool bPreprocessOnly => Inner.bPreprocessOnly;
 
-		public bool bPreprocessDepends
-		{
-			get { return Inner.bPreprocessDepends; }
-		}
+		public bool bPreprocessDepends => Inner.bPreprocessDepends;
 
-		public bool bWithAssembly
-		{
-			get { return Inner.bWithAssembly; }
-		}
+		public bool bWithAssembly => Inner.bWithAssembly;
 
-		public StaticAnalyzer StaticAnalyzer
-		{
-			get { return Inner.StaticAnalyzer; }
-		}
+		public StaticAnalyzer StaticAnalyzer => Inner.StaticAnalyzer;
 
-		public StaticAnalyzerOutputType StaticAnalyzerOutputType
-		{
-			get { return Inner.StaticAnalyzerOutputType; }
-		}
+		public StaticAnalyzerOutputType StaticAnalyzerOutputType => Inner.StaticAnalyzerOutputType;
 
-		public StaticAnalyzerMode StaticAnalyzerMode
-		{
-			get { return Inner.StaticAnalyzerMode; }
-		}
+		public StaticAnalyzerMode StaticAnalyzerMode => Inner.StaticAnalyzerMode;
 
-		public int MinFilesUsingPrecompiledHeader
-		{
-			get { return Inner.MinFilesUsingPrecompiledHeader; }
-		}
+		public int MinFilesUsingPrecompiledHeader => Inner.MinFilesUsingPrecompiledHeader;
 
-		public bool bForcePrecompiledHeaderForGameModules
-		{
-			get { return Inner.bForcePrecompiledHeaderForGameModules; }
-		}
+		public bool bForcePrecompiledHeaderForGameModules => Inner.bForcePrecompiledHeaderForGameModules;
 
-		public bool bUseIncrementalLinking
-		{
-			get { return Inner.bUseIncrementalLinking; }
-		}
+		public bool bUseIncrementalLinking => Inner.bUseIncrementalLinking;
 
-		public bool bAllowLTCG
-		{
-			get { return Inner.bAllowLTCG; }
-		}
+		public bool bAllowLTCG => Inner.bAllowLTCG;
 
-		public bool bPreferThinLTO
-		{
-			get { return Inner.bPreferThinLTO; }
-		}
+		public bool bPreferThinLTO => Inner.bPreferThinLTO;
 
-		public bool bPGOProfile
-		{
-			get { return Inner.bPGOProfile; }
-		}
+		public bool bPGOProfile => Inner.bPGOProfile;
 
-		public bool bPGOOptimize
-		{
-			get { return Inner.bPGOOptimize; }
-		}
+		public bool bPGOOptimize => Inner.bPGOOptimize;
 
-		public bool bSupportEditAndContinue
-		{
-			get { return Inner.bSupportEditAndContinue; }
-		}
+		public bool bSupportEditAndContinue => Inner.bSupportEditAndContinue;
 
-		public bool bOmitFramePointers
-		{
-			get { return Inner.bOmitFramePointers; }
-		}
+		public bool bOmitFramePointers => Inner.bOmitFramePointers;
 
-		public bool bEnableCppModules
-		{
-			get { return Inner.bEnableCppModules; }
-		}
+		public bool bEnableCppModules => Inner.bEnableCppModules;
 
-		public bool bEnableCppCoroutinesForEvaluation
-		{
-			get { return Inner.bEnableCppCoroutinesForEvaluation; }
-		}
+		public bool bEnableCppCoroutinesForEvaluation => Inner.bEnableCppCoroutinesForEvaluation;
 
-		public bool bEnableProcessPriorityControl
-		{
-			get { return Inner.bEnableProcessPriorityControl; }
-		}
+		public bool bEnableProcessPriorityControl => Inner.bEnableProcessPriorityControl;
 
 		[Obsolete("Deprecated in UE5.3 - No longer used as MallocProfiler is removed in favor of UnrealInsights.")]
-		public bool bUseMallocProfiler
-		{
-			get { return Inner.bUseMallocProfiler; }
-		}
+		public bool bUseMallocProfiler => Inner.bUseMallocProfiler;
 
-		public bool bShaderCompilerWorkerTrace
-        {
-			get { return Inner.bShaderCompilerWorkerTrace;  }
-        }
+		public bool bShaderCompilerWorkerTrace => Inner.bShaderCompilerWorkerTrace;
 
-		public bool bUseSharedPCHs
-		{
-			get { return Inner.bUseSharedPCHs; }
-		}
+		public bool bUseSharedPCHs => Inner.bUseSharedPCHs;
 
-		public bool bUseShippingPhysXLibraries
-		{
-			get { return Inner.bUseShippingPhysXLibraries; }
-		}
+		public bool bUseShippingPhysXLibraries => Inner.bUseShippingPhysXLibraries;
 
-		public bool bUseCheckedPhysXLibraries
-		{
-			get { return Inner.bUseCheckedPhysXLibraries; }
-		}
+		public bool bUseCheckedPhysXLibraries => Inner.bUseCheckedPhysXLibraries;
 
-		public bool bCheckLicenseViolations
-		{
-			get { return Inner.bCheckLicenseViolations; }
-		}
+		public bool bCheckLicenseViolations => Inner.bCheckLicenseViolations;
 
-		public bool bBreakBuildOnLicenseViolation
-		{
-			get { return Inner.bBreakBuildOnLicenseViolation; }
-		}
+		public bool bBreakBuildOnLicenseViolation => Inner.bBreakBuildOnLicenseViolation;
 
-		public bool? bUseFastPDBLinking
-		{
-			get { return Inner.bUseFastPDBLinking; }
-		}
+		public bool? bUseFastPDBLinking => Inner.bUseFastPDBLinking;
 
-		public bool bCreateMapFile
-		{
-			get { return Inner.bCreateMapFile; }
-		}
+		public bool bCreateMapFile => Inner.bCreateMapFile;
 
-		public bool bAllowRuntimeSymbolFiles
-		{
-			get { return Inner.bAllowRuntimeSymbolFiles; }
-		}
+		public bool bAllowRuntimeSymbolFiles => Inner.bAllowRuntimeSymbolFiles;
 
-		public string? PackagePath
-		{
-			get { return Inner.PackagePath; }
-		}
+		public string? PackagePath => Inner.PackagePath;
 
-		public string? CrashDiagnosticDirectory
-		{
-			get { return Inner.CrashDiagnosticDirectory; }
-		}
+		public string? CrashDiagnosticDirectory => Inner.CrashDiagnosticDirectory;
 
-		public string? ThinLTOCacheDirectory
-		{
-			get { return Inner.ThinLTOCacheDirectory; }
-		}
+		public string? ThinLTOCacheDirectory => Inner.ThinLTOCacheDirectory;
 
-		public string? ThinLTOCachePruningArguments
-		{
-			get { return Inner.ThinLTOCachePruningArguments; }
-		}
+		public string? ThinLTOCachePruningArguments => Inner.ThinLTOCachePruningArguments;
 
-		public string? BundleVersion
-		{
-			get { return Inner.BundleVersion; }
-		}
+		public string? BundleVersion => Inner.BundleVersion;
 
-		public bool bDeployAfterCompile
-		{
-			get { return Inner.bDeployAfterCompile; }
-		}
+		public bool bDeployAfterCompile => Inner.bDeployAfterCompile;
 
-		public bool bAllowRemotelyCompiledPCHs
-		{
-			get { return Inner.bAllowRemotelyCompiledPCHs; }
-		}
+		public bool bAllowRemotelyCompiledPCHs => Inner.bAllowRemotelyCompiledPCHs;
 
-		public bool bCheckSystemHeadersForModification
-		{
-			get { return Inner.bCheckSystemHeadersForModification; }
-		}
+		public bool bCheckSystemHeadersForModification => Inner.bCheckSystemHeadersForModification;
 
-		public bool bDisableLinking
-		{
-			get { return Inner.bDisableLinking; }
-		}
+		public bool bDisableLinking => Inner.bDisableLinking;
 
-		public bool bIgnoreBuildOutputs
-		{
-			get { return Inner.bIgnoreBuildOutputs; }
-		}
+		public bool bIgnoreBuildOutputs => Inner.bIgnoreBuildOutputs;
 
-		public bool bFormalBuild
-		{
-			get { return Inner.bFormalBuild; }
-		}
+		public bool bFormalBuild => Inner.bFormalBuild;
 
-		public bool bUseAdaptiveUnityBuild
-		{
-			get { return Inner.bUseAdaptiveUnityBuild; }
-		}
+		public bool bUseAdaptiveUnityBuild => Inner.bUseAdaptiveUnityBuild;
 
-		public bool bFlushBuildDirOnRemoteMac
-		{
-			get { return Inner.bFlushBuildDirOnRemoteMac; }
-		}
+		public bool bFlushBuildDirOnRemoteMac => Inner.bFlushBuildDirOnRemoteMac;
 
-		public bool bPrintToolChainTimingInfo
-		{
-			get { return Inner.bPrintToolChainTimingInfo; }
-		}
+		public bool bPrintToolChainTimingInfo => Inner.bPrintToolChainTimingInfo;
 
-		public bool bParseTimingInfoForTracing
-		{
-			get { return Inner.bParseTimingInfoForTracing; }
-		}
+		public bool bParseTimingInfoForTracing => Inner.bParseTimingInfoForTracing;
 
-		public bool bPublicSymbolsByDefault
-		{
-			get { return Inner.bPublicSymbolsByDefault; }
-		}
+		public bool bPublicSymbolsByDefault => Inner.bPublicSymbolsByDefault;
 
-		public string? ToolChainName
-		{
-			get { return Inner.ToolChainName; }
-		}
+		public string? ToolChainName => Inner.ToolChainName;
 
-		public float MSVCCompileActionWeight
-		{
-			get { return Inner.MSVCCompileActionWeight; }
-		}
+		public float MSVCCompileActionWeight => Inner.MSVCCompileActionWeight;
 
-		public float ClangCompileActionWeight
-		{
-			get { return Inner.ClangCompileActionWeight; }
-		}
+		public float ClangCompileActionWeight => Inner.ClangCompileActionWeight;
 
-		public bool bLegacyPublicIncludePaths
-		{
-			get { return Inner.bLegacyPublicIncludePaths; }
-		}
+		public bool bLegacyPublicIncludePaths => Inner.bLegacyPublicIncludePaths;
 
-		public bool bLegacyParentIncludePaths
-		{
-			get { return Inner.bLegacyParentIncludePaths; }
-		}
+		public bool bLegacyParentIncludePaths => Inner.bLegacyParentIncludePaths;
 
-		public CppStandardVersion CppStandard
-		{
-			get { return Inner.CppStandard; }
-		}
+		public CppStandardVersion CppStandard => Inner.CppStandard;
 
-		public CStandardVersion CStandard
-		{
-			get { return Inner.CStandard; }
-		}
+		public CStandardVersion CStandard => Inner.CStandard;
 
-		public MinimumCpuArchitectureX64 MinCpuArchX64
-		{
-			get { return Inner.MinCpuArchX64; }
-		}
+		public MinimumCpuArchitectureX64 MinCpuArchX64 => Inner.MinCpuArchX64;
 
-		internal bool bNoManifestChanges
-		{
-			get { return Inner.bNoManifestChanges; }
-		}
+		internal bool bNoManifestChanges => Inner.bNoManifestChanges;
 
-		public string? BuildVersion
-		{
-			get { return Inner.BuildVersion; }
-		}
+		public string? BuildVersion => Inner.BuildVersion;
 
-		public TargetLinkType LinkType
-		{
-			get { return Inner.LinkType; }
-		}
+		public TargetLinkType LinkType => Inner.LinkType;
 
-		public IReadOnlyList<string> GlobalDefinitions
-		{
-			get { return Inner.GlobalDefinitions.AsReadOnly(); }
-		}
+		public IReadOnlyList<string> GlobalDefinitions => Inner.GlobalDefinitions.AsReadOnly();
 
-		public IReadOnlyList<string> ProjectDefinitions
-		{
-			get { return Inner.ProjectDefinitions.AsReadOnly(); }
-		}
+		public IReadOnlyList<string> ProjectDefinitions => Inner.ProjectDefinitions.AsReadOnly();
 
-		public string? LaunchModuleName
-		{
-			get { return Inner.LaunchModuleName; }
-		}
+		public string? LaunchModuleName => Inner.LaunchModuleName;
 
-		public string? ExportPublicHeader
-		{
-			get { return Inner.ExportPublicHeader; }
-		}
+		public string? ExportPublicHeader => Inner.ExportPublicHeader;
 
-		public IReadOnlyList<string> ExtraModuleNames
-		{
-			get { return Inner.ExtraModuleNames.AsReadOnly(); }
-		}
+		public IReadOnlyList<string> ExtraModuleNames => Inner.ExtraModuleNames.AsReadOnly();
 
-		public IReadOnlyList<FileReference> ManifestFileNames
-		{
-			get { return Inner.ManifestFileNames.AsReadOnly(); }
-		}
+		public IReadOnlyList<FileReference> ManifestFileNames => Inner.ManifestFileNames.AsReadOnly();
 
-		public IReadOnlyList<FileReference> DependencyListFileNames
-		{
-			get { return Inner.DependencyListFileNames.AsReadOnly(); }
-		}
+		public IReadOnlyList<FileReference> DependencyListFileNames => Inner.DependencyListFileNames.AsReadOnly();
 
-		public TargetBuildEnvironment BuildEnvironment
-		{
-			get { return Inner.BuildEnvironment; }
-		}
+		public TargetBuildEnvironment BuildEnvironment => Inner.BuildEnvironment;
 
-		public bool bOverrideBuildEnvironment
-		{
-			get { return Inner.bOverrideBuildEnvironment; }
-		}
+		public bool bOverrideBuildEnvironment => Inner.bOverrideBuildEnvironment;
 
-		public IReadOnlyList<TargetInfo> PreBuildTargets
-		{
-			get { return Inner.PreBuildTargets; }
-		}
+		public IReadOnlyList<TargetInfo> PreBuildTargets => Inner.PreBuildTargets;
 
-		public IReadOnlyList<string> PreBuildSteps
-		{
-			get { return Inner.PreBuildSteps; }
-		}
+		public IReadOnlyList<string> PreBuildSteps => Inner.PreBuildSteps;
 
-		public IReadOnlyList<string> PostBuildSteps
-		{
-			get { return Inner.PostBuildSteps; }
-		}
+		public IReadOnlyList<string> PostBuildSteps => Inner.PostBuildSteps;
 
-		public IReadOnlyList<string> AdditionalBuildProducts
-		{
-			get { return Inner.AdditionalBuildProducts; }
-		}
+		public IReadOnlyList<string> AdditionalBuildProducts => Inner.AdditionalBuildProducts;
 
-		public string? AdditionalCompilerArguments
-		{
-			get { return Inner.AdditionalCompilerArguments; }
-		}
+		public string? AdditionalCompilerArguments => Inner.AdditionalCompilerArguments;
 
-		public string? AdditionalLinkerArguments
-		{
-			get { return Inner.AdditionalLinkerArguments; }
-		}
+		public string? AdditionalLinkerArguments => Inner.AdditionalLinkerArguments;
 
-		public double MemoryPerActionGB
-		{
-			get { return Inner.MemoryPerActionGB; }
-		}
+		public double MemoryPerActionGB => Inner.MemoryPerActionGB;
 
-		public string? GeneratedProjectName
-		{
-			get { return Inner.GeneratedProjectName; }
-		}
+		public string? GeneratedProjectName => Inner.GeneratedProjectName;
 
 		public ReadOnlyAndroidTargetRules AndroidPlatform
 		{
@@ -3983,56 +3249,26 @@ namespace UnrealBuildTool
 			private set;
 		}
 
-		public bool bShouldCompileAsDLL
-		{
-			get { return Inner.bShouldCompileAsDLL; }
-		}
+		public bool bShouldCompileAsDLL => Inner.bShouldCompileAsDLL;
 
-		public bool bGenerateProjectFiles
-		{
-			get { return Inner.bGenerateProjectFiles; }
-		}
+		public bool bGenerateProjectFiles => Inner.bGenerateProjectFiles;
 
-		public bool bIsEngineInstalled
-		{
-			get { return Inner.bIsEngineInstalled; }
-		}
+		public bool bIsEngineInstalled => Inner.bIsEngineInstalled;
 
 
-		public IReadOnlyList<string>? DisableUnityBuildForModules
-		{
-			get { return Inner.DisableUnityBuildForModules; }
-		}
+		public IReadOnlyList<string>? DisableUnityBuildForModules => Inner.DisableUnityBuildForModules;
 
-		public IReadOnlyList<string>? EnableOptimizeCodeForModules
-		{
-			get { return Inner.EnableOptimizeCodeForModules; }
-		}
+		public IReadOnlyList<string>? EnableOptimizeCodeForModules => Inner.EnableOptimizeCodeForModules;
 
-		public IReadOnlyList<string>? DisableOptimizeCodeForModules
-		{
-			get { return Inner.DisableOptimizeCodeForModules; }
-		}
+		public IReadOnlyList<string>? DisableOptimizeCodeForModules => Inner.DisableOptimizeCodeForModules;
 
-		public IReadOnlyList<string>? OptimizeForSizeModules
-		{
-			get { return Inner.OptimizeForSizeModules; }
-		}
+		public IReadOnlyList<string>? OptimizeForSizeModules => Inner.OptimizeForSizeModules;
 
-		public IReadOnlyList<string>? OptimizeForSizeAndSpeedModules
-		{
-			get { return Inner.OptimizeForSizeAndSpeedModules; }
-		}
+		public IReadOnlyList<string>? OptimizeForSizeAndSpeedModules => Inner.OptimizeForSizeAndSpeedModules;
 
-		public IReadOnlyList<UnrealTargetPlatform>? OptedInModulePlatforms
-		{
-			get { return Inner.OptedInModulePlatforms; } 
-		}
+		public IReadOnlyList<UnrealTargetPlatform>? OptedInModulePlatforms => Inner.OptedInModulePlatforms;
 
-		public TestTargetRules? InnerTestTargetRules
-		{
-			get { return Inner as TestTargetRules; }
-		}
+		public TestTargetRules? InnerTestTargetRules => Inner as TestTargetRules;
 
 #pragma warning restore C1591
 		#endregion
@@ -4040,52 +3276,34 @@ namespace UnrealBuildTool
 		/// <summary>
 		/// Provide access to the RelativeEnginePath property for code referencing ModuleRules.BuildConfiguration.
 		/// </summary>
-		public string RelativeEnginePath
-		{
-			get { return Unreal.EngineDirectory.MakeRelativeTo(DirectoryReference.GetCurrentDirectory()) + Path.DirectorySeparatorChar; }
-		}
+		public string RelativeEnginePath => Unreal.EngineDirectory.MakeRelativeTo(DirectoryReference.GetCurrentDirectory()) + Path.DirectorySeparatorChar;
 
 		/// <summary>
 		/// Provide access to the UEThirdPartySourceDirectory property for code referencing ModuleRules.UEBuildConfiguration.
 		/// </summary>
-		public string UEThirdPartySourceDirectory
-		{
-			get { return "ThirdParty/"; }
-		}
+		public string UEThirdPartySourceDirectory => "ThirdParty/";
 
 		/// <summary>
 		/// Provide access to the UEThirdPartyBinariesDirectory property for code referencing ModuleRules.UEBuildConfiguration.
 		/// </summary>
-		public string UEThirdPartyBinariesDirectory
-		{
-			get { return "../Binaries/ThirdParty/"; }
-		}
+		public string UEThirdPartyBinariesDirectory => "../Binaries/ThirdParty/";
 
 		/// <summary>
 		/// Whether this is a low level tests target.
 		/// </summary>
-		public bool IsTestTarget
-		{
-			get { return Inner.IsTestTarget; }
-		}
+		public bool IsTestTarget => Inner.IsTestTarget;
 
 		/// <summary>
 		/// Whether this is a test target explicitly defined.
 		/// Explicitley defined test targets always inherit from TestTargetRules and define their own tests.
 		/// Implicit test targets are created from existing targets when building with -Mode=Test and they include tests from all dependencies.
 		/// </summary>
-		public bool ExplicitTestsTarget
-		{
-			get { return Inner.ExplicitTestsTarget; }
-		}
+		public bool ExplicitTestsTarget => Inner.ExplicitTestsTarget;
 
 		/// <summary>
 		/// Controls the value of WITH_LOW_LEVEL_TESTS that dictates whether module-specific low level tests are compiled in or not.
 		/// </summary>
-		public bool WithLowLevelTests
-		{
-			get { return Inner.WithLowLevelTests; }
-		}
+		public bool WithLowLevelTests => Inner.WithLowLevelTests;
 
 		/// <summary>
 		/// Checks if current platform is part of a given platform group
diff --git a/Engine/Source/Programs/UnrealBuildTool/Configuration/TestModuleRules.cs b/Engine/Source/Programs/UnrealBuildTool/Configuration/TestModuleRules.cs
--- a/Engine/Source/Programs/UnrealBuildTool/Configuration/TestModuleRules.cs
+++ b/Engine/Source/Programs/UnrealBuildTool/Configuration/TestModuleRules.cs
@@ -448,8 +448,8 @@ namespace UnrealBuildTool
 			/// </summary>
 			public bool UsesCatch2
 			{
-				get { return UsesCatch2Private; }
-				set { UsesCatch2Private = value; }
+				get => UsesCatch2Private;
+				set => UsesCatch2Private = value;
 			}
 
 			private Dictionary<UnrealTargetPlatform, string> PlatformTagsPrivate = new Dictionary<UnrealTargetPlatform, string>();
@@ -458,8 +458,8 @@ namespace UnrealBuildTool
 			/// </summary>
 			public Dictionary<UnrealTargetPlatform, string> PlatformTags
 			{
-				get { return PlatformTagsPrivate; }
-				set { PlatformTagsPrivate = value; }
+				get => PlatformTagsPrivate;
+				set => PlatformTagsPrivate = value;
 			}
 
 			private Dictionary<UnrealTargetPlatform, string> PlatformCompilationExtraArgsPrivate = new Dictionary<UnrealTargetPlatform, string>();
@@ -468,8 +468,8 @@ namespace UnrealBuildTool
 			/// </summary>
 			public Dictionary<UnrealTargetPlatform, string> PlatformCompilationExtraArgs
 			{
-				get { return PlatformCompilationExtraArgsPrivate; }
-				set { PlatformCompilationExtraArgsPrivate = value; }
+				get => PlatformCompilationExtraArgsPrivate;
+				set => PlatformCompilationExtraArgsPrivate = value;
 			}
 
 			private List<UnrealTargetPlatform> PlatformsRunUnsupportedPrivate = new List<UnrealTargetPlatform>() { UnrealTargetPlatform.Android, UnrealTargetPlatform.IOS };
@@ -478,8 +478,8 @@ namespace UnrealBuildTool
 			/// </summary>
 			public List<UnrealTargetPlatform> PlatformsRunUnsupported
 			{
-				get { return PlatformsRunUnsupportedPrivate;}
-				set { PlatformsRunUnsupportedPrivate = value; }
+				get => PlatformsRunUnsupportedPrivate;
+				set => PlatformsRunUnsupportedPrivate = value;
 			}
 		}
 #pragma warning restore 8618
diff --git a/Engine/Source/Programs/UnrealBuildTool/Configuration/TestTargetRules.cs b/Engine/Source/Programs/UnrealBuildTool/Configuration/TestTargetRules.cs
--- a/Engine/Source/Programs/UnrealBuildTool/Configuration/TestTargetRules.cs
+++ b/Engine/Source/Programs/UnrealBuildTool/Configuration/TestTargetRules.cs
@@ -50,8 +50,8 @@ namespace UnrealBuildTool
 		/// </summary>
 		public override bool bCompileAgainstApplicationCore
 		{
-			get { return bTestsRequireApplicationCore; }
-			set { bTestsRequireApplicationCore = value; }
+			get => bTestsRequireApplicationCore;
+			set => bTestsRequireApplicationCore = value;
 		}
 
 		/// <summary>
@@ -65,8 +65,8 @@ namespace UnrealBuildTool
 		/// </summary>
 		public override bool bCompileAgainstCoreUObject
 		{
-			get { return bTestsRequireCoreUObject; }
-			set { bTestsRequireCoreUObject = value; }
+			get => bTestsRequireCoreUObject;
+			set => bTestsRequireCoreUObject = value;
 		}
 
 		/// <summary>
@@ -80,8 +80,8 @@ namespace UnrealBuildTool
 		/// </summary>
 		public override bool bCompileAgainstEngine
 		{
-			get { return bTestsRequireEngine; }
-			set { bTestsRequireEngine = value; }
+			get => bTestsRequireEngine;
+			set => bTestsRequireEngine = value;
 		}
 
 		/// <summary>
@@ -95,8 +95,8 @@ namespace UnrealBuildTool
 		/// </summary>
 		public override bool bCompileAgainstEditor
 		{
-			get { return bTestsRequireEditor; }
-			set { bTestsRequireEditor = value; }
+			get => bTestsRequireEditor;
+			set => bTestsRequireEditor = value;
 		}
 
 		/// <summary>
@@ -108,8 +108,8 @@ namespace UnrealBuildTool
 		/// Whether to stub the platform file.
 		/// </summary>
 		public bool bUsePlatformFileStub
-		{ 
-			get { return bUsePlatformFileStubPrivate; }
+		{
+			get => bUsePlatformFileStubPrivate;
 			set
 			{
 				bUsePlatformFileStubPrivate = value;
@@ -125,7 +125,7 @@ namespace UnrealBuildTool
 		/// </summary>
 		public bool bMockEngineDefaults
 		{
-			get { return bMockEngineDefaultsPrivate; }
+			get => bMockEngineDefaultsPrivate;
 			set
 			{
 				bMockEngineDefaultsPrivate = value;
diff --git a/Engine/Source/Programs/UnrealBuildTool/Configuration/UEBuildFramework.cs b/Engine/Source/Programs/UnrealBuildTool/Configuration/UEBuildFramework.cs
--- a/Engine/Source/Programs/UnrealBuildTool/Configuration/UEBuildFramework.cs
+++ b/Engine/Source/Programs/UnrealBuildTool/Configuration/UEBuildFramework.cs
@@ -55,10 +55,7 @@ namespace UnrealBuildTool
 		/// <summary>
 		/// File created after the framework has been extracted. Used to add dependencies into the action graph, used by Modern Xcode as well
 		/// </summary>
-		public FileItem? ExtractedTokenFile
-		{
-			get { return ZipOutputDirectory == null ? null : FileItem.GetItemByFileReference(new FileReference(ZipOutputDirectory!.FullName + ".extracted")); }
-		}
+		public FileItem? ExtractedTokenFile => ZipOutputDirectory == null ? null : FileItem.GetItemByFileReference(new FileReference(ZipOutputDirectory!.FullName + ".extracted"));
 		/// <summary>
 		/// For legacy xcode projects, we unzip in UBT (isntead of xcode), so we track if we've made an action for this framework yet (if two
 		/// modules use the same framework, we only want to unzip it once. Other than time waste, there could be conflicts
diff --git a/Engine/Source/Programs/UnrealBuildTool/Configuration/UEBuildModule.cs b/Engine/Source/Programs/UnrealBuildTool/Configuration/UEBuildModule.cs
--- a/Engine/Source/Programs/UnrealBuildTool/Configuration/UEBuildModule.cs
+++ b/Engine/Source/Programs/UnrealBuildTool/Configuration/UEBuildModule.cs
@@ -38,18 +38,12 @@ namespace UnrealBuildTool
 		/// <summary>
 		/// The name that uniquely identifies the module.
 		/// </summary>
-		public string Name
-		{
-			get { return Rules.Name; }
-		}
+		public string Name => Rules.Name;
 
 		/// <summary>
 		/// Path to the module directory
 		/// </summary>
-		public DirectoryReference ModuleDirectory
-		{
-			get { return Rules.Directory; }
-		}
+		public DirectoryReference ModuleDirectory => Rules.Directory;
 
 		/// <summary>
 		/// Paths to all potential module source directories (with platform extension directories added in)
@@ -59,10 +53,7 @@ namespace UnrealBuildTool
 		/// <summary>
 		/// The name of the .Build.cs file this module was created from, if any
 		/// </summary>
-		public FileReference RulesFile
-		{
-			get { return Rules.File; }
-		}
+		public FileReference RulesFile => Rules.File;
 
 		/// <summary>
 		/// The binary the module will be linked into for the current target.  Only set after UEBuildBinary.BindModules is called.
@@ -192,18 +183,12 @@ namespace UnrealBuildTool
 		/// <summary>
 		/// The Verse source code directory associated with this module if any
 		/// </summary>
-		public virtual DirectoryReference? VerseDirectory
-		{
-			get { return null; }
-		}
+		public virtual DirectoryReference? VerseDirectory => null;
 
 		/// <summary>
 		/// If this module has Verse code associated with it (convenience function)
 		/// </summary>
-		public bool bHasVerse
-		{
-			get { return VerseDirectory != null; }
-		}
+		public bool bHasVerse => VerseDirectory != null;
 
 		/// <summary>
 		/// If this module or any of its dependencies has Verse code associated with it
diff --git a/Engine/Source/Programs/UnrealBuildTool/Configuration/UEBuildModuleCPP.cs b/Engine/Source/Programs/UnrealBuildTool/Configuration/UEBuildModuleCPP.cs
--- a/Engine/Source/Programs/UnrealBuildTool/Configuration/UEBuildModuleCPP.cs
+++ b/Engine/Source/Programs/UnrealBuildTool/Configuration/UEBuildModuleCPP.cs
@@ -62,18 +62,12 @@ namespace UnrealBuildTool
 		/// <summary>
 		/// The directory for this module's generated UHT code
 		/// </summary>
-		public DirectoryReference? GeneratedCodeDirectoryUHT
-		{
-			get { return GeneratedCodeDirectory != null ? DirectoryReference.Combine(GeneratedCodeDirectory!, "UHT") : null; }
-		}
+		public DirectoryReference? GeneratedCodeDirectoryUHT => GeneratedCodeDirectory != null ? DirectoryReference.Combine(GeneratedCodeDirectory!, "UHT") : null;
 
 		/// <summary>
 		/// The directory for this module's generated VNI code
 		/// </summary>
-		public DirectoryReference? GeneratedCodeDirectoryVNI
-		{
-			get { return GeneratedCodeDirectory != null ? DirectoryReference.Combine(GeneratedCodeDirectory!, "VNI") : null; }
-		}
+		public DirectoryReference? GeneratedCodeDirectoryVNI => GeneratedCodeDirectory != null ? DirectoryReference.Combine(GeneratedCodeDirectory!, "VNI") : null;
 
 		/// <summary>
 		/// Global override to force all include paths to be always added
@@ -103,10 +97,7 @@ namespace UnrealBuildTool
 		/// <summary>
 		/// The Verse source code directory associated with this module if any
 		/// </summary>
-		public override DirectoryReference? VerseDirectory
-		{
-			get { return AssociatedVerseDirectory; }
-		}
+		public override DirectoryReference? VerseDirectory => AssociatedVerseDirectory;
 
 		protected override void GetReferencedDirectories(HashSet<DirectoryReference> Directories)
 		{
@@ -315,10 +306,7 @@ namespace UnrealBuildTool
 		/// <summary>
 		/// Path to the precompiled manifest location
 		/// </summary>
-		public virtual FileReference PrecompiledManifestLocation
-		{
-			get { return FileReference.Combine(IntermediateDirectoryNoArch, String.Format("{0}.precompiled", Name)); }
-		}
+		public virtual FileReference PrecompiledManifestLocation => FileReference.Combine(IntermediateDirectoryNoArch, String.Format("{0}.precompiled", Name));
 
 		/// <summary>
 		/// Sets up the environment for compiling any module that includes the public interface of this module.
diff --git a/Engine/Source/Programs/UnrealBuildTool/Configuration/UEBuildPlugin.cs b/Engine/Source/Programs/UnrealBuildTool/Configuration/UEBuildPlugin.cs
--- a/Engine/Source/Programs/UnrealBuildTool/Configuration/UEBuildPlugin.cs
+++ b/Engine/Source/Programs/UnrealBuildTool/Configuration/UEBuildPlugin.cs
@@ -58,50 +58,32 @@ namespace UnrealBuildTool
 		/// <summary>
 		/// Accessor for the name of this plugin
 		/// </summary>
-		public string Name
-		{
-			get { return Info.Name; }
-		}
+		public string Name => Info.Name;
 
 		/// <summary>
 		/// Accessor for the file for this plugin
 		/// </summary>
-		public FileReference File
-		{
-			get { return Info.File; }
-		}
+		public FileReference File => Info.File;
 
 		/// <summary>
 		/// Accessor for the child files for this plugin
 		/// </summary>
-		public List<FileReference> ChildFiles
-		{
-			get { return Info.ChildFiles; }
-		}
+		public List<FileReference> ChildFiles => Info.ChildFiles;
 
 		/// <summary>
 		/// Accessor for the type of the plugin
 		/// </summary>
-		public PluginType Type
-		{
-			get { return Info.Type; }
-		}
+		public PluginType Type => Info.Type;
 
 		/// <summary>
 		/// Accessor for this plugin's root directory
 		/// </summary>
-		public DirectoryReference Directory
-		{
-			get { return Info.Directory; }
-		}
+		public DirectoryReference Directory => Info.Directory;
 
 		/// <summary>
 		/// Accessor for this plugin's descriptor
 		/// </summary>
-		public PluginDescriptor Descriptor
-		{
-			get { return Info.Descriptor; }
-		}
+		public PluginDescriptor Descriptor => Info.Descriptor;
 
 		/// <summary>
 		/// Returns the name of this plugin for debugging
diff --git a/Engine/Source/Programs/UnrealBuildTool/Configuration/UEBuildTarget.cs b/Engine/Source/Programs/UnrealBuildTool/Configuration/UEBuildTarget.cs
--- a/Engine/Source/Programs/UnrealBuildTool/Configuration/UEBuildTarget.cs
+++ b/Engine/Source/Programs/UnrealBuildTool/Configuration/UEBuildTarget.cs
@@ -1001,10 +1001,7 @@ namespace UnrealBuildTool
 		/// <summary>
 		/// True if there is is more than one architecture specified
 		/// </summary>
-		public bool bIsMultiArch
-		{
-			get { return Architectures.Count > 1; }
-		}
+		public bool bIsMultiArch => Architectures.Count > 1;
 
 		/// <summary>
 		/// Gets the Architectures list as a + delimited string
diff --git a/Engine/Source/Programs/UnrealBuildTool/Executors/Experimental/FASTBuild.cs b/Engine/Source/Programs/UnrealBuildTool/Executors/Experimental/FASTBuild.cs
--- a/Engine/Source/Programs/UnrealBuildTool/Executors/Experimental/FASTBuild.cs
+++ b/Engine/Source/Programs/UnrealBuildTool/Executors/Experimental/FASTBuild.cs
@@ -172,10 +172,7 @@ namespace UnrealBuildTool
 			this.LocalExecutor = new ParallelExecutor(MaxLocalActions, bAllCores, bCompactOutput, Logger);
 		}
 
-		public override string Name
-		{
-			get { return "FASTBuild"; }
-		}
+		public override string Name => "FASTBuild";
 
 		public static string GetExecutableName()
 		{
diff --git a/Engine/Source/Programs/UnrealBuildTool/Executors/ParallelExecutor.cs b/Engine/Source/Programs/UnrealBuildTool/Executors/ParallelExecutor.cs
--- a/Engine/Source/Programs/UnrealBuildTool/Executors/ParallelExecutor.cs
+++ b/Engine/Source/Programs/UnrealBuildTool/Executors/ParallelExecutor.cs
@@ -138,10 +138,7 @@ namespace UnrealBuildTool
 		/// <summary>
 		/// Returns the name of this executor
 		/// </summary>
-		public override string Name
-		{
-			get { return "Parallel"; }
-		}
+		public override string Name => "Parallel";
 
 		/// <summary>
 		/// Checks whether the task executor can be used
diff --git a/Engine/Source/Programs/UnrealBuildTool/Executors/XGE.cs b/Engine/Source/Programs/UnrealBuildTool/Executors/XGE.cs
--- a/Engine/Source/Programs/UnrealBuildTool/Executors/XGE.cs
+++ b/Engine/Source/Programs/UnrealBuildTool/Executors/XGE.cs
@@ -101,10 +101,7 @@ namespace UnrealBuildTool
 			XmlConfig.ApplyTo(this);
 		}
 
-		public override string Name
-		{
-			get { return "XGE"; }
-		}
+		public override string Name => "XGE";
 
 		public static bool TryGetXgConsoleExecutable([NotNullWhen(true)] out string? OutXgConsoleExe)
 		{
diff --git a/Engine/Source/Programs/UnrealBuildTool/Modes/IWYUMode.cs b/Engine/Source/Programs/UnrealBuildTool/Modes/IWYUMode.cs
--- a/Engine/Source/Programs/UnrealBuildTool/Modes/IWYUMode.cs
+++ b/Engine/Source/Programs/UnrealBuildTool/Modes/IWYUMode.cs
@@ -291,10 +291,7 @@ namespace UnrealBuildTool
 				}
 				return EngineIncludeOrderHelper.GetLatestDeprecationDefine();
 			}
-			set
-			{
-				HeaderDeprecationTagOverride = value;
-			}
+			set => HeaderDeprecationTagOverride = value;
 		}
 
 		/// <summary>
diff --git a/Engine/Source/Programs/UnrealBuildTool/Modes/ServerMode.cs b/Engine/Source/Programs/UnrealBuildTool/Modes/ServerMode.cs
--- a/Engine/Source/Programs/UnrealBuildTool/Modes/ServerMode.cs
+++ b/Engine/Source/Programs/UnrealBuildTool/Modes/ServerMode.cs
@@ -80,7 +80,7 @@ namespace UnrealBuildTool
 	struct DiscoverTargetsResponse : ICommandResponse
 	{
 		[JsonConverter(typeof(JsonStringEnumConverter))]
-		public CommandType Type { get { return CommandType.DiscoverTargets; } }
+		public CommandType Type => CommandType.DiscoverTargets;
 		public long Sequence { get; set; }
 		public List<string> Configurations { get; set; }
 		public List<string> Platforms { get; set; }
@@ -94,13 +94,13 @@ namespace UnrealBuildTool
 	struct SetTargetResponse : ICommandResponse
 	{
 		[JsonConverter(typeof(JsonStringEnumConverter))]
-		public CommandType Type { get { return CommandType.SetTarget; } }
+		public CommandType Type => CommandType.SetTarget;
 		public long Sequence { get; set; }
 	}
 	struct QueryFileResponse : ICommandResponse
 	{
 		[JsonConverter(typeof(JsonStringEnumConverter))]
-		public CommandType Type { get { return CommandType.QueryFile; } }
+		public CommandType Type => CommandType.QueryFile;
 		public long Sequence { get; set; }
 		public bool Found { get; set; }
 	}
@@ -108,7 +108,7 @@ namespace UnrealBuildTool
 	struct GetBrowseConfigurationResponse : ICommandResponse
 	{
 		[JsonConverter(typeof(JsonStringEnumConverter))]
-		public CommandType Type { get { return CommandType.GetBrowseConfiguration; } }
+		public CommandType Type => CommandType.GetBrowseConfiguration;
 		public long Sequence { get; set; }
 		public bool Success { get; set; }
 
@@ -120,7 +120,7 @@ namespace UnrealBuildTool
 	class GetCompileSettingsResponse : ICommandResponse
 	{
 		[JsonConverter(typeof(JsonStringEnumConverter))]
-		public CommandType Type { get { return CommandType.GetCompileSettings; } }
+		public CommandType Type => CommandType.GetCompileSettings;
 		public long Sequence { get; set; }
 		public List<TargetIntellisenseInfo.CompileSettings?> Settings { get; set; } = new();
 	}
diff --git a/Engine/Source/Programs/UnrealBuildTool/Platform/Android/UEBuildAndroid.cs b/Engine/Source/Programs/UnrealBuildTool/Platform/Android/UEBuildAndroid.cs
--- a/Engine/Source/Programs/UnrealBuildTool/Platform/Android/UEBuildAndroid.cs
+++ b/Engine/Source/Programs/UnrealBuildTool/Platform/Android/UEBuildAndroid.cs
@@ -78,40 +78,19 @@ namespace UnrealBuildTool
 		#region Read-only accessor properties 
 		#pragma warning disable CS1591
 
-		public bool bEnableAddressSanitizer
-		{
-			get { return Inner.bEnableAddressSanitizer; }
-		}
+		public bool bEnableAddressSanitizer => Inner.bEnableAddressSanitizer;
 
-		public bool bEnableHWAddressSanitizer
-		{
-			get { return Inner.bEnableHWAddressSanitizer; }
-		}
+		public bool bEnableHWAddressSanitizer => Inner.bEnableHWAddressSanitizer;
 
-		public bool bEnableThreadSanitizer
-		{
-			get { return Inner.bEnableThreadSanitizer; }
-		}
+		public bool bEnableThreadSanitizer => Inner.bEnableThreadSanitizer;
 
-		public bool bEnableUndefinedBehaviorSanitizer
-		{
-			get { return Inner.bEnableUndefinedBehaviorSanitizer; }
-		}
+		public bool bEnableUndefinedBehaviorSanitizer => Inner.bEnableUndefinedBehaviorSanitizer;
 
-		public bool bEnableMinimalUndefinedBehaviorSanitizer
-		{
-			get { return Inner.bEnableMinimalUndefinedBehaviorSanitizer; }
-		}
+		public bool bEnableMinimalUndefinedBehaviorSanitizer => Inner.bEnableMinimalUndefinedBehaviorSanitizer;
 
-		public bool bEnableRayTracing
-		{
-			get { return Inner.bEnableRayTracing; }
-		}
+		public bool bEnableRayTracing => Inner.bEnableRayTracing;
 
-		public AndroidTargetRules TargetRules
-		{
-			get { return Inner; }
-		}
+		public AndroidTargetRules TargetRules => Inner;
 
 		#pragma warning restore CS1591
 		#endregion
@@ -626,10 +605,7 @@ namespace UnrealBuildTool
 
 	class AndroidPlatformFactory : UEBuildPlatformFactory
 	{
-		public override UnrealTargetPlatform TargetPlatform
-		{
-			get { return UnrealTargetPlatform.Android; }
-		}
+		public override UnrealTargetPlatform TargetPlatform => UnrealTargetPlatform.Android;
 
 		public override void RegisterBuildPlatforms(ILogger Logger)
 		{
diff --git a/Engine/Source/Programs/UnrealBuildTool/Platform/IOS/IOSCodeSigning.cs b/Engine/Source/Programs/UnrealBuildTool/Platform/IOS/IOSCodeSigning.cs
--- a/Engine/Source/Programs/UnrealBuildTool/Platform/IOS/IOSCodeSigning.cs
+++ b/Engine/Source/Programs/UnrealBuildTool/Platform/IOS/IOSCodeSigning.cs
@@ -2141,13 +2141,7 @@ namespace UnrealBuildTool
 		/// <summary>
 		/// The shared (Engine) build directory (on PC)
 		/// </summary>
-		public static string EngineBuildDirectory
-		{
-			get
-			{
-				return Path.Combine(Unreal.EngineDirectory.FullName, "Build", CodeSigningConfig.OSString);
-			}
-		}
+		public static string EngineBuildDirectory => Path.Combine(Unreal.EngineDirectory.FullName, "Build", CodeSigningConfig.OSString);
 
 		/// <summary>
 		/// The local build intermediate directory (on PC)
@@ -2167,43 +2161,28 @@ namespace UnrealBuildTool
 		/// <summary>
 		/// The local directory cooked files are placed (on PC)
 		/// </summary>
-		public static string CookedDirectory
-		{
-			get { return Path.GetFullPath(GameDirectory + @"\Saved\Cooked\" + CodeSigningConfig.OSString); }
-		}
+		public static string CookedDirectory => Path.GetFullPath(GameDirectory + @"\Saved\Cooked\" + CodeSigningConfig.OSString);
 
 		/// <summary>
 		/// The local directory config files are placed (on PC)
 		/// </summary>
-		public static string ConfigDirectory
-		{
-			get { return Path.GetFullPath(GameDirectory + @"\Saved\Config"); }
-		}
+		public static string ConfigDirectory => Path.GetFullPath(GameDirectory + @"\Saved\Config");
 
 		/// <summary>
 		/// The engine config files are placed (on PC)
 		/// </summary>
-		public static string DefaultConfigDirectory
-		{
-			get { return Path.Combine(Unreal.EngineDirectory.FullName, "Engine/Config"); }
-		}
+		public static string DefaultConfigDirectory => Path.Combine(Unreal.EngineDirectory.FullName, "Engine/Config");
 
 		/// <summary>
 		/// The engine directory
 		/// </summary>
-		public static string EngineDirectory
-		{
-			get { return Unreal.EngineDirectory.FullName; }
-		}
+		public static string EngineDirectory => Unreal.EngineDirectory.FullName;
 
 		/// <summary>
 		/// The local directory that a payload (GameName.app) is assembled into before being copied to the Mac (on PC)
 		/// </summary>
-		public static string PayloadRootDirectory
-		{
-			get { return Path.GetFullPath(PCStagingRootDir + @"\Payload"); }
-		}
-		
+		public static string PayloadRootDirectory => Path.GetFullPath(PCStagingRootDir + @"\Payload");
+
 		/// <summary>
 		/// Program_Architecture
 		/// </summary>
diff --git a/Engine/Source/Programs/UnrealBuildTool/Platform/IOS/IOSPlatformSDK.Versions.cs b/Engine/Source/Programs/UnrealBuildTool/Platform/IOS/IOSPlatformSDK.Versions.cs
--- a/Engine/Source/Programs/UnrealBuildTool/Platform/IOS/IOSPlatformSDK.Versions.cs
+++ b/Engine/Source/Programs/UnrealBuildTool/Platform/IOS/IOSPlatformSDK.Versions.cs
@@ -22,15 +22,10 @@ namespace UnrealBuildTool
 		/// <summary>
 		/// The minimum Clang version that a static library can be built with
 		/// </summary>
-		public Version MinimumStaticLibClangVersion
-		{
-			get
-			{
+		public Version MinimumStaticLibClangVersion =>
 				// return new Version("1200.00.32.29");	// UE 5.0 (Xcode 12.4)
 				// return new Version("1300.00.29.30"); // (old) UE 5.1 sets Xcode13.2.1 as minimum supported
 				// return new Version("1316.0.21.2"); // UE 5.1 sets Xcode13.4 as minimum supported
-				return new Version("1400.0.29.202"); // UE 5.2 sets Xcode14.1 as minimum supported
-			}
-		}
+				new Version("1400.0.29.202"); // UE 5.2 sets Xcode14.1 as minimum supported
 	}
 }
diff --git a/Engine/Source/Programs/UnrealBuildTool/Platform/IOS/UEBuildIOS.cs b/Engine/Source/Programs/UnrealBuildTool/Platform/IOS/UEBuildIOS.cs
--- a/Engine/Source/Programs/UnrealBuildTool/Platform/IOS/UEBuildIOS.cs
+++ b/Engine/Source/Programs/UnrealBuildTool/Platform/IOS/UEBuildIOS.cs
@@ -172,75 +172,33 @@ namespace UnrealBuildTool
 		/// </summary>
 		#region Read-only accessor properties 
 #pragma warning disable CS1591
-		public bool bStripSymbols
-		{
-			get { return Inner.bStripSymbols; }
-		}
-			
-		public bool bGenerateFrameworkWrapperProject
-		{
-			get { return Inner.bGenerateFrameworkWrapperProject; }
-		}
+		public bool bStripSymbols => Inner.bStripSymbols;
 
-		public bool bGeneratedSYM
-		{
-			get { return Inner.bGeneratedSYM; }
-		}
+		public bool bGenerateFrameworkWrapperProject => Inner.bGenerateFrameworkWrapperProject;
 
-		public bool bCreateStubIPA
-		{
-			get { return Inner.bCreateStubIPA; }
-		}
+		public bool bGeneratedSYM => Inner.bGeneratedSYM;
 
-		public bool bSkipCrashlytics
-		{
-			get { return Inner.bSkipCrashlytics; }
-		}
+		public bool bCreateStubIPA => Inner.bCreateStubIPA;
 
-		public bool bSkipClangValidation
-		{
-			get { return Inner.bSkipClangValidation; }
-		}
+		public bool bSkipCrashlytics => Inner.bSkipCrashlytics;
 
-		public bool bForDistribution
-		{
-			get { return Inner.bForDistribution; }
-		}
+		public bool bSkipClangValidation => Inner.bSkipClangValidation;
 
-		public string? ImportProvision
-		{
-			get { return Inner.ImportProvision; }
-		}
+		public bool bForDistribution => Inner.bForDistribution;
 
-		public string? ImportCertificate
-		{
-			get { return Inner.ImportCertificate; }
-		}
+		public string? ImportProvision => Inner.ImportProvision;
 
-		public string? ImportCertificatePassword
-		{
-			get { return Inner.ImportCertificatePassword; }
-		}
+		public string? ImportCertificate => Inner.ImportCertificate;
 
-		public float RuntimeVersion
-		{
-			get { return float.Parse(Inner.ProjectSettings!.RuntimeVersion, System.Globalization.CultureInfo.InvariantCulture); }
-		}
+		public string? ImportCertificatePassword => Inner.ImportCertificatePassword;
 
-		public bool bEnableAddressSanitizer
-		{
-			get { return Inner.bEnableAddressSanitizer; }
-		}
+		public float RuntimeVersion => float.Parse(Inner.ProjectSettings!.RuntimeVersion, System.Globalization.CultureInfo.InvariantCulture);
 
-		public bool bEnableThreadSanitizer
-		{
-			get { return Inner.bEnableThreadSanitizer; }
-		}
+		public bool bEnableAddressSanitizer => Inner.bEnableAddressSanitizer;
 
-		public bool bEnableUndefinedBehaviorSanitizer
-		{
-			get { return Inner.bEnableUndefinedBehaviorSanitizer; }
-		}
+		public bool bEnableThreadSanitizer => Inner.bEnableThreadSanitizer;
+
+		public bool bEnableUndefinedBehaviorSanitizer => Inner.bEnableUndefinedBehaviorSanitizer;
 
 #pragma warning restore CS1591
 		#endregion
@@ -480,10 +438,7 @@ namespace UnrealBuildTool
 		public string? BundleIdentifier;
 		public bool bHaveCertificate = false;
 
-		public string? MobileProvision
-		{
-			get => MobileProvisionFile?.GetFileName();
-		}
+		public string? MobileProvision => MobileProvisionFile?.GetFileName();
 
 		public IOSProvisioningData(IOSProjectSettings ProjectSettings, bool bForDistribution, ILogger Logger)
 			: this(ProjectSettings, false, bForDistribution, Logger)
@@ -1289,10 +1244,7 @@ namespace UnrealBuildTool
 
 	class IOSPlatformFactory : UEBuildPlatformFactory
 	{
-		public override UnrealTargetPlatform TargetPlatform
-		{
-			get { return UnrealTargetPlatform.IOS; }
-		}
+		public override UnrealTargetPlatform TargetPlatform => UnrealTargetPlatform.IOS;
 
 		/// <summary>
 		/// Register the platform with the UEBuildPlatform class
diff --git a/Engine/Source/Programs/UnrealBuildTool/Platform/IOS/UEDeployIOS.cs b/Engine/Source/Programs/UnrealBuildTool/Platform/IOS/UEDeployIOS.cs
--- a/Engine/Source/Programs/UnrealBuildTool/Platform/IOS/UEDeployIOS.cs
+++ b/Engine/Source/Programs/UnrealBuildTool/Platform/IOS/UEDeployIOS.cs
@@ -27,8 +27,8 @@ namespace UnrealBuildTool
 
 		public bool ForDistribution
 		{
-			get { return bForDistribution; }
-			set { bForDistribution = value; }
+			get => bForDistribution;
+			set => bForDistribution = value;
 		}
 		bool bForDistribution = false;
 
@@ -45,12 +45,9 @@ namespace UnrealBuildTool
 				set;
 			}
 
-			public static bool bCustomLaunchscreenStoryboard = false;			
+			public static bool bCustomLaunchscreenStoryboard = false;
 
-			static string RunningVersionFilename
-			{
-				get { return Path.Combine(BuildDirectory!, GameName + ".PackageVersionCounter"); }
-			}
+			static string RunningVersionFilename => Path.Combine(BuildDirectory!, GameName + ".PackageVersionCounter");
 
 			/// <summary>
 			/// Reads the GameName.PackageVersionCounter from disk and bumps the minor version number in it
diff --git a/Engine/Source/Programs/UnrealBuildTool/Platform/Linux/UEBuildLinux.cs b/Engine/Source/Programs/UnrealBuildTool/Platform/Linux/UEBuildLinux.cs
--- a/Engine/Source/Programs/UnrealBuildTool/Platform/Linux/UEBuildLinux.cs
+++ b/Engine/Source/Programs/UnrealBuildTool/Platform/Linux/UEBuildLinux.cs
@@ -131,50 +131,23 @@ namespace UnrealBuildTool
 		#region Read-only accessor properties 
 		#pragma warning disable CS1591
 
-		public bool bPreservePSYM
-		{
-			get { return Inner.bPreservePSYM; }
-		}
+		public bool bPreservePSYM => Inner.bPreservePSYM;
 
-		public bool bEnableAddressSanitizer
-		{
-			get { return Inner.bEnableAddressSanitizer; }
-		}
+		public bool bEnableAddressSanitizer => Inner.bEnableAddressSanitizer;
 
-		public bool bEnableLibFuzzer
-		{
-			get { return Inner.bEnableLibFuzzer; }
-		}
+		public bool bEnableLibFuzzer => Inner.bEnableLibFuzzer;
 
-		public bool bEnableThreadSanitizer
-		{
-			get { return Inner.bEnableThreadSanitizer; }
-		}
+		public bool bEnableThreadSanitizer => Inner.bEnableThreadSanitizer;
 
-		public bool bEnableUndefinedBehaviorSanitizer
-		{
-			get { return Inner.bEnableUndefinedBehaviorSanitizer; }
-		}
+		public bool bEnableUndefinedBehaviorSanitizer => Inner.bEnableUndefinedBehaviorSanitizer;
 
-		public bool bEnableMemorySanitizer
-		{
-			get { return Inner.bEnableMemorySanitizer; }
-		}
+		public bool bEnableMemorySanitizer => Inner.bEnableMemorySanitizer;
 
-		public bool bTuneDebugInfoForLLDB
-		{
-			get { return Inner.bTuneDebugInfoForLLDB; }
-		}
+		public bool bTuneDebugInfoForLLDB => Inner.bTuneDebugInfoForLLDB;
 
-		public bool bDisableDumpSyms
-		{
-			get { return Inner.bDisableDumpSyms; }
-		}
+		public bool bDisableDumpSyms => Inner.bDisableDumpSyms;
 
-		public bool bEnableRayTracing
-		{
-			get { return Inner.bEnableRayTracing; }
-		}
+		public bool bEnableRayTracing => Inner.bEnableRayTracing;
 
 		#pragma warning restore CS1591
 		#endregion
@@ -681,10 +654,7 @@ namespace UnrealBuildTool
 
 	class LinuxPlatformFactory : UEBuildPlatformFactory
 	{
-		public override UnrealTargetPlatform TargetPlatform
-		{
-			get { return UnrealTargetPlatform.Linux; }
-		}
+		public override UnrealTargetPlatform TargetPlatform => UnrealTargetPlatform.Linux;
 
 		/// <summary>
 		/// Register the platform with the UEBuildPlatform class
diff --git a/Engine/Source/Programs/UnrealBuildTool/Platform/Mac/ApplePlatformSDK.Versions.cs b/Engine/Source/Programs/UnrealBuildTool/Platform/Mac/ApplePlatformSDK.Versions.cs
--- a/Engine/Source/Programs/UnrealBuildTool/Platform/Mac/ApplePlatformSDK.Versions.cs
+++ b/Engine/Source/Programs/UnrealBuildTool/Platform/Mac/ApplePlatformSDK.Versions.cs
@@ -41,13 +41,7 @@ namespace UnrealBuildTool
 		/// <summary>
 		/// The minimum macOS SDK version that a dynamic library can be built with
 		/// </summary>
-		public virtual Version MinimumDynamicLibSDKVersion
-		{
-			get
-			{
-				return new Version("12.1");		// SDK used in Xcode13.2.1
-			}
-		}
+		public virtual Version MinimumDynamicLibSDKVersion => new Version("12.1");      // SDK used in Xcode13.2.1
 
 	}
 }
diff --git a/Engine/Source/Programs/UnrealBuildTool/Platform/Mac/MacExports.cs b/Engine/Source/Programs/UnrealBuildTool/Platform/Mac/MacExports.cs
--- a/Engine/Source/Programs/UnrealBuildTool/Platform/Mac/MacExports.cs
+++ b/Engine/Source/Programs/UnrealBuildTool/Platform/Mac/MacExports.cs
@@ -19,13 +19,7 @@ namespace UnrealBuildTool
 		/// Describes the architecture of the host. Note - this ignores translation.
 		/// IsRunningUnderRosetta can be used to detect that we're running under translation
 		/// </summary>
-		public static UnrealArch HostArchitecture
-		{
-			get
-			{
-				return IsRunningOnAppleArchitecture ? UnrealArch.Arm64 : UnrealArch.X64;
-			}
-		}
+		public static UnrealArch HostArchitecture => IsRunningOnAppleArchitecture ? UnrealArch.Arm64 : UnrealArch.X64;
 
 		/// <summary>
 		/// Cached result for AppleArch check
diff --git a/Engine/Source/Programs/UnrealBuildTool/Platform/Mac/MacToolChain.cs b/Engine/Source/Programs/UnrealBuildTool/Platform/Mac/MacToolChain.cs
--- a/Engine/Source/Programs/UnrealBuildTool/Platform/Mac/MacToolChain.cs
+++ b/Engine/Source/Programs/UnrealBuildTool/Platform/Mac/MacToolChain.cs
@@ -81,15 +81,9 @@ namespace UnrealBuildTool
 
 		public static Lazy<MacToolChainSettings> SettingsPrivate = new Lazy<MacToolChainSettings>(() => new MacToolChainSettings(false, Log.Logger));
 
-		public static MacToolChainSettings Settings
-		{
-			get { return SettingsPrivate.Value; }
-		}
+		public static MacToolChainSettings Settings => SettingsPrivate.Value;
 
-		public static string SDKPath
-		{
-			get { return Settings.BaseSDKDir + "/MacOSX.sdk"; }
-		}
+		public static string SDKPath => Settings.BaseSDKDir + "/MacOSX.sdk";
 
 		/// <summary>
 		/// Which compiler\linker frontend to use
diff --git a/Engine/Source/Programs/UnrealBuildTool/Platform/Mac/UEBuildMac.cs b/Engine/Source/Programs/UnrealBuildTool/Platform/Mac/UEBuildMac.cs
--- a/Engine/Source/Programs/UnrealBuildTool/Platform/Mac/UEBuildMac.cs
+++ b/Engine/Source/Programs/UnrealBuildTool/Platform/Mac/UEBuildMac.cs
@@ -93,35 +93,17 @@ namespace UnrealBuildTool
 		#region Read-only accessor properties 
 		#pragma warning disable CS1591
 
-		public bool bEnableAddressSanitizer
-		{
-			get { return Inner.bEnableAddressSanitizer; }
-		}
+		public bool bEnableAddressSanitizer => Inner.bEnableAddressSanitizer;
 
-		public bool bEnableLibFuzzer
-		{
-			get { return Inner.bEnableLibFuzzer; }
-		}
+		public bool bEnableLibFuzzer => Inner.bEnableLibFuzzer;
 
-		public bool bEnableThreadSanitizer
-		{
-			get { return Inner.bEnableThreadSanitizer; }
-		}
+		public bool bEnableThreadSanitizer => Inner.bEnableThreadSanitizer;
 
-		public bool bEnableUndefinedBehaviorSanitizer
-		{
-			get { return Inner.bEnableUndefinedBehaviorSanitizer; }
-		}
+		public bool bEnableUndefinedBehaviorSanitizer => Inner.bEnableUndefinedBehaviorSanitizer;
 
-		public bool bSkipClangValidation
-		{
-			get { return Inner.bSkipClangValidation; }
-		}
+		public bool bSkipClangValidation => Inner.bSkipClangValidation;
 
-		public bool bEnableRayTracing
-		{
-			get { return Inner.bEnableRayTracing; }
-		}
+		public bool bEnableRayTracing => Inner.bEnableRayTracing;
 
 #pragma warning restore CS1591
 		#endregion
@@ -608,10 +590,7 @@ namespace UnrealBuildTool
 
 	class MacPlatformFactory : UEBuildPlatformFactory
 	{
-		public override UnrealTargetPlatform TargetPlatform
-		{
-			get { return UnrealTargetPlatform.Mac; }
-		}
+		public override UnrealTargetPlatform TargetPlatform => UnrealTargetPlatform.Mac;
 
 		/// <summary>
 		/// Register the platform with the UEBuildPlatform class
diff --git a/Engine/Source/Programs/UnrealBuildTool/Platform/TVOS/UEBuildTVOS.cs b/Engine/Source/Programs/UnrealBuildTool/Platform/TVOS/UEBuildTVOS.cs
--- a/Engine/Source/Programs/UnrealBuildTool/Platform/TVOS/UEBuildTVOS.cs
+++ b/Engine/Source/Programs/UnrealBuildTool/Platform/TVOS/UEBuildTVOS.cs
@@ -12,10 +12,7 @@ namespace UnrealBuildTool
 		/// <summary>
 		/// which devices the game is allowed to run on
 		/// </summary>
-		public override string RuntimeDevices
-		{
-			get { return "3"; }
-		}
+		public override string RuntimeDevices => "3";
 
 		public TVOSProjectSettings(FileReference? ProjectFile, String? Bundle)
 			: base(ProjectFile, UnrealTargetPlatform.TVOS, Bundle)
@@ -127,10 +124,7 @@ namespace UnrealBuildTool
 
 	class TVOSPlatformFactory : UEBuildPlatformFactory
 	{
-		public override UnrealTargetPlatform TargetPlatform
-		{
-			get { return UnrealTargetPlatform.TVOS; }
-		}
+		public override UnrealTargetPlatform TargetPlatform => UnrealTargetPlatform.TVOS;
 
 		/// <summary>
 		/// Register the platform with the UEBuildPlatform class
diff --git a/Engine/Source/Programs/UnrealBuildTool/Platform/Windows/AppXManifestGeneratorBase.cs b/Engine/Source/Programs/UnrealBuildTool/Platform/Windows/AppXManifestGeneratorBase.cs
--- a/Engine/Source/Programs/UnrealBuildTool/Platform/Windows/AppXManifestGeneratorBase.cs
+++ b/Engine/Source/Programs/UnrealBuildTool/Platform/Windows/AppXManifestGeneratorBase.cs
@@ -22,10 +22,10 @@ namespace UnrealBuildTool
 	abstract public class AppXManifestGeneratorBase
 	{
 		/// config section for platform-specific target settings
-		protected virtual string IniSection_PlatformTargetSettings { get { return string.Format( "/Script/{0}PlatformEditor.{0}TargetSettings", Platform.ToString() ); } }
-		
+		protected virtual string IniSection_PlatformTargetSettings => string.Format("/Script/{0}PlatformEditor.{0}TargetSettings", Platform.ToString());
+
 		/// config section for general target settings
-		protected virtual string IniSection_GeneralProjectSettings { get { return "/Script/EngineSettings.GeneralProjectSettings"; } }
+		protected virtual string IniSection_GeneralProjectSettings => "/Script/EngineSettings.GeneralProjectSettings";
 
 		/// default subdirectory for build resources
 		protected const string BuildResourceSubPath = "Resources";
diff --git a/Engine/Source/Programs/UnrealBuildTool/Platform/Windows/PVSToolChain.cs b/Engine/Source/Programs/UnrealBuildTool/Platform/Windows/PVSToolChain.cs
--- a/Engine/Source/Programs/UnrealBuildTool/Platform/Windows/PVSToolChain.cs
+++ b/Engine/Source/Programs/UnrealBuildTool/Platform/Windows/PVSToolChain.cs
@@ -209,7 +209,7 @@ namespace UnrealBuildTool
 		public PVSAnalysisModeFlags ModeFlags
 		{
 			get
- 			{
+			{
 				if (ModePrivate.HasValue)
 				{
 					return ModePrivate.Value;
@@ -223,10 +223,7 @@ namespace UnrealBuildTool
 					return PVSAnalysisModeFlags.GeneralAnalysis;
 				}
 			}
-			set
-			{
-				ModePrivate = value;
-			}
+			set => ModePrivate = value;
 		}
 
 		/// <summary>
@@ -254,10 +251,7 @@ namespace UnrealBuildTool
 					return AnalysisTimeoutFlags.After_30_minutes;
 				}
 			}
-			set
-			{
-				TimeoutPrivate = value;
-			}
+			set => TimeoutPrivate = value;
 		}
 
 		/// <summary>
@@ -285,10 +279,7 @@ namespace UnrealBuildTool
 					return false;
 				}
 			}
-			set
-			{
-				EnableNoNoisePrivate = value;
-			}
+			set => EnableNoNoisePrivate = value;
 		}
 
 		/// <summary>
@@ -316,10 +307,7 @@ namespace UnrealBuildTool
 					return false;
 				}
 			}
-			set
-			{
-				EnableReportDisabledRulesPrivate = value;
-			}
+			set => EnableReportDisabledRulesPrivate = value;
 		}
 
 		/// <summary>
@@ -350,50 +338,32 @@ namespace UnrealBuildTool
 		/// <summary>
 		/// Accessor for the Application settings
 		/// </summary>
-		internal PVSApplicationSettings? ApplicationSettings
-		{
-			get { return Inner.ApplicationSettings.Value; }
-		}
+		internal PVSApplicationSettings? ApplicationSettings => Inner.ApplicationSettings.Value;
 
 		/// <summary>
 		/// Whether to use the application settings for the mode
 		/// </summary>
-		public bool UseApplicationSettings
-		{
-			get { return Inner.UseApplicationSettings; }
-		}
+		public bool UseApplicationSettings => Inner.UseApplicationSettings;
 
 		/// <summary>
 		/// Override for the analysis mode to use
 		/// </summary>
-		public PVSAnalysisModeFlags ModeFlags
-		{
-			get { return Inner.ModeFlags; }
-		}
-    
+		public PVSAnalysisModeFlags ModeFlags => Inner.ModeFlags;
+
 		/// <summary>
 		/// Override for the analysis timeout to use
 		/// </summary>
-		public AnalysisTimeoutFlags AnalysisTimeoutFlag
-		{
-			get { return Inner.AnalysisTimeoutFlag; }
-		}
+		public AnalysisTimeoutFlags AnalysisTimeoutFlag => Inner.AnalysisTimeoutFlag;
 
 		/// <summary>
 		/// Override NoNoise analysis setting to use
 		/// </summary>
-		public bool EnableNoNoise
-		{
-			get { return Inner.EnableNoNoise; }
-		}
+		public bool EnableNoNoise => Inner.EnableNoNoise;
 
 		/// <summary>
 		/// Override EnableReportDisabledRules analysis setting to use
 		/// </summary>
-		public bool EnableReportDisabledRules
-		{
-			get { return Inner.EnableReportDisabledRules; }
-		}
+		public bool EnableReportDisabledRules => Inner.EnableReportDisabledRules;
 	}
 
 	/// <summary>
diff --git a/Engine/Source/Programs/UnrealBuildTool/Platform/Windows/UEBuildWindows.cs b/Engine/Source/Programs/UnrealBuildTool/Platform/Windows/UEBuildWindows.cs
--- a/Engine/Source/Programs/UnrealBuildTool/Platform/Windows/UEBuildWindows.cs
+++ b/Engine/Source/Programs/UnrealBuildTool/Platform/Windows/UEBuildWindows.cs
@@ -382,10 +382,7 @@ namespace UnrealBuildTool
 
 		/// VS2015 updated some of the CRT definitions but not all of the Windows SDK has been updated to match.
 		/// Microsoft provides legacy_stdio_definitions library to enable building with VS2015 until they fix everything up.
-		public bool bNeedsLegacyStdioDefinitionsLib
-		{
-			get { return Compiler.IsMSVC() || Compiler.IsClang(); }
-		}
+		public bool bNeedsLegacyStdioDefinitionsLib => Compiler.IsMSVC() || Compiler.IsClang();
 
 		/// <summary>
 		/// The stack size when linking
@@ -418,8 +415,8 @@ namespace UnrealBuildTool
 		/// </summary>
 		public bool bCreateHotPatchableImage
 		{
-			get { return bCreateHotPatchableImagePrivate ?? Target.bWithLiveCoding; }
-			set { bCreateHotPatchableImagePrivate = value; }
+			get => bCreateHotPatchableImagePrivate ?? Target.bWithLiveCoding;
+			set => bCreateHotPatchableImagePrivate = value;
 		}
 		private bool? bCreateHotPatchableImagePrivate;
 
@@ -428,8 +425,8 @@ namespace UnrealBuildTool
 		/// </summary>
 		public bool bStripUnreferencedSymbols
 		{
-			get { return bStripUnreferencedSymbolsPrivate ?? ((Target.Configuration == UnrealTargetConfiguration.Test || Target.Configuration == UnrealTargetConfiguration.Shipping) && !Target.bWithLiveCoding); }
-			set { bStripUnreferencedSymbolsPrivate = value; }
+			get => bStripUnreferencedSymbolsPrivate ?? ((Target.Configuration == UnrealTargetConfiguration.Test || Target.Configuration == UnrealTargetConfiguration.Shipping) && !Target.bWithLiveCoding);
+			set => bStripUnreferencedSymbolsPrivate = value;
 		}
 		private bool? bStripUnreferencedSymbolsPrivate;
 			
@@ -438,8 +435,8 @@ namespace UnrealBuildTool
 		/// </summary>
 		public bool bMergeIdenticalCOMDATs
 		{
-			get { return bMergeIdenticalCOMDATsPrivate ?? ((Target.Configuration == UnrealTargetConfiguration.Test || Target.Configuration == UnrealTargetConfiguration.Shipping) && !Target.bWithLiveCoding); }
-			set { bMergeIdenticalCOMDATsPrivate = value; }
+			get => bMergeIdenticalCOMDATsPrivate ?? ((Target.Configuration == UnrealTargetConfiguration.Test || Target.Configuration == UnrealTargetConfiguration.Shipping) && !Target.bWithLiveCoding);
+			set => bMergeIdenticalCOMDATsPrivate = value;
 		}
 		private bool? bMergeIdenticalCOMDATsPrivate;
 
@@ -493,26 +490,17 @@ namespace UnrealBuildTool
 		/// <summary>
 		/// Directory containing the toolchain
 		/// </summary>
-		public string? ToolChainDir
-		{
-			get { return (Environment == null)? null : Environment.ToolChainDir.FullName; }
-		}
+		public string? ToolChainDir => (Environment == null) ? null : Environment.ToolChainDir.FullName;
 
 		/// <summary>
 		/// The version number of the toolchain
 		/// </summary>
-		public string? ToolChainVersion
-		{
-			get { return (Environment == null)? null : Environment.ToolChainVersion.ToString(); }
-		}
+		public string? ToolChainVersion => (Environment == null) ? null : Environment.ToolChainVersion.ToString();
 
 		/// <summary>
 		/// Root directory containing the Windows Sdk
 		/// </summary>
-		public string? WindowsSdkDir
-		{
-			get { return (Environment == null)? null : Environment.WindowsSdkDir.FullName; }
-		}
+		public string? WindowsSdkDir => (Environment == null) ? null : Environment.WindowsSdkDir.FullName;
 
 
 		/// <summary>
@@ -535,10 +523,7 @@ namespace UnrealBuildTool
 		/// <summary>
 		/// Directory containing the DIA SDK
 		/// </summary>
-		public string? DiaSdkDir
-		{
-			get { return MicrosoftPlatformSDK.FindDiaSdkDirs(Environment!.ToolChain).Select(x => x.FullName).FirstOrDefault(); }
-		}
+		public string? DiaSdkDir => MicrosoftPlatformSDK.FindDiaSdkDirs(Environment!.ToolChain).Select(x => x.FullName).FirstOrDefault();
 
 		/// <summary>
 		/// Directory containing the IDE package (Professional, Community, etc...)
@@ -621,265 +606,115 @@ namespace UnrealBuildTool
 		#region Read-only accessor properties 
 		#pragma warning disable CS1591
 
-		public bool bUseFastGenProfile 
-		{
-			get { return Inner.bUseFastGenProfile; }
-		}
+		public bool bUseFastGenProfile => Inner.bUseFastGenProfile;
 
-		public bool bPGONoExtraCounters 
-		{
-			get { return Inner.bPGONoExtraCounters; }
-		}
+		public bool bPGONoExtraCounters => Inner.bPGONoExtraCounters;
 
-		public WindowsCompiler Compiler
-		{
-			get { return Inner.Compiler; }
-		}
-		
-		public UnrealArch Architecture
-		{
-			get { return Inner.Architecture; }
-		}
+		public WindowsCompiler Compiler => Inner.Compiler;
 
-		public string? CompilerVersion
-		{
-			get { return Inner.CompilerVersion; }
-		}
+		public UnrealArch Architecture => Inner.Architecture;
 
-		public string? WindowsSdkVersion
-		{
-			get { return Inner.WindowsSdkVersion; }
-		}
+		public string? CompilerVersion => Inner.CompilerVersion;
 
-		public int TargetWindowsVersion
-		{
-			get { return Inner.TargetWindowsVersion; }
-		}
+		public string? WindowsSdkVersion => Inner.WindowsSdkVersion;
 
-		public bool bPixProfilingEnabled
-		{
-			get { return Inner.bPixProfilingEnabled; }
-		}
+		public int TargetWindowsVersion => Inner.TargetWindowsVersion;
 
-		public bool bUseWindowsSDK10
-		{
-			get { return Inner.bUseWindowsSDK10; }
-		}
+		public bool bPixProfilingEnabled => Inner.bPixProfilingEnabled;
 
-		public bool bUseCPPWinRT
-		{
-			get { return Inner.bUseCPPWinRT; }
-		}
+		public bool bUseWindowsSDK10 => Inner.bUseWindowsSDK10;
 
-		public bool bVCFastFail
-		{
-			get { return Inner.bVCFastFail; }
-		}
+		public bool bUseCPPWinRT => Inner.bUseCPPWinRT;
 
-		public bool bAllowClangLinker
-		{
-			get { return Inner.bAllowClangLinker; }
-		}
+		public bool bVCFastFail => Inner.bVCFastFail;
 
-		public bool bEnableRayTracing
-		{
-			get { return Inner.bEnableRayTracing; }
-		}
+		public bool bAllowClangLinker => Inner.bAllowClangLinker;
 
-		public string? CompanyName
-		{
-			get { return Inner.CompanyName; }
-		}
+		public bool bEnableRayTracing => Inner.bEnableRayTracing;
 
-		public string? CopyrightNotice
-		{
-			get { return Inner.CopyrightNotice; }
-		}
+		public string? CompanyName => Inner.CompanyName;
 
-		public string? ProductName
-		{
-			get { return Inner.ProductName; }
-		}
+		public string? CopyrightNotice => Inner.CopyrightNotice;
 
-		public bool bEnableAddressSanitizer
-		{ 
-			get { return Inner.bEnableAddressSanitizer; }
-		}
+		public string? ProductName => Inner.ProductName;
 
-		public bool bEnableLibFuzzer
-		{ 
-			get { return Inner.bEnableLibFuzzer; }
-		}
-        
-		public bool bWriteSarif
-		{
-			get { return Inner.bWriteSarif; }
-		}
+		public bool bEnableAddressSanitizer => Inner.bEnableAddressSanitizer;
 
-		public string? ObjSrcMapFile
-		{
-			get { return Inner.ObjSrcMapFile; }
-		}
+		public bool bEnableLibFuzzer => Inner.bEnableLibFuzzer;
 
-		public string? ModuleDefinitionFile
-		{
-			get { return Inner.ModuleDefinitionFile; }
-		}
+		public bool bWriteSarif => Inner.bWriteSarif;
 
-		public string? ManifestFile
-		{
-			get { return Inner.ManifestFile; }
-		}
+		public string? ObjSrcMapFile => Inner.ObjSrcMapFile;
 
-		public string? ApplicationIcon
-		{
-			get { return Inner.ApplicationIcon; }
-		}
+		public string? ModuleDefinitionFile => Inner.ModuleDefinitionFile;
 
-		public bool bNeedsLegacyStdioDefinitionsLib
-		{
-			get { return Inner.bNeedsLegacyStdioDefinitionsLib; }
-		}
+		public string? ManifestFile => Inner.ManifestFile;
 
-		public bool bStrictConformanceMode
-		{
-			get { return Inner.bStrictConformanceMode; }
-		}
+		public string? ApplicationIcon => Inner.ApplicationIcon;
 
-		public bool bUpdatedCPPMacro
-		{
-			get { return Inner.bUpdatedCPPMacro; }
-		}
+		public bool bNeedsLegacyStdioDefinitionsLib => Inner.bNeedsLegacyStdioDefinitionsLib;
 
-		public bool bStrictInlineConformance
-		{
-			get { return Inner.bStrictInlineConformance; }
-		}
+		public bool bStrictConformanceMode => Inner.bStrictConformanceMode;
 
-		public bool bStrictPreprocessorConformance
-		{
-			get { return Inner.bStrictPreprocessorConformance; }
-		}
+		public bool bUpdatedCPPMacro => Inner.bUpdatedCPPMacro;
 
-		public bool bStrictEnumTypesConformance
-		{
-			get { return Inner.bStrictEnumTypesConformance; }
-		}
+		public bool bStrictInlineConformance => Inner.bStrictInlineConformance;
 
-		public int DefaultStackSize
-		{
-			get { return Inner.DefaultStackSize; }
-		}
+		public bool bStrictPreprocessorConformance => Inner.bStrictPreprocessorConformance;
 
-		public int DefaultStackSizeCommit
-		{
-			get { return Inner.DefaultStackSizeCommit; }
-		}
+		public bool bStrictEnumTypesConformance => Inner.bStrictEnumTypesConformance;
 
-		public int PCHMemoryAllocationFactor
-		{
-			get { return Inner.PCHMemoryAllocationFactor; }
-		}
+		public int DefaultStackSize => Inner.DefaultStackSize;
 
-		public string AdditionalLinkerOptions
-		{
-			get { return Inner.AdditionalLinkerOptions; }
-		}
+		public int DefaultStackSizeCommit => Inner.DefaultStackSizeCommit;
 
-		public bool bCreateHotpatchableImage
-		{
-			get { return Inner.bCreateHotPatchableImage; }
-		}
+		public int PCHMemoryAllocationFactor => Inner.PCHMemoryAllocationFactor;
 
-		public bool bStripUnreferencedSymbols
-		{
-			get { return Inner.bStripUnreferencedSymbols; }
-		}
+		public string AdditionalLinkerOptions => Inner.AdditionalLinkerOptions;
 
-		public bool bMergeIdenticalCOMDATs
-		{
-			get { return Inner.bMergeIdenticalCOMDATs; }
-		}
+		public bool bCreateHotpatchableImage => Inner.bCreateHotPatchableImage;
 
-		public bool bOptimizeGlobalData
-		{
-			get { return Inner.bOptimizeGlobalData; }
-		}
+		public bool bStripUnreferencedSymbols => Inner.bStripUnreferencedSymbols;
 
-		public bool bClangTimeTrace
-		{
-			get { return Inner.bClangTimeTrace; }
-		}
+		public bool bMergeIdenticalCOMDATs => Inner.bMergeIdenticalCOMDATs;
 
-		public bool bCompilerTrace
-		{
-			get { return Inner.bCompilerTrace; }
-		}
+		public bool bOptimizeGlobalData => Inner.bOptimizeGlobalData;
 
-		public bool bShowIncludes
-		{
-			get { return Inner.bShowIncludes; }
-		}
+		public bool bClangTimeTrace => Inner.bClangTimeTrace;
+
+		public bool bCompilerTrace => Inner.bCompilerTrace;
+
+		public bool bShowIncludes => Inner.bShowIncludes;
 
 		public string GetVisualStudioCompilerVersionName()
 		{
 			return Inner.GetVisualStudioCompilerVersionName();
 		}
 
-		public bool bUseBundledDbgHelp
-		{
-			get { return Inner.bUseBundledDbgHelp; }
-		}
+		public bool bUseBundledDbgHelp => Inner.bUseBundledDbgHelp;
 
-		public bool bUseXCurl
-		{
-			get { return Inner.bUseXCurl; }
-		}
+		public bool bUseXCurl => Inner.bUseXCurl;
 
 		public ReadOnlyPVSTargetSettings PVS
 		{
 			get; private set;
 		}
 
-		internal VCEnvironment? Environment
-		{
-			get { return Inner.Environment; }
-		}
+		internal VCEnvironment? Environment => Inner.Environment;
 
-		public WindowsCompiler ToolChain
-		{
-			get { return Inner.ToolChain; }
-		}
+		public WindowsCompiler ToolChain => Inner.ToolChain;
 
-		public string? ToolChainDir
-		{
-			get { return Inner.ToolChainDir; }
-		}
+		public string? ToolChainDir => Inner.ToolChainDir;
 
-		public string? ToolChainVersion
-		{
-			get { return Inner.ToolChainVersion; }
-		}
+		public string? ToolChainVersion => Inner.ToolChainVersion;
 
-		public string? WindowsSdkDir
-		{
-			get { return Inner.WindowsSdkDir; }
-		}
+		public string? WindowsSdkDir => Inner.WindowsSdkDir;
 
-		public string? NetFxSdkDir
-		{
-			get { return Inner.NetFxSdkDir; }
-		}
+		public string? NetFxSdkDir => Inner.NetFxSdkDir;
 
-		public string? DiaSdkDir
-		{
-			get { return Inner.DiaSdkDir; }
-		}
-		
-		public string? IDEDir
-		{
-			get { return Inner.IDEDir; }
-		}
+		public string? DiaSdkDir => Inner.DiaSdkDir;
+
+		public string? IDEDir => Inner.IDEDir;
 
 		#pragma warning restore CS1591
 		#endregion
@@ -1725,10 +1560,7 @@ namespace UnrealBuildTool
 
 	class WindowsPlatformFactory : UEBuildPlatformFactory
 	{
-		public override UnrealTargetPlatform TargetPlatform
-		{
-			get { return UnrealTargetPlatform.Win64; }
-		}
+		public override UnrealTargetPlatform TargetPlatform => UnrealTargetPlatform.Win64;
 
 		/// <summary>
 		/// Register the platform with the UEBuildPlatform class
diff --git a/Engine/Source/Programs/UnrealBuildTool/Platform/Windows/VCCompileAction.cs b/Engine/Source/Programs/UnrealBuildTool/Platform/Windows/VCCompileAction.cs
--- a/Engine/Source/Programs/UnrealBuildTool/Platform/Windows/VCCompileAction.cs
+++ b/Engine/Source/Programs/UnrealBuildTool/Platform/Windows/VCCompileAction.cs
@@ -184,7 +184,7 @@ namespace UnrealBuildTool
 		public bool bUseActionHistory => true;
 
 		/// <inheritdoc/>
-		public bool bIsHighPriority { get => CreatePchFile != null; }
+		public bool bIsHighPriority => CreatePchFile != null;
 
 		/// <inheritdoc/>
 		public double Weight { get; set; } = 1.0;
@@ -306,14 +306,8 @@ namespace UnrealBuildTool
 		}
 
 		/// <inheritdoc/>
-		string IExternalAction.CommandVersion
-		{
-			get
-			{
-				return ToolChainVersion;
-			}
-		}
-		
+		string IExternalAction.CommandVersion => ToolChainVersion;
+
 		#endregion
 
 		/// <summary>
diff --git a/Engine/Source/Programs/UnrealBuildTool/Platform/Windows/VCSpecificFileAction.cs b/Engine/Source/Programs/UnrealBuildTool/Platform/Windows/VCSpecificFileAction.cs
--- a/Engine/Source/Programs/UnrealBuildTool/Platform/Windows/VCSpecificFileAction.cs
+++ b/Engine/Source/Programs/UnrealBuildTool/Platform/Windows/VCSpecificFileAction.cs
@@ -48,7 +48,7 @@ namespace UnrealBuildTool
 			BaseAction.Write(Writer);
 		}
 
-		public DirectoryReference RootDirectory { get => SourceDir; }
+		public DirectoryReference RootDirectory => SourceDir;
 
 		public IExternalAction? CreateAction(FileItem SourceFile, ILogger Logger)
 		{
diff --git a/Engine/Source/Programs/UnrealBuildTool/ProjectFiles/CMake/CMakefileGenerator.cs b/Engine/Source/Programs/UnrealBuildTool/ProjectFiles/CMake/CMakefileGenerator.cs
--- a/Engine/Source/Programs/UnrealBuildTool/ProjectFiles/CMake/CMakefileGenerator.cs
+++ b/Engine/Source/Programs/UnrealBuildTool/ProjectFiles/CMake/CMakefileGenerator.cs
@@ -47,171 +47,78 @@ namespace UnrealBuildTool
 		/// This determines if engine files are included in the source lists.
 		/// </remarks>
 		/// <returns><value>true</value> if we should treat this as a project build; <value>false</value> otherwise.</returns>
-		public bool IsProjectBuild
-		{
-			get { return !string.IsNullOrEmpty(GameProjectName); }
-		}
+		public bool IsProjectBuild => !string.IsNullOrEmpty(GameProjectName);
 
 		/// <summary>
 		/// The file extension for this project file.
 		/// </summary>
-		public override string ProjectFileExtension
-		{
-			get
-			{
-				return ".txt";
-			}
-		}
+		public override string ProjectFileExtension => ".txt";
 
-		public string ProjectFileName
-		{
-			get
-			{
-				return "CMakeLists" + ProjectFileExtension;
-			}
-		}
+		public string ProjectFileName => "CMakeLists" + ProjectFileExtension;
 
 		/// <summary>
 		/// The CMake helper file extension
 		/// </summary>
-		public string CMakeExtension
-		{
-			get
-			{
-				return ".cmake";
-			}
-		}
+		public string CMakeExtension => ".cmake";
 
 		/// <summary>
 		/// The CMake file used to store the list of includes for the project.
 		/// </summary>
-		public string CMakeIncludesFileName
-		{
-			get
-			{
-				return "cmake-includes" + CMakeExtension;
-			}
-		}
+		public string CMakeIncludesFileName => "cmake-includes" + CMakeExtension;
 
 		/// <summary>
 		/// The CMake file used to store the configuration files (INI) for the engine.
 		/// </summary>
-		public string CMakeEngineConfigsFileName
-		{
-			get
-			{
-				return "cmake-config-engine" + CMakeExtension;
-			}
-		}
+		public string CMakeEngineConfigsFileName => "cmake-config-engine" + CMakeExtension;
 
 		/// <summary>
 		/// The CMake file used to store the configuration files (INI) for the project.
 		/// </summary>
-		public string CMakeProjectConfigsFileName
-		{
-			get
-			{
-				return "cmake-config-project" + CMakeExtension;
-			}
-		}
+		public string CMakeProjectConfigsFileName => "cmake-config-project" + CMakeExtension;
 
 		/// <summary>
 		/// The CMake file used to store the additional build configuration files (CSharp) for the engine.
 		/// </summary>
-		public string CMakeEngineCSFileName
-		{
-			get
-			{
-				return "cmake-csharp-engine" + CMakeExtension;
-			}
-		}
+		public string CMakeEngineCSFileName => "cmake-csharp-engine" + CMakeExtension;
 
 		/// <summary>
 		/// The CMake file used to store the additional configuration files (CSharp) for the project.
 		/// </summary>
-		public string CMakeProjectCSFileName
-		{
-			get
-			{
-				return "cmake-csharp-project" + CMakeExtension;
-			}
-		}
+		public string CMakeProjectCSFileName => "cmake-csharp-project" + CMakeExtension;
 
 		/// <summary>
 		/// The CMake file used to store the additional shader files (usf/ush) for the engine.
 		/// </summary>
-		public string CMakeEngineShadersFileName
-		{
-			get
-			{
-				return "cmake-shaders-engine" + CMakeExtension;
-			}
-		}
+		public string CMakeEngineShadersFileName => "cmake-shaders-engine" + CMakeExtension;
 
 		/// <summary>
 		/// The CMake file used to store the additional shader files (usf/ush) for the project.
 		/// </summary>
-		public string CMakeProjectShadersFileName
-		{
-			get
-			{
-				return "cmake-shaders-project" + CMakeExtension;
-			}
-		}
+		public string CMakeProjectShadersFileName => "cmake-shaders-project" + CMakeExtension;
 
 		/// <summary>
 		/// The CMake file used to store the list of engine headers.
 		/// </summary>
-		public string CMakeEngineHeadersFileName
-		{
-			get
-			{
-				return "cmake-headers-ue" + CMakeExtension;
-			}
-		}
+		public string CMakeEngineHeadersFileName => "cmake-headers-ue" + CMakeExtension;
 		/// <summary>
 		/// The CMake file used to store the list of engine headers.
 		/// </summary>
-		public string CMakeProjectHeadersFileName
-		{
-			get
-			{
-				return "cmake-headers-project" + CMakeExtension;
-			}
-		}
+		public string CMakeProjectHeadersFileName => "cmake-headers-project" + CMakeExtension;
 
 		/// <summary>
 		/// The CMake file used to store the list of sources for the engine.
 		/// </summary>
-		public string CMakeEngineSourcesFileName
-		{
-			get
-			{
-				return "cmake-sources-engine" + CMakeExtension;
-			}
-		}
-		
+		public string CMakeEngineSourcesFileName => "cmake-sources-engine" + CMakeExtension;
+
 		/// <summary>
 		/// The CMake file used to store the list of sources for the project.
 		/// </summary>
-		public string CMakeProjectSourcesFileName
-		{
-			get
-			{
-				return "cmake-sources-project" +  CMakeExtension;
-			}
-		}
+		public string CMakeProjectSourcesFileName => "cmake-sources-project" + CMakeExtension;
 
 		/// <summary>
 		/// The CMake file used to store the list of definitions for the project.
 		/// </summary>
-		public string CMakeDefinitionsFileName
-		{
-			get
-			{
-				return "cmake-definitions" + CMakeExtension;
-			}
-		}
+		public string CMakeDefinitionsFileName => "cmake-definitions" + CMakeExtension;
 
 
 		/// <summary>
diff --git a/Engine/Source/Programs/UnrealBuildTool/ProjectFiles/CodeLite/CodeLiteGenerator.cs b/Engine/Source/Programs/UnrealBuildTool/ProjectFiles/CodeLite/CodeLiteGenerator.cs
--- a/Engine/Source/Programs/UnrealBuildTool/ProjectFiles/CodeLite/CodeLiteGenerator.cs
+++ b/Engine/Source/Programs/UnrealBuildTool/ProjectFiles/CodeLite/CodeLiteGenerator.cs
@@ -41,13 +41,7 @@ namespace UnrealBuildTool
 		//
 		// Returns CodeLite's project filename extension.
 		//
-		override public string ProjectFileExtension
-		{
-			get
-			{
-				return ".project";
-			}
-		}
+		override public string ProjectFileExtension => ".project";
 		protected override bool WritePrimaryProjectFile(ProjectFile? UBTProject, PlatformProjectGeneratorCollection PlatformProjectGenerators, ILogger Logger)
 		{
 			string SolutionFileName = PrimaryProjectName + SolutionExtension;
diff --git a/Engine/Source/Programs/UnrealBuildTool/ProjectFiles/Eddie/EddieGenerator.cs b/Engine/Source/Programs/UnrealBuildTool/ProjectFiles/Eddie/EddieGenerator.cs
--- a/Engine/Source/Programs/UnrealBuildTool/ProjectFiles/Eddie/EddieGenerator.cs
+++ b/Engine/Source/Programs/UnrealBuildTool/ProjectFiles/Eddie/EddieGenerator.cs
@@ -15,15 +15,9 @@ namespace UnrealBuildTool
 			: base(InOnlyGameProject)
 		{
 		}
-		
-		override public string ProjectFileExtension
-		{
-			get
-			{
-				return ".wkst";
-			}
-		}
-		
+
+		override public string ProjectFileExtension => ".wkst";
+
 		public override void CleanProjectFiles(DirectoryReference InPrimaryProjectDirectory, string InPrimaryProjectName, DirectoryReference InIntermediateProjectFilesPath, ILogger Logger)
 		{
 			FileReference PrimaryProjDeleteFilename = FileReference.Combine(InPrimaryProjectDirectory, InPrimaryProjectName + ".wkst");
diff --git a/Engine/Source/Programs/UnrealBuildTool/ProjectFiles/KDevelop/KDevelopGenerator.cs b/Engine/Source/Programs/UnrealBuildTool/ProjectFiles/KDevelop/KDevelopGenerator.cs
--- a/Engine/Source/Programs/UnrealBuildTool/ProjectFiles/KDevelop/KDevelopGenerator.cs
+++ b/Engine/Source/Programs/UnrealBuildTool/ProjectFiles/KDevelop/KDevelopGenerator.cs
@@ -30,13 +30,7 @@ namespace UnrealBuildTool
 		}
 
 		/// File extension for project files we'll be generating (e.g. ".vcxproj")
-		override public string ProjectFileExtension
-		{
-			get
-			{
-				return ".kdev4";
-			}
-		}
+		override public string ProjectFileExtension => ".kdev4";
 
 		protected override bool WritePrimaryProjectFile(ProjectFile? UBTProject, PlatformProjectGeneratorCollection PlatformProjectGenerators, ILogger Logger)
 		{
diff --git a/Engine/Source/Programs/UnrealBuildTool/ProjectFiles/Make/MakefileGenerator.cs b/Engine/Source/Programs/UnrealBuildTool/ProjectFiles/Make/MakefileGenerator.cs
--- a/Engine/Source/Programs/UnrealBuildTool/ProjectFiles/Make/MakefileGenerator.cs
+++ b/Engine/Source/Programs/UnrealBuildTool/ProjectFiles/Make/MakefileGenerator.cs
@@ -41,13 +41,7 @@ namespace UnrealBuildTool
 		}
 
 		/// File extension for project files we'll be generating (e.g. ".vcxproj")
-		override public string ProjectFileExtension
-		{
-			get
-			{
-				return ".mk";
-			}
-		}
+		override public string ProjectFileExtension => ".mk";
 
 		protected override bool WritePrimaryProjectFile(ProjectFile? UBTProject, PlatformProjectGeneratorCollection PlatformProjectGenerators, ILogger Logger)
 		{
diff --git a/Engine/Source/Programs/UnrealBuildTool/ProjectFiles/ProjectFile.cs b/Engine/Source/Programs/UnrealBuildTool/ProjectFiles/ProjectFile.cs
--- a/Engine/Source/Programs/UnrealBuildTool/ProjectFiles/ProjectFile.cs
+++ b/Engine/Source/Programs/UnrealBuildTool/ProjectFiles/ProjectFile.cs
@@ -532,18 +532,12 @@ namespace UnrealBuildTool
 		/// <summary>
 		/// Legacy accessor for user search paths
 		/// </summary>
-		public List<string> IntelliSenseIncludeSearchPaths
-		{
-			get { return UserIncludePaths.RelativePaths; }
-		}
+		public List<string> IntelliSenseIncludeSearchPaths => UserIncludePaths.RelativePaths;
 
 		/// <summary>
 		/// Legacy accessor for system include paths
 		/// </summary>
-		public List<string> IntelliSenseSystemIncludeSearchPaths
-		{
-			get { return SystemIncludePaths.RelativePaths; }
-		}
+		public List<string> IntelliSenseSystemIncludeSearchPaths => SystemIncludePaths.RelativePaths;
 
 		/// <summary>
 		/// List of preprocessor definitions for the project
diff --git a/Engine/Source/Programs/UnrealBuildTool/ProjectFiles/ProjectTarget.cs b/Engine/Source/Programs/UnrealBuildTool/ProjectFiles/ProjectTarget.cs
--- a/Engine/Source/Programs/UnrealBuildTool/ProjectFiles/ProjectTarget.cs
+++ b/Engine/Source/Programs/UnrealBuildTool/ProjectFiles/ProjectTarget.cs
@@ -20,10 +20,7 @@ namespace UnrealBuildTool
 		/// Delegate for creating a rules instance for a given platform/configuration
 		public Func<UnrealTargetPlatform, UnrealTargetConfiguration, TargetRules> CreateRulesDelegate;
 
-		public string Name
-		{
-			get { return TargetFilePath.GetFileNameWithoutAnyExtensions(); }
-		}
+		public string Name => TargetFilePath.GetFileNameWithoutAnyExtensions();
 
 		public override string ToString()
 		{
diff --git a/Engine/Source/Programs/UnrealBuildTool/ProjectFiles/QMake/QMakefileGenerator.cs b/Engine/Source/Programs/UnrealBuildTool/ProjectFiles/QMake/QMakefileGenerator.cs
--- a/Engine/Source/Programs/UnrealBuildTool/ProjectFiles/QMake/QMakefileGenerator.cs
+++ b/Engine/Source/Programs/UnrealBuildTool/ProjectFiles/QMake/QMakefileGenerator.cs
@@ -31,13 +31,7 @@ namespace UnrealBuildTool
 		}
 
 		/// File extension for project files we'll be generating (e.g. ".vcxproj")
-		override public string ProjectFileExtension
-		{
-			get
-			{
-				return ".pro";
-			}
-		}
+		override public string ProjectFileExtension => ".pro";
 
 		protected override bool WritePrimaryProjectFile(ProjectFile? UBTProject, PlatformProjectGeneratorCollection PlatformProjectGenerators, ILogger Logger)
 		{
diff --git a/Engine/Source/Programs/UnrealBuildTool/ProjectFiles/VisualStudio/VCProject.cs b/Engine/Source/Programs/UnrealBuildTool/ProjectFiles/VisualStudio/VCProject.cs
--- a/Engine/Source/Programs/UnrealBuildTool/ProjectFiles/VisualStudio/VCProject.cs
+++ b/Engine/Source/Programs/UnrealBuildTool/ProjectFiles/VisualStudio/VCProject.cs
@@ -57,10 +57,7 @@ namespace UnrealBuildTool
 		/// <summary>
 		/// The name of this context
 		/// </summary>
-		public string Name
-		{
-			get { return String.Format("{0}|{1}", ConfigurationName, PlatformName); }
-		}
+		public string Name => String.Format("{0}|{1}", ConfigurationName, PlatformName);
 
 		/// <summary>
 		/// Serializes this context to a string for debugging
@@ -92,10 +89,7 @@ namespace UnrealBuildTool
 		/// <summary>
 		/// The Guid representing the project type e.g. C# or C++
 		/// </summary>
-		public virtual string ProjectTypeGUID
-		{
-			get { throw new BuildException("Unrecognized type of project file for Visual Studio solution"); }
-		}
+		public virtual string ProjectTypeGUID => throw new BuildException("Unrecognized type of project file for Visual Studio solution");
 
 		static Guid PathNamespaceGuid { get; } = new Guid("2D8570D5-7FFC-4E6D-A9D7-E860E117D717");
 
@@ -296,10 +290,7 @@ namespace UnrealBuildTool
 		public const string DefaultPlatformName = "Win64";
 
 		// This is the GUID that Visual Studio uses to identify a C++ project file in the solution
-		public override string ProjectTypeGUID
-		{
-			get { return "{8BC9CEB8-8B4A-11D0-8D11-00A0C91BC942}"; }
-		}
+		public override string ProjectTypeGUID => "{8BC9CEB8-8B4A-11D0-8D11-00A0C91BC942}";
 
 		/// <summary>
 		/// Constructs a new project file object
@@ -530,10 +521,7 @@ namespace UnrealBuildTool
 				ProjectTarget = InProjectTarget;
 			}
 
-			public string? ProjectConfigurationAndPlatformName
-			{
-				get { return (ProjectPlatformName == null) ? null : (ProjectConfigurationName + "|" + ProjectPlatformName); }
-			}
+			public string? ProjectConfigurationAndPlatformName => (ProjectPlatformName == null) ? null : (ProjectConfigurationName + "|" + ProjectPlatformName);
 
 			public override string ToString()
 			{
@@ -2161,10 +2149,7 @@ namespace UnrealBuildTool
 		/// <summary>
 		/// This is the GUID that Visual Studio uses to identify a C# project file in the solution
 		/// </summary>
-		public override string ProjectTypeGUID
-		{
-			get { return "{FAE04EC0-301F-11D3-BF4B-00C04F79EFBC}"; }
-		}
+		public override string ProjectTypeGUID => "{FAE04EC0-301F-11D3-BF4B-00C04F79EFBC}";
 
 		/// <summary>
 		/// Platforms that this project supports
diff --git a/Engine/Source/Programs/UnrealBuildTool/ProjectFiles/VisualStudio/VCProjectFileGenerator.cs b/Engine/Source/Programs/UnrealBuildTool/ProjectFiles/VisualStudio/VCProjectFileGenerator.cs
--- a/Engine/Source/Programs/UnrealBuildTool/ProjectFiles/VisualStudio/VCProjectFileGenerator.cs
+++ b/Engine/Source/Programs/UnrealBuildTool/ProjectFiles/VisualStudio/VCProjectFileGenerator.cs
@@ -151,13 +151,7 @@ namespace UnrealBuildTool
 		}
 
 		/// File extension for project files we'll be generating (e.g. ".vcxproj")
-		override public string ProjectFileExtension
-		{
-			get
-			{
-				return ".vcxproj";
-			}
-		}
+		override public string ProjectFileExtension => ".vcxproj";
 
 		/// <summary>
 		/// </summary>
diff --git a/Engine/Source/Programs/UnrealBuildTool/ProjectFiles/VisualStudioCode/VSCodeProjectFileGenerator.cs b/Engine/Source/Programs/UnrealBuildTool/ProjectFiles/VisualStudioCode/VSCodeProjectFileGenerator.cs
--- a/Engine/Source/Programs/UnrealBuildTool/ProjectFiles/VisualStudioCode/VSCodeProjectFileGenerator.cs
+++ b/Engine/Source/Programs/UnrealBuildTool/ProjectFiles/VisualStudioCode/VSCodeProjectFileGenerator.cs
@@ -239,13 +239,7 @@ namespace UnrealBuildTool
 			private string TabString = "";
 		}
 
-		override public string ProjectFileExtension
-		{
-			get
-			{
-				return ".vscode";
-			}
-		}
+		override public string ProjectFileExtension => ".vscode";
 
 		public override void CleanProjectFiles(DirectoryReference InPrimaryProjectDirectory, string InPrimaryProjectName, DirectoryReference InIntermediateProjectFilesPath, ILogger Logger)
 		{
diff --git a/Engine/Source/Programs/UnrealBuildTool/ProjectFiles/Xcode/XcodeProjectFileGenerator.cs b/Engine/Source/Programs/UnrealBuildTool/ProjectFiles/Xcode/XcodeProjectFileGenerator.cs
--- a/Engine/Source/Programs/UnrealBuildTool/ProjectFiles/Xcode/XcodeProjectFileGenerator.cs
+++ b/Engine/Source/Programs/UnrealBuildTool/ProjectFiles/Xcode/XcodeProjectFileGenerator.cs
@@ -161,13 +161,7 @@ namespace UnrealBuildTool
 		}
 
 		/// File extension for project files we'll be generating (e.g. ".vcxproj")
-		override public string ProjectFileExtension
-		{
-			get
-			{
-				return ".xcodeproj";
-			}
-		}
+		override public string ProjectFileExtension => ".xcodeproj";
 
 		/// <summary>
 		/// </summary>
diff --git a/Engine/Source/Programs/UnrealBuildTool/System/BuildVersion.cs b/Engine/Source/Programs/UnrealBuildTool/System/BuildVersion.cs
--- a/Engine/Source/Programs/UnrealBuildTool/System/BuildVersion.cs
+++ b/Engine/Source/Programs/UnrealBuildTool/System/BuildVersion.cs
@@ -73,10 +73,7 @@ namespace UnrealBuildTool
 		/// Returns the value which can be used as the compatible changelist. Requires that the regular changelist is also set, and defaults to the 
 		/// regular changelist if a specific compatible changelist is not set.
 		/// </summary>
-		public int EffectiveCompatibleChangelist
-		{
-			get { return (Changelist != 0 && CompatibleChangelist != 0)? CompatibleChangelist : Changelist; }
-		}
+		public int EffectiveCompatibleChangelist => (Changelist != 0 && CompatibleChangelist != 0) ? CompatibleChangelist : Changelist;
 
 		/// <summary>
 		/// Try to read a version file from disk
@@ -288,55 +285,25 @@ namespace UnrealBuildTool
 		#region Read-only accessor properties 
 		#pragma warning disable CS1591
 
-		public int MajorVersion
-		{
-			get { return Inner.MajorVersion; }
-		}
+		public int MajorVersion => Inner.MajorVersion;
 
-		public int MinorVersion
-		{
-			get { return Inner.MinorVersion; }
-		}
+		public int MinorVersion => Inner.MinorVersion;
 
-		public int PatchVersion
-		{
-			get { return Inner.PatchVersion; }
-		}
+		public int PatchVersion => Inner.PatchVersion;
 
-		public int Changelist
-		{
-			get { return Inner.Changelist; }
-		}
+		public int Changelist => Inner.Changelist;
 
-		public int CompatibleChangelist
-		{
-			get { return Inner.CompatibleChangelist; }
-		}
+		public int CompatibleChangelist => Inner.CompatibleChangelist;
 
-		public int EffectiveCompatibleChangelist
-		{
-			get { return Inner.EffectiveCompatibleChangelist; }
-		}
+		public int EffectiveCompatibleChangelist => Inner.EffectiveCompatibleChangelist;
 
-		public bool IsLicenseeVersion
-		{
-			get { return Inner.IsLicenseeVersion; }
-		}
+		public bool IsLicenseeVersion => Inner.IsLicenseeVersion;
 
-		public bool IsPromotedBuild
-		{
-			get { return Inner.IsPromotedBuild; }
-		}
+		public bool IsPromotedBuild => Inner.IsPromotedBuild;
 
-		public string? BranchName
-		{
-			get { return Inner.BranchName; }
-		}
+		public string? BranchName => Inner.BranchName;
 
-		public string? BuildVersionString
-		{
-			get { return Inner.BuildVersionString; }
-		}
+		public string? BuildVersionString => Inner.BuildVersionString;
 
 		#pragma warning restore C1591
 		#endregion
diff --git a/Engine/Source/Programs/UnrealBuildTool/System/ConfigFile.cs b/Engine/Source/Programs/UnrealBuildTool/System/ConfigFile.cs
--- a/Engine/Source/Programs/UnrealBuildTool/System/ConfigFile.cs
+++ b/Engine/Source/Programs/UnrealBuildTool/System/ConfigFile.cs
@@ -461,10 +461,7 @@ namespace UnrealBuildTool
 		/// <summary>
 		/// Names of sections in this file
 		/// </summary>
-		public IEnumerable<string> SectionNames
-		{
-			get { return Sections.Keys; }
-		}
+		public IEnumerable<string> SectionNames => Sections.Keys;
 
 		/// <summary>
 		/// Tries to get a config section by name, or creates one if it doesn't exist
diff --git a/Engine/Source/Programs/UnrealBuildTool/System/ConfigHierarchy.cs b/Engine/Source/Programs/UnrealBuildTool/System/ConfigHierarchy.cs
--- a/Engine/Source/Programs/UnrealBuildTool/System/ConfigHierarchy.cs
+++ b/Engine/Source/Programs/UnrealBuildTool/System/ConfigHierarchy.cs
@@ -127,10 +127,7 @@ namespace UnrealBuildTool
 		/// <summary>
 		/// Returns a list of key names
 		/// </summary>
-		public IEnumerable<string> KeyNames
-		{
-			get { return KeyToValue.Keys; }
-		}
+		public IEnumerable<string> KeyNames => KeyToValue.Keys;
 
 		/// <summary>
 		/// Tries to find the value for a given key
diff --git a/Engine/Source/Programs/UnrealBuildTool/System/CppCompileEnvironment.cs b/Engine/Source/Programs/UnrealBuildTool/System/CppCompileEnvironment.cs
--- a/Engine/Source/Programs/UnrealBuildTool/System/CppCompileEnvironment.cs
+++ b/Engine/Source/Programs/UnrealBuildTool/System/CppCompileEnvironment.cs
@@ -583,13 +583,7 @@ namespace UnrealBuildTool
 		/// <summary>
 		/// The file containing the precompiled header data.
 		/// </summary>
-		public FileItem? PrecompiledHeaderFile
-		{
-			get
-			{
-				return GetPrecompiledHeaderFile(PCHInstance);
-			}
-		}
+		public FileItem? PrecompiledHeaderFile => GetPrecompiledHeaderFile(PCHInstance);
 
 		/// <summary>
 		/// The parent PCH instance used when creating this PCH.
@@ -599,13 +593,7 @@ namespace UnrealBuildTool
 		/// <summary>
 		/// The parent's PCH header file.
 		/// </summary>
-		public FileItem? ParentPrecompiledHeaderFile
-		{
-			get
-			{
-				return GetPrecompiledHeaderFile(ParentPCHInstance);
-			}
-		}
+		public FileItem? ParentPrecompiledHeaderFile => GetPrecompiledHeaderFile(ParentPCHInstance);
 
 		/// <summary>
 		/// True if a single PRecompiledHeader exists, or at least one PerArchPrecompiledHeaderFile exists
diff --git a/Engine/Source/Programs/UnrealBuildTool/System/Timeline.cs b/Engine/Source/Programs/UnrealBuildTool/System/Timeline.cs
--- a/Engine/Source/Programs/UnrealBuildTool/System/Timeline.cs
+++ b/Engine/Source/Programs/UnrealBuildTool/System/Timeline.cs
@@ -100,10 +100,7 @@ namespace UnrealBuildTool
 		/// <summary>
 		/// Property for the total time elapsed
 		/// </summary>
-		public static TimeSpan Elapsed
-		{
-			get { return Stopwatch.Elapsed; }
-		}
+		public static TimeSpan Elapsed => Stopwatch.Elapsed;
 
 		/// <summary>
 		/// Start the stopwatch
diff --git a/Engine/Source/Programs/UnrealBuildTool/ToolChain/ClangSpecificFileAction.cs b/Engine/Source/Programs/UnrealBuildTool/ToolChain/ClangSpecificFileAction.cs
--- a/Engine/Source/Programs/UnrealBuildTool/ToolChain/ClangSpecificFileAction.cs
+++ b/Engine/Source/Programs/UnrealBuildTool/ToolChain/ClangSpecificFileAction.cs
@@ -44,7 +44,7 @@ namespace UnrealBuildTool
 			Writer.WriteList(RspLines.ToList(), (Str) => Writer.WriteString(Str));
 		}
 
-		public DirectoryReference RootDirectory { get => SourceDir; }
+		public DirectoryReference RootDirectory => SourceDir;
 
 		public IExternalAction? CreateAction(FileItem SourceFile, ILogger Logger)
 		{
