UnrealBuildTool: Automated code cleanup

#rnx

diff --git a/Engine/Source/Programs/UnrealBuildTool.Tests/TokenReaderTests.cs b/Engine/Source/Programs/UnrealBuildTool.Tests/TokenReaderTests.cs
--- a/Engine/Source/Programs/UnrealBuildTool.Tests/TokenReaderTests.cs
+++ b/Engine/Source/Programs/UnrealBuildTool.Tests/TokenReaderTests.cs
@@ -1,11 +1,10 @@
 // Copyright Epic Games, Inc. All Rights Reserved.
 
-using Microsoft.VisualStudio.TestTools.UnitTesting;
 using System;
 using System.Collections.Generic;
 using System.Linq;
 using System.Text;
-using System.Threading.Tasks;
+using Microsoft.VisualStudio.TestTools.UnitTesting;
 using UnrealBuildTool;
 
 namespace UnrealBuildToolTests
diff --git a/Engine/Source/Programs/UnrealBuildTool/Artifacts/Artifacts.cs b/Engine/Source/Programs/UnrealBuildTool/Artifacts/Artifacts.cs
--- a/Engine/Source/Programs/UnrealBuildTool/Artifacts/Artifacts.cs
+++ b/Engine/Source/Programs/UnrealBuildTool/Artifacts/Artifacts.cs
@@ -1,13 +1,11 @@
 // Copyright Epic Games, Inc. All Rights Reserved.
 
-using EpicGames.Core;
-using EpicGames.Serialization;
 using System;
-using System.Collections.Generic;
 using System.IO;
-using System.Linq;
 using System.Threading;
 using System.Threading.Tasks;
+using EpicGames.Core;
+using EpicGames.Serialization;
 
 namespace UnrealBuildTool.Artifacts
 {
@@ -90,7 +88,7 @@ namespace UnrealBuildTool.Artifacts
 		/// <inheritdoc/>
 		public override int GetHashCode()
 		{
-			return HashCode.Combine(Key.GetHashCode(), MappingKey.GetHashCode(), Inputs.GetHashCode(), Outputs.GetHashCode()); 
+			return HashCode.Combine(Key.GetHashCode(), MappingKey.GetHashCode(), Inputs.GetHashCode(), Outputs.GetHashCode());
 		}
 	}
 
diff --git a/Engine/Source/Programs/UnrealBuildTool/Configuration/BuildConfiguration.cs b/Engine/Source/Programs/UnrealBuildTool/Configuration/BuildConfiguration.cs
--- a/Engine/Source/Programs/UnrealBuildTool/Configuration/BuildConfiguration.cs
+++ b/Engine/Source/Programs/UnrealBuildTool/Configuration/BuildConfiguration.cs
@@ -1,10 +1,6 @@
 // Copyright Epic Games, Inc. All Rights Reserved.
 
 using System;
-using System.IO;
-using System.Collections.Generic;
-using System.Xml;
-using System.Reflection;
 using EpicGames.Core;
 
 namespace UnrealBuildTool
@@ -14,13 +10,13 @@ namespace UnrealBuildTool
 	/// </summary>
 	class BuildConfiguration
 	{
-        /// <summary>
-        /// Whether to ignore import library files that are out of date when building targets. Set this to true to improve iteration time.
-        /// By default, we do not bother re-linking targets if only a dependent .lib has changed, as chances are that
-        /// the import library was not actually different unless a dependent header file of this target was also changed,
-        /// in which case the target would automatically be rebuilt.
-        /// </summary>
-        [XmlConfigFile]
+		/// <summary>
+		/// Whether to ignore import library files that are out of date when building targets. Set this to true to improve iteration time.
+		/// By default, we do not bother re-linking targets if only a dependent .lib has changed, as chances are that
+		/// the import library was not actually different unless a dependent header file of this target was also changed,
+		/// in which case the target would automatically be rebuilt.
+		/// </summary>
+		[XmlConfigFile]
 		public bool bIgnoreOutdatedImportLibraries = true;
 
 		/// <summary>
@@ -137,7 +133,7 @@ namespace UnrealBuildTool
 		/// <summary>
 		/// True if hot-reload from IDE is allowed.
 		/// </summary>
-		[CommandLine("-NoHotReloadFromIDE", Value="false")]
+		[CommandLine("-NoHotReloadFromIDE", Value = "false")]
 		[XmlConfigFile(Category = "UEBuildConfiguration")]
 		public bool bAllowHotReloadFromIDE = true;
 
diff --git a/Engine/Source/Programs/UnrealBuildTool/Configuration/ModuleRules.cs b/Engine/Source/Programs/UnrealBuildTool/Configuration/ModuleRules.cs
--- a/Engine/Source/Programs/UnrealBuildTool/Configuration/ModuleRules.cs
+++ b/Engine/Source/Programs/UnrealBuildTool/Configuration/ModuleRules.cs
@@ -6,8 +6,6 @@ using System.Diagnostics.CodeAnalysis;
 using System.IO;
 using System.Linq;
 using System.Reflection;
-using System.Text;
-using System.Threading.Tasks;
 using EpicGames.Core;
 using Microsoft.Extensions.Logging;
 using UnrealBuildBase;
@@ -409,7 +407,7 @@ namespace UnrealBuildTool
 			/// <summary>
 			/// How to handle linking and copying the framework
 			/// </summary>
-			public  enum FrameworkMode
+			public enum FrameworkMode
 			{
 				/// <summary>
 				/// Pass this framework to the linker
@@ -834,7 +832,7 @@ namespace UnrealBuildTool
 
 		/// <inheritdoc cref="DeterministicWarningLevelPrivate"/>
 		private WarningLevel DeterministicWarningLevelPrivate;
-		
+
 		/// <summary>
 		/// How to treat shadow variable warnings
 		/// </summary>
@@ -1313,7 +1311,7 @@ namespace UnrealBuildTool
 		/// Which stanard to use for compiling this module
 		/// </summary>
 		public CppStandardVersion? CppStandard;
-		
+
 		/// <summary>
 		/// Which standard to use for compiling this module
 		/// </summary>
@@ -1470,7 +1468,7 @@ namespace UnrealBuildTool
 					if (Target.Type == TargetType.Editor)
 					{
 						PublicDependencyModuleNames.Add("GameplayDebuggerEditor");
-					}					
+					}
 				}
 				else
 				{
@@ -1675,7 +1673,7 @@ namespace UnrealBuildTool
 		internal void PrepareModuleForTests()
 		{
 			TestTargetRules? TestTargetRules = Target.InnerTestTargetRules;
-			if (TestTargetRules == null) 
+			if (TestTargetRules == null)
 			{
 				return;
 			}
diff --git a/Engine/Source/Programs/UnrealBuildTool/Configuration/ModuleRulesContext.cs b/Engine/Source/Programs/UnrealBuildTool/Configuration/ModuleRulesContext.cs
--- a/Engine/Source/Programs/UnrealBuildTool/Configuration/ModuleRulesContext.cs
+++ b/Engine/Source/Programs/UnrealBuildTool/Configuration/ModuleRulesContext.cs
@@ -1,10 +1,5 @@
 // Copyright Epic Games, Inc. All Rights Reserved.
 
-using System;
-using System.Collections.Generic;
-using System.Linq;
-using System.Text;
-using System.Threading.Tasks;
 using EpicGames.Core;
 
 namespace UnrealBuildTool
diff --git a/Engine/Source/Programs/UnrealBuildTool/Configuration/TargetDescriptor.cs b/Engine/Source/Programs/UnrealBuildTool/Configuration/TargetDescriptor.cs
--- a/Engine/Source/Programs/UnrealBuildTool/Configuration/TargetDescriptor.cs
+++ b/Engine/Source/Programs/UnrealBuildTool/Configuration/TargetDescriptor.cs
@@ -6,9 +6,8 @@ using System.IO;
 using System.Linq;
 using System.Text;
 using EpicGames.Core;
-using UnrealBuildBase;
 using Microsoft.Extensions.Logging;
-using System.Runtime.InteropServices;
+using UnrealBuildBase;
 
 namespace UnrealBuildTool
 {
@@ -58,7 +57,7 @@ namespace UnrealBuildTool
 		/// Lists of files to compile
 		/// </summary>
 		[CommandLine("-FileList=")]
-		public List<FileReference> FileLists = new List<FileReference>(); 
+		public List<FileReference> FileLists = new List<FileReference>();
 
 		/// <summary>
 		/// Individual file(s) to compile
@@ -142,7 +141,7 @@ namespace UnrealBuildTool
 
 			// If there are any additional command line arguments
 			List<string> AdditionalArguments = new List<string>();
-			if(Arguments != null)
+			if (Arguments != null)
 			{
 				// Apply the arguments to this object
 				Arguments.ApplyTo(this);
@@ -158,10 +157,10 @@ namespace UnrealBuildTool
 				}
 
 				// Parse all the hot-reload module names
-				foreach(string ModuleWithSuffix in Arguments.GetValues("-ModuleWithSuffix="))
+				foreach (string ModuleWithSuffix in Arguments.GetValues("-ModuleWithSuffix="))
 				{
 					int SuffixIdx = ModuleWithSuffix.LastIndexOf(',');
-					if(SuffixIdx == -1)
+					if (SuffixIdx == -1)
 					{
 						throw new BuildException("Missing suffix argument from -ModuleWithSuffix=Name,Suffix");
 					}
@@ -169,7 +168,7 @@ namespace UnrealBuildTool
 					string ModuleName = ModuleWithSuffix.Substring(0, SuffixIdx);
 
 					int Suffix;
-					if(!Int32.TryParse(ModuleWithSuffix.Substring(SuffixIdx + 1), out Suffix))
+					if (!Int32.TryParse(ModuleWithSuffix.Substring(SuffixIdx + 1), out Suffix))
 					{
 						throw new BuildException("Suffix for modules must be an integer");
 					}
@@ -178,9 +177,9 @@ namespace UnrealBuildTool
 				}
 
 				// Pull out all the arguments that haven't been used so far
-				for(int Idx = 0; Idx < Arguments.Count; Idx++)
+				for (int Idx = 0; Idx < Arguments.Count; Idx++)
 				{
-					if(!Arguments.HasBeenUsed(Idx))
+					if (!Arguments.HasBeenUsed(Idx))
 					{
 						AdditionalArguments.Add(Arguments[Idx]);
 					}
@@ -232,16 +231,16 @@ namespace UnrealBuildTool
 			List<string> Targets;
 			Arguments = Arguments.Remove("-Target=", out Targets);
 
-			if(TargetLists.Count > 0 || Targets.Count > 0)
+			if (TargetLists.Count > 0 || Targets.Count > 0)
 			{
 				// Try to parse multiple arguments from a single command line
-				foreach(string TargetList in TargetLists)
+				foreach (string TargetList in TargetLists)
 				{
 					string[] Lines = File.ReadAllLines(TargetList);
-					foreach(string Line in Lines)
+					foreach (string Line in Lines)
 					{
 						string TrimLine = Line.Trim();
-						if(TrimLine.Length > 0 && TrimLine[0] != ';')
+						if (TrimLine.Length > 0 && TrimLine[0] != ';')
 						{
 							CommandLineArguments NewArguments = Arguments.Append(CommandLineArguments.Split(TrimLine));
 							ParseCommandLine(NewArguments, bUsePrecompiled, bSkipRulesCompile, bForceRulesCompile, TargetDescriptors, Logger);
@@ -249,7 +248,7 @@ namespace UnrealBuildTool
 					}
 				}
 
-				foreach(string Target in Targets)
+				foreach (string Target in Targets)
 				{
 					CommandLineArguments NewArguments = Arguments.Append(CommandLineArguments.Split(Target));
 					ParseCommandLine(NewArguments, bUsePrecompiled, bSkipRulesCompile, bForceRulesCompile, TargetDescriptors, Logger);
@@ -282,16 +281,16 @@ namespace UnrealBuildTool
 			for (int ArgumentIndex = 0; ArgumentIndex < Arguments.Count; ArgumentIndex++)
 			{
 				string Argument = Arguments[ArgumentIndex];
-				if(Argument.Length > 0 && Argument[0] != '-')
+				if (Argument.Length > 0 && Argument[0] != '-')
 				{
 					// Mark this argument as used. We'll interpret it as one thing or another.
 					Arguments.MarkAsUsed(ArgumentIndex);
 
 					// Check if it's a project file argument
-					if(Argument.EndsWith(".uproject", StringComparison.OrdinalIgnoreCase))
+					if (Argument.EndsWith(".uproject", StringComparison.OrdinalIgnoreCase))
 					{
 						FileReference NewProjectFile = new FileReference(Argument);
-						if(ProjectFile != null && ProjectFile != NewProjectFile)
+						if (ProjectFile != null && ProjectFile != NewProjectFile)
 						{
 							throw new BuildException("Multiple project files specified on command line (first {0}, then {1})", ProjectFile, NewProjectFile);
 						}
@@ -304,10 +303,10 @@ namespace UnrealBuildTool
 
 					// Try to parse them as platforms
 					UnrealTargetPlatform ParsedPlatform;
-					if(UnrealTargetPlatform.TryParse(InlineArguments[0], out ParsedPlatform))
+					if (UnrealTargetPlatform.TryParse(InlineArguments[0], out ParsedPlatform))
 					{
 						Platforms.Add(ParsedPlatform);
-						for(int InlineArgumentIdx = 1; InlineArgumentIdx < InlineArguments.Length; InlineArgumentIdx++)
+						for (int InlineArgumentIdx = 1; InlineArgumentIdx < InlineArguments.Length; InlineArgumentIdx++)
 						{
 							Platforms.Add(UnrealTargetPlatform.Parse(InlineArguments[InlineArgumentIdx]));
 						}
@@ -316,13 +315,13 @@ namespace UnrealBuildTool
 
 					// Try to parse them as configurations
 					UnrealTargetConfiguration ParsedConfiguration;
-					if(Enum.TryParse(InlineArguments[0], true, out ParsedConfiguration))
+					if (Enum.TryParse(InlineArguments[0], true, out ParsedConfiguration))
 					{
 						Configurations.Add(ParsedConfiguration);
-						for(int InlineArgumentIdx = 1; InlineArgumentIdx < InlineArguments.Length; InlineArgumentIdx++)
+						for (int InlineArgumentIdx = 1; InlineArgumentIdx < InlineArguments.Length; InlineArgumentIdx++)
 						{
 							string InlineArgument = InlineArguments[InlineArgumentIdx];
-							if(!Enum.TryParse(InlineArgument, true, out ParsedConfiguration))
+							if (!Enum.TryParse(InlineArgument, true, out ParsedConfiguration))
 							{
 								throw new BuildException("Invalid configuration '{0}'", InlineArgument);
 							}
@@ -346,10 +345,10 @@ namespace UnrealBuildTool
 			}
 
 			// Make sure the project file exists, and make sure we're using the correct case.
-			if(ProjectFile != null)
+			if (ProjectFile != null)
 			{
 				FileInfo ProjectFileInfo = FileUtils.FindCorrectCase(ProjectFile.ToFileInfo());
-				if(!ProjectFileInfo.Exists)
+				if (!ProjectFileInfo.Exists)
 				{
 					throw new BuildException("Unable to find project '{0}'.", ProjectFile);
 				}
@@ -357,7 +356,7 @@ namespace UnrealBuildTool
 			}
 
 			// Expand all the platforms, architectures and configurations
-			foreach(UnrealTargetPlatform Platform in Platforms)
+			foreach (UnrealTargetPlatform Platform in Platforms)
 			{
 				// Make sure the platform is valid
 				if (!InstalledPlatformInfo.IsValid(null, Platform, null, EProjectType.Code, InstalledPlatformState.Downloaded))
@@ -379,13 +378,13 @@ namespace UnrealBuildTool
 				UnrealArchitectures? ParamArchitectures = UnrealArchitectures.FromString(ParamArchitectureList, Platform);
 
 
-				foreach(UnrealTargetConfiguration Configuration in Configurations)
+				foreach (UnrealTargetConfiguration Configuration in Configurations)
 				{
 					// Create all the target descriptors for targets specified by type
-					foreach(string TargetTypeString in Arguments.GetValues("-TargetType="))
+					foreach (string TargetTypeString in Arguments.GetValues("-TargetType="))
 					{
 						TargetType TargetType;
-						if(!Enum.TryParse(TargetTypeString, out TargetType))
+						if (!Enum.TryParse(TargetTypeString, out TargetType))
 						{
 							throw new BuildException("Invalid target type '{0}'", TargetTypeString);
 						}
@@ -407,7 +406,7 @@ namespace UnrealBuildTool
 					}
 
 					// Create all the target descriptors
-					foreach(string TargetName in TargetNames)
+					foreach (string TargetName in TargetNames)
 					{
 						FileReference? TargetProjectFile = ProjectFile;
 						UnrealArchitectures Architectures;
@@ -448,7 +447,7 @@ namespace UnrealBuildTool
 							foreach (UnrealArch Architecture in Architectures.Architectures)
 							{
 								TargetDescriptors.Add(new TargetDescriptor(TargetProjectFile, TargetName, Platform, Configuration, new UnrealArchitectures(Architecture), Arguments));
-						}
+							}
 						}
 						else
 						{
@@ -468,15 +467,15 @@ namespace UnrealBuildTool
 		public static bool TryParseProjectFileArgument(CommandLineArguments Arguments, [NotNullWhen(true)] out FileReference? ProjectFile)
 		{
 			FileReference? ExplicitProjectFile;
-			if(Arguments.TryGetValue("-Project=", out ExplicitProjectFile))
+			if (Arguments.TryGetValue("-Project=", out ExplicitProjectFile))
 			{
 				ProjectFile = ExplicitProjectFile;
 				return true;
 			}
 
-			for(int Idx = 0; Idx < Arguments.Count; Idx++)
+			for (int Idx = 0; Idx < Arguments.Count; Idx++)
 			{
-				if(Arguments[Idx][0] != '-' && Arguments[Idx].EndsWith(".uproject", StringComparison.OrdinalIgnoreCase))
+				if (Arguments[Idx][0] != '-' && Arguments[Idx].EndsWith(".uproject", StringComparison.OrdinalIgnoreCase))
 				{
 					Arguments.MarkAsUsed(Idx);
 					ProjectFile = new FileReference(Arguments[Idx]);
@@ -484,7 +483,7 @@ namespace UnrealBuildTool
 				}
 			}
 
-			if(UnrealBuildTool.IsProjectInstalled())
+			if (UnrealBuildTool.IsProjectInstalled())
 			{
 				ProjectFile = UnrealBuildTool.GetInstalledProjectFile()!;
 				return true;
@@ -503,11 +502,11 @@ namespace UnrealBuildTool
 			StringBuilder Result = new StringBuilder();
 			Result.AppendFormat("{0} {1} {2}", Name, Platform, Configuration);
 			Result.AppendFormat($" -Architecture={Architectures}");
-			if(ProjectFile != null)
+			if (ProjectFile != null)
 			{
 				Result.AppendFormat(" -Project={0}", Utils.MakePathSafeToUseWithCommandLine(ProjectFile));
 			}
-			if(AdditionalArguments != null && AdditionalArguments.Count > 0)
+			if (AdditionalArguments != null && AdditionalArguments.Count > 0)
 			{
 				Result.AppendFormat(" {0}", AdditionalArguments);
 			}
@@ -517,15 +516,15 @@ namespace UnrealBuildTool
 		public override int GetHashCode()
 		{
 #pragma warning disable RS1024
-			return String.GetHashCode(ProjectFile?.FullName) + 
-				Name.GetHashCode() + 
-				Platform.GetHashCode() + 
-				Configuration.GetHashCode() + 
+			return String.GetHashCode(ProjectFile?.FullName) +
+				Name.GetHashCode() +
+				Platform.GetHashCode() +
+				Configuration.GetHashCode() +
 				Architectures.GetHashCode();
 #pragma warning restore RE1024
 		}
 
-		public override bool Equals(object? Obj) 
+		public override bool Equals(object? Obj)
 		{
 			TargetDescriptor? Other = Obj as TargetDescriptor;
 			if (Other != null)
diff --git a/Engine/Source/Programs/UnrealBuildTool/Configuration/TargetInfo.cs b/Engine/Source/Programs/UnrealBuildTool/Configuration/TargetInfo.cs
--- a/Engine/Source/Programs/UnrealBuildTool/Configuration/TargetInfo.cs
+++ b/Engine/Source/Programs/UnrealBuildTool/Configuration/TargetInfo.cs
@@ -1,13 +1,8 @@
 // Copyright Epic Games, Inc. All Rights Reserved.
 
 using System;
-using System.Collections.Generic;
 using System.Linq;
-using System.Runtime.Serialization;
-using System.Text;
-using System.Threading.Tasks;
 using EpicGames.Core;
-using UnrealBuildBase;
 
 namespace UnrealBuildTool
 {
diff --git a/Engine/Source/Programs/UnrealBuildTool/Configuration/TargetRules.cs b/Engine/Source/Programs/UnrealBuildTool/Configuration/TargetRules.cs
--- a/Engine/Source/Programs/UnrealBuildTool/Configuration/TargetRules.cs
+++ b/Engine/Source/Programs/UnrealBuildTool/Configuration/TargetRules.cs
@@ -5,10 +5,10 @@ using System.Collections.Generic;
 using System.IO;
 using System.Linq;
 using System.Reflection;
+using System.Runtime.Serialization;
 using EpicGames.Core;
-using UnrealBuildBase;
 using Microsoft.Extensions.Logging;
-using System.Runtime.Serialization;
+using UnrealBuildBase;
 
 namespace UnrealBuildTool
 {
@@ -784,7 +784,7 @@ namespace UnrealBuildTool
 		/// </summary>
 		[RequiresUniqueBuildEnvironment]
 		public bool bUseChaosMemoryTracking = false;
-		
+
 		/// <summary>
 		/// Whether to compile in Chaos Visual Debugger (CVD) support features to record the state of the physics simulation
 		/// </summary>
@@ -1202,7 +1202,7 @@ namespace UnrealBuildTool
 		/// Whether to include PerfCounters support.
 		/// </summary>
 		[RequiresUniqueBuildEnvironment]
-        public bool bWithPerfCounters
+		public bool bWithPerfCounters
 		{
 			get => bWithPerfCountersOverride ?? (Type == TargetType.Editor || Type == TargetType.Server);
 			set => bWithPerfCountersOverride = value;
@@ -1508,7 +1508,7 @@ namespace UnrealBuildTool
 		[RequiresUniqueBuildEnvironment]
 		public WarningLevel DefaultWarningLevel
 		{
-			get => (DefaultWarningLevelPrivate == WarningLevel.Default)? (bWarningsAsErrors ? WarningLevel.Error : WarningLevel.Warning) : DefaultWarningLevelPrivate;
+			get => (DefaultWarningLevelPrivate == WarningLevel.Default) ? (bWarningsAsErrors ? WarningLevel.Error : WarningLevel.Warning) : DefaultWarningLevelPrivate;
 			set => DefaultWarningLevelPrivate = value;
 		}
 
@@ -1521,7 +1521,7 @@ namespace UnrealBuildTool
 		/// </summary>
 		public WarningLevel DeprecationWarningLevel
 		{
-			get => (DeprecationWarningLevelPrivate == WarningLevel.Default)? DefaultWarningLevel : DeprecationWarningLevelPrivate;
+			get => (DeprecationWarningLevelPrivate == WarningLevel.Default) ? DefaultWarningLevel : DeprecationWarningLevelPrivate;
 			set => DeprecationWarningLevelPrivate = value;
 		}
 
@@ -1600,9 +1600,9 @@ namespace UnrealBuildTool
 		/// <summary>
 		/// Forces frame pointers to be retained this is usually required when you want reliable callstacks e.g. mallocframeprofiler
 		/// </summary>
-		[XmlConfigFile(Category = "BuildConfiguration", Name="bRetainFramePointers")]
-		[CommandLine("-RetainFramePointers", Value="true")]
-		[CommandLine("-NoRetainFramePointers", Value="false")]
+		[XmlConfigFile(Category = "BuildConfiguration", Name = "bRetainFramePointers")]
+		[CommandLine("-RetainFramePointers", Value = "true")]
+		[CommandLine("-NoRetainFramePointers", Value = "false")]
 		public bool? bRetainFramePointersOverride = null;
 
 		/// <summary>
@@ -1942,7 +1942,7 @@ namespace UnrealBuildTool
 		/// Whether to force skipping deployment for platforms that require deployment by default.
 		/// </summary>
 		[CommandLine("-SkipDeploy")]
-		private bool bForceSkipDeploy = false; 
+		private bool bForceSkipDeploy = false;
 
 		/// <summary>
 		/// When enabled, allows XGE to compile pre-compiled header files on remote machines.  Otherwise, PCHs are always generated locally.
@@ -2114,8 +2114,8 @@ namespace UnrealBuildTool
 		/// Backing storage for the LinkType property.
 		/// </summary>
 		[RequiresUniqueBuildEnvironment]
-		[CommandLine("-Monolithic", Value ="Monolithic")]
-		[CommandLine("-Modular", Value ="Modular")]
+		[CommandLine("-Monolithic", Value = "Monolithic")]
+		[CommandLine("-Modular", Value = "Modular")]
 		TargetLinkType LinkTypePrivate = TargetLinkType.Default;
 
 		/// <summary>
@@ -2415,9 +2415,9 @@ namespace UnrealBuildTool
 			bDeployAfterCompile = bForceSkipDeploy ? false : bDeployAfterCompile;
 
 			// Set the default build version
-			if(String.IsNullOrEmpty(BuildVersion))
+			if (String.IsNullOrEmpty(BuildVersion))
 			{
-				if(String.IsNullOrEmpty(Target.Version.BuildVersionString))
+				if (String.IsNullOrEmpty(Target.Version.BuildVersionString))
 				{
 					BuildVersion = String.Format("{0}-CL-{1}", Target.Version.BranchName, Target.Version.Changelist);
 				}
@@ -2472,7 +2472,7 @@ namespace UnrealBuildTool
 		/// </summary>
 		internal void SetOverridesForTargetType()
 		{
-			if(Type == global::UnrealBuildTool.TargetType.Game)
+			if (Type == global::UnrealBuildTool.TargetType.Game)
 			{
 				GlobalDefinitions.Add("UE_GAME=1");
 			}
@@ -2622,12 +2622,12 @@ namespace UnrealBuildTool
 		{
 			yield return this;
 
-			foreach(FieldInfo Field in GetType().GetFields(BindingFlags.Public | BindingFlags.Instance))
+			foreach (FieldInfo Field in GetType().GetFields(BindingFlags.Public | BindingFlags.Instance))
 			{
-				if(Field.GetCustomAttribute<ConfigSubObjectAttribute>() != null)
+				if (Field.GetCustomAttribute<ConfigSubObjectAttribute>() != null)
 				{
 					object? Value = Field.GetValue(this);
-					if(Value != null)
+					if (Value != null)
 					{
 						yield return Value;
 					}
@@ -2681,7 +2681,7 @@ namespace UnrealBuildTool
 				}
 			}
 
-			if (DefaultBuildSettings < LatestVersion) 
+			if (DefaultBuildSettings < LatestVersion)
 			{
 				Diagnostics.Add("[Upgrade]");
 				Diagnostics.Add("[Upgrade] Using backward-compatible build settings. The latest version of UE sets the following values by default, which may require code changes:");
diff --git a/Engine/Source/Programs/UnrealBuildTool/Configuration/TestModuleRules.cs b/Engine/Source/Programs/UnrealBuildTool/Configuration/TestModuleRules.cs
--- a/Engine/Source/Programs/UnrealBuildTool/Configuration/TestModuleRules.cs
+++ b/Engine/Source/Programs/UnrealBuildTool/Configuration/TestModuleRules.cs
@@ -1,12 +1,12 @@
 // Copyright Epic Games, Inc. All Rights Reserved.
 
+using System;
 using System.Collections.Generic;
 using System.IO;
 using System.Linq;
 using System.Xml.Linq;
-using UnrealBuildBase;
 using EpicGames.Core;
-using System;
+using UnrealBuildBase;
 
 namespace UnrealBuildTool
 {
@@ -215,7 +215,7 @@ namespace UnrealBuildTool
 				if (TestNames != null)
 				{
 					XElement lastUpdatedNode = TestNames;
-						
+
 					InsertOrUpdateTestFlagProperty(ref lastUpdatedNode, TestMetadata.TestName, "Disabled", Convert.ToString(TestMetadata.Disabled));
 					InsertOrUpdateTestFlagProperty(ref lastUpdatedNode, TestMetadata.TestName, "Short", Convert.ToString(TestMetadata.TestShortName));
 					InsertOrUpdateTestFlagProperty(ref lastUpdatedNode, TestMetadata.TestName, "Target", Convert.ToString(TestTargetName));
@@ -293,7 +293,7 @@ namespace UnrealBuildTool
 					}
 					using (FileStream FileStream = System.IO.File.Create(GeneratedPropertiesPlatformFile))
 					{
-						 new XDocument(new XElement(BuildGraphNamespace + "BuildGraph", new XAttribute(XNamespace.Xmlns + "xsi", SchemaInstance), new XAttribute(SchemaInstance + "schemaLocation", SchemaLocation))).Save(FileStream);
+						new XDocument(new XElement(BuildGraphNamespace + "BuildGraph", new XAttribute(XNamespace.Xmlns + "xsi", SchemaInstance), new XAttribute(SchemaInstance + "schemaLocation", SchemaLocation))).Save(FileStream);
 					}
 				}
 
diff --git a/Engine/Source/Programs/UnrealBuildTool/Configuration/UEBuildBinary.cs b/Engine/Source/Programs/UnrealBuildTool/Configuration/UEBuildBinary.cs
--- a/Engine/Source/Programs/UnrealBuildTool/Configuration/UEBuildBinary.cs
+++ b/Engine/Source/Programs/UnrealBuildTool/Configuration/UEBuildBinary.cs
@@ -2,15 +2,11 @@
 
 using System;
 using System.Collections.Generic;
-using System.Text;
-using System.Diagnostics;
 using System.IO;
 using System.Linq;
-using System.Xml;
-using System.Globalization;
 using EpicGames.Core;
-using UnrealBuildBase;
 using Microsoft.Extensions.Logging;
+using UnrealBuildBase;
 
 namespace UnrealBuildTool
 {
@@ -149,7 +145,7 @@ namespace UnrealBuildTool
 			this.bBuildAdditionalConsoleApp = bBuildAdditionalConsoleApp;
 			this.PrimaryModule = PrimaryModule;
 			this.bUsePrecompiled = bUsePrecompiled;
-			
+
 			Modules.Add(PrimaryModule);
 		}
 
@@ -221,7 +217,7 @@ namespace UnrealBuildTool
 			OutputFiles.AddRange(Executables);
 
 			// Save all the output items for this binary. This is used for hot-reload, and excludes any items added in PostBuild (such as additional files copied into the app).
-			if(Target.LinkType == TargetLinkType.Modular)
+			if (Target.LinkType == TargetLinkType.Modular)
 			{
 				Graph.SetOutputItemsForModule(PrimaryModule.Name, OutputFiles.ToArray());
 			}
@@ -232,7 +228,7 @@ namespace UnrealBuildTool
 				// Produce additional binary but link it as a console app
 				LinkEnvironment ConsoleAppLinkEvironment = new LinkEnvironment(BinaryLinkEnvironment);
 				ConsoleAppLinkEvironment.bIsBuildingConsoleApplication = true;
-				ConsoleAppLinkEvironment.WindowsEntryPointOverride = "WinMainCRTStartup";		// For WinMain() instead of "main()" for Launch module
+				ConsoleAppLinkEvironment.WindowsEntryPointOverride = "WinMainCRTStartup";       // For WinMain() instead of "main()" for Launch module
 				ConsoleAppLinkEvironment.OutputFilePaths = ConsoleAppLinkEvironment.OutputFilePaths.Select(Path => GetAdditionalConsoleAppPath(Path)).ToList();
 
 				// Link the console app executable
@@ -261,11 +257,11 @@ namespace UnrealBuildTool
 		/// <param name="ExeDir">Output directory for the executable</param>
 		public void PrepareRuntimeDependencies(List<RuntimeDependency> RuntimeDependencies, Dictionary<FileReference, FileReference> TargetFileToSourceFile, DirectoryReference ExeDir)
 		{
-			foreach(UEBuildModule Module in Modules)
+			foreach (UEBuildModule Module in Modules)
 			{
 				foreach (ModuleRules.RuntimeDependency Dependency in Module.Rules.RuntimeDependencies.Inner)
 				{
-					if(Dependency.SourcePath == null)
+					if (Dependency.SourcePath == null)
 					{
 						// Expand the target path
 						string ExpandedPath = Module.ExpandPathVariables(Dependency.Path, OutputDir, ExeDir);
@@ -296,22 +292,22 @@ namespace UnrealBuildTool
 						{
 							Mapping = FilePattern.CreateMapping(null, ref SourcePattern, ref TargetPattern);
 						}
-						catch(FilePatternException Ex)
+						catch (FilePatternException Ex)
 						{
 							ExceptionUtils.AddContext(Ex, "while creating runtime dependencies for module '{0}'", Module.Name);
 							throw;
 						}
 
 						// Add actions to copy everything
-						foreach(KeyValuePair<FileReference, FileReference> Pair in Mapping)
+						foreach (KeyValuePair<FileReference, FileReference> Pair in Mapping)
 						{
 							FileReference? ExistingSourceFile;
-							if(!TargetFileToSourceFile.TryGetValue(Pair.Key, out ExistingSourceFile))
+							if (!TargetFileToSourceFile.TryGetValue(Pair.Key, out ExistingSourceFile))
 							{
 								TargetFileToSourceFile[Pair.Key] = Pair.Value;
 								RuntimeDependencies.Add(new RuntimeDependency(Pair.Key, Dependency.Type));
 							}
-							else if(ExistingSourceFile != Pair.Value)
+							else if (ExistingSourceFile != Pair.Value)
 							{
 								throw new BuildException("Runtime dependency '{0}' is configured to be staged from '{1}' and '{2}'", Pair.Key, Pair.Value, ExistingSourceFile);
 							}
@@ -358,7 +354,7 @@ namespace UnrealBuildTool
 			List<UEBuildModule> GameModules = new List<UEBuildModule>();
 			foreach (UEBuildModule Module in Modules)
 			{
-				if(Module.Rules.Context.bCanHotReload)
+				if (Module.Rules.Context.bCanHotReload)
 				{
 					GameModules.Add(Module);
 				}
@@ -398,7 +394,7 @@ namespace UnrealBuildTool
 		public void FindModuleReferences(Dictionary<UEBuildModule, UEBuildModule?> ReferencedBy)
 		{
 			List<UEBuildModule> ReferencedModules = new List<UEBuildModule>();
-			foreach(UEBuildModule Module in Modules)
+			foreach (UEBuildModule Module in Modules)
 			{
 				ReferencedModules.Add(Module);
 				ReferencedBy.Add(Module, null);
@@ -406,7 +402,7 @@ namespace UnrealBuildTool
 
 			List<UEBuildModule> DirectlyReferencedModules = new List<UEBuildModule>();
 			HashSet<UEBuildModule> VisitedModules = new HashSet<UEBuildModule>();
-			for(int Idx = 0; Idx < ReferencedModules.Count; Idx++)
+			for (int Idx = 0; Idx < ReferencedModules.Count; Idx++)
 			{
 				UEBuildModule SourceModule = ReferencedModules[Idx];
 
@@ -415,9 +411,9 @@ namespace UnrealBuildTool
 				SourceModule.GetAllDependencyModules(DirectlyReferencedModules, VisitedModules, false, false, true);
 
 				// Set up the references for all the new modules
-				foreach(UEBuildModule DirectlyReferencedModule in DirectlyReferencedModules)
+				foreach (UEBuildModule DirectlyReferencedModule in DirectlyReferencedModules)
 				{
-					if(!ReferencedBy.ContainsKey(DirectlyReferencedModule))
+					if (!ReferencedBy.ContainsKey(DirectlyReferencedModule))
 					{
 						ReferencedBy.Add(DirectlyReferencedModule, SourceModule);
 						ReferencedModules.Add(DirectlyReferencedModule);
@@ -457,22 +453,22 @@ namespace UnrealBuildTool
 		public void GetBuildProducts(ReadOnlyTargetRules Target, UEToolChain ToolChain, Dictionary<FileReference, BuildProductType> BuildProducts, bool bCreateDebugInfo)
 		{
 			// Add all the precompiled outputs
-			foreach(UEBuildModuleCPP Module in Modules.OfType<UEBuildModuleCPP>())
+			foreach (UEBuildModuleCPP Module in Modules.OfType<UEBuildModuleCPP>())
 			{
-				if(Module.Rules.bPrecompile)
+				if (Module.Rules.bPrecompile)
 				{
-					if(Module.GeneratedCodeDirectory != null && DirectoryReference.Exists(Module.GeneratedCodeDirectory))
+					if (Module.GeneratedCodeDirectory != null && DirectoryReference.Exists(Module.GeneratedCodeDirectory))
 					{
-						foreach(FileReference GeneratedCodeFile in DirectoryReference.EnumerateFiles(Module.GeneratedCodeDirectory, "*", SearchOption.AllDirectories))
+						foreach (FileReference GeneratedCodeFile in DirectoryReference.EnumerateFiles(Module.GeneratedCodeDirectory, "*", SearchOption.AllDirectories))
 						{
 							// Exclude timestamp files, since they're always updated and cause collisions between builds
-							if(!GeneratedCodeFile.GetFileName().Equals("Timestamp", StringComparison.OrdinalIgnoreCase) && !GeneratedCodeFile.HasExtension(".cpp"))
+							if (!GeneratedCodeFile.GetFileName().Equals("Timestamp", StringComparison.OrdinalIgnoreCase) && !GeneratedCodeFile.HasExtension(".cpp"))
 							{
 								BuildProducts.Add(GeneratedCodeFile, BuildProductType.BuildResource);
 							}
 						}
 					}
-					if(Target.LinkType == TargetLinkType.Monolithic)
+					if (Target.LinkType == TargetLinkType.Monolithic)
 					{
 						FileReference PrecompiledManifestLocation = Module.PrecompiledManifestLocation;
 						BuildProducts.Add(PrecompiledManifestLocation, BuildProductType.BuildResource);
@@ -493,7 +489,7 @@ namespace UnrealBuildTool
 			}
 
 			// Add all the binary outputs
-			if(!Target.bDisableLinking)
+			if (!Target.bDisableLinking)
 			{
 				// Get the type of build products we're creating
 				BuildProductType OutputType = BuildProductType.RequiredResource;
@@ -568,7 +564,7 @@ namespace UnrealBuildTool
 		/// <param name="BundleResources">List of bundle resources required by this module</param>
 		public void GatherAdditionalResources(List<string> Libraries, List<UEBuildBundleResource> BundleResources)
 		{
-			foreach(UEBuildModule Module in Modules)
+			foreach (UEBuildModule Module in Modules)
 			{
 				Module.GatherAdditionalResources(Libraries, BundleResources);
 			}
@@ -620,7 +616,7 @@ namespace UnrealBuildTool
 					string? Alias;
 					if (PrimaryModule.AliasRestrictedFolders.TryGetValue(BinaryFolder.ToString(), out Alias))
 					{
-						foreach(RestrictedFolder Folder in RestrictedFolder.GetValues())
+						foreach (RestrictedFolder Folder in RestrictedFolder.GetValues())
 						{
 							if (Folder.ToString().Equals(Alias))
 							{
@@ -632,7 +628,7 @@ namespace UnrealBuildTool
 				BinaryFolders.AddRange(AliasedBinaryFolders);
 
 				// Check all the dependent modules
-				foreach(UEBuildModule Module in ModuleReferencedBy.Keys)
+				foreach (UEBuildModule Module in ModuleReferencedBy.Keys)
 				{
 					// Find the restricted folders for this module
 					Dictionary<RestrictedFolder, DirectoryReference>? ModuleRestrictedFolders;
@@ -643,12 +639,12 @@ namespace UnrealBuildTool
 					}
 
 					// Write errors for any missing paths in the output files
-					foreach(KeyValuePair<RestrictedFolder, DirectoryReference> Pair in ModuleRestrictedFolders)
+					foreach (KeyValuePair<RestrictedFolder, DirectoryReference> Pair in ModuleRestrictedFolders)
 					{
-						if(!BinaryFolders.Contains(Pair.Key))
+						if (!BinaryFolders.Contains(Pair.Key))
 						{
 							List<string> ReferenceChain = new List<string>();
-							for(UEBuildModule? ReferencedModule = Module; ReferencedModule != null; ReferencedModule = ModuleReferencedBy[ReferencedModule])
+							for (UEBuildModule? ReferencedModule = Module; ReferencedModule != null; ReferencedModule = ModuleReferencedBy[ReferencedModule])
 							{
 								ReferenceChain.Insert(0, ReferencedModule.Name);
 							}
@@ -671,7 +667,7 @@ namespace UnrealBuildTool
 			Writer.WriteValue("Type", Type.ToString());
 
 			Writer.WriteArrayStart("Modules");
-			foreach(UEBuildModule Module in Modules)
+			foreach (UEBuildModule Module in Modules)
 			{
 				Writer.WriteValue(Module.Name);
 			}
@@ -706,7 +702,7 @@ namespace UnrealBuildTool
 
 			CppCompileEnvironment BinaryCompileEnvironment = CreateBinaryCompileEnvironment(CompileEnvironment);
 
-			if(BinaryCompileEnvironment.bUseSharedBuildEnvironment && Target.ProjectFile != null && IntermediateDirectory.IsUnderDirectory(Target.ProjectFile.Directory))
+			if (BinaryCompileEnvironment.bUseSharedBuildEnvironment && Target.ProjectFile != null && IntermediateDirectory.IsUnderDirectory(Target.ProjectFile.Directory))
 			{
 				BinaryCompileEnvironment.bUseSharedBuildEnvironment = false;
 			}
@@ -737,7 +733,7 @@ namespace UnrealBuildTool
 					}
 
 					// Force a reference to initialize module for this binary
-					if(Module.Rules.bRequiresImplementModule ?? true)
+					if (Module.Rules.bRequiresImplementModule ?? true)
 					{
 						BinaryLinkEnvironment.IncludeFunctions.Add(String.Format("IMPLEMENT_MODULE_{0}", Module.Name));
 					}
@@ -784,13 +780,13 @@ namespace UnrealBuildTool
 			BinaryLinkEnvironment.bIsBuildingLibrary = IsBuildingLibrary(Type);
 
 			// If we don't have any resource file, use the default or compile a custom one for this module
-			if(BinaryLinkEnvironment.Platform.IsInGroup(UnrealPlatformGroup.Windows))
+			if (BinaryLinkEnvironment.Platform.IsInGroup(UnrealPlatformGroup.Windows))
 			{
 				// Figure out if this binary has any custom resource files. Hacky check to ignore the resource file in the Launch module, since it contains dialogs that the engine needs and always needs to be included.
 				FileItem[] CustomResourceFiles = BinaryLinkEnvironment.InputFiles.Where(x => x.Location.HasExtension(".res") && !x.Location.FullName.EndsWith("\\Launch\\PCLaunch.rc.res", StringComparison.OrdinalIgnoreCase)).ToArray();
-				if(CustomResourceFiles.Length == 0)
+				if (CustomResourceFiles.Length == 0)
 				{
-					if(BinaryLinkEnvironment.DefaultResourceFiles.Count > 0)
+					if (BinaryLinkEnvironment.DefaultResourceFiles.Count > 0)
 					{
 						// Use the default resource file if possible
 						BinaryLinkEnvironment.InputFiles.AddRange(BinaryLinkEnvironment.DefaultResourceFiles);
@@ -823,7 +819,7 @@ namespace UnrealBuildTool
 			BinaryLinkEnvironment.InputFiles.AddRange(BinaryLinkEnvironment.CommonResourceFiles);
 
 			// Allow the platform to modify the binary link environment for platform-specific resources etc.
-			UEBuildPlatform.GetBuildPlatform(Target.Platform).ModifyBinaryLinkEnvironment( BinaryLinkEnvironment, BinaryCompileEnvironment, Target, ToolChain, Graph);
+			UEBuildPlatform.GetBuildPlatform(Target.Platform).ModifyBinaryLinkEnvironment(BinaryLinkEnvironment, BinaryCompileEnvironment, Target, ToolChain, Graph);
 
 
 			return BinaryLinkEnvironment;
diff --git a/Engine/Source/Programs/UnrealBuildTool/Configuration/UEBuildBundleResource.cs b/Engine/Source/Programs/UnrealBuildTool/Configuration/UEBuildBundleResource.cs
--- a/Engine/Source/Programs/UnrealBuildTool/Configuration/UEBuildBundleResource.cs
+++ b/Engine/Source/Programs/UnrealBuildTool/Configuration/UEBuildBundleResource.cs
@@ -1,11 +1,5 @@
 // Copyright Epic Games, Inc. All Rights Reserved.
 
-using System;
-using System.Collections.Generic;
-using System.Linq;
-using System.Text;
-using System.Threading.Tasks;
-
 namespace UnrealBuildTool
 {
 	/// <summary>
diff --git a/Engine/Source/Programs/UnrealBuildTool/Configuration/UEBuildDeploy.cs b/Engine/Source/Programs/UnrealBuildTool/Configuration/UEBuildDeploy.cs
--- a/Engine/Source/Programs/UnrealBuildTool/Configuration/UEBuildDeploy.cs
+++ b/Engine/Source/Programs/UnrealBuildTool/Configuration/UEBuildDeploy.cs
@@ -1,9 +1,7 @@
 // Copyright Epic Games, Inc. All Rights Reserved.
 
-using EpicGames.Core;
-using System.Linq;
-using System.IO;
 using System;
+using System.IO;
 using Microsoft.Extensions.Logging;
 
 namespace UnrealBuildTool
@@ -46,7 +44,8 @@ namespace UnrealBuildTool
 					}
 
 					FromFile.CopyTo(ToDestFile, true);
-				} catch (Exception ex)
+				}
+				catch (Exception ex)
 				{
 					Logger.LogError(ex, "Error while copying {FromFile} from {FromDir} to {ToDir}: {Ex}", FromFile, FromDir, ToDir, ex);
 					return false;
diff --git a/Engine/Source/Programs/UnrealBuildTool/Configuration/UEBuildFramework.cs b/Engine/Source/Programs/UnrealBuildTool/Configuration/UEBuildFramework.cs
--- a/Engine/Source/Programs/UnrealBuildTool/Configuration/UEBuildFramework.cs
+++ b/Engine/Source/Programs/UnrealBuildTool/Configuration/UEBuildFramework.cs
@@ -2,13 +2,10 @@
 
 using System;
 using System.Collections.Generic;
-using System.Linq;
-using System.Text;
 using System.Xml;
-using System.Threading.Tasks;
 using EpicGames.Core;
-using UnrealBuildBase;
 using Microsoft.Extensions.Logging;
+using UnrealBuildBase;
 
 namespace UnrealBuildTool
 {
@@ -113,7 +110,7 @@ namespace UnrealBuildTool
 			this.FrameworkDirectory = FrameworkDirectory;
 			this.CopyBundledAssets = CopyBundledAssets;
 			this.bCopyFramework = bCopyFramework;
-			if(this.FrameworkDirectory.FullName.EndsWith(".xcframework"))
+			if (this.FrameworkDirectory.FullName.EndsWith(".xcframework"))
 			{
 				this.XCFrameworkVariants = LoadXCFrameworkVariants(Logger);
 			}
@@ -127,10 +124,10 @@ namespace UnrealBuildTool
 		/// <param name="Logger"></param>
 		public DirectoryReference? GetFrameworkDirectory(UnrealTargetPlatform? Platform, UnrealArch? Architecture, ILogger Logger)
 		{
-			if( XCFrameworkVariants != null && Platform != null && Architecture != null)
+			if (XCFrameworkVariants != null && Platform != null && Architecture != null)
 			{
 				XCFrameworkVariantEntry? FrameworkVariant = XCFrameworkVariants.Find(x => x.Matches(Platform.Value, Architecture.Value));
-				if( FrameworkVariant != null )
+				if (FrameworkVariant != null)
 				{
 					return DirectoryReference.Combine(FrameworkDirectory!, FrameworkVariant.LibraryIdentifier);
 				}
@@ -145,29 +142,29 @@ namespace UnrealBuildTool
 		List<XCFrameworkVariantEntry>? LoadXCFrameworkVariants(ILogger Logger)
 		{
 			XmlDocument PlistDoc = new XmlDocument();
-			PlistDoc.Load(FileReference.Combine(FrameworkDirectory!, "Info.plist" ).FullName);
+			PlistDoc.Load(FileReference.Combine(FrameworkDirectory!, "Info.plist").FullName);
 
 			// Check the plist type
 			XmlNode? CFBundlePackageType = PlistDoc.SelectSingleNode("/plist/dict[key='CFBundlePackageType']/string[1]");
-			if( CFBundlePackageType == null || CFBundlePackageType.NodeType != XmlNodeType.Element || CFBundlePackageType.InnerText != "XFWK" )
+			if (CFBundlePackageType == null || CFBundlePackageType.NodeType != XmlNodeType.Element || CFBundlePackageType.InnerText != "XFWK")
 			{
 				Logger.LogWarning("CFBundlePackageType is not set to XFWK in Info.plist for XCFramework {Name}", Name);
 				return null;
 			}
 
 			// Load Framework variants data from dictionary nodes
-			XmlNodeList? FrameworkVariantDicts = PlistDoc.SelectNodes("/plist/dict[key='AvailableLibraries']/array/dict");			
-			if( FrameworkVariantDicts == null )
+			XmlNodeList? FrameworkVariantDicts = PlistDoc.SelectNodes("/plist/dict[key='AvailableLibraries']/array/dict");
+			if (FrameworkVariantDicts == null)
 			{
 				Logger.LogWarning("Invalid Info.plist file for XCFramework {Name}. It will be used as a regular Framework", Name);
 				return null;
 			}
 
 			List<XCFrameworkVariantEntry> Variants = new List<XCFrameworkVariantEntry>();
-			foreach( XmlNode VariantDict in FrameworkVariantDicts )
+			foreach (XmlNode VariantDict in FrameworkVariantDicts)
 			{
 				XCFrameworkVariantEntry? Variant = XCFrameworkVariantEntry.Parse(VariantDict, Logger);
-				if( Variant == null )
+				if (Variant == null)
 				{
 					Logger.LogWarning("Failed to load variant from Info.plist file for XCFramework {Name}", Name);
 				}
@@ -236,18 +233,18 @@ namespace UnrealBuildTool
 			/// <param name="Architecture"></param>
 			public bool Matches(UnrealTargetPlatform Platform, UnrealArch Architecture)
 			{
-				if(( Platform == UnrealTargetPlatform.IOS && SupportedPlatform == "ios" ) ||
-				   ( Platform == UnrealTargetPlatform.Mac && SupportedPlatform == "macos" ) ||
-				   ( Platform == UnrealTargetPlatform.TVOS && SupportedPlatform == "tvos" ))
+				if ((Platform == UnrealTargetPlatform.IOS && SupportedPlatform == "ios") ||
+				   (Platform == UnrealTargetPlatform.Mac && SupportedPlatform == "macos") ||
+				   (Platform == UnrealTargetPlatform.TVOS && SupportedPlatform == "tvos"))
 				{
-					if( Architecture == UnrealArch.IOSSimulator || Architecture == UnrealArch.TVOSSimulator)
+					if (Architecture == UnrealArch.IOSSimulator || Architecture == UnrealArch.TVOSSimulator)
 					{
 						// When using -simulator we don't have the actual architecture. Assume arm64 as other parts of UBT already are doing
-						return ( SupportedPlatformVariant == "simulator" ) && SupportedArchitectures.Contains("arm64");
+						return (SupportedPlatformVariant == "simulator") && SupportedArchitectures.Contains("arm64");
 					}
 					else
 					{
-						return ( SupportedPlatformVariant == null ) && SupportedArchitectures.Contains(Architecture.AppleName);
+						return (SupportedPlatformVariant == null) && SupportedArchitectures.Contains(Architecture.AppleName);
 					}
 				}
 				return false;
@@ -261,7 +258,7 @@ namespace UnrealBuildTool
 			public static XCFrameworkVariantEntry? Parse(XmlNode DictNode, ILogger Logger)
 			{
 				// Entries from a <dict> in a plist file always contains a <key> node and a value node
-				if(( DictNode.ChildNodes.Count % 2 ) != 0)
+				if ((DictNode.ChildNodes.Count % 2) != 0)
 				{
 					Logger.LogDebug("Invalid dict in Info.plist file for XCFramework");
 					return null;
@@ -273,14 +270,14 @@ namespace UnrealBuildTool
 				string? SupportedPlatform = null;
 				string? SupportedPlatformVariant = null;
 
-				for( int i = 0; i < DictNode.ChildNodes.Count; i += 2 )
-            	{
-                	XmlNode? Key = DictNode.ChildNodes[i];
-                	XmlNode? Value = DictNode.ChildNodes[i + 1];
+				for (int i = 0; i < DictNode.ChildNodes.Count; i += 2)
+				{
+					XmlNode? Key = DictNode.ChildNodes[i];
+					XmlNode? Value = DictNode.ChildNodes[i + 1];
 
-					if( Value != null && Key != null && Key.Name == "key" )
+					if (Value != null && Key != null && Key.Name == "key")
 					{
-						switch( Key.InnerText )
+						switch (Key.InnerText)
 						{
 							case "LibraryIdentifier":
 								LibraryIdentifier = ParseString(Value, Logger);
@@ -304,12 +301,12 @@ namespace UnrealBuildTool
 				}
 
 				// All fields but SupportedPlatformVariant are allowed to be present in the framework variant descriprion
-				if( LibraryIdentifier == null || LibraryPath == null || SupportedArchitectures == null || SupportedPlatform == null )
+				if (LibraryIdentifier == null || LibraryPath == null || SupportedArchitectures == null || SupportedPlatform == null)
 				{
 					Logger.LogDebug("Missing data in Info.plist file for XCFramework");
 					return null;
 				}
-				return new XCFrameworkVariantEntry( LibraryIdentifier, LibraryPath, SupportedArchitectures, SupportedPlatform, SupportedPlatformVariant ); 
+				return new XCFrameworkVariantEntry(LibraryIdentifier, LibraryPath, SupportedArchitectures, SupportedPlatform, SupportedPlatformVariant);
 			}
 
 			/// <summary>
@@ -317,9 +314,9 @@ namespace UnrealBuildTool
 			/// </summary>
 			/// <param name="ValueNode">XmlNode we expect to contain a &lt;string&gt; value</param>			
 			/// <param name="Logger">Logger for diagnostic output</param>
-			static string? ParseString( XmlNode ValueNode, ILogger Logger )
+			static string? ParseString(XmlNode ValueNode, ILogger Logger)
 			{
-				if( ValueNode.Name != "string" )
+				if (ValueNode.Name != "string")
 				{
 					Logger.LogDebug("Unexpected tag \"{Tag}\" while expecting \"string\" in Info.plist file for XCFramework", ValueNode.Name);
 					return null;
@@ -332,21 +329,21 @@ namespace UnrealBuildTool
 			/// </summary>
 			/// <param name="ValueNode">XmlNode we expect to contain a &lt;array&gt; value with several &lt;string&gt; entries</param>			
 			/// <param name="Logger">Logger for diagnostic output</param>
-			static List<string>? ParseListOfStrings( XmlNode ValueNode, ILogger Logger )
+			static List<string>? ParseListOfStrings(XmlNode ValueNode, ILogger Logger)
 			{
-				if( ValueNode.Name != "array" )
+				if (ValueNode.Name != "array")
 				{
 					Logger.LogDebug("Unexpected tag \"{Name}\" while expecting \"array\" in Info.plist file for XCFramework", ValueNode.Name);
 					return null;
 				}
 
 				List<string> ListOfStrings = new List<string>();
-				foreach( XmlNode ChildNode in ValueNode.ChildNodes )
+				foreach (XmlNode ChildNode in ValueNode.ChildNodes)
 				{
-					string? ParsedString = ParseString( ChildNode, Logger );
-					if( ParsedString != null )
+					string? ParsedString = ParseString(ChildNode, Logger);
+					if (ParsedString != null)
 					{
-						ListOfStrings.Add( ParsedString );
+						ListOfStrings.Add(ParsedString);
 					}
 				}
 				return ListOfStrings;
diff --git a/Engine/Source/Programs/UnrealBuildTool/Configuration/UEBuildModule.cs b/Engine/Source/Programs/UnrealBuildTool/Configuration/UEBuildModule.cs
--- a/Engine/Source/Programs/UnrealBuildTool/Configuration/UEBuildModule.cs
+++ b/Engine/Source/Programs/UnrealBuildTool/Configuration/UEBuildModule.cs
@@ -2,16 +2,13 @@
 
 using System;
 using System.Collections.Generic;
-using System.Text;
-using System.Diagnostics;
 using System.IO;
 using System.Linq;
-using System.Text.RegularExpressions;
-using System.Xml;
+using System.Text;
+using System.Threading.Tasks;
 using EpicGames.Core;
-using UnrealBuildBase;
 using Microsoft.Extensions.Logging;
-using System.Threading.Tasks;
+using UnrealBuildBase;
 
 namespace UnrealBuildTool
 {
@@ -246,9 +243,9 @@ namespace UnrealBuildTool
 			}
 
 			PublicAdditionalFrameworks = new HashSet<UEBuildFramework>();
-			if(Rules.PublicAdditionalFrameworks != null)
+			if (Rules.PublicAdditionalFrameworks != null)
 			{
-				foreach(ModuleRules.Framework FrameworkRules in Rules.PublicAdditionalFrameworks)
+				foreach (ModuleRules.Framework FrameworkRules in Rules.PublicAdditionalFrameworks)
 				{
 					bool bLinkFramework = FrameworkRules.Mode == ModuleRules.Framework.FrameworkMode.Link || FrameworkRules.Mode == ModuleRules.Framework.FrameworkMode.LinkAndCopy;
 					bool bCopyFramework = FrameworkRules.Mode == ModuleRules.Framework.FrameworkMode.Copy || FrameworkRules.Mode == ModuleRules.Framework.FrameworkMode.LinkAndCopy;
@@ -261,16 +258,16 @@ namespace UnrealBuildTool
 							Name: FrameworkRules.Name,
 							ZipFile: FileReference.Combine(ModuleDirectory, FrameworkRules.Path),
 							OutputDirectory: DirectoryReference.Combine(Unreal.EngineDirectory, "Intermediate", "UnzippedFrameworks", FrameworkRules.Name, Path.GetFileNameWithoutExtension(FrameworkRules.Path)),
-							CopyBundledAssets: FrameworkRules.CopyBundledAssets, 
+							CopyBundledAssets: FrameworkRules.CopyBundledAssets,
 							bLinkFramework, bCopyFramework);
 					}
 					else
 					{
 						// Framework on disk
 						Framework = new UEBuildFramework(
-							Name: FrameworkRules.Name, 
-							FrameworkDirectory: DirectoryReference.Combine(ModuleDirectory, FrameworkRules.Path), 
-							CopyBundledAssets: FrameworkRules.CopyBundledAssets, 
+							Name: FrameworkRules.Name,
+							FrameworkDirectory: DirectoryReference.Combine(ModuleDirectory, FrameworkRules.Path),
+							CopyBundledAssets: FrameworkRules.CopyBundledAssets,
 							bLinkFramework, bCopyFramework, Logger);
 					}
 					PublicAdditionalFrameworks.Add(Framework);
@@ -279,7 +276,7 @@ namespace UnrealBuildTool
 
 			PublicAdditionalBundleResources = Rules.AdditionalBundleResources == null ? new HashSet<UEBuildBundleResource>() : new HashSet<UEBuildBundleResource>(Rules.AdditionalBundleResources.Select(x => new UEBuildBundleResource(x)));
 			PublicDelayLoadDLLs = HashSetFromOptionalEnumerableStringParameter(Rules.PublicDelayLoadDLLs);
-			if(Rules.bUsePrecompiled)
+			if (Rules.bUsePrecompiled)
 			{
 				PrivateIncludePaths = new HashSet<DirectoryReference>();
 			}
@@ -341,26 +338,26 @@ namespace UnrealBuildTool
 			HashSet<UEBuildModule> Modules = new HashSet<UEBuildModule>();
 			Modules.UnionWith(PublicDependencyModules!);
 			Modules.UnionWith(PrivateDependencyModules!);
-			if(bWithIncludePathModules)
+			if (bWithIncludePathModules)
 			{
 				Modules.UnionWith(PublicIncludePathModules!);
 				Modules.UnionWith(PrivateIncludePathModules!);
 			}
-			if(bWithDynamicallyLoadedModules)
+			if (bWithDynamicallyLoadedModules)
 			{
 				Modules.UnionWith(DynamicallyLoadedModules!);
 			}
 			return Modules;
-        }
+		}
 
-  		/// <summary>
+		/// <summary>
 		/// Returns a list of this module's frameworks.
 		/// </summary>
 		/// <returns>A List containing the frameworks this module requires.</returns>
-        public List<string> GetPublicFrameworks()
-        {
-            return new List<string>(PublicFrameworks);
-        }
+		public List<string> GetPublicFrameworks()
+		{
+			return new List<string>(PublicFrameworks);
+		}
 
 		/// <summary>
 		/// Returns a list of this module's immediate dependencies.
@@ -377,12 +374,12 @@ namespace UnrealBuildTool
 		protected HashSet<DirectoryReference> CreateDirectoryHashSet(IEnumerable<string> InEnumerableStrings)
 		{
 			HashSet<DirectoryReference> Directories = new HashSet<DirectoryReference>();
-			if(InEnumerableStrings != null)
+			if (InEnumerableStrings != null)
 			{
-				foreach(string InputString in InEnumerableStrings)
+				foreach (string InputString in InEnumerableStrings)
 				{
 					DirectoryReference Dir = new DirectoryReference(ExpandPathVariables(InputString, null, null));
-					if(DirectoryLookupCache.DirectoryExists(Dir))
+					if (DirectoryLookupCache.DirectoryExists(Dir))
 					{
 						Directories.Add(Dir);
 					}
@@ -444,7 +441,7 @@ namespace UnrealBuildTool
 				ReferencedDirs.ExceptWith(PrivateDependencyModules!.SelectMany(x => x.RestrictedFoldersAllowList));
 
 				// Add flags for each of them
-				foreach(DirectoryReference ReferencedDir in ReferencedDirs)
+				foreach (DirectoryReference ReferencedDir in ReferencedDirs)
 				{
 					// Find the base directory containing this reference
 					DirectoryReference? BaseDir = RootDirectories.FirstOrDefault(x => ReferencedDir.IsUnderDirectory(x));
@@ -456,9 +453,9 @@ namespace UnrealBuildTool
 
 					// Add references to each of the restricted folders
 					List<RestrictedFolder> Folders = RestrictedFolders.FindRestrictedFolders(BaseDir, ReferencedDir);
-					foreach(RestrictedFolder Folder in Folders)
+					foreach (RestrictedFolder Folder in Folders)
 					{
-						if(!References.ContainsKey(Folder))
+						if (!References.ContainsKey(Folder))
 						{
 							References.Add(Folder, ReferencedDir);
 						}
@@ -476,7 +473,7 @@ namespace UnrealBuildTool
 		{
 			Directories.Add(ModuleDirectory);
 
-			foreach(DirectoryReference PublicIncludePath in PublicIncludePaths)
+			foreach (DirectoryReference PublicIncludePath in PublicIncludePaths)
 			{
 				Directories.Add(PublicIncludePath);
 			}
@@ -488,7 +485,7 @@ namespace UnrealBuildTool
 			{
 				Directories.Add(PrivateIncludePath);
 			}
-			foreach(DirectoryReference PublicSystemIncludePath in PublicSystemIncludePaths)
+			foreach (DirectoryReference PublicSystemIncludePath in PublicSystemIncludePaths)
 			{
 				Directories.Add(PublicSystemIncludePath);
 			}
@@ -647,7 +644,7 @@ namespace UnrealBuildTool
 			// Add the import or export declaration for the module
 			if (Rules.Type == ModuleRules.ModuleType.CPlusPlus)
 			{
-				if(Rules.Target.LinkType == TargetLinkType.Monolithic)
+				if (Rules.Target.LinkType == TargetLinkType.Monolithic)
 				{
 					if (Rules.Target.bShouldCompileAsDLL && (Rules.Target.bHasExports || Rules.ModuleSymbolVisibility == ModuleRules.SymbolVisibility.VisibileForDll))
 					{
@@ -658,11 +655,11 @@ namespace UnrealBuildTool
 						Definitions.Add(ModuleApiDefine + "=");
 					}
 				}
-				else if(Binary == null || SourceBinary != Binary)
+				else if (Binary == null || SourceBinary != Binary)
 				{
 					Definitions.Add(ModuleApiDefine + "=DLLIMPORT");
 				}
-				else if(!Binary.bAllowExports)
+				else if (!Binary.bAllowExports)
 				{
 					Definitions.Add(ModuleApiDefine + "=");
 				}
@@ -729,20 +726,20 @@ namespace UnrealBuildTool
 		/// <returns>The path with variables expanded</returns>
 		public string ExpandPathVariables(string Path, DirectoryReference? BinaryOutputDir, DirectoryReference? TargetOutputDir)
 		{
-			if(Path.StartsWith("$(", StringComparison.Ordinal))
+			if (Path.StartsWith("$(", StringComparison.Ordinal))
 			{
 				int StartIdx = 2;
-				for(int EndIdx = StartIdx; EndIdx < Path.Length; EndIdx++)
+				for (int EndIdx = StartIdx; EndIdx < Path.Length; EndIdx++)
 				{
-					if(Path[EndIdx] == ')')
+					if (Path[EndIdx] == ')')
 					{
-						if(MatchVariableName(Path, StartIdx, EndIdx, "EngineDir"))
+						if (MatchVariableName(Path, StartIdx, EndIdx, "EngineDir"))
 						{
 							Path = Unreal.EngineDirectory + Path.Substring(EndIdx + 1);
 						}
-						else if(MatchVariableName(Path, StartIdx, EndIdx, "ProjectDir"))
+						else if (MatchVariableName(Path, StartIdx, EndIdx, "ProjectDir"))
 						{
-							if(Rules.Target.ProjectFile == null)
+							if (Rules.Target.ProjectFile == null)
 							{
 								Path = Unreal.EngineDirectory + Path.Substring(EndIdx + 1);
 							}
@@ -751,19 +748,19 @@ namespace UnrealBuildTool
 								Path = Rules.Target.ProjectFile.Directory + Path.Substring(EndIdx + 1);
 							}
 						}
-						else if(MatchVariableName(Path, StartIdx, EndIdx, "ModuleDir"))
+						else if (MatchVariableName(Path, StartIdx, EndIdx, "ModuleDir"))
 						{
 							Path = Rules.ModuleDirectory + Path.Substring(EndIdx + 1);
 						}
-						else if(MatchVariableName(Path, StartIdx, EndIdx, "PluginDir"))
+						else if (MatchVariableName(Path, StartIdx, EndIdx, "PluginDir"))
 						{
 							Path = Rules.PluginDirectory + Path.Substring(EndIdx + 1);
 						}
-						else if(BinaryOutputDir != null && MatchVariableName(Path, StartIdx, EndIdx, "BinaryOutputDir"))
+						else if (BinaryOutputDir != null && MatchVariableName(Path, StartIdx, EndIdx, "BinaryOutputDir"))
 						{
 							Path = BinaryOutputDir.FullName + Path.Substring(EndIdx + 1);
 						}
-						else if(TargetOutputDir != null && MatchVariableName(Path, StartIdx, EndIdx, "TargetOutputDir"))
+						else if (TargetOutputDir != null && MatchVariableName(Path, StartIdx, EndIdx, "TargetOutputDir"))
 						{
 							Path = TargetOutputDir.FullName + Path.Substring(EndIdx + 1);
 						}
@@ -771,7 +768,7 @@ namespace UnrealBuildTool
 						{
 							string Name = Path.Substring(StartIdx, EndIdx - StartIdx);
 							string? Value = Environment.GetEnvironmentVariable(Name);
-							if(String.IsNullOrEmpty(Value))
+							if (String.IsNullOrEmpty(Value))
 							{
 								throw new BuildException("Environment variable '{0}' is not defined (referenced by {1})", Name, Rules.File);
 							}
@@ -806,7 +803,7 @@ namespace UnrealBuildTool
 		/// <returns>The path with variables expanded</returns>
 		private IEnumerable<string> ExpandPathVariables(IEnumerable<string> Paths, DirectoryReference? BinaryDir, DirectoryReference? ExeDir)
 		{
-			foreach(string Path in Paths)
+			foreach (string Path in Paths)
 			{
 				yield return ExpandPathVariables(Path, BinaryDir, ExeDir);
 			}
@@ -932,7 +929,7 @@ namespace UnrealBuildTool
 		public virtual List<FileItem> Compile(ReadOnlyTargetRules Target, UEToolChain ToolChain, CppCompileEnvironment CompileEnvironment, ISourceFileWorkingSet WorkingSet, IActionGraphBuilder Graph, ILogger Logger)
 		{
 			// Generate type libraries for Windows
-			foreach(ModuleRules.TypeLibrary TypeLibrary in Rules.TypeLibraries)
+			foreach (ModuleRules.TypeLibrary TypeLibrary in Rules.TypeLibraries)
 			{
 				FileReference OutputFile = FileReference.Combine(IntermediateDirectory, TypeLibrary.Header);
 				ToolChain.GenerateTypeLibraryHeader(CompileEnvironment, TypeLibrary, OutputFile, Graph);
@@ -940,26 +937,26 @@ namespace UnrealBuildTool
 
 			return new List<FileItem>();
 		}
-		
+
 
 		public void CopyDebuggerVisualizers(UEToolChain ToolChain, IActionGraphBuilder Graph, ILogger Logger)
 		{
 			if (NatvisSourceFile != null)
 			{
 				ToolChain.CopyDebuggerVisualizer(NatvisSourceFile, IntermediateDirectory, Graph);
-			}	
+			}
 		}
-		
+
 		public void LinkDebuggerVisualizers(List<FileItem> OutFiles, UEToolChain ToolChain, ILogger Logger)
 		{
-			if (NatvisSourceFile != null) 
+			if (NatvisSourceFile != null)
 			{
 				FileItem? Item = ToolChain.LinkDebuggerVisualizer(NatvisSourceFile, IntermediateDirectory);
 				if (Item != null)
 				{
 					OutFiles.Add(Item);
 				}
-			}	
+			}
 		}
 
 		// Object interface.
@@ -1052,7 +1049,7 @@ namespace UnrealBuildTool
 			// Add us to the reference stack - note do this before checking for cycles as we allow them if an element in the stack
 			// declares the next element at leading to a dependency cycle.
 			ReferenceStack.Add(this);
-			
+
 			// Check if this module is invalid for the current target, based on Module attributes
 			if (!ModuleRules.IsValidForTarget(Rules.GetType(), Rules.Target, out string? InvalidReason))
 			{
@@ -1091,7 +1088,7 @@ namespace UnrealBuildTool
 					Logger.LogError("Circular dependency on {Name} detected.\n" +
 						"\tFull Route: {FullRoute}\n" +
 						"\tCycled Route: is {CycleRoot}.\n" +
-						"Break this loop by moving dependencies into a separate module or using Private/PublicIncludePathModuleNames to reference declarations\n", 
+						"Break this loop by moving dependencies into a separate module or using Private/PublicIncludePathModuleNames to reference declarations\n",
 						ThisRefName, NextReferenceChain, CycleChain);
 
 				}
@@ -1099,11 +1096,11 @@ namespace UnrealBuildTool
 				{
 					Logger.LogDebug("Found circular reference to {ThisRefName}, but {GuiltyModule} declares a cycle on {VictimModule} which breaks the chain", ThisRefName, GuiltyModule, VictimModule);
 				}
-			}		
+			}
 
 			// Create all the referenced modules. This path can be recursive, so we check against PublicDependencyModules to ensure we don't recurse through the 
 			// same module twice (it produces better errors if something fails).
-			if(PublicDependencyModules == null)
+			if (PublicDependencyModules == null)
 			{
 				// Log dependencies if required
 				if (Logger.IsEnabled((LogLevel)LogEventType.VeryVerbose))
@@ -1214,7 +1211,7 @@ namespace UnrealBuildTool
 		{
 			if (Rules.Type == ModuleRules.ModuleType.CPlusPlus)
 			{
-				return new[] {ModuleApiDefine + "="};
+				return new[] { ModuleApiDefine + "=" };
 			}
 
 			return new string[0];
@@ -1264,7 +1261,7 @@ namespace UnrealBuildTool
 			ExportJsonStringArray(Writer, "PublicDefinitions", PublicDefinitions);
 
 			Writer.WriteArrayStart("CircularlyReferencedModules");
-			foreach(string ModuleName in Rules.CircularlyReferencedDependentModules)
+			foreach (string ModuleName in Rules.CircularlyReferencedDependentModules)
 			{
 				Writer.WriteValue(ModuleName);
 			}
@@ -1307,7 +1304,7 @@ namespace UnrealBuildTool
 			}
 			Writer.WriteArrayEnd();
 		}
-		
+
 		/// <summary>
 		/// Write an array of strings to a JSON writer
 		/// </summary>
@@ -1319,7 +1316,7 @@ namespace UnrealBuildTool
 			Writer.WriteArrayStart(ArrayName);
 			if (Strings != null)
 			{
-				foreach(string String in Strings)
+				foreach (string String in Strings)
 				{
 					Writer.WriteValue(String);
 				}
diff --git a/Engine/Source/Programs/UnrealBuildTool/Configuration/UEBuildModuleCPP.cs b/Engine/Source/Programs/UnrealBuildTool/Configuration/UEBuildModuleCPP.cs
--- a/Engine/Source/Programs/UnrealBuildTool/Configuration/UEBuildModuleCPP.cs
+++ b/Engine/Source/Programs/UnrealBuildTool/Configuration/UEBuildModuleCPP.cs
@@ -8,7 +8,6 @@ using System.Text;
 using EpicGames.Core;
 using Microsoft.Extensions.Logging;
 using UnrealBuildBase;
-using static UnrealBuildTool.ProjectFile;
 
 namespace UnrealBuildTool
 {
@@ -103,9 +102,9 @@ namespace UnrealBuildTool
 		{
 			base.GetReferencedDirectories(Directories);
 
-			if(!Rules.bUsePrecompiled)
+			if (!Rules.bUsePrecompiled)
 			{
-				if(SourceDirectories == null)
+				if (SourceDirectories == null)
 				{
 					throw new BuildException("GetReferencedDirectories() should not be called before building.");
 				}
@@ -207,8 +206,8 @@ namespace UnrealBuildTool
 				foreach (string CircularlyReferencedModuleName in Rules.CircularlyReferencedDependentModules)
 				{
 					if (CircularlyReferencedModuleName != "BlueprintContext" &&
-					    !CircularDependenciesAllowList.Any(x =>
-						    x.Key == Name && x.Value == CircularlyReferencedModuleName))
+						!CircularDependenciesAllowList.Any(x =>
+							x.Key == Name && x.Value == CircularlyReferencedModuleName))
 					{
 						Logger.LogWarning(
 							"Found reference between '{Source}' and '{Target}'. Support for circular references is being phased out; please do not introduce new ones.",
@@ -293,9 +292,9 @@ namespace UnrealBuildTool
 		/// <param name="LegacyPublicIncludePaths">List populated with the discovered directories</param>
 		static void EnumerateLegacyIncludePaths(DirectoryItem BaseDirectory, IReadOnlySet<string> ExcludeNames, HashSet<DirectoryReference> LegacyPublicIncludePaths)
 		{
-			foreach(DirectoryItem SubDirectory in BaseDirectory.EnumerateDirectories())
+			foreach (DirectoryItem SubDirectory in BaseDirectory.EnumerateDirectories())
 			{
-				if(!ExcludeNames.Contains(SubDirectory.Name))
+				if (!ExcludeNames.Contains(SubDirectory.Name))
 				{
 					LegacyPublicIncludePaths.Add(SubDirectory.Location);
 					EnumerateLegacyIncludePaths(SubDirectory, ExcludeNames, LegacyPublicIncludePaths);
@@ -361,19 +360,19 @@ namespace UnrealBuildTool
 			CppCompileEnvironment ModuleCompileEnvironment = CreateModuleCompileEnvironment(Target, BinaryCompileEnvironment, Logger);
 
 			// If the module is precompiled, read the object files from the manifest
-			if(Rules.bUsePrecompiled && Target.LinkType == TargetLinkType.Monolithic)
+			if (Rules.bUsePrecompiled && Target.LinkType == TargetLinkType.Monolithic)
 			{
-				if(!FileReference.Exists(PrecompiledManifestLocation))
+				if (!FileReference.Exists(PrecompiledManifestLocation))
 				{
 					throw new BuildException("Missing precompiled manifest for '{0}', '{1}'. This module was most likely not flagged for being included in a precompiled build - set 'PrecompileForTargets = PrecompileTargetsType.Any;' in {0}.build.cs to override." +
 						" If part of a plugin, also check if its 'Type' is correct.", Name, PrecompiledManifestLocation);
 				}
 
 				PrecompiledManifest Manifest = PrecompiledManifest.Read(PrecompiledManifestLocation);
-				foreach(FileReference OutputFile in Manifest.OutputFiles)
+				foreach (FileReference OutputFile in Manifest.OutputFiles)
 				{
 					FileItem ObjectFile = FileItem.GetItemByFileReference(OutputFile);
-					if(!ObjectFile.Exists)
+					if (!ObjectFile.Exists)
 					{
 						throw new BuildException("Missing object file {0} listed in {1}", OutputFile, PrecompiledManifestLocation);
 					}
@@ -658,13 +657,13 @@ namespace UnrealBuildTool
 			// Compile CPP files
 			if (bModuleUsesUnityBuild)
 			{
-				Unity.GenerateUnityCPPs(Target, CPPFiles, InputFiles.HeaderFiles, CompileEnvironment, WorkingSet, Rules.ShortName ?? Name, IntermediateDirectory, Graph, SourceFileToUnityFile, 
+				Unity.GenerateUnityCPPs(Target, CPPFiles, InputFiles.HeaderFiles, CompileEnvironment, WorkingSet, Rules.ShortName ?? Name, IntermediateDirectory, Graph, SourceFileToUnityFile,
 					out List<FileItem> NormalFiles, out List<FileItem> AdaptiveFiles, NumIncludedBytesPerUnityCPP);
 				LinkInputFiles.AddRange(CompileFilesWithToolChain(Target, ToolChain, CompileEnvironment, ModuleCompileEnvironment, NormalFiles, AdaptiveFiles, Graph, Logger).ObjectFiles);
 			}
 			else
 			{
-				Unity.GetAdaptiveFiles(Target, CPPFiles, InputFiles.HeaderFiles, CompileEnvironment, WorkingSet, Rules.ShortName ?? Name, IntermediateDirectory, Graph, 
+				Unity.GetAdaptiveFiles(Target, CPPFiles, InputFiles.HeaderFiles, CompileEnvironment, WorkingSet, Rules.ShortName ?? Name, IntermediateDirectory, Graph,
 					out List<FileItem> NormalFiles, out List<FileItem> AdaptiveFiles);
 				if (NormalFiles.Where(file => !file.HasExtension(".gen.cpp")).Count() == 0)
 				{
@@ -681,7 +680,7 @@ namespace UnrealBuildTool
 			}
 
 			// Compile C files directly. Do not use a PCH here, because a C++ PCH is not compatible with C source files.
-			if(InputFiles.CFiles.Count > 0)
+			if (InputFiles.CFiles.Count > 0)
 			{
 				CppCompileEnvironment CCompileEnvironment = new CppCompileEnvironment(ModuleCompileEnvironment);
 				CreateHeaderForDefinitions(CCompileEnvironment, IntermediateDirectory, "c", Graph);
@@ -689,13 +688,13 @@ namespace UnrealBuildTool
 			}
 
 			// Compile CC files directly.
-			if(InputFiles.CCFiles.Count > 0)
+			if (InputFiles.CCFiles.Count > 0)
 			{
 				LinkInputFiles.AddRange(ToolChain.CompileAllCPPFiles(CompileEnvironment, InputFiles.CCFiles, IntermediateDirectory, Name, Graph).ObjectFiles);
 			}
 
 			// Compile MM files directly.
-			if(InputFiles.MMFiles.Count > 0)
+			if (InputFiles.MMFiles.Count > 0)
 			{
 				LinkInputFiles.AddRange(ToolChain.CompileAllCPPFiles(CompileEnvironment, InputFiles.MMFiles, IntermediateDirectory, Name, Graph).ObjectFiles);
 			}
@@ -703,10 +702,10 @@ namespace UnrealBuildTool
 			// Compile RC files. The resource compiler does not work with response files, and using the regular compile environment can easily result in the 
 			// command line length exceeding the OS limit. Use the binary compile environment to keep the size down, and require that all include paths
 			// must be specified relative to the resource file itself or Engine/Source.
-			if(InputFiles.RCFiles.Count > 0)
+			if (InputFiles.RCFiles.Count > 0)
 			{
 				CppCompileEnvironment ResourceCompileEnvironment = new CppCompileEnvironment(BinaryCompileEnvironment);
-				if(Binary != null)
+				if (Binary != null)
 				{
 					// @todo: This should be in some Windows code somewhere...
 					ResourceCompileEnvironment.Definitions.Add("ORIGINAL_FILE_NAME=\"" + Binary.OutputFilePaths[0].GetFileName() + "\"");
@@ -715,7 +714,7 @@ namespace UnrealBuildTool
 			}
 
 			// Write the compiled manifest
-			if(Rules.bPrecompile && Target.LinkType == TargetLinkType.Monolithic)
+			if (Rules.bPrecompile && Target.LinkType == TargetLinkType.Monolithic)
 			{
 				DirectoryReference.CreateDirectory(PrecompiledManifestLocation.Directory);
 
@@ -725,7 +724,7 @@ namespace UnrealBuildTool
 			}
 
 			// Write a mapping of unity object file to standalone object file for live coding
-			if(Rules.Target.bWithLiveCoding)
+			if (Rules.Target.bWithLiveCoding)
 			{
 				StringWriter StringWriter = new();
 				using (JsonWriter Writer = new JsonWriter(StringWriter))
@@ -826,7 +825,7 @@ namespace UnrealBuildTool
 			FileItem HeaderFile = FileItem.GetItemByFileReference(FileReference.Combine(ModuleDirectory, Rules.SharedPCHHeaderFile!));
 
 			DirectoryReference PrecompiledHeaderDir;
-			if(Rules.bUsePrecompiled)
+			if (Rules.bUsePrecompiled)
 			{
 				PrecompiledHeaderDir = DirectoryReference.Combine(Target.ProjectIntermediateDirectory, Name);
 			}
@@ -902,7 +901,7 @@ namespace UnrealBuildTool
 		public PrecompiledHeaderInstance FindOrCreateSharedPCH(UEToolChain? ToolChain, PrecompiledHeaderTemplate Template, CppCompileEnvironment ModuleCompileEnvironment, IActionGraphBuilder Graph)
 		{
 			PrecompiledHeaderInstance? Instance = Template.Instances.Find(x => IsCompatibleForSharedPCH(x.CompileEnvironment, ModuleCompileEnvironment));
-			if(Instance == null)
+			if (Instance == null)
 			{
 				List<string> Definitions = Template.BaseCompileEnvironment.Definitions;
 
@@ -1014,27 +1013,27 @@ namespace UnrealBuildTool
 		/// <returns>True if the two compile enviroments are compatible</returns>
 		static internal bool IsCompatibleForSharedPCH(CppCompileEnvironment ModuleCompileEnvironment, CppCompileEnvironment CompileEnvironment)
 		{
-			if(ModuleCompileEnvironment.bOptimizeCode != CompileEnvironment.bOptimizeCode)
+			if (ModuleCompileEnvironment.bOptimizeCode != CompileEnvironment.bOptimizeCode)
 			{
 				return false;
 			}
-			if(ModuleCompileEnvironment.bUseRTTI != CompileEnvironment.bUseRTTI)
+			if (ModuleCompileEnvironment.bUseRTTI != CompileEnvironment.bUseRTTI)
 			{
 				return false;
 			}
-			if(ModuleCompileEnvironment.bEnableExceptions != CompileEnvironment.bEnableExceptions)
+			if (ModuleCompileEnvironment.bEnableExceptions != CompileEnvironment.bEnableExceptions)
 			{
 				return false;
 			}
-			if(ModuleCompileEnvironment.ShadowVariableWarningLevel != CompileEnvironment.ShadowVariableWarningLevel)
+			if (ModuleCompileEnvironment.ShadowVariableWarningLevel != CompileEnvironment.ShadowVariableWarningLevel)
 			{
 				return false;
 			}
-			if(ModuleCompileEnvironment.UnsafeTypeCastWarningLevel != CompileEnvironment.UnsafeTypeCastWarningLevel)
+			if (ModuleCompileEnvironment.UnsafeTypeCastWarningLevel != CompileEnvironment.UnsafeTypeCastWarningLevel)
 			{
 				return false;
 			}
-			if(ModuleCompileEnvironment.bEnableUndefinedIdentifierWarnings != CompileEnvironment.bEnableUndefinedIdentifierWarnings)
+			if (ModuleCompileEnvironment.bEnableUndefinedIdentifierWarnings != CompileEnvironment.bEnableUndefinedIdentifierWarnings)
 			{
 				return false;
 			}
@@ -1042,12 +1041,12 @@ namespace UnrealBuildTool
 			{
 				return false;
 			}
-			
+
 			if (ModuleCompileEnvironment.CStandard != CompileEnvironment.CStandard)
 			{
 				return false;
 			}
-			
+
 			if (ModuleCompileEnvironment.IncludeOrderVersion != CompileEnvironment.IncludeOrderVersion)
 			{
 				return false;
@@ -1064,9 +1063,9 @@ namespace UnrealBuildTool
 		static private string GetSuffixForSharedPCH(CppCompileEnvironment CompileEnvironment, CppCompileEnvironment BaseCompileEnvironment)
 		{
 			string Variant = "";
-			if(CompileEnvironment.bOptimizeCode != BaseCompileEnvironment.bOptimizeCode)
+			if (CompileEnvironment.bOptimizeCode != BaseCompileEnvironment.bOptimizeCode)
 			{
-				if(CompileEnvironment.bOptimizeCode)
+				if (CompileEnvironment.bOptimizeCode)
 				{
 					Variant += ".Optimized";
 				}
@@ -1075,7 +1074,7 @@ namespace UnrealBuildTool
 					Variant += ".NonOptimized";
 				}
 			}
-			if(CompileEnvironment.bUseRTTI != BaseCompileEnvironment.bUseRTTI)
+			if (CompileEnvironment.bUseRTTI != BaseCompileEnvironment.bUseRTTI)
 			{
 				if (CompileEnvironment.bUseRTTI)
 				{
@@ -1128,7 +1127,7 @@ namespace UnrealBuildTool
 					Variant += ".NoTypeCast";
 				}
 			}
-			
+
 			if (CompileEnvironment.bEnableUndefinedIdentifierWarnings != BaseCompileEnvironment.bEnableUndefinedIdentifierWarnings)
 			{
 				if (CompileEnvironment.bEnableUndefinedIdentifierWarnings)
@@ -1220,9 +1219,9 @@ namespace UnrealBuildTool
 			ILogger Logger)
 		{
 			bool bAdaptiveUnityDisablesPCH = false;
-			if(Rules.PCHUsage == ModuleRules.PCHUsageMode.UseExplicitOrSharedPCHs)
+			if (Rules.PCHUsage == ModuleRules.PCHUsageMode.UseExplicitOrSharedPCHs)
 			{
-				if(Rules.bTreatAsEngineModule || Rules.PrivatePCHHeaderFile == null)
+				if (Rules.bTreatAsEngineModule || Rules.PrivatePCHHeaderFile == null)
 				{
 					bAdaptiveUnityDisablesPCH = Target.bAdaptiveUnityDisablesPCH;
 				}
@@ -1281,7 +1280,7 @@ namespace UnrealBuildTool
 			{
 				// Create the new compile environment. Always turn off PCH due to different compiler settings.
 				CppCompileEnvironment AdaptiveUnityEnvironment = new CppCompileEnvironment(ModuleCompileEnvironment);
-				if(Target.bAdaptiveUnityDisablesOptimizations)
+				if (Target.bAdaptiveUnityDisablesOptimizations)
 				{
 					AdaptiveUnityEnvironment.bOptimizeCode = false;
 				}
@@ -1292,15 +1291,15 @@ namespace UnrealBuildTool
 
 				// Create a per-file PCH
 				CPPOutput AdaptiveOutput;
-				if(Target.bAdaptiveUnityCreatesDedicatedPCH)
+				if (Target.bAdaptiveUnityCreatesDedicatedPCH)
 				{
 					AdaptiveOutput = CompileAdaptiveNonUnityFilesWithDedicatedPCH(ToolChain, AdaptiveUnityEnvironment, AdaptiveFiles, IntermediateDirectory, Name, Graph);
 				}
-				else if(bAdaptiveUnityDisablesPCH)
+				else if (bAdaptiveUnityDisablesPCH)
 				{
 					AdaptiveOutput = CompileAdaptiveNonUnityFilesWithoutPCH(ToolChain, AdaptiveUnityEnvironment, AdaptiveFiles, IntermediateDirectory, Name, Graph);
 				}
-				else if(AdaptiveUnityEnvironment.bOptimizeCode != CompileEnvironment.bOptimizeCode || AdaptiveUnityEnvironment.bSupportEditAndContinue != CompileEnvironment.bSupportEditAndContinue)
+				else if (AdaptiveUnityEnvironment.bOptimizeCode != CompileEnvironment.bOptimizeCode || AdaptiveUnityEnvironment.bSupportEditAndContinue != CompileEnvironment.bSupportEditAndContinue)
 				{
 					AdaptiveOutput = CompileAdaptiveNonUnityFiles(ToolChain, AdaptiveUnityEnvironment, AdaptiveFiles, IntermediateDirectory, Name, Graph);
 				}
@@ -1340,7 +1339,7 @@ namespace UnrealBuildTool
 		static CPPOutput CompileAdaptiveNonUnityFilesWithDedicatedPCH(UEToolChain ToolChain, CppCompileEnvironment CompileEnvironment, List<FileItem> Files, DirectoryReference IntermediateDirectory, string ModuleName, IActionGraphBuilder Graph)
 		{
 			CPPOutput Output = new CPPOutput();
-			foreach(FileItem File in Files)
+			foreach (FileItem File in Files)
 			{
 				// Build the contents of the wrapper file
 				StringBuilder WrapperContents = new StringBuilder();
@@ -1355,7 +1354,7 @@ namespace UnrealBuildTool
 					WrapperContents.AppendLine();
 					WriteDefinitions(CompileEnvironment.Definitions, WrapperContents);
 					WrapperContents.AppendLine();
-					using(StreamReader Reader = new StreamReader(File.Location.FullName))
+					using (StreamReader Reader = new StreamReader(File.Location.FullName))
 					{
 						CppIncludeParser.CopyIncludeDirectives(Reader, WrapperContents);
 					}
@@ -1403,7 +1402,7 @@ namespace UnrealBuildTool
 			if (Target.bUsePCHFiles && Rules.PCHUsage != ModuleRules.PCHUsageMode.NoPCHs)
 			{
 				// If this module doesn't need a shared PCH, configure that
-				if(Rules.PrivatePCHHeaderFile != null && (Rules.PCHUsage == ModuleRules.PCHUsageMode.NoSharedPCHs || Rules.PCHUsage == ModuleRules.PCHUsageMode.UseExplicitOrSharedPCHs))
+				if (Rules.PrivatePCHHeaderFile != null && (Rules.PCHUsage == ModuleRules.PCHUsageMode.NoSharedPCHs || Rules.PCHUsage == ModuleRules.PCHUsageMode.UseExplicitOrSharedPCHs))
 				{
 					PrecompiledHeaderInstance Instance = CreatePrivatePCH(ToolChain, FileItem.GetItemByFileReference(FileReference.Combine(ModuleDirectory, Rules.PrivatePCHHeaderFile)), CompileEnvironment, Graph);
 
@@ -1425,7 +1424,7 @@ namespace UnrealBuildTool
 
 					// Find the first shared PCH module we can use
 					PrecompiledHeaderTemplate? Template = CompileEnvironment.SharedPCHs.FirstOrDefault(x => ReferencedModules.Contains(x.Module));
-					if(Template != null && Template.IsValidFor(CompileEnvironment))
+					if (Template != null && Template.IsValidFor(CompileEnvironment))
 					{
 						PrecompiledHeaderInstance Instance = FindOrCreateSharedPCH(ToolChain, Template, CompileEnvironment, Graph);
 
@@ -1438,7 +1437,7 @@ namespace UnrealBuildTool
 							Writer.AppendLine($"#include \"{GetIncludeString(Instance.DefinitionsFile)}\"");
 
 							// Games may choose to use shared PCHs from the engine, so allow them to change the value of these macros
-							if(!Rules.bTreatAsEngineModule)
+							if (!Rules.bTreatAsEngineModule)
 							{
 								Writer.AppendLine("#undef UE_IS_ENGINE_MODULE");
 								Writer.AppendLine("#undef DEPRECATED_FORGAME");
@@ -1499,11 +1498,11 @@ namespace UnrealBuildTool
 		/// <param name="Graph">The action graph being built</param>
 		static FileItem? CreateHeaderForDefinitions(CppCompileEnvironment CompileEnvironment, DirectoryReference IntermediateDirectory, string? HeaderSuffix, IActionGraphBuilder Graph)
 		{
-			if(CompileEnvironment.Definitions.Count > 0)
+			if (CompileEnvironment.Definitions.Count > 0)
 			{
 				string PrivateDefinitionsName = "Definitions.h";
 
-				if(!String.IsNullOrEmpty(HeaderSuffix))
+				if (!String.IsNullOrEmpty(HeaderSuffix))
 				{
 					PrivateDefinitionsName = $"Definitions.{HeaderSuffix}.h";
 				}
@@ -1591,10 +1590,10 @@ namespace UnrealBuildTool
 		static void WriteDefinitions(IEnumerable<string> Definitions, StringBuilder Writer)
 		{
 			Writer.AppendLine("// Generated by UnrealBuildTool (UEBuildModuleCPP.cs)");
-			foreach(string Definition in Definitions)
+			foreach (string Definition in Definitions)
 			{
 				int EqualsIdx = Definition.IndexOf('=');
-				if(EqualsIdx == -1)
+				if (EqualsIdx == -1)
 				{
 					Writer.Append("#define ").Append(Definition).AppendLine(" 1");
 				}
@@ -1614,13 +1613,13 @@ namespace UnrealBuildTool
 		/// <param name="CppFiles">List of C++ source files</param>
 		private void CheckFirstIncludeMatchesEachCppFile(ReadOnlyTargetRules Target, CppCompileEnvironment ModuleCompileEnvironment, List<FileItem> HeaderFiles, List<FileItem> CppFiles)
 		{
-			if(Rules.PCHUsage == ModuleRules.PCHUsageMode.UseExplicitOrSharedPCHs)
+			if (Rules.PCHUsage == ModuleRules.PCHUsageMode.UseExplicitOrSharedPCHs)
 			{
-				if(InvalidIncludeDirectives == null)
+				if (InvalidIncludeDirectives == null)
 				{
 					// Find headers used by the source file.
 					Dictionary<string, FileReference> NameToHeaderFile = new Dictionary<string, FileReference>();
-					foreach(FileItem HeaderFile in HeaderFiles)
+					foreach (FileItem HeaderFile in HeaderFiles)
 					{
 						NameToHeaderFile[HeaderFile.Location.GetFileNameWithoutExtension()] = HeaderFile.Location;
 					}
@@ -1632,7 +1631,7 @@ namespace UnrealBuildTool
 						foreach (FileItem CppFile in CppFiles)
 						{
 							string? FirstInclude = ModuleCompileEnvironment.MetadataCache.GetFirstInclude(CppFile);
-							if(FirstInclude != null)
+							if (FirstInclude != null)
 							{
 								string IncludeName = Path.GetFileNameWithoutExtension(FirstInclude);
 								string ExpectedName = CppFile.Location.GetFileNameWithoutExtension();
@@ -1653,7 +1652,7 @@ namespace UnrealBuildTool
 
 		private bool IgnoreMismatchedHeader(string ExpectedName)
 		{
-			switch(Name)
+			switch (Name)
 			{
 				case "D3D11RHI":
 				case "D3D12RHI":
@@ -1674,13 +1673,13 @@ namespace UnrealBuildTool
 		/// <returns>True if optimization should be enabled</returns>
 		public static bool ShouldEnableOptimization(ModuleRules.CodeOptimization Setting, UnrealTargetConfiguration Configuration, bool bIsEngineModule)
 		{
-			switch(Setting)
+			switch (Setting)
 			{
 				case ModuleRules.CodeOptimization.Never:
 					return false;
 				case ModuleRules.CodeOptimization.Default:
 				case ModuleRules.CodeOptimization.InNonDebugBuilds:
-					return (Configuration == UnrealTargetConfiguration.Debug)? false : (Configuration != UnrealTargetConfiguration.DebugGame || bIsEngineModule);
+					return (Configuration == UnrealTargetConfiguration.Debug) ? false : (Configuration != UnrealTargetConfiguration.DebugGame || bIsEngineModule);
 				case ModuleRules.CodeOptimization.InShippingBuildsOnly:
 					return (Configuration == UnrealTargetConfiguration.Shipping);
 				default:
@@ -1817,7 +1816,7 @@ namespace UnrealBuildTool
 			Result.Definitions.AddRange(Rules.PrivateDefinitions);
 
 			// Add the project definitions
-			if(!Rules.bTreatAsEngineModule)
+			if (!Rules.bTreatAsEngineModule)
 			{
 				Result.Definitions.AddRange(Rules.Target.ProjectDefinitions);
 			}
@@ -1922,7 +1921,7 @@ namespace UnrealBuildTool
 				}
 			}
 
-			if(SourceFiles.Length > 0)
+			if (SourceFiles.Length > 0)
 			{
 				SourceDirectories.Add(BaseDirectory.Location);
 			}
@@ -1935,7 +1934,7 @@ namespace UnrealBuildTool
 		{
 			".h", ".ipsh",
 			".cpp", ".ixx", ".c", ".cc",
-			".m", ".mm", 
+			".m", ".mm",
 			".rc", ".ispc",
 		};
 
@@ -1951,7 +1950,7 @@ namespace UnrealBuildTool
 		{
 			bIgnoreFileFound = false;
 			List<FileItem> SourceFiles = new List<FileItem>();
-			foreach(FileItem InputFile in BaseDirectory.EnumerateFiles())
+			foreach (FileItem InputFile in BaseDirectory.EnumerateFiles())
 			{
 				if (InputFile.Name == ".ubtignore")
 				{
diff --git a/Engine/Source/Programs/UnrealBuildTool/Configuration/UEBuildModuleExternal.cs b/Engine/Source/Programs/UnrealBuildTool/Configuration/UEBuildModuleExternal.cs
--- a/Engine/Source/Programs/UnrealBuildTool/Configuration/UEBuildModuleExternal.cs
+++ b/Engine/Source/Programs/UnrealBuildTool/Configuration/UEBuildModuleExternal.cs
@@ -1,10 +1,5 @@
 // Copyright Epic Games, Inc. All Rights Reserved.
 
-using System;
-using System.Collections.Generic;
-using System.Linq;
-using System.Text;
-using System.Threading.Tasks;
 using EpicGames.Core;
 using Microsoft.Extensions.Logging;
 
diff --git a/Engine/Source/Programs/UnrealBuildTool/Configuration/UEBuildPlatform.cs b/Engine/Source/Programs/UnrealBuildTool/Configuration/UEBuildPlatform.cs
--- a/Engine/Source/Programs/UnrealBuildTool/Configuration/UEBuildPlatform.cs
+++ b/Engine/Source/Programs/UnrealBuildTool/Configuration/UEBuildPlatform.cs
@@ -2,15 +2,14 @@
 
 using System;
 using System.Collections.Generic;
-using System.Text;
-using System.Diagnostics;
+using System.Diagnostics.CodeAnalysis;
 using System.IO;
 using System.Linq;
-using EpicGames.Core;
 using System.Reflection;
-using System.Diagnostics.CodeAnalysis;
-using OpenTracing.Util;
+using System.Text;
+using EpicGames.Core;
 using Microsoft.Extensions.Logging;
+using OpenTracing.Util;
 
 namespace UnrealBuildTool
 {
@@ -102,7 +101,7 @@ namespace UnrealBuildTool
 			return ActiveArchitectures(ProjectFile, TargetName);
 		}
 
-	
+
 		/// <summary>
 		/// Returns the set all architectures potentially supported by this project. Can be used by project file gnenerators to restrict IDE architecture options
 		/// Defaults to AllSupportedArchitectures
@@ -291,8 +290,8 @@ namespace UnrealBuildTool
 						using (GlobalTracer.Instance.BuildSpan(CheckType.Name).StartActive())
 						{
 							UEBuildPlatformFactory TempInst = (UEBuildPlatformFactory)Activator.CreateInstance(CheckType)!;
-							
-							if(bHostPlatformOnly && TempInst.TargetPlatform != BuildHostPlatform.Current.Platform)
+
+							if (bHostPlatformOnly && TempInst.TargetPlatform != BuildHostPlatform.Current.Platform)
 							{
 								continue;
 							}
@@ -315,7 +314,7 @@ namespace UnrealBuildTool
 		/// <returns>Array of platform folders</returns>
 		public static string[] GetPlatformFolderNames()
 		{
-			if(CachedPlatformFolderNames == null)
+			if (CachedPlatformFolderNames == null)
 			{
 				List<string> PlatformFolderNames = new List<string>();
 
@@ -336,12 +335,12 @@ namespace UnrealBuildTool
 		/// </summary>
 		public IReadOnlySet<string> GetIncludedFolderNames()
 		{
-			if(CachedIncludedFolderNames == null)
+			if (CachedIncludedFolderNames == null)
 			{
 				HashSet<string> Names = new HashSet<string>(DirectoryReference.Comparer);
 
 				Names.Add(Platform.ToString());
-				foreach(UnrealPlatformGroup Group in UEBuildPlatform.GetPlatformGroups(Platform))
+				foreach (UnrealPlatformGroup Group in UEBuildPlatform.GetPlatformGroups(Platform))
 				{
 					Names.Add(Group.ToString());
 				}
@@ -356,7 +355,7 @@ namespace UnrealBuildTool
 		/// </summary>
 		public IReadOnlySet<string> GetExcludedFolderNames()
 		{
-			if(CachedExcludedFolderNames == null)
+			if (CachedExcludedFolderNames == null)
 			{
 				CachedExcludedFolderNames = new HashSet<string>(GetPlatformFolderNames().Except(GetIncludedFolderNames()), DirectoryReference.Comparer);
 			}
@@ -598,7 +597,7 @@ namespace UnrealBuildTool
 		{
 			// find or add the list of groups for this platform
 			List<UnrealTargetPlatform>? Platforms;
-			if(!PlatformGroupDictionary.TryGetValue(InGroup, out Platforms))
+			if (!PlatformGroupDictionary.TryGetValue(InGroup, out Platforms))
 			{
 				Platforms = new List<UnrealTargetPlatform>();
 				PlatformGroupDictionary.Add(InGroup, Platforms);
@@ -637,12 +636,12 @@ namespace UnrealBuildTool
 		public static UEBuildPlatform GetBuildPlatform(UnrealTargetPlatform InPlatform)
 		{
 			UEBuildPlatform? Platform;
-			if(!TryGetBuildPlatform(InPlatform, out Platform))
+			if (!TryGetBuildPlatform(InPlatform, out Platform))
 			{
 				throw new BuildException("GetBuildPlatform: No BuildPlatform found for {0}", InPlatform.ToString());
 			}
 			return Platform;
-			}
+		}
 
 		/// <summary>
 		/// Retrieve the IUEBuildPlatform instance for the given TargetPlatform
@@ -1233,7 +1232,7 @@ namespace UnrealBuildTool
 		/// <param name="Target">The target rules in use</param>
 		/// <param name="ToolChain">The toolchain being used</param>
 		/// <param name="Graph">Action graph that is used to build the binary</param>
-		public virtual void ModifyBinaryLinkEnvironment( LinkEnvironment BinaryLinkEnvironment, CppCompileEnvironment BinaryCompileEnvironment, ReadOnlyTargetRules Target, UEToolChain ToolChain, IActionGraphBuilder Graph)
+		public virtual void ModifyBinaryLinkEnvironment(LinkEnvironment BinaryLinkEnvironment, CppCompileEnvironment BinaryCompileEnvironment, ReadOnlyTargetRules Target, UEToolChain ToolChain, IActionGraphBuilder Graph)
 		{
 		}
 
diff --git a/Engine/Source/Programs/UnrealBuildTool/Configuration/UEBuildPlugin.cs b/Engine/Source/Programs/UnrealBuildTool/Configuration/UEBuildPlugin.cs
--- a/Engine/Source/Programs/UnrealBuildTool/Configuration/UEBuildPlugin.cs
+++ b/Engine/Source/Programs/UnrealBuildTool/Configuration/UEBuildPlugin.cs
@@ -1,10 +1,6 @@
 // Copyright Epic Games, Inc. All Rights Reserved.
 
-using System;
 using System.Collections.Generic;
-using System.Linq;
-using System.Text;
-using System.Threading.Tasks;
 using EpicGames.Core;
 
 namespace UnrealBuildTool
diff --git a/Engine/Source/Programs/UnrealBuildTool/Configuration/UEBuildTarget.cs b/Engine/Source/Programs/UnrealBuildTool/Configuration/UEBuildTarget.cs
--- a/Engine/Source/Programs/UnrealBuildTool/Configuration/UEBuildTarget.cs
+++ b/Engine/Source/Programs/UnrealBuildTool/Configuration/UEBuildTarget.cs
@@ -4,25 +4,19 @@ using System;
 using System.Collections.Concurrent;
 using System.Collections.Generic;
 using System.ComponentModel;
-using System.Linq;
-using System.Text;
-using System.Diagnostics;
 using System.Globalization;
 using System.IO;
-using System.Xml;
+using System.Linq;
+using System.Reflection;
 using System.Runtime.Serialization;
+using System.Text;
 using System.Text.Json.Serialization;
-using System.Text.RegularExpressions;
+using System.Threading.Tasks;
 using EpicGames.Core;
-using System.Reflection;
-using OpenTracing.Util;
-using UnrealBuildBase;
 using EpicGames.UHT.Utils;
 using Microsoft.Extensions.Logging;
-using System.Runtime.InteropServices;
-using System.Threading.Tasks;
-using UnrealBuildTool;
-using System.Xml.Linq;
+using OpenTracing.Util;
+using UnrealBuildBase;
 
 
 namespace UnrealBuildTool
@@ -909,7 +903,7 @@ namespace UnrealBuildTool
 		/// </summary>
 		/// <param name="Architectures"></param>
 		/// <param name="ValidationPlatform"></param>
-		public UnrealArchitectures(IEnumerable<string> Architectures, UnrealTargetPlatform? ValidationPlatform=null)
+		public UnrealArchitectures(IEnumerable<string> Architectures, UnrealTargetPlatform? ValidationPlatform = null)
 		{
 			// validate before conversion (make sure all given architectures can be converted to an UnrealArch that is supported by the platform)
 			if (ValidationPlatform != null)
@@ -1934,7 +1928,7 @@ namespace UnrealBuildTool
 		Dictionary<FileReference, LoadOrderManifest> PrepareLoadOrderManifests(ILogger Logger)
 		{
 			Dictionary<FileReference, LoadOrderManifest> FileNameToManifest = new Dictionary<FileReference, LoadOrderManifest>();
-			
+
 			UEBuildPlatform BuildPlatform = UEBuildPlatform.GetBuildPlatform(Platform);
 
 			if (!bCompileMonolithic && BuildPlatform.RequiresLoadOrderManifest(Rules))
@@ -1970,7 +1964,7 @@ namespace UnrealBuildTool
 				Sorter.CycleHandling = TopologicalSorter<UEBuildModule>.CycleMode.BreakWithInfo;
 				Sorter.Logger = Logger;
 				Sorter.NodeToString = Module => Module.Name;
-				
+
 				if (!Sorter.Sort())
 				{
 					Logger.LogError("Failed to generate {ManifestFileName}: Couldn't sort dynamic modules in a way that would respect all dependencies (probably circular dependencies)", ManifestFileName);
@@ -2301,9 +2295,9 @@ namespace UnrealBuildTool
 
 			// Create the makefile
 			string ExternalMetadata = UEBuildPlatform.GetBuildPlatform(Platform).GetExternalBuildMetadata(ProjectFile);
-			TargetMakefile Makefile = new TargetMakefile(ExternalMetadata, Binaries[0].OutputFilePaths[0], ReceiptFileName, 
-				ProjectIntermediateDirectory, ProjectIntermediateDirectoryNoArch, TargetType, 
-				Rules.ConfigValueTracker, bDeployAfterCompile, UbtPlugins?.ToArray(), EnabledUbtPlugins?.ToArray(), 
+			TargetMakefile Makefile = new TargetMakefile(ExternalMetadata, Binaries[0].OutputFilePaths[0], ReceiptFileName,
+				ProjectIntermediateDirectory, ProjectIntermediateDirectoryNoArch, TargetType,
+				Rules.ConfigValueTracker, bDeployAfterCompile, UbtPlugins?.ToArray(), EnabledUbtPlugins?.ToArray(),
 				EnabledUhtPlugins?.ToArray());
 			Makefile.IsTestTarget = Rules.IsTestTarget;
 			TargetMakefileBuilder MakefileBuilder = new TargetMakefileBuilder(Makefile, Logger);
@@ -2447,9 +2441,9 @@ namespace UnrealBuildTool
 					}
 				}
 			}
-			
+
 			// Copy debugger visualizer files for each module to their intermediate directory 
-			foreach (UEBuildModule Module in Modules.Values) 
+			foreach (UEBuildModule Module in Modules.Values)
 			{
 				Module.CopyDebuggerVisualizers(TargetToolChain, MakefileBuilder, Logger);
 			}
@@ -3776,18 +3770,18 @@ namespace UnrealBuildTool
 				UEBuildModule Module = FindOrCreateModuleByName(ModuleName, PrecompileReferenceChain, Logger);
 				AllModules.Add(Module);
 				Module.RecursivelyCreateModules(
-					(string ModuleName, string ReferenceChain) => 
-					{ 
+					(string ModuleName, string ReferenceChain) =>
+					{
 						UEBuildModule FoundModule = FindOrCreateModuleByName(ModuleName, ReferenceChain, Logger);
 						AllModules.Add(FoundModule);
-						return FoundModule; 
-					}, 
+						return FoundModule;
+					},
 					PrecompileReferenceChain, Logger);
 			}
 
 			// Exclude additional modules that were added only for include-path-only purposes, and those that will have their Verse dependency satisfied
 			HashSet<UEBuildModuleCPP> ValidModules = new HashSet<UEBuildModuleCPP>(
-				AllModules.OfType<UEBuildModuleCPP>() .Where(x => x.PrivateIncludePathModules != null && (!x.bDependsOnVerse || Rules.bUseVerse)));
+				AllModules.OfType<UEBuildModuleCPP>().Where(x => x.PrivateIncludePathModules != null && (!x.bDependsOnVerse || Rules.bUseVerse)));
 
 			// Make sure precompiled modules don't reference any non-precompiled modules
 			foreach (UEBuildModuleCPP ValidModule in ValidModules)
@@ -4849,7 +4843,7 @@ namespace UnrealBuildTool
 			// Check if server-only code should be compiled out.
 			GlobalCompileEnvironment.Definitions.Add(String.Format("WITH_SERVER_CODE={0}", Rules.bWithServerCode ? 1 : 0));
 
-			GlobalCompileEnvironment.Definitions.Add(String.Format("UE_FNAME_OUTLINE_NUMBER={0}", Rules.bFNameOutlineNumber? 1 : 0));
+			GlobalCompileEnvironment.Definitions.Add(String.Format("UE_FNAME_OUTLINE_NUMBER={0}", Rules.bFNameOutlineNumber ? 1 : 0));
 
 			// Set the defines for Push Model
 			if (Rules.bWithPushModel)
