UnrealBuildTool: Automated code cleanup

diff --git a/Engine/Source/Programs/UnrealBuildTool/Executors/Experimental/FASTBuild.cs b/Engine/Source/Programs/UnrealBuildTool/Executors/Experimental/FASTBuild.cs
--- a/Engine/Source/Programs/UnrealBuildTool/Executors/Experimental/FASTBuild.cs
+++ b/Engine/Source/Programs/UnrealBuildTool/Executors/Experimental/FASTBuild.cs
@@ -15,16 +15,15 @@
 using System;
 using System.Collections;
 using System.Collections.Generic;
-using System.IO;
 using System.Diagnostics;
+using System.IO;
 using System.Linq;
-using EpicGames.Core;
 using System.Runtime.Serialization;
-using System.Text.RegularExpressions;
-using UnrealBuildBase;
 using System.Runtime.Versioning;
-using Microsoft.Extensions.Logging;
 using System.Threading.Tasks;
+using EpicGames.Core;
+using Microsoft.Extensions.Logging;
+using UnrealBuildBase;
 using UnrealBuildTool.Artifacts;
 
 namespace UnrealBuildTool
@@ -77,7 +76,7 @@ namespace UnrealBuildTool
 		/// </summary>
 		ActionExecutor LocalExecutor;
 
-		public readonly static string DefaultExecutableBasePath	= Path.Combine(Unreal.EngineDirectory.FullName, "Extras", "ThirdPartyNotUE", "FASTBuild");
+		public readonly static string DefaultExecutableBasePath = Path.Combine(Unreal.EngineDirectory.FullName, "Extras", "ThirdPartyNotUE", "FASTBuild");
 
 		//////////////////////////////////////////
 		// Tweakables
@@ -89,7 +88,7 @@ namespace UnrealBuildTool
 		/// Used to specify the location of fbuild.exe if the distributed binary isn't being used
 		/// </summary>
 		[XmlConfigFile]
-		public static string? FBuildExecutablePath	= null;
+		public static string? FBuildExecutablePath = null;
 
 		/////////////////
 		// Distribution
@@ -98,13 +97,13 @@ namespace UnrealBuildTool
 		/// Controls network build distribution
 		/// </summary>
 		[XmlConfigFile]
-		public static bool bEnableDistribution		= true;
+		public static bool bEnableDistribution = true;
 
 		/// <summary>
 		/// Used to specify the location of the brokerage. If null, FASTBuild will fall back to checking FASTBUILD_BROKERAGE_PATH
 		/// </summary>
 		[XmlConfigFile]
-		public static string? FBuildBrokeragePath	= null;
+		public static string? FBuildBrokeragePath = null;
 
 		/// <summary>
 		/// Used to specify the FASTBuild coordinator IP or network name. If null, FASTBuild will fall back to checking FASTBUILD_COORDINATOR
@@ -119,19 +118,19 @@ namespace UnrealBuildTool
 		/// Controls whether to use caching at all. CachePath and FASTCacheMode are only relevant if this is enabled.
 		/// </summary>
 		[XmlConfigFile]
-		public static bool bEnableCaching			= true;
+		public static bool bEnableCaching = true;
 
 		/// <summary>
 		/// Cache access mode - only relevant if bEnableCaching is true;
 		/// </summary>
 		[XmlConfigFile]
-		public static FASTBuildCacheMode CacheMode	= FASTBuildCacheMode.ReadOnly;
+		public static FASTBuildCacheMode CacheMode = FASTBuildCacheMode.ReadOnly;
 
 		/// <summary>
 		/// Used to specify the location of the cache. If null, FASTBuild will fall back to checking FASTBUILD_CACHE_PATH
 		/// </summary>
 		[XmlConfigFile]
-		public static string? FBuildCachePath		= null;
+		public static string? FBuildCachePath = null;
 
 		/////////////////
 		// Misc Options
@@ -140,13 +139,13 @@ namespace UnrealBuildTool
 		/// Whether to force remote
 		/// </summary>
 		[XmlConfigFile]
-		public static bool bForceRemote				= false;
+		public static bool bForceRemote = false;
 
 		/// <summary>
 		/// Whether to stop on error
 		/// </summary>
 		[XmlConfigFile]
-		public static bool bStopOnError				= false;
+		public static bool bStopOnError = false;
 
 		/// <summary>
 		/// Which MSVC CRT Redist version to use
@@ -355,15 +354,15 @@ namespace UnrealBuildTool
 
 		private readonly HashSet<string> ForceOverwriteCompilerOptionModules = new HashSet<string>()
 		{
-			"Module.USDStageImporter", 
-			"Module.USDUtilities", 
-			"Module.UnrealUSDWrapper", 
-			"Module.USDStage", 
-			"Module.USDSchemas", 
+			"Module.USDStageImporter",
+			"Module.USDUtilities",
+			"Module.UnrealUSDWrapper",
+			"Module.USDStage",
+			"Module.USDSchemas",
 			"Module.GeometryCacheUSD",
-			"Module.USDStageEditorViewModels", 
-			"Module.USDTests", 
-			"Module.USDStageEditor", 
+			"Module.USDStageEditorViewModels",
+			"Module.USDTests",
+			"Module.USDStageEditor",
 			"Module.USDExporter"
 		};
 
@@ -422,7 +421,7 @@ namespace UnrealBuildTool
 				foreach (Tuple<string, Func<LinkedAction, string>, FBBuildType> BuildTypeSearchParam in BuildTypeSearchParams)
 				{
 					if (BuildTypeSearchParam.Item3.Equals(FBBuildType.Apple) &&
-						(BuildTypeSearchParam.Item2(Action).Contains("Win64", StringComparison.OrdinalIgnoreCase) || 
+						(BuildTypeSearchParam.Item2(Action).Contains("Win64", StringComparison.OrdinalIgnoreCase) ||
 						BuildTypeSearchParam.Item2(Action).Contains("X64", StringComparison.OrdinalIgnoreCase)))
 					{
 						continue;
@@ -444,16 +443,16 @@ namespace UnrealBuildTool
 			return false;
 		}
 
-		private bool IsMSVC()						{ return BuildType == FBBuildType.Windows;	}
-		private bool IsApple()						{ return BuildType == FBBuildType.Apple;	}
-		
+		private bool IsMSVC() { return BuildType == FBBuildType.Windows; }
+		private bool IsApple() { return BuildType == FBBuildType.Apple; }
+
 		private string GetCompilerName()
 		{
 			switch (BuildType)
 			{
 				default:
-				case FBBuildType.Windows:	return "UECompiler";
-				case FBBuildType.Apple:		return "UEAppleCompiler";
+				case FBBuildType.Windows: return "UECompiler";
+				case FBBuildType.Apple: return "UEAppleCompiler";
 			}
 		}
 
@@ -474,7 +473,7 @@ namespace UnrealBuildTool
 				{
 					return false;
 				}
-				
+
 				return ExecuteBffFile(FASTBuildFilePath, Logger);
 			}
 
@@ -696,23 +695,23 @@ namespace UnrealBuildTool
 				}
 
 				// Skip the following tokens:
-				if ((Token == "/I")						||
-					(Token == "/external:I")			||
-					(Token == "/l")						||
-					(Token == "/D")						||
-					(Token == "-D")						||
-					(Token == "-x")						||
-					(Token == "-F")						||
-					(Token == "-arch")					||
-					(Token == "-isysroot")				||
-					(Token == "-include")				||
-					(Token == "-current_version")		||
-					(Token == "-compatibility_version")	||
-					(Token == "-rpath")					||
-					(Token == "-weak_library")			||
-					(Token == "-weak_framework")		||
-					(Token == "-framework")				||
-					(Token == "/sourceDependencies")	||
+				if ((Token == "/I") ||
+					(Token == "/external:I") ||
+					(Token == "/l") ||
+					(Token == "/D") ||
+					(Token == "-D") ||
+					(Token == "-x") ||
+					(Token == "-F") ||
+					(Token == "-arch") ||
+					(Token == "-isysroot") ||
+					(Token == "-include") ||
+					(Token == "-current_version") ||
+					(Token == "-compatibility_version") ||
+					(Token == "-rpath") ||
+					(Token == "-weak_library") ||
+					(Token == "-weak_framework") ||
+					(Token == "-framework") ||
+					(Token == "/sourceDependencies") ||
 					(Token == "/sourceDependencies:directives"))
 				{
 					++i;
@@ -800,8 +799,8 @@ namespace UnrealBuildTool
 			{
 				// This may fail if the caller emptied PATH; we try to ignore the problem since
 				// it probably means we are building for another platform.
-                if(BuildType == FBBuildType.Windows)
-                {
+				if (BuildType == FBBuildType.Windows)
+				{
 					VCEnv = VCEnvironment.Create(
 						Compiler: WindowsPlatform.GetDefaultCompiler(null, UnrealArch.X64, Logger),
 						ToolChain: WindowsCompiler.Default,
@@ -911,7 +910,7 @@ namespace UnrealBuildTool
 					{
 						AddText("\t\t'$Root$/{cluiSubDirName}/clui.dll'\n");
 					}
-					cluiDllPath = new FileReference(VCEnv.GetToolPath() + "{cluiSubDirName}/clui.dll");					
+					cluiDllPath = new FileReference(VCEnv.GetToolPath() + "{cluiSubDirName}/clui.dll");
 				}
 				else
 				{
@@ -996,7 +995,7 @@ namespace UnrealBuildTool
 				}
 
 
-                AddText("\t}\n"); //End extra files
+				AddText("\t}\n"); //End extra files
 
 				AddText($"\t.CompilerFamily = 'msvc'\n");
 				AddText("}\n\n"); //End compiler
@@ -1279,29 +1278,29 @@ namespace UnrealBuildTool
 				}
 			}
 
-			string DistArgument				= bEnableDistribution ? "-dist" : "";
-			string ForceRemoteArgument		= bForceRemote ? "-forceremote" : "";
-			string NoStopOnErrorArgument	= bStopOnError ? "" : "-nostoponerror";
-			string IDEArgument				= IsApple() ? "" : "-ide";
-			string MaxProcesses				= "-j" + ((ParallelExecutor)LocalExecutor).NumParallelProcesses;
+			string DistArgument = bEnableDistribution ? "-dist" : "";
+			string ForceRemoteArgument = bForceRemote ? "-forceremote" : "";
+			string NoStopOnErrorArgument = bStopOnError ? "" : "-nostoponerror";
+			string IDEArgument = IsApple() ? "" : "-ide";
+			string MaxProcesses = "-j" + ((ParallelExecutor)LocalExecutor).NumParallelProcesses;
 
 			// Interesting flags for FASTBuild:
 			// -nostoponerror, -verbose, -monitor (if FASTBuild Monitor Visual Studio Extension is installed!)
 			// Yassine: The -clean is to bypass the FASTBuild internal
 			// dependencies checks (cached in the fdb) as it could create some conflicts with UBT.
 			// Basically we want FB to stupidly compile what UBT tells it to.
-			string FBCommandLine	= $"-monitor -summary {DistArgument} {CacheArgument} {IDEArgument} {MaxProcesses} -clean -config \"{BffFilePath}\" {NoStopOnErrorArgument} {ForceRemoteArgument}";
+			string FBCommandLine = $"-monitor -summary {DistArgument} {CacheArgument} {IDEArgument} {MaxProcesses} -clean -config \"{BffFilePath}\" {NoStopOnErrorArgument} {ForceRemoteArgument}";
 
 			Logger.LogInformation("FBuild Command Line Arguments: '{FBCommandLine}", FBCommandLine);
 
-			string FBExecutable		= GetExecutablePath()!;
-			string WorkingDirectory	= Path.GetFullPath(Path.Combine(Unreal.EngineDirectory.MakeRelativeTo(DirectoryReference.GetCurrentDirectory()), "Source"));
+			string FBExecutable = GetExecutablePath()!;
+			string WorkingDirectory = Path.GetFullPath(Path.Combine(Unreal.EngineDirectory.MakeRelativeTo(DirectoryReference.GetCurrentDirectory()), "Source"));
 
-			ProcessStartInfo FBStartInfo		= new ProcessStartInfo(FBExecutable, FBCommandLine);
-			FBStartInfo.UseShellExecute			= false;
-			FBStartInfo.WorkingDirectory		= WorkingDirectory;
-			FBStartInfo.RedirectStandardError	= true;
-			FBStartInfo.RedirectStandardOutput	= true;
+			ProcessStartInfo FBStartInfo = new ProcessStartInfo(FBExecutable, FBCommandLine);
+			FBStartInfo.UseShellExecute = false;
+			FBStartInfo.WorkingDirectory = WorkingDirectory;
+			FBStartInfo.RedirectStandardError = true;
+			FBStartInfo.RedirectStandardOutput = true;
 
 			string? Coordinator = GetCoordinator();
 			if (!string.IsNullOrEmpty(Coordinator) && !FBStartInfo.EnvironmentVariables.ContainsKey("FASTBUILD_COORDINATOR"))
@@ -1322,9 +1321,9 @@ namespace UnrealBuildTool
 
 			try
 			{
-				Process FBProcess				= new Process();
-				FBProcess.StartInfo				= FBStartInfo;
-				FBProcess.EnableRaisingEvents	= true;
+				Process FBProcess = new Process();
+				FBProcess.StartInfo = FBStartInfo;
+				FBProcess.EnableRaisingEvents = true;
 
 				DataReceivedEventHandler OutputEventHandler = (Sender, Args) =>
 				{
diff --git a/Engine/Source/Programs/UnrealBuildTool/Executors/HybridExecutor.cs b/Engine/Source/Programs/UnrealBuildTool/Executors/HybridExecutor.cs
--- a/Engine/Source/Programs/UnrealBuildTool/Executors/HybridExecutor.cs
+++ b/Engine/Source/Programs/UnrealBuildTool/Executors/HybridExecutor.cs
@@ -1,10 +1,9 @@
 // Copyright Epic Games, Inc. All Rights Reserved.
 
-using EpicGames.Core;
-using Microsoft.Extensions.Logging;
 using System.Collections.Generic;
 using System.Linq;
 using System.Threading.Tasks;
+using Microsoft.Extensions.Logging;
 using UnrealBuildTool.Artifacts;
 
 namespace UnrealBuildTool
@@ -93,7 +92,7 @@ namespace UnrealBuildTool
 
 			// Build up a set of leaf actions in several iterations, ensuring that the number of leaf actions in each 
 			HashSet<LinkedAction> LeafActions = new HashSet<LinkedAction>();
-			for (;;)
+			for (; ; )
 			{
 				// Find all the leaf actions in the graph
 				List<LinkedAction> NewLeafActions = new List<LinkedAction>();
diff --git a/Engine/Source/Programs/UnrealBuildTool/Executors/ImmediateActionQueue.cs b/Engine/Source/Programs/UnrealBuildTool/Executors/ImmediateActionQueue.cs
--- a/Engine/Source/Programs/UnrealBuildTool/Executors/ImmediateActionQueue.cs
+++ b/Engine/Source/Programs/UnrealBuildTool/Executors/ImmediateActionQueue.cs
@@ -1,15 +1,13 @@
 // Copyright Epic Games, Inc. All Rights Reserved.
 
-using EpicGames.Core;
-using Microsoft.CodeAnalysis;
-using Microsoft.CodeAnalysis.CSharp.Syntax;
-using Microsoft.Extensions.Logging;
 using System;
-using System.Collections;
 using System.Collections.Generic;
 using System.Linq;
 using System.Threading;
 using System.Threading.Tasks;
+using EpicGames.Core;
+using Microsoft.CodeAnalysis;
+using Microsoft.Extensions.Logging;
 using UnrealBuildBase;
 using UnrealBuildTool.Artifacts;
 
@@ -327,10 +325,10 @@ namespace UnrealBuildTool
 			foreach (LinkedAction action in actions)
 			{
 				action.SortIndex = index;
-				Actions[index++] = new ActionState 
-				{ 
-					Action = action, 
-					Status = ActionStatus.Queued, 
+				Actions[index++] = new ActionState
+				{
+					Action = action,
+					Status = ActionStatus.Queued,
 					Phase = action.ArtifactMode.HasFlag(ArtifactMode.Enabled) ? initialPhase : ActionPhase.Compile,
 					Results = null,
 				};
@@ -431,12 +429,12 @@ namespace UnrealBuildTool
 		{
 			foreach (ActionState actionState in Actions)
 			{
-				if (actionState.Status == ActionStatus.Queued && 
-					actionState.Phase == ActionPhase.Compile && 
+				if (actionState.Status == ActionStatus.Queued &&
+					actionState.Phase == ActionPhase.Compile &&
 					GetActionReadyState(actionState.Action) == ActionReadyState.Ready)
 				{
 					yield return actionState.Action;
-				}				
+				}
 			}
 		}
 
@@ -449,7 +447,7 @@ namespace UnrealBuildTool
 		{
 
 			// Loop until we have an action or no completed actions (i.e. no propagated errors)
-			for(; ; )
+			for (; ; )
 			{
 				bool hasCanceled = false;
 
diff --git a/Engine/Source/Programs/UnrealBuildTool/Executors/XGE.cs b/Engine/Source/Programs/UnrealBuildTool/Executors/XGE.cs
--- a/Engine/Source/Programs/UnrealBuildTool/Executors/XGE.cs
+++ b/Engine/Source/Programs/UnrealBuildTool/Executors/XGE.cs
@@ -3,23 +3,20 @@
 using System;
 using System.Collections;
 using System.Collections.Generic;
-using System.IO;
 using System.Diagnostics;
-using System.Xml;
-using System.Text.RegularExpressions;
+using System.Diagnostics.CodeAnalysis;
+using System.IO;
 using System.Linq;
-using System.Reflection;
-using Microsoft.Win32;
-using System.Text;
-using EpicGames.Core;
 using System.Net;
 using System.Net.NetworkInformation;
 using System.Runtime.InteropServices;
-using UnrealBuildBase;
-using System.Diagnostics.CodeAnalysis;
 using System.Runtime.Versioning;
-using Microsoft.Extensions.Logging;
 using System.Threading.Tasks;
+using System.Xml;
+using EpicGames.Core;
+using Microsoft.Extensions.Logging;
+using Microsoft.Win32;
+using UnrealBuildBase;
 using UnrealBuildTool.Artifacts;
 
 namespace UnrealBuildTool
@@ -109,12 +106,12 @@ namespace UnrealBuildTool
 			if (OperatingSystem.IsWindows())
 			{
 				string? XgConsoleExe;
-				if(TryGetXgConsoleExecutableFromRegistry(RegistryView.Registry32, out XgConsoleExe))
+				if (TryGetXgConsoleExecutableFromRegistry(RegistryView.Registry32, out XgConsoleExe))
 				{
 					OutXgConsoleExe = XgConsoleExe;
 					return true;
 				}
-				if(TryGetXgConsoleExecutableFromRegistry(RegistryView.Registry64, out XgConsoleExe))
+				if (TryGetXgConsoleExecutableFromRegistry(RegistryView.Registry64, out XgConsoleExe))
 				{
 					OutXgConsoleExe = XgConsoleExe;
 					return true;
@@ -163,17 +160,17 @@ namespace UnrealBuildTool
 		{
 			try
 			{
-				using(RegistryKey BaseKey = RegistryKey.OpenBaseKey(RegistryHive.LocalMachine, View))
+				using (RegistryKey BaseKey = RegistryKey.OpenBaseKey(RegistryHive.LocalMachine, View))
 				{
 					using (RegistryKey? Key = BaseKey.OpenSubKey("SOFTWARE\\Xoreax\\IncrediBuild\\Builder", false))
 					{
-						if(Key != null)
+						if (Key != null)
 						{
 							string? Folder = Key.GetValue("Folder", null) as string;
-							if(!String.IsNullOrEmpty(Folder))
+							if (!String.IsNullOrEmpty(Folder))
 							{
 								string FileName = Path.Combine(Folder, "xgConsole.exe");
-								if(File.Exists(FileName))
+								if (File.Exists(FileName))
 								{
 									OutXgConsoleExe = FileName;
 									return true;
@@ -183,7 +180,7 @@ namespace UnrealBuildTool
 					}
 				}
 			}
-			catch(Exception Ex)
+			catch (Exception Ex)
 			{
 				Log.WriteException(Ex, null);
 			}
@@ -332,7 +329,7 @@ namespace UnrealBuildTool
 						return false;
 					}
 				}
-				catch(Exception Ex)
+				catch (Exception Ex)
 				{
 					Logger.LogDebug("Unable to query for status of Incredibuild service: {Ex}", ExceptionUtils.FormatExceptionDetails(Ex));
 					return false;
@@ -377,10 +374,10 @@ namespace UnrealBuildTool
 
 		public static void ExportActions(List<LinkedAction> ActionsToExecute, ILogger Logger)
 		{
-			for(int FileNum = 0;;FileNum++)
+			for (int FileNum = 0; ; FileNum++)
 			{
 				string OutFile = Path.Combine(Unreal.EngineDirectory.FullName, "Intermediate", "Build", String.Format("UBTExport.{0}.xge.xml", FileNum.ToString("D3")));
-				if(!File.Exists(OutFile))
+				if (!File.Exists(OutFile))
 				{
 					ExportActions(ActionsToExecute, OutFile, Logger);
 					break;
@@ -481,7 +478,7 @@ namespace UnrealBuildTool
 
 				// Don't allow remote linking if on VPN.
 				bool CanExecuteRemotely = Action.bCanExecuteRemotely && Action.bCanExecuteRemotelyWithXGE;
-				if(CanExecuteRemotely && Action.ActionType == ActionType.Link)
+				if (CanExecuteRemotely && Action.ActionType == ActionType.Link)
 				{
 					if (HostOnVpn || !bAllowRemoteLinking)
 					{
@@ -511,7 +508,7 @@ namespace UnrealBuildTool
 				{
 					ToolElement.SetAttribute("OutputPrefix", OutputPrefix);
 				}
-				if(Action.GroupNames.Count > 0)
+				if (Action.GroupNames.Count > 0)
 				{
 					ToolElement.SetAttribute("GroupPrefix", String.Format("** For {0} **", String.Join(" + ", Action.GroupNames)));
 				}
@@ -540,7 +537,7 @@ namespace UnrealBuildTool
 					string.Join(
 						",",
 						Action.ProducedItems.Select(
-							delegate(FileItem ProducedItem) { return ProducedItem.Location.GetFileName(); }
+							delegate (FileItem ProducedItem) { return ProducedItem.Location.GetFileName(); }
 							).ToArray()
 						)
 					);
@@ -578,7 +575,7 @@ namespace UnrealBuildTool
 
 				// Create a semi-colon separated list of the other tasks this task depends on the results of.
 				List<string> DependencyNames = new List<string>();
-				foreach(LinkedAction PrerequisiteAction in Action.PrerequisiteActions)
+				foreach (LinkedAction PrerequisiteAction in Action.PrerequisiteActions)
 				{
 					if (Actions.Contains(PrerequisiteAction))
 					{
@@ -641,7 +638,7 @@ namespace UnrealBuildTool
 			}
 
 			string? XgConsolePath;
-			if(!TryGetXgConsoleExecutable(out XgConsolePath))
+			if (!TryGetXgConsoleExecutable(out XgConsolePath))
 			{
 				throw new BuildException("Unable to find xgConsole executable.");
 			}
@@ -721,7 +718,7 @@ namespace UnrealBuildTool
 				// Create a wrapper delegate that will parse the output actions
 				DataReceivedEventHandler EventHandlerWrapper = (Sender, Args) =>
 				{
-					if(Args.Data != null)
+					if (Args.Data != null)
 					{
 						string Text = Args.Data;
 						if (Text.StartsWith(ProgressMarkupPrefix))
@@ -753,7 +750,7 @@ namespace UnrealBuildTool
 
 							// Strip out anything that is just an XGE timer. Some programs don't output anything except the progress text.
 							Text = Args.Data.Substring(MarkupLength);
-							if(Text.StartsWith(" (") && Text.EndsWith(")"))
+							if (Text.StartsWith(" (") && Text.EndsWith(")"))
 							{
 								// Write the progress text with the next line of output if the current doesn't have any status.
 								ProgressText = Text.Trim();
diff --git a/Engine/Source/Programs/UnrealBuildTool/Matchers/XoreaxEventMatcher.cs b/Engine/Source/Programs/UnrealBuildTool/Matchers/XoreaxEventMatcher.cs
--- a/Engine/Source/Programs/UnrealBuildTool/Matchers/XoreaxEventMatcher.cs
+++ b/Engine/Source/Programs/UnrealBuildTool/Matchers/XoreaxEventMatcher.cs
@@ -1,7 +1,5 @@
 // Copyright Epic Games, Inc. All Rights Reserved.
 
-using System;
-using System.Diagnostics.CodeAnalysis;
 using System.Text.RegularExpressions;
 using EpicGames.Core;
 using Microsoft.Extensions.Logging;
@@ -14,7 +12,7 @@ namespace UnrealBuildTool.Matchers
 		public const string PropertyAgent = "agent";
 		public const string PropertyDuration = "duration";
 		public const string PropertyStartTime = "startTime";
-		
+
 		static readonly Regex s_buildService = new Regex(@"\(BuildService.exe\) is not running");
 
 		static readonly Regex s_xgConsole = new Regex(@"BUILD FAILED: (.*)xgConsole\.exe(.*)");
@@ -57,7 +55,7 @@ namespace UnrealBuildTool.Matchers
 				LogEventBuilder builder = new LogEventBuilder(cursor);
 				builder.MoveNext();
 
-				for(int idx = 0; idx < 100; idx++)
+				for (int idx = 0; idx < 100; idx++)
 				{
 					if (builder.Current.IsMatch(idx, s_buildSystemErrorEnd))
 					{
@@ -118,7 +116,7 @@ namespace UnrealBuildTool.Matchers
 					return builder.ToMatch(LogEventPriority.High, LogLevel.Information, eventId);
 				}
 			}
-			
+
 			return null;
 		}
 	}
diff --git a/Engine/Source/Programs/UnrealBuildTool/Modes/AggregateClangTimingInfo.cs b/Engine/Source/Programs/UnrealBuildTool/Modes/AggregateClangTimingInfo.cs
--- a/Engine/Source/Programs/UnrealBuildTool/Modes/AggregateClangTimingInfo.cs
+++ b/Engine/Source/Programs/UnrealBuildTool/Modes/AggregateClangTimingInfo.cs
@@ -1,16 +1,15 @@
 // Copyright Epic Games, Inc. All Rights Reserved.
 
-using EpicGames.Core;
-using Microsoft.Extensions.Logging;
 using System;
 using System.Collections.Concurrent;
 using System.Collections.Generic;
 using System.IO;
 using System.IO.Compression;
 using System.Linq;
-using System.Reflection;
 using System.Text.Json;
 using System.Threading.Tasks;
+using EpicGames.Core;
+using Microsoft.Extensions.Logging;
 
 namespace UnrealBuildTool
 {
@@ -71,7 +70,7 @@ namespace UnrealBuildTool
 			public TraceData(FileReference inputFile, ClangTrace? trace)
 			{
 				SourceFile = inputFile;
-				
+
 				TotalExecuteCompiler = trace?.traceEvents?.FindLast(x => string.Equals(x.name, "Total ExecuteCompiler"))?.dur ?? 0;
 
 				// Subset of execute compiler
diff --git a/Engine/Source/Programs/UnrealBuildTool/Modes/AnalyzeMode.cs b/Engine/Source/Programs/UnrealBuildTool/Modes/AnalyzeMode.cs
--- a/Engine/Source/Programs/UnrealBuildTool/Modes/AnalyzeMode.cs
+++ b/Engine/Source/Programs/UnrealBuildTool/Modes/AnalyzeMode.cs
@@ -1,15 +1,12 @@
 // Copyright Epic Games, Inc. All Rights Reserved.
 
-using EpicGames.Core;
 using System;
 using System.Collections.Generic;
 using System.IO;
 using System.Linq;
-using System.Text;
 using System.Threading.Tasks;
-using System.Xml;
+using EpicGames.Core;
 using Microsoft.Extensions.Logging;
-using UnrealBuildBase;
 
 namespace UnrealBuildTool.Modes
 {
diff --git a/Engine/Source/Programs/UnrealBuildTool/Modes/BuildMode.cs b/Engine/Source/Programs/UnrealBuildTool/Modes/BuildMode.cs
--- a/Engine/Source/Programs/UnrealBuildTool/Modes/BuildMode.cs
+++ b/Engine/Source/Programs/UnrealBuildTool/Modes/BuildMode.cs
@@ -1,16 +1,16 @@
 // Copyright Epic Games, Inc. All Rights Reserved.
 
 using System;
+using System.Collections.Concurrent;
 using System.Collections.Generic;
 using System.Diagnostics;
 using System.Linq;
 using System.Text;
+using System.Threading.Tasks;
 using EpicGames.Core;
+using Microsoft.Extensions.Logging;
 using OpenTracing.Util;
 using UnrealBuildBase;
-using Microsoft.Extensions.Logging;
-using System.Collections.Concurrent;
-using System.Threading.Tasks;
 using UnrealBuildTool.Artifacts;
 
 namespace UnrealBuildTool
@@ -242,7 +242,7 @@ namespace UnrealBuildTool
 					{
 						Options |= BuildOptions.XGEExport;
 					}
-					if(bNoEngineChanges)
+					if (bNoEngineChanges)
 					{
 						Options |= BuildOptions.NoEngineChanges;
 					}
@@ -281,9 +281,9 @@ namespace UnrealBuildTool
 			ProjectDescriptor ProjectDescriptor = ProjectDescriptor.FromFile(TargetDescriptor.ProjectFile);
 			List<string[]> InitCommandBatches = new List<string[]>();
 
-			if(ProjectDescriptor != null && ProjectDescriptor.InitSteps != null)
+			if (ProjectDescriptor != null && ProjectDescriptor.InitSteps != null)
 			{
-				if(ProjectDescriptor.InitSteps.TryGetCommands(BuildHostPlatform.Current.Platform, out string[]? Commands))
+				if (ProjectDescriptor.InitSteps.TryGetCommands(BuildHostPlatform.Current.Platform, out string[]? Commands))
 				{
 					InitCommandBatches.Add(Commands);
 				}
@@ -314,7 +314,7 @@ namespace UnrealBuildTool
 		private static FileReference[] WriteInitScripts(TargetDescriptor TargetDescriptor, UnrealTargetPlatform HostPlatform, DirectoryReference Directory, string FilePrefix, List<string[]> CommandBatches)
 		{
 			List<FileReference> ScriptFiles = new List<FileReference>();
-			foreach(string[] CommandBatch in CommandBatches)
+			foreach (string[] CommandBatch in CommandBatches)
 			{
 				// Find all the standard variables
 				Dictionary<string, string> Variables = GetTargetVariables(TargetDescriptor);
@@ -325,15 +325,15 @@ namespace UnrealBuildTool
 
 				// Write it to disk
 				List<string> Contents = new List<string>();
-				if(RuntimePlatform.IsWindows)
+				if (RuntimePlatform.IsWindows)
 				{
 					Contents.Insert(0, "@echo off");
 				}
-				foreach(string Command in CommandBatch)
+				foreach (string Command in CommandBatch)
 				{
 					Contents.Add(Utils.ExpandVariables(Command, Variables));
 				}
-				if(!DirectoryReference.Exists(ScriptFile.Directory))
+				if (!DirectoryReference.Exists(ScriptFile.Directory))
 				{
 					DirectoryReference.CreateDirectory(ScriptFile.Directory);
 				}
@@ -352,13 +352,13 @@ namespace UnrealBuildTool
 		/// <returns>Map of variable names to values</returns>
 		private static Dictionary<string, string> GetTargetVariables(TargetDescriptor TargetDescriptor)
 		{
-			Dictionary<string, string> Variables = new Dictionary<string,string>();
+			Dictionary<string, string> Variables = new Dictionary<string, string>();
 			Variables.Add("RootDir", Unreal.RootDirectory.FullName);
 			Variables.Add("EngineDir", Unreal.EngineDirectory.FullName);
 			Variables.Add("TargetName", TargetDescriptor.Name);
 			Variables.Add("TargetPlatform", TargetDescriptor.Platform.ToString());
 			Variables.Add("TargetConfiguration", TargetDescriptor.Configuration.ToString());
-			if(TargetDescriptor.ProjectFile != null)
+			if (TargetDescriptor.ProjectFile != null)
 			{
 				Variables.Add("ProjectDir", TargetDescriptor.ProjectFile.Directory.FullName);
 				Variables.Add("ProjectFile", TargetDescriptor.ProjectFile.FullName);
@@ -515,7 +515,7 @@ namespace UnrealBuildTool
 					using (GlobalTracer.Instance.BuildSpan("Reading action history").StartActive())
 					{
 						TargetDescriptor TargetDescriptor = TargetDescriptors[TargetIdx];
-						if(TargetDescriptor.ProjectFile != null)
+						if (TargetDescriptor.ProjectFile != null)
 						{
 							History.Mount(TargetDescriptor.ProjectFile.Directory);
 						}
@@ -548,7 +548,7 @@ namespace UnrealBuildTool
 					};
 
 					Logger.LogInformation("Experimental artifact system using '{directory}' in {mode} mode", artifactDirectory.FullName, mode);
-				
+
 					actionArtifactCache = ActionArtifactCache.CreateHordeFileCache(artifactDirectory, CppDependencies, Logger);
 					actionArtifactCache.EnableReads = BuildConfiguration.bArtifactRead;
 					actionArtifactCache.EnableWrites = BuildConfiguration.bArtifactWrites;
@@ -589,7 +589,7 @@ namespace UnrealBuildTool
 						Dictionary<FileItem, List<string>> ModuleImports = new Dictionary<FileItem, List<string>>();
 
 						List<FileItem> CompiledModuleInterfaces = new List<FileItem>();
-						foreach(LinkedAction ModuleDependencyAction in ModuleDependencyActions)
+						foreach (LinkedAction ModuleDependencyAction in ModuleDependencyActions)
 						{
 							ICppCompileAction? CppModulesAction = ModuleDependencyAction.Inner as ICppCompileAction;
 							if (CppModulesAction != null && CppModulesAction.CompiledModuleInterfaceFile != null)
@@ -618,7 +618,7 @@ namespace UnrealBuildTool
 								ICppCompileAction CppModulesAction = (ICppCompileAction)PrerequisiteAction.Inner;
 
 								List<string>? ImportedModules;
-								if(ModuleImports.TryGetValue(CppModulesAction.CompiledModuleInterfaceFile!, out ImportedModules))
+								if (ModuleImports.TryGetValue(CppModulesAction.CompiledModuleInterfaceFile!, out ImportedModules))
 								{
 									foreach (string ImportedModule in ImportedModules)
 									{
@@ -637,12 +637,12 @@ namespace UnrealBuildTool
 									}
 								}
 							}
-							else if(PrerequisiteAction.ActionType == ActionType.Compile)
+							else if (PrerequisiteAction.ActionType == ActionType.Compile)
 							{
-								foreach(FileItem PrerequisiteItem in PrerequisiteAction.PrerequisiteItems)
+								foreach (FileItem PrerequisiteItem in PrerequisiteAction.PrerequisiteItems)
 								{
 									string? ModuleName;
-									if(ModuleInputs.TryGetValue(PrerequisiteItem, out ModuleName))
+									if (ModuleInputs.TryGetValue(PrerequisiteItem, out ModuleName))
 									{
 										Action NewAction = new Action(PrerequisiteAction.Inner);
 										NewAction.CommandArguments += String.Format(" /reference \"{0}={1}\"", ModuleName, PrerequisiteItem.AbsolutePath);
@@ -672,7 +672,7 @@ namespace UnrealBuildTool
 
 				// Allow hot reload to override the actions
 				int HotReloadTargetIdx = -1;
-				for(int Idx = 0; Idx < TargetDescriptors.Count; Idx++)
+				for (int Idx = 0; Idx < TargetDescriptors.Count; Idx++)
 				{
 					if (TargetDescriptors[Idx].HotReloadMode != HotReloadMode.Disabled)
 					{
@@ -690,7 +690,7 @@ namespace UnrealBuildTool
 					{
 						HotReload.CheckForLiveCodingSessionActive(TargetDescriptors[Idx], Makefiles[Idx], BuildConfiguration, Logger);
 					}
-				}	
+				}
 
 				if (HotReloadTargetIdx != -1)
 				{
@@ -764,7 +764,7 @@ namespace UnrealBuildTool
 				else
 				{
 					// Execute the actions
-					if(MergedActionsToExecute.Count == 0)
+					if (MergedActionsToExecute.Count == 0)
 					{
 						if (TargetDescriptors.Any(x => !x.bQuiet))
 						{
@@ -808,7 +808,7 @@ namespace UnrealBuildTool
 					}
 
 					// Run the deployment steps
-					foreach(TargetMakefile Makefile in Makefiles)
+					foreach (TargetMakefile Makefile in Makefiles)
 					{
 						// Receipt file may not exist when compiling specific files
 						if (Makefile.bDeployAfterCompile && FileReference.Exists(Makefile.ReceiptFile))
@@ -919,7 +919,7 @@ namespace UnrealBuildTool
 		/// <param name="Logger">Logger for output</param>
 		internal static List<FileReference> GetAllSourceFilesIncludingHeader(List<FileReference> SpecificFilesToCompile, FileReference? ProjectFile, List<IExternalAction> Actions, ILogger Logger)
 		{
-			Func<FileReference, bool> IsCode = x => IsHeader(x)|| x.HasExtension(".cpp") || x.HasExtension(".c");
+			Func<FileReference, bool> IsCode = x => IsHeader(x) || x.HasExtension(".cpp") || x.HasExtension(".c");
 			List<FileReference> SpecificHeaderFiles = SpecificFilesToCompile.Where(IsHeader).ToList();
 			if (SpecificHeaderFiles.Count == 0)
 			{
@@ -1031,18 +1031,18 @@ namespace UnrealBuildTool
 				}
 			}
 
-			if(OutputIndices.Count == 1)
+			if (OutputIndices.Count == 1)
 			{
-				foreach(string Diagnostic in Makefiles[OutputIndices[0]].Diagnostics)
+				foreach (string Diagnostic in Makefiles[OutputIndices[0]].Diagnostics)
 				{
 					Logger.LogInformation("{Diagnostic}", Diagnostic);
 				}
 			}
 			else
 			{
-				foreach(int OutputIndex in OutputIndices)
+				foreach (int OutputIndex in OutputIndices)
 				{
-					foreach(string Diagnostic in Makefiles[OutputIndex].Diagnostics)
+					foreach (string Diagnostic in Makefiles[OutputIndex].Diagnostics)
 					{
 						Logger.LogInformation("{Name}: {Diagnostic}", TargetDescriptors[OutputIndex].Name, Diagnostic);
 					}
@@ -1068,14 +1068,14 @@ namespace UnrealBuildTool
 		{
 			// Get the path to the makefile for this target
 			FileReference? MakefileLocation = null;
-			if(BuildConfiguration.bUseUBTMakefiles)
+			if (BuildConfiguration.bUseUBTMakefiles)
 			{
 				MakefileLocation = TargetMakefile.GetLocation(TargetDescriptor.ProjectFile, TargetDescriptor.Name, TargetDescriptor.Platform, TargetDescriptor.Architectures, TargetDescriptor.Configuration);
 			}
 
 			// Try to load an existing makefile
 			TargetMakefile? Makefile = null;
-			if(MakefileLocation != null)
+			if (MakefileLocation != null)
 			{
 				using (GlobalTracer.Instance.BuildSpan("TargetMakefile.Load()").StartActive())
 				{
@@ -1090,7 +1090,7 @@ namespace UnrealBuildTool
 
 			// If we have a makefile, execute the pre-build steps and check it's still valid
 			bool bHasRunPreBuildScripts = false;
-			if(Makefile != null)
+			if (Makefile != null)
 			{
 				// Execute the scripts.
 				Utils.ExecuteCustomBuildSteps(Makefile.PreBuildScripts, Logger);
@@ -1100,7 +1100,7 @@ namespace UnrealBuildTool
 
 				// Check that the makefile is still valid
 				string? Reason;
-				if(!TargetMakefile.IsValidForSourceFiles(Makefile, TargetDescriptor.ProjectFile, TargetDescriptor.Platform, WorkingSet, Logger, out Reason))
+				if (!TargetMakefile.IsValidForSourceFiles(Makefile, TargetDescriptor.ProjectFile, TargetDescriptor.Platform, WorkingSet, Logger, out Reason))
 				{
 					Logger.LogInformation("Invalidating makefile for {TargetDescriptorName} ({Reason})", TargetDescriptor.Name, Reason);
 					Makefile = null;
@@ -1108,7 +1108,7 @@ namespace UnrealBuildTool
 			}
 
 			// If we couldn't load a makefile, create a new one
-			if(Makefile == null)
+			if (Makefile == null)
 			{
 				// Create the target
 				UEBuildTarget Target;
@@ -1121,7 +1121,7 @@ namespace UnrealBuildTool
 				FileReference[] PreBuildScripts = Target.CreatePreBuildScripts();
 
 				// Execute the pre-build scripts
-				if(!bHasRunPreBuildScripts)
+				if (!bHasRunPreBuildScripts)
 				{
 					Utils.ExecuteCustomBuildSteps(PreBuildScripts, Logger);
 					bHasRunPreBuildScripts = true;
@@ -1151,7 +1151,7 @@ namespace UnrealBuildTool
 				}
 
 				// Save the makefile for next time
-				if(MakefileLocation != null)
+				if (MakefileLocation != null)
 				{
 					using (GlobalTracer.Instance.BuildSpan("TargetMakefile.Save()").StartActive())
 					{
@@ -1197,16 +1197,16 @@ namespace UnrealBuildTool
 		/// <returns>List of actions that need to be executed</returns>
 		static void GatherOutputItems(TargetDescriptor TargetDescriptor, TargetMakefile Makefile, HashSet<FileItem> OutputItems)
 		{
-			if(TargetDescriptor.SpecificFilesToCompile.Count > 0)
+			if (TargetDescriptor.SpecificFilesToCompile.Count > 0)
 			{
 			}
-			else if(TargetDescriptor.OnlyModuleNames.Count > 0)
+			else if (TargetDescriptor.OnlyModuleNames.Count > 0)
 			{
 				// Find the output items for this module
-				foreach(string OnlyModuleName in TargetDescriptor.OnlyModuleNames)
+				foreach (string OnlyModuleName in TargetDescriptor.OnlyModuleNames)
 				{
 					FileItem[]? OutputItemsForModule;
-					if(!Makefile.ModuleNameToOutputItems.TryGetValue(OnlyModuleName, out OutputItemsForModule))
+					if (!Makefile.ModuleNameToOutputItems.TryGetValue(OnlyModuleName, out OutputItemsForModule))
 					{
 						throw new BuildException("Unable to find output items for module '{0}'", OnlyModuleName);
 					}
@@ -1231,16 +1231,16 @@ namespace UnrealBuildTool
 			// Set of all output items. Knowing that there are no conflicts in produced items, we use this to eliminate duplicate actions.
 			Dictionary<FileItem, LinkedAction> OutputItemToProducingAction = new Dictionary<FileItem, LinkedAction>();
 			HashSet<LinkedAction> IgnoreConflictActions = new HashSet<LinkedAction>();
-			for(int TargetIdx = 0; TargetIdx < TargetDescriptors.Count; TargetIdx++)
+			for (int TargetIdx = 0; TargetIdx < TargetDescriptors.Count; TargetIdx++)
 			{
 				string GroupPrefix = String.Format("{0}-{1}-{2}", TargetDescriptors[TargetIdx].Name, TargetDescriptors[TargetIdx].Platform, TargetDescriptors[TargetIdx].Configuration);
-				foreach(LinkedAction TargetAction in TargetActions[TargetIdx])
+				foreach (LinkedAction TargetAction in TargetActions[TargetIdx])
 				{
 					if (!TargetAction.ProducedItems.Any())
 					{
 						continue;
 					}
-					
+
 					if (TargetAction.IgnoreConflicts())
 					{
 						IgnoreConflictActions.Add(TargetAction);
@@ -1250,7 +1250,7 @@ namespace UnrealBuildTool
 					FileItem ProducedItem = TargetAction.ProducedItems.First();
 
 					LinkedAction? ExistingAction;
-					if(!OutputItemToProducingAction.TryGetValue(ProducedItem, out ExistingAction))
+					if (!OutputItemToProducingAction.TryGetValue(ProducedItem, out ExistingAction))
 					{
 						ExistingAction = new LinkedAction(TargetAction, TargetDescriptors[TargetIdx]);
 						OutputItemToProducingAction[ProducedItem] = ExistingAction;
@@ -1355,7 +1355,7 @@ namespace UnrealBuildTool
 					{
 						FileReference.Copy(CrashDump, FileReference.Combine(SaveCrashDumpDirectory, CrashDump.GetFileName()));
 					}
-					catch(Exception Ex)
+					catch (Exception Ex)
 					{
 						// don't stop if there was a problem copying one of the files
 						Logger.LogWarning("Failed to copy crash dump {CrashDump}: {Ex}", CrashDump, ExceptionUtils.FormatException(Ex));
diff --git a/Engine/Source/Programs/UnrealBuildTool/Modes/CleanMode.cs b/Engine/Source/Programs/UnrealBuildTool/Modes/CleanMode.cs
--- a/Engine/Source/Programs/UnrealBuildTool/Modes/CleanMode.cs
+++ b/Engine/Source/Programs/UnrealBuildTool/Modes/CleanMode.cs
@@ -2,16 +2,14 @@
 
 using System;
 using System.Collections.Generic;
-using System.Diagnostics;
 using System.IO;
 using System.Linq;
-using System.Text;
 using System.Threading.Tasks;
 using EpicGames.Core;
+using Microsoft.Extensions.Logging;
 using OpenTracing;
 using OpenTracing.Util;
 using UnrealBuildBase;
-using Microsoft.Extensions.Logging;
 
 namespace UnrealBuildTool
 {
diff --git a/Engine/Source/Programs/UnrealBuildTool/Modes/DeployMode.cs b/Engine/Source/Programs/UnrealBuildTool/Modes/DeployMode.cs
--- a/Engine/Source/Programs/UnrealBuildTool/Modes/DeployMode.cs
+++ b/Engine/Source/Programs/UnrealBuildTool/Modes/DeployMode.cs
@@ -1,9 +1,5 @@
 // Copyright Epic Games, Inc. All Rights Reserved.
 
-using System;
-using System.Collections.Generic;
-using System.Linq;
-using System.Text;
 using System.Threading.Tasks;
 using EpicGames.Core;
 using Microsoft.Extensions.Logging;
@@ -19,7 +15,7 @@ namespace UnrealBuildTool
 		/// <summary>
 		/// If we are just running the deployment step, specifies the path to the given deployment settings
 		/// </summary>
-		[CommandLine("-Receipt", Required=true)]
+		[CommandLine("-Receipt", Required = true)]
 		public FileReference? ReceiptFile = null;
 
 		/// <summary>
diff --git a/Engine/Source/Programs/UnrealBuildTool/Modes/ExecuteMode.cs b/Engine/Source/Programs/UnrealBuildTool/Modes/ExecuteMode.cs
--- a/Engine/Source/Programs/UnrealBuildTool/Modes/ExecuteMode.cs
+++ b/Engine/Source/Programs/UnrealBuildTool/Modes/ExecuteMode.cs
@@ -1,16 +1,10 @@
 // Copyright Epic Games, Inc. All Rights Reserved.
 
-using System;
 using System.Collections.Generic;
-using System.Diagnostics;
-using System.IO;
-using System.Linq;
-using System.Text;
-using System.Threading;
 using System.Threading.Tasks;
 using EpicGames.Core;
-using OpenTracing.Util;
 using Microsoft.Extensions.Logging;
+using OpenTracing.Util;
 
 namespace UnrealBuildTool
 {
diff --git a/Engine/Source/Programs/UnrealBuildTool/Modes/FixIncludePathsMode.cs b/Engine/Source/Programs/UnrealBuildTool/Modes/FixIncludePathsMode.cs
--- a/Engine/Source/Programs/UnrealBuildTool/Modes/FixIncludePathsMode.cs
+++ b/Engine/Source/Programs/UnrealBuildTool/Modes/FixIncludePathsMode.cs
@@ -2,13 +2,13 @@
 
 using System;
 using System.Collections.Generic;
+using System.Diagnostics;
 using System.Linq;
 using System.Text.RegularExpressions;
+using System.Threading.Tasks;
 using EpicGames.Core;
-using UnrealBuildBase;
 using Microsoft.Extensions.Logging;
-using System.Threading.Tasks;
-using System.Diagnostics;
+using UnrealBuildBase;
 
 namespace UnrealBuildTool
 {
@@ -283,7 +283,7 @@ namespace UnrealBuildTool
 													PreferredPathCache[Include] = PreferredInclude;
 												}
 											}
-												
+
 											if (PreferredInclude == null)
 											{
 												Logger.LogInformation("{FileName}({LineNumber}): Could not find path to '{IncludePath}'", InputFile.FullName, LineNumber, Include);
diff --git a/Engine/Source/Programs/UnrealBuildTool/Modes/GenerateProjectFilesMode.cs b/Engine/Source/Programs/UnrealBuildTool/Modes/GenerateProjectFilesMode.cs
--- a/Engine/Source/Programs/UnrealBuildTool/Modes/GenerateProjectFilesMode.cs
+++ b/Engine/Source/Programs/UnrealBuildTool/Modes/GenerateProjectFilesMode.cs
@@ -2,16 +2,14 @@
 
 using System;
 using System.Collections.Generic;
-using System.IO;
+using System.Diagnostics.CodeAnalysis;
 using System.Linq;
 using System.Reflection;
 using System.Text;
 using System.Threading.Tasks;
 using EpicGames.Core;
-using System.Diagnostics;
-using UnrealBuildBase;
-using System.Diagnostics.CodeAnalysis;
 using Microsoft.Extensions.Logging;
+using UnrealBuildBase;
 
 namespace UnrealBuildTool
 {
@@ -38,9 +36,9 @@ namespace UnrealBuildTool
 		[CommandLine("-VSMac", Value = nameof(ProjectFileFormat.VisualStudioMac))]
 		[CommandLine("-CLion", Value = nameof(ProjectFileFormat.CLion))]
 		[CommandLine("-Rider", Value = nameof(ProjectFileFormat.Rider))]
-		#if __VPROJECT_AVAILABLE__
-			[CommandLine("-VProject", Value = nameof(ProjectFileFormat.VProject))]
-		#endif
+#if __VPROJECT_AVAILABLE__
+		[CommandLine("-VProject", Value = nameof(ProjectFileFormat.VProject))]
+#endif
 		HashSet<ProjectFileFormat> ProjectFileFormats = new HashSet<ProjectFileFormat>();
 
 		/// <summary>
@@ -86,7 +84,7 @@ namespace UnrealBuildTool
 			// Apply the XML config again with a project specific BuildConfiguration.xml 
 			XmlConfig.ReadConfigFiles(null, ProjectFile?.Directory, Logger);
 			XmlConfig.ApplyTo(this);
-			
+
 			// Warn if there are explicit project file formats specified
 			if (ProjectFileFormats.Count > 0 && !bAutomated)
 			{
@@ -100,7 +98,7 @@ namespace UnrealBuildTool
 				Configuration.Append("<?xml version=\"1.0\" encoding=\"utf-8\" ?>\n");
 				Configuration.Append("<Configuration xmlns=\"https://www.unrealengine.com/BuildConfiguration\">\n");
 				Configuration.Append("  <ProjectFileGenerator>\n");
-				foreach(ProjectFileFormat ProjectFileFormat in ProjectFileFormats)
+				foreach (ProjectFileFormat ProjectFileFormat in ProjectFileFormats)
 				{
 					Configuration.AppendFormat("    <Format>{0}</Format>\n", ProjectFileFormat);
 				}
@@ -139,9 +137,9 @@ namespace UnrealBuildTool
 				if (CheckType.IsClass && !CheckType.IsAbstract && CheckType.IsSubclassOf(typeof(PlatformProjectGenerator)))
 				{
 					PlatformProjectGenerator Generator = (PlatformProjectGenerator)Activator.CreateInstance(CheckType, Arguments, Logger)!;
-					foreach(UnrealTargetPlatform Platform in Generator.GetPlatforms())
+					foreach (UnrealTargetPlatform Platform in Generator.GetPlatforms())
 					{
-						if(DisablePlatformProjectGenerators == null || !DisablePlatformProjectGenerators.Any(x => x.Equals(Platform.ToString(), StringComparison.OrdinalIgnoreCase)))
+						if (DisablePlatformProjectGenerators == null || !DisablePlatformProjectGenerators.Any(x => x.Equals(Platform.ToString(), StringComparison.OrdinalIgnoreCase)))
 						{
 							Logger.LogDebug("Registering project generator {CheckType} for {Platform}", CheckType, Platform);
 							PlatformProjectGenerators.RegisterPlatformProjectGenerator(Platform, Generator, Logger);
@@ -223,11 +221,11 @@ namespace UnrealBuildTool
 					case ProjectFileFormat.Rider:
 						Generator = new RiderProjectFileGenerator(ProjectFile, Arguments);
 						break;
-					#if __VPROJECT_AVAILABLE__
-						case ProjectFileFormat.VProject:
-							Generator = new VProjectFileGenerator(ProjectFile);
-							break;
-					#endif
+#if __VPROJECT_AVAILABLE__
+					case ProjectFileFormat.VProject:
+						Generator = new VProjectFileGenerator(ProjectFile);
+						break;
+#endif
 					default:
 						throw new BuildException("Unhandled project file type '{0}", ProjectFileFormat);
 				}
@@ -242,7 +240,7 @@ namespace UnrealBuildTool
 			ProjectFileGenerator.bGenerateProjectFiles = true;
 			// perform anything that only needs to happen one time, in the first project genereator
 			bool bPerformOneTimeOperations = true;
-			foreach(KeyValuePair< ProjectFileFormat, ProjectFileGenerator> Pair in Generators)
+			foreach (KeyValuePair<ProjectFileFormat, ProjectFileGenerator> Pair in Generators)
 			{
 				Logger.LogInformation("");
 				Logger.LogInformation($"Generating {Pair.Key} project files:");
@@ -287,7 +285,7 @@ namespace UnrealBuildTool
 					{
 						CandidateProjectPath = Arguments[Idx];
 						Arguments.MarkAsUsed(Idx);
-						break; 
+						break;
 					}
 				}
 			}
@@ -320,9 +318,9 @@ namespace UnrealBuildTool
 				ProjectFile = CandidateProjectFile;
 				return true;
 			}
-			
+
 			FileReference? InstalledProjectFile = UnrealBuildTool.GetInstalledProjectFile();
-			if(InstalledProjectFile != null)
+			if (InstalledProjectFile != null)
 			{
 				ProjectFile = InstalledProjectFile;
 				return true;
diff --git a/Engine/Source/Programs/UnrealBuildTool/Modes/IWYUMode.cs b/Engine/Source/Programs/UnrealBuildTool/Modes/IWYUMode.cs
--- a/Engine/Source/Programs/UnrealBuildTool/Modes/IWYUMode.cs
+++ b/Engine/Source/Programs/UnrealBuildTool/Modes/IWYUMode.cs
@@ -1,18 +1,18 @@
 // Copyright Epic Games, Inc. All Rights Reserved.
 
 using System;
+using System.Collections.Concurrent;
 using System.Collections.Generic;
-using EpicGames.Core;
-using UnrealBuildBase;
-using Microsoft.Extensions.Logging;
-using System.Linq;
-using System.Text.Json;
+using System.Diagnostics.CodeAnalysis;
 using System.IO;
-using System.Threading.Tasks;
+using System.Linq;
 using System.Text;
-using System.Diagnostics.CodeAnalysis;
+using System.Text.Json;
 using System.Threading;
-using System.Collections.Concurrent;
+using System.Threading.Tasks;
+using EpicGames.Core;
+using Microsoft.Extensions.Logging;
+using UnrealBuildBase;
 
 
 // IWYUMode is a mode that can be used to clean up includes in source code. It uses the clang based tool include-what-you-use (IWYU) to figure out what is needed in each .h/.cpp file
@@ -190,7 +190,7 @@ namespace UnrealBuildTool
 		/// <summary>
 		/// The different headers and source files that was covered by this execution of iwyu.
 		/// </summary>
-		public List<IWYUInfo> Files { get; set; } = new ();
+		public List<IWYUInfo> Files { get; set; } = new();
 
 		/// <summary>
 		/// Name of .iwyu file
@@ -1967,7 +1967,7 @@ namespace UnrealBuildTool
 
 			ConcurrentDictionary<string, long> AllSeenIncludes = new();
 
-			List<ValueTuple<IWYUInfo, Dictionary<string, string> >> InfosToCompare = new();
+			List<ValueTuple<IWYUInfo, Dictionary<string, string>>> InfosToCompare = new();
 
 			Logger.LogInformation($"Generating transitive seen include lists...");
 			Parallel.ForEach(Infos.Values, Info =>
@@ -2050,7 +2050,7 @@ namespace UnrealBuildTool
 				long OptimizedSize = 0;
 				foreach (var Include in File.TransitiveIncludes.Keys)
 					OptimizedSize += FileToSize[Include];
-				float Percent = 100.0f - (((float)OptimizedSize)/(OptimizedSize + Saved) * 100.0f);
+				float Percent = 100.0f - (((float)OptimizedSize) / (OptimizedSize + Saved) * 100.0f);
 				Logger.LogInformation($"{Path.GetFileName(File.File)}   {PrettySize(OptimizedSize + Saved)} -> {PrettySize(OptimizedSize)}  (Saved {PrettySize(Saved)} or {Percent.ToString("0.0")}%)");
 			}
 			Logger.LogInformation("");
diff --git a/Engine/Source/Programs/UnrealBuildTool/Modes/InlineGeneratedCpps.cs b/Engine/Source/Programs/UnrealBuildTool/Modes/InlineGeneratedCpps.cs
--- a/Engine/Source/Programs/UnrealBuildTool/Modes/InlineGeneratedCpps.cs
+++ b/Engine/Source/Programs/UnrealBuildTool/Modes/InlineGeneratedCpps.cs
@@ -3,12 +3,11 @@
 using System;
 using System.Collections.Generic;
 using System.Linq;
-using System.Text;
 using System.Text.RegularExpressions;
 using System.Threading.Tasks;
 using EpicGames.Core;
-using UnrealBuildBase;
 using Microsoft.Extensions.Logging;
+using UnrealBuildBase;
 
 namespace UnrealBuildTool
 {
diff --git a/Engine/Source/Programs/UnrealBuildTool/Modes/JsonExportMode.cs b/Engine/Source/Programs/UnrealBuildTool/Modes/JsonExportMode.cs
--- a/Engine/Source/Programs/UnrealBuildTool/Modes/JsonExportMode.cs
+++ b/Engine/Source/Programs/UnrealBuildTool/Modes/JsonExportMode.cs
@@ -1,13 +1,11 @@
 // Copyright Epic Games, Inc. All Rights Reserved.
 
-using System;
 using System.Collections.Generic;
 using System.Linq;
-using System.Text;
 using System.Threading.Tasks;
 using EpicGames.Core;
-using UnrealBuildBase;
 using Microsoft.Extensions.Logging;
+using UnrealBuildBase;
 
 namespace UnrealBuildTool
 {
@@ -34,14 +32,14 @@ namespace UnrealBuildTool
 			Arguments.ApplyTo(this);
 
 			List<TargetDescriptor> TargetDescriptors = TargetDescriptor.ParseCommandLine(Arguments, false, false, false, Logger);
-			foreach(TargetDescriptor TargetDescriptor in TargetDescriptors)
+			foreach (TargetDescriptor TargetDescriptor in TargetDescriptors)
 			{
 				// Create the target
 				UEBuildTarget Target = UEBuildTarget.Create(TargetDescriptor, false, false, false, Logger);
 
 				// Get the output file
 				FileReference? OutputFile = TargetDescriptor.AdditionalArguments.GetFileReferenceOrDefault("-OutputFile=", null);
-				if(OutputFile == null)
+				if (OutputFile == null)
 				{
 					OutputFile = Target.ReceiptFileName.ChangeExtension(".json");
 				}
diff --git a/Engine/Source/Programs/UnrealBuildTool/Modes/PrintBuildGraphInfo.cs b/Engine/Source/Programs/UnrealBuildTool/Modes/PrintBuildGraphInfo.cs
--- a/Engine/Source/Programs/UnrealBuildTool/Modes/PrintBuildGraphInfo.cs
+++ b/Engine/Source/Programs/UnrealBuildTool/Modes/PrintBuildGraphInfo.cs
@@ -4,10 +4,8 @@ using System;
 using System.Collections.Generic;
 using System.Linq;
 using System.Text;
-using System.Text.RegularExpressions;
 using System.Threading.Tasks;
 using EpicGames.Core;
-using UnrealBuildBase;
 using Microsoft.Extensions.Logging;
 
 namespace UnrealBuildTool
@@ -67,7 +65,7 @@ namespace UnrealBuildTool
 							{
 								ModuleCPP.CreateCompileEnvironmentForIntellisense(Target.Rules, BinaryCompileEnvironment, Logger);
 							}
-							
+
 							Modules.Add(Module);
 						}
 					}
@@ -202,7 +200,7 @@ namespace UnrealBuildTool
 								}
 
 								var DepModuleArray = DepModuleStack.Reverse().ToArray();
-								for (int DepModuleIndex = 0; DepModuleIndex < DepModuleArray.Length; DepModuleIndex++) 
+								for (int DepModuleIndex = 0; DepModuleIndex < DepModuleArray.Length; DepModuleIndex++)
 								{
 									Dependencies.Append(DepModuleArray[DepModuleIndex].Name);
 									if (DepModuleIndex != DepModuleArray.Length - 1)
diff --git a/Engine/Source/Programs/UnrealBuildTool/Modes/ProfileUnitySizesMode.cs b/Engine/Source/Programs/UnrealBuildTool/Modes/ProfileUnitySizesMode.cs
--- a/Engine/Source/Programs/UnrealBuildTool/Modes/ProfileUnitySizesMode.cs
+++ b/Engine/Source/Programs/UnrealBuildTool/Modes/ProfileUnitySizesMode.cs
@@ -2,12 +2,12 @@
 
 using System;
 using System.Collections.Generic;
+using System.Linq;
 using System.Text.RegularExpressions;
 using System.Threading.Tasks;
 using EpicGames.Core;
-using UnrealBuildBase;
 using Microsoft.Extensions.Logging;
-using System.Linq;
+using UnrealBuildBase;
 
 namespace UnrealBuildTool
 {
@@ -168,7 +168,7 @@ namespace UnrealBuildTool
 				// build each Module
 				ModuleList.SortBy(module => module.Name);
 				foreach (UEBuildModule Module in ModuleList)
-				{					
+				{
 					await CompileModuleAsync(BuildConfiguration, TargetDescriptor, Target, Module, Logger);
 				}
 			}
@@ -194,7 +194,7 @@ namespace UnrealBuildTool
 
 			int BuildNum = 1;
 			await CompileModuleAsync($"  [{BuildNum++}/{TotalBuilds}] ", BuildConfiguration, TargetDescriptor, Module, Logger, CurrentModuleUnitySize, true, false);
-			
+
 			TimingData CurrentCompileTime = await GetBestCompileModuleTimeAsync($"  [{BuildNum++}/{TotalBuilds}] ", BuildConfiguration, TargetDescriptor, Module, Logger, CurrentModuleUnitySize, false, false);
 			if (!CurrentCompileTime.IsValid())
 			{
diff --git a/Engine/Source/Programs/UnrealBuildTool/Modes/QueryMode.cs b/Engine/Source/Programs/UnrealBuildTool/Modes/QueryMode.cs
--- a/Engine/Source/Programs/UnrealBuildTool/Modes/QueryMode.cs
+++ b/Engine/Source/Programs/UnrealBuildTool/Modes/QueryMode.cs
@@ -1,15 +1,13 @@
 // Copyright Epic Games, Inc. All Rights Reserved.
 
-using EpicGames.Core;
-using Microsoft.Extensions.Logging;
 using System;
+using System.Collections.Generic;
 using System.Linq;
 using System.Text.Json;
-using System.Text.Json.Serialization;
 using System.Threading.Tasks;
-using System.Collections.Generic;
+using EpicGames.Core;
+using Microsoft.Extensions.Logging;
 using OpenTracing.Util;
-using UnrealBuildTool;
 using UnrealBuildBase;
 
 namespace UnrealBuildTool
@@ -74,7 +72,7 @@ namespace UnrealBuildTool
 
 	[ToolMode("Query", ToolModeOptions.BuildPlatforms | ToolModeOptions.XmlConfig | ToolModeOptions.UseStartupTraceListener)]
 	class QueryMode : ToolMode
-	{ 
+	{
 		[CommandLine("-LogDirectory=")]
 		public DirectoryReference? LogDirectory = null;
 
@@ -128,7 +126,7 @@ namespace UnrealBuildTool
 				PropertyNamingPolicy = JsonNamingPolicy.CamelCase,
 				WriteIndented = bIndented,
 			};
-			switch(Query)
+			switch (Query)
 			{
 				case QueryType.Capabilities:
 					Logger.LogInformation("QueryCapabilities");
@@ -148,10 +146,10 @@ namespace UnrealBuildTool
 		{
 			var Reply = new
 			{
-				Queries = new List<string>{ QueryType.Capabilities.ToString(), QueryType.AvailableTargets.ToString(), QueryType.TargetDetails.ToString() }
+				Queries = new List<string> { QueryType.Capabilities.ToString(), QueryType.AvailableTargets.ToString(), QueryType.TargetDetails.ToString() }
 			};
 			Console.WriteLine(JsonSerializer.Serialize(Reply, JsonOptions));
-			return 0;		
+			return 0;
 		}
 
 		private int QueryAvailableTargets(CommandLineArguments Arguments, ILogger Logger, JsonSerializerOptions JsonOptions)
@@ -198,9 +196,9 @@ namespace UnrealBuildTool
 				foreach (FileReference TargetFilePath in AllTargetFiles)
 				{
 					string TargetName = TargetFilePath.GetFileNameWithoutAnyExtensions();
-					FileReference? ProjectPath = Projects.FirstOrDefault(p => TargetFilePath.IsUnderDirectory( p.Directory));
+					FileReference? ProjectPath = Projects.FirstOrDefault(p => TargetFilePath.IsUnderDirectory(p.Directory));
 					Targets.Add(TargetName, new TargetConfigs() { ProjectPath = ProjectPath?.ToString() ?? "", Configurations = Configurations, Platforms = Platforms.Select(x => x.ToString()).ToList() });
-					if ( DefaultTarget == null || TargetName == "UnrealEditor")
+					if (DefaultTarget == null || TargetName == "UnrealEditor")
 					{
 						DefaultTarget = TargetName;
 					}
diff --git a/Engine/Source/Programs/UnrealBuildTool/Modes/QueryTargetsMode.cs b/Engine/Source/Programs/UnrealBuildTool/Modes/QueryTargetsMode.cs
--- a/Engine/Source/Programs/UnrealBuildTool/Modes/QueryTargetsMode.cs
+++ b/Engine/Source/Programs/UnrealBuildTool/Modes/QueryTargetsMode.cs
@@ -3,11 +3,10 @@
 using System;
 using System.Collections.Generic;
 using System.Linq;
-using System.Text;
 using System.Threading.Tasks;
 using EpicGames.Core;
-using UnrealBuildBase;
 using Microsoft.Extensions.Logging;
+using UnrealBuildBase;
 
 namespace UnrealBuildTool
 {
@@ -34,7 +33,7 @@ namespace UnrealBuildTool
 		/// </summary>
 		[CommandLine("-IncludeAllTargets")]
 		bool bIncludeAllTargets = false;
-		
+
 		/// <summary>
 		/// Execute the mode
 		/// </summary>
@@ -51,7 +50,7 @@ namespace UnrealBuildTool
 			Arguments.ApplyTo(BuildConfiguration);
 
 			// Ensure the path to the output file is valid
-			if(OutputFile == null)
+			if (OutputFile == null)
 			{
 				OutputFile = GetDefaultOutputFile(ProjectFile);
 			}
@@ -114,8 +113,8 @@ namespace UnrealBuildTool
 				foreach (string TargetName in TargetNames)
 				{
 					// skip target rules that are platform extension or platform group specializations
-					string[] TargetPathSplit = TargetName.Split(new char[]{'_'}, StringSplitOptions.RemoveEmptyEntries );
-					if (TargetPathSplit.Length > 1 && (UnrealTargetPlatform.IsValidName(TargetPathSplit.Last()) || UnrealPlatformGroup.IsValidName(TargetPathSplit.Last()) ) )
+					string[] TargetPathSplit = TargetName.Split(new char[] { '_' }, StringSplitOptions.RemoveEmptyEntries);
+					if (TargetPathSplit.Length > 1 && (UnrealTargetPlatform.IsValidName(TargetPathSplit.Last()) || UnrealPlatformGroup.IsValidName(TargetPathSplit.Last())))
 					{
 						continue;
 					}
@@ -133,7 +132,7 @@ namespace UnrealBuildTool
 						Logger.LogDebug("{Ex}", ExceptionUtils.FormatException(Ex));
 						continue;
 					}
-					
+
 					// Skip non-default targets if one is specified.
 					if (ProjectFile != null && !bIncludeAllTargets)
 					{
diff --git a/Engine/Source/Programs/UnrealBuildTool/Modes/ServerMode.cs b/Engine/Source/Programs/UnrealBuildTool/Modes/ServerMode.cs
--- a/Engine/Source/Programs/UnrealBuildTool/Modes/ServerMode.cs
+++ b/Engine/Source/Programs/UnrealBuildTool/Modes/ServerMode.cs
@@ -1,15 +1,14 @@
 // Copyright Epic Games, Inc. All Rights Reserved.
 
-using EpicGames.Core;
-using Microsoft.Extensions.Logging;
 using System;
+using System.Collections.Generic;
 using System.Linq;
 using System.Text.Json;
 using System.Text.Json.Serialization;
 using System.Threading.Tasks;
-using System.Collections.Generic;
+using EpicGames.Core;
+using Microsoft.Extensions.Logging;
 using OpenTracing.Util;
-using UnrealBuildTool;
 using UnrealBuildBase;
 
 namespace UnrealBuildTool
@@ -351,7 +350,7 @@ namespace UnrealBuildTool
 
 					foreach (UEBuildModuleCPP Module in Binary.Modules.OfType<UEBuildModuleCPP>())
 					{
-						if (Module.Binary != null && Module.Binary != Binary) 
+						if (Module.Binary != null && Module.Binary != Binary)
 						{
 							continue;
 						}
@@ -404,20 +403,20 @@ namespace UnrealBuildTool
 		private QueryFileResponse Handle_QueryFile(QueryFileCommand Command, ILogger Logger)
 		{
 			var FailResponse = new QueryFileResponse { Found = false };
-			if (CurrentTargetIntellisenseInfo == null )
+			if (CurrentTargetIntellisenseInfo == null)
 			{
 				return FailResponse;
 			}
-			try 
+			try
 			{
-				var Path = System.IO.Path.GetFullPath(Command.AbsolutePath); 
+				var Path = System.IO.Path.GetFullPath(Command.AbsolutePath);
 				if (CurrentTargetIntellisenseInfo.FindModuleForFile(new FileReference(Path)) != null)
 				{
 					return new QueryFileResponse { Found = true };
 				}
 				return FailResponse;
 			}
-			catch (Exception )
+			catch (Exception)
 			{
 				return FailResponse;
 			}
@@ -431,7 +430,7 @@ namespace UnrealBuildTool
 			}
 			if (Command.AbsolutePath != null)
 			{
-				var Path = System.IO.Path.GetFullPath(Command.AbsolutePath); 
+				var Path = System.IO.Path.GetFullPath(Command.AbsolutePath);
 				UEBuildModule? Module = CurrentTargetIntellisenseInfo.FindModuleForDirectory(new DirectoryReference(Path));
 				if (Module != null)
 				{
@@ -453,8 +452,8 @@ namespace UnrealBuildTool
 
 		private GetCompileSettingsResponse Handle_GetCompileSettings(GetCompileSettingsCommand Command, ILogger Logger)
 		{
-			GetCompileSettingsResponse FailResponse = new GetCompileSettingsResponse { }; 
-			try 
+			GetCompileSettingsResponse FailResponse = new GetCompileSettingsResponse { };
+			try
 			{
 				GetCompileSettingsResponse Response = new GetCompileSettingsResponse { };
 				foreach (string AbsolutePath in Command.AbsolutePaths)
diff --git a/Engine/Source/Programs/UnrealBuildTool/Modes/SetupPlatformsMode.cs b/Engine/Source/Programs/UnrealBuildTool/Modes/SetupPlatformsMode.cs
--- a/Engine/Source/Programs/UnrealBuildTool/Modes/SetupPlatformsMode.cs
+++ b/Engine/Source/Programs/UnrealBuildTool/Modes/SetupPlatformsMode.cs
@@ -1,9 +1,5 @@
 // Copyright Epic Games, Inc. All Rights Reserved.
 
-using System;
-using System.Collections.Generic;
-using System.Linq;
-using System.Text;
 using System.Threading.Tasks;
 using EpicGames.Core;
 using Microsoft.Extensions.Logging;
diff --git a/Engine/Source/Programs/UnrealBuildTool/Modes/ValidatePlatformsMode.cs b/Engine/Source/Programs/UnrealBuildTool/Modes/ValidatePlatformsMode.cs
--- a/Engine/Source/Programs/UnrealBuildTool/Modes/ValidatePlatformsMode.cs
+++ b/Engine/Source/Programs/UnrealBuildTool/Modes/ValidatePlatformsMode.cs
@@ -1,9 +1,6 @@
 // Copyright Epic Games, Inc. All Rights Reserved.
 
-using System;
 using System.Collections.Generic;
-using System.Linq;
-using System.Text;
 using System.Threading.Tasks;
 using EpicGames.Core;
 using Microsoft.Extensions.Logging;
@@ -47,7 +44,7 @@ namespace UnrealBuildTool
 			Arguments.CheckAllArgumentsUsed();
 
 			// If the -AllPlatforms argument is specified, add all the known platforms into the list
-			if(bAllPlatforms)
+			if (bAllPlatforms)
 			{
 				Platforms.UnionWith(UnrealTargetPlatform.GetValidPlatforms());
 			}
@@ -70,7 +67,7 @@ namespace UnrealBuildTool
 				{
 					Logger.LogInformation("##PlatformValidate: {Platform} INVALID {PlatformSdkString}", Platform.ToString(), PlatformSDKString);
 				}
-			} 
+			}
 			return Task.FromResult(0);
 		}
 	}
diff --git a/Engine/Source/Programs/UnrealBuildTool/Modes/WriteDocumentationMode.cs b/Engine/Source/Programs/UnrealBuildTool/Modes/WriteDocumentationMode.cs
--- a/Engine/Source/Programs/UnrealBuildTool/Modes/WriteDocumentationMode.cs
+++ b/Engine/Source/Programs/UnrealBuildTool/Modes/WriteDocumentationMode.cs
@@ -1,12 +1,7 @@
 // Copyright Epic Games, Inc. All Rights Reserved.
 
-using System;
-using System.Collections.Generic;
-using System.Linq;
-using System.Text;
 using System.Threading.Tasks;
 using EpicGames.Core;
-using UnrealBuildBase;
 using Microsoft.Extensions.Logging;
 
 namespace UnrealBuildTool
@@ -48,7 +43,7 @@ namespace UnrealBuildTool
 			Arguments.ApplyTo(this);
 			Arguments.CheckAllArgumentsUsed();
 
-			switch(Type)
+			switch (Type)
 			{
 				case DocumentationType.BuildConfiguration:
 					XmlConfig.WriteDocumentation(OutputFile, Logger);
diff --git a/Engine/Source/Programs/UnrealBuildTool/Modes/WriteMetadataMode.cs b/Engine/Source/Programs/UnrealBuildTool/Modes/WriteMetadataMode.cs
--- a/Engine/Source/Programs/UnrealBuildTool/Modes/WriteMetadataMode.cs
+++ b/Engine/Source/Programs/UnrealBuildTool/Modes/WriteMetadataMode.cs
@@ -4,13 +4,10 @@ using System;
 using System.Collections.Generic;
 using System.Diagnostics.CodeAnalysis;
 using System.IO;
-using System.Linq;
-using System.Runtime.Serialization.Formatters.Binary;
-using System.Text;
 using System.Threading.Tasks;
 using EpicGames.Core;
-using UnrealBuildBase;
 using Microsoft.Extensions.Logging;
+using UnrealBuildBase;
 
 namespace UnrealBuildTool
 {
@@ -100,7 +97,7 @@ namespace UnrealBuildTool
 			// Acquire a different mutex to the regular UBT instance, since this mode will be called as part of a build. We need the mutex to ensure that building two modular configurations 
 			// in parallel don't clash over writing shared *.modules files (eg. DebugGame and Development editors).
 			string MutexName = SingleInstanceMutex.GetUniqueMutexForPath("UnrealBuildTool_WriteMetadata", Unreal.RootDirectory.FullName);
-			using(new SingleInstanceMutex(MutexName, true))
+			using (new SingleInstanceMutex(MutexName, true))
 			{
 				return ExecuteInternal(Arguments, Logger);
 			}
@@ -121,7 +118,7 @@ namespace UnrealBuildTool
 			Arguments.CheckAllArgumentsUsed();
 
 			// Make sure the version number is correct
-			if(VersionNumber != CurrentVersionNumber)
+			if (VersionNumber != CurrentVersionNumber)
 			{
 				throw new BuildException("Version number to WriteMetadataMode is incorrect (expected {0}, got {1})", CurrentVersionNumber, VersionNumber);
 			}
@@ -146,10 +143,10 @@ namespace UnrealBuildTool
 			}
 
 			// Read all the existing manifests and merge them into the new ones if they have the same build id
-			foreach(KeyValuePair<FileReference, ModuleManifest> Pair in TargetInfo.FileToManifest)
+			foreach (KeyValuePair<FileReference, ModuleManifest> Pair in TargetInfo.FileToManifest)
 			{
 				ModuleManifest? SourceManifest;
-				if(TryReadManifest(Pair.Key, Logger, out SourceManifest) && SourceManifest.BuildId == BuildId)
+				if (TryReadManifest(Pair.Key, Logger, out SourceManifest) && SourceManifest.BuildId == BuildId)
 				{
 					MergeManifests(SourceManifest, Pair.Value);
 				}
@@ -159,12 +156,12 @@ namespace UnrealBuildTool
 			foreach (KeyValuePair<FileReference, ModuleManifest> Pair in TargetInfo.FileToManifest)
 			{
 				FileReference ManifestFile = Pair.Key;
-				if(!UnrealBuildTool.IsFileInstalled(ManifestFile))
+				if (!UnrealBuildTool.IsFileInstalled(ManifestFile))
 				{
 					ModuleManifest Manifest = Pair.Value;
 					Manifest.BuildId = BuildId ?? String.Empty;
 
-					if(!FileReference.Exists(ManifestFile))
+					if (!FileReference.Exists(ManifestFile))
 					{
 						// If the file doesn't already exist, just write it out
 						DirectoryReference.CreateDirectory(ManifestFile.Directory);
@@ -292,14 +289,14 @@ namespace UnrealBuildTool
 		/// <returns>True if the manifest was read correctly, false otherwise</returns>
 		public static bool TryReadManifest(FileReference ManifestFileName, ILogger Logger, [NotNullWhen(true)] out ModuleManifest? Manifest)
 		{
-			if(FileReference.Exists(ManifestFileName))
+			if (FileReference.Exists(ManifestFileName))
 			{
 				try
 				{
 					Manifest = ModuleManifest.Read(ManifestFileName);
 					return true;
 				}
-				catch(Exception Ex)
+				catch (Exception Ex)
 				{
 					Logger.LogWarning("Unable to read '{ManifestFileName}'; ignoring.", ManifestFileName);
 					Logger.LogDebug("{Ex}", ExceptionUtils.FormatExceptionDetails(Ex));
@@ -317,9 +314,9 @@ namespace UnrealBuildTool
 		/// <param name="TargetManifest">The target manifest to merge into</param>
 		static void MergeManifests(ModuleManifest SourceManifest, ModuleManifest TargetManifest)
 		{
-			foreach(KeyValuePair<string, string> ModulePair in SourceManifest.ModuleNameToFileName)
+			foreach (KeyValuePair<string, string> ModulePair in SourceManifest.ModuleNameToFileName)
 			{
-				if(!TargetManifest.ModuleNameToFileName.ContainsKey(ModulePair.Key))
+				if (!TargetManifest.ModuleNameToFileName.ContainsKey(ModulePair.Key))
 				{
 					TargetManifest.ModuleNameToFileName.Add(ModulePair.Key, ModulePair.Value);
 				}
diff --git a/Engine/Source/Programs/UnrealBuildTool/Platform/Android/AndroidAARHandler.cs b/Engine/Source/Programs/UnrealBuildTool/Platform/Android/AndroidAARHandler.cs
--- a/Engine/Source/Programs/UnrealBuildTool/Platform/Android/AndroidAARHandler.cs
+++ b/Engine/Source/Programs/UnrealBuildTool/Platform/Android/AndroidAARHandler.cs
@@ -2,15 +2,8 @@
 
 using System;
 using System.Collections.Generic;
-using System.Linq;
-using System.Text;
-using System.Xml;
-using System.Xml.Linq;
-using System.Diagnostics;
 using System.IO;
-using EpicGames.Core;
-
-using Ionic.Zip;
+using System.Xml.Linq;
 using Microsoft.Extensions.Logging;
 
 namespace UnrealBuildTool
@@ -140,7 +133,7 @@ namespace UnrealBuildTool
 			{
 				string PackageDirectory = Path.Combine(Repository, PackagePath, BaseName, Version);
 
-                if (Directory.Exists(PackageDirectory))
+				if (Directory.Exists(PackageDirectory))
 				{
 					return PackageDirectory;
 				}
@@ -484,13 +477,13 @@ namespace UnrealBuildTool
 				string Filename = Entry.Filename + ".jar";
 				string BaseName = Path.GetFileName(Filename);
 				string TargetPath = Path.Combine(DestinationPath, BaseName);
-                //Logger.LogInformation("Attempting to copy JAR {Filename} {BaseName} {TargetPath}", Filename, BaseName, TargetPath);
+				//Logger.LogInformation("Attempting to copy JAR {Filename} {BaseName} {TargetPath}", Filename, BaseName, TargetPath);
 
-                if (!File.Exists(Filename))
-                {
-                    Logger.LogInformation("JAR doesn't exist! {Filename}", Filename);
-                }
-                if (!File.Exists(TargetPath))
+				if (!File.Exists(Filename))
+				{
+					Logger.LogInformation("JAR doesn't exist! {Filename}", Filename);
+				}
+				if (!File.Exists(TargetPath))
 				{
 					Logger.LogInformation("Copying JAR {BaseName}", BaseName);
 					File.Copy(Filename, TargetPath);
@@ -521,7 +514,8 @@ namespace UnrealBuildTool
 				if (!File.Exists(TargetManifestFileName))
 				{
 					Logger.LogInformation("Extracting AAR {BaseName}", BaseName);
-					/*IEnumerable<string> FileNames =*/ UnzipFiles(Entry.Filename + ".aar", TargetPath, Logger);
+					/*IEnumerable<string> FileNames =*/
+					UnzipFiles(Entry.Filename + ".aar", TargetPath, Logger);
 
 					// Must have a src directory (even if empty)
 					string SrcDirectory = Path.Combine(TargetPath, "src");
diff --git a/Engine/Source/Programs/UnrealBuildTool/Platform/Android/AndroidPlatformSDK.cs b/Engine/Source/Programs/UnrealBuildTool/Platform/Android/AndroidPlatformSDK.cs
--- a/Engine/Source/Programs/UnrealBuildTool/Platform/Android/AndroidPlatformSDK.cs
+++ b/Engine/Source/Programs/UnrealBuildTool/Platform/Android/AndroidPlatformSDK.cs
@@ -3,9 +3,9 @@
 using System;
 using System.Collections.Generic;
 using System.IO;
-using EpicGames.Core;
-using System.Text.RegularExpressions;
 using System.Linq;
+using System.Text.RegularExpressions;
+using EpicGames.Core;
 using Microsoft.Extensions.Logging;
 
 ///////////////////////////////////////////////////////////////////
@@ -57,7 +57,7 @@ namespace UnrealBuildTool
 					}
 				}
 			}
-	
+
 			return null;
 		}
 
@@ -301,7 +301,7 @@ namespace UnrealBuildTool
 				Logger.LogInformation("NDKROOT llvm missing");
 				return false;
 			}
-			
+
 			// check JDK is valid
 			if (String.IsNullOrEmpty(JavaPath))
 			{
diff --git a/Engine/Source/Programs/UnrealBuildTool/Platform/Android/UEBuildAndroid.cs b/Engine/Source/Programs/UnrealBuildTool/Platform/Android/UEBuildAndroid.cs
--- a/Engine/Source/Programs/UnrealBuildTool/Platform/Android/UEBuildAndroid.cs
+++ b/Engine/Source/Programs/UnrealBuildTool/Platform/Android/UEBuildAndroid.cs
@@ -2,11 +2,9 @@
 
 using System;
 using System.Collections.Generic;
-using System.Linq;
 using System.IO;
 using EpicGames.Core;
 using Microsoft.Extensions.Logging;
-using Microsoft.Extensions.Logging.Abstractions;
 using UnrealBuildBase;
 
 namespace UnrealBuildTool
@@ -76,7 +74,7 @@ namespace UnrealBuildTool
 		/// Accessors for fields on the inner TargetRules instance
 		/// </summary>
 		#region Read-only accessor properties 
-		#pragma warning disable CS1591
+#pragma warning disable CS1591
 
 		public bool bEnableAddressSanitizer => Inner.bEnableAddressSanitizer;
 
@@ -92,7 +90,7 @@ namespace UnrealBuildTool
 
 		public AndroidTargetRules TargetRules => Inner;
 
-		#pragma warning restore CS1591
+#pragma warning restore CS1591
 		#endregion
 	}
 
@@ -168,7 +166,7 @@ namespace UnrealBuildTool
 	{
 		UEBuildPlatformSDK SDK;
 
-		public AndroidPlatform(UnrealTargetPlatform InTargetPlatform, UEBuildPlatformSDK InSDK, ILogger InLogger) 
+		public AndroidPlatform(UnrealTargetPlatform InTargetPlatform, UEBuildPlatformSDK InSDK, ILogger InLogger)
 			: base(InTargetPlatform, InSDK, new AndroidArchitectureConfig(), InLogger)
 		{
 			SDK = InSDK;
@@ -266,14 +264,14 @@ namespace UnrealBuildTool
 
 		public override string[] GetDebugInfoExtensions(ReadOnlyTargetRules InTarget, UEBuildBinaryType InBinaryType)
 		{
-			return new string [] {};
+			return new string[] { };
 		}
 
 		public override void FindAdditionalBuildProductsToClean(ReadOnlyTargetRules Target, List<FileReference> FilesToDelete, List<DirectoryReference> DirectoriesToDelete)
 		{
 			base.FindAdditionalBuildProductsToClean(Target, FilesToDelete, DirectoriesToDelete);
 
-			if(Target.ProjectFile != null)
+			if (Target.ProjectFile != null)
 			{
 				DirectoriesToDelete.Add(DirectoryReference.Combine(DirectoryReference.FromFile(Target.ProjectFile), "Intermediate", "Android"));
 			}
@@ -282,7 +280,7 @@ namespace UnrealBuildTool
 		public virtual bool HasSpecificDefaultBuildConfig(UnrealTargetPlatform Platform, DirectoryReference ProjectPath, ILogger Logger)
 		{
 			string[] BoolKeys = new string[] {
-				"bBuildForArm64", "bBuildForX8664", 
+				"bBuildForArm64", "bBuildForX8664",
 				"bBuildForES31", "bBuildWithHiddenSymbolVisibility", "bSaveSymbols"
 			};
 			string[] StringKeys = new string[] {
@@ -305,7 +303,7 @@ namespace UnrealBuildTool
 			{
 				return false;
 			}
-			
+
 			// any shared-between-all-androids would be here
 
 			// check the base settings
diff --git a/Engine/Source/Programs/UnrealBuildTool/Platform/Android/UEDeployAndroid.cs b/Engine/Source/Programs/UnrealBuildTool/Platform/Android/UEDeployAndroid.cs
--- a/Engine/Source/Programs/UnrealBuildTool/Platform/Android/UEDeployAndroid.cs
+++ b/Engine/Source/Programs/UnrealBuildTool/Platform/Android/UEDeployAndroid.cs
@@ -2,18 +2,16 @@
 
 using System;
 using System.Collections.Generic;
-using System.Linq;
-using System.Text;
-using System.Xml;
 using System.Diagnostics;
 using System.IO;
 using System.IO.Compression;
-using Microsoft.Win32;
+using System.Linq;
+using System.Security.Cryptography;
+using System.Text;
 using System.Xml.Linq;
 using EpicGames.Core;
-using System.Security.Cryptography;
-using UnrealBuildBase;
 using Microsoft.Extensions.Logging;
+using UnrealBuildBase;
 
 namespace UnrealBuildTool
 {
@@ -153,7 +151,7 @@ namespace UnrealBuildTool
 
 			UPL = new UnrealPluginLanguage(ProjectFile, inPluginExtraData, NDKArches, "http://schemas.android.com/apk/res/android", "xmlns:android=\"http://schemas.android.com/apk/res/android\" xmlns:tools=\"http://schemas.android.com/tools\"", UnrealTargetPlatform.Android, Logger);
 			UPLHashCode = UPL.GetUPLHash();
-//			APL.SetTrace();
+			//			APL.SetTrace();
 		}
 
 		private void SetMinimumSDKLevelForGradle()
@@ -166,7 +164,7 @@ namespace UnrealBuildTool
 			{
 				MinimumSDKLevelForGradle = Math.Max(MinimumSDKLevelForGradle, 19);
 			}
-			if(EOSSDKPluginEnabled)
+			if (EOSSDKPluginEnabled)
 			{
 				MinimumSDKLevelForGradle = Math.Max(MinimumSDKLevelForGradle, 23);
 			}
@@ -386,10 +384,10 @@ namespace UnrealBuildTool
 			}
 
 			// don't allow higher than 30.0.3 for now (will be installed by Gradle if missing)
-//			if (BestVersion > ((30 << 24) | (0 << 16) | (3 << 8)))
-//			{
-//				BestVersionString = "30.0.3";
-//			}
+			//			if (BestVersion > ((30 << 24) | (0 << 16) | (3 << 8)))
+			//			{
+			//				BestVersionString = "30.0.3";
+			//			}
 
 			CachedBuildToolsVersion = BestVersionString;
 			LastAndroidHomePath = HomePath;
@@ -415,7 +413,7 @@ namespace UnrealBuildTool
 		/// </summary>
 		/// <param name="Ini"></param>
 		protected bool ReadPackageDataInsideApkFromIni(ConfigHierarchy? Ini)
-		{		
+		{
 			// make a new one if one wasn't passed in
 			if (Ini == null)
 			{
@@ -464,7 +462,7 @@ namespace UnrealBuildTool
 			}
 
 			List<string>? GoogleVRCaps = new List<string>();
-			if(Ini.GetArray("/Script/AndroidRuntimeSettings.AndroidRuntimeSettings", "GoogleVRCaps", out GoogleVRCaps))
+			if (Ini.GetArray("/Script/AndroidRuntimeSettings.AndroidRuntimeSettings", "GoogleVRCaps", out GoogleVRCaps))
 			{
 				return GoogleVRCaps.Contains("Daydream33") || GoogleVRCaps.Contains("Daydream63") || GoogleVRCaps.Contains("Daydream66");
 			}
@@ -581,7 +579,7 @@ namespace UnrealBuildTool
 				}
 
 				// make the dst filename with the same structure as it was in SourceDir
-				string DestFilename = Path.Combine(DestDir, Utils.MakePathRelativeTo(Filename, SourceDir)).Replace('\\', Path.DirectorySeparatorChar).Replace('/', Path.DirectorySeparatorChar); 
+				string DestFilename = Path.Combine(DestDir, Utils.MakePathRelativeTo(Filename, SourceDir)).Replace('\\', Path.DirectorySeparatorChar).Replace('/', Path.DirectorySeparatorChar);
 
 				// make the subdirectory if needed
 				string DestSubdir = Path.GetDirectoryName(DestFilename)!;
@@ -940,7 +938,7 @@ namespace UnrealBuildTool
 			public string SourceFile;
 			public string DestinationFile;
 
-			public TemplateFile(string SourceFile, string DestinationFile) 
+			public TemplateFile(string SourceFile, string DestinationFile)
 			{
 				this.SourceFile = SourceFile;
 				this.DestinationFile = DestinationFile;
@@ -958,7 +956,7 @@ namespace UnrealBuildTool
 
 		private int CachedStoreVersion = -1;
 		private int CachedStoreVersionOffsetArm64 = 0;
-		private int CachedStoreVersionOffsetX8664= 0;
+		private int CachedStoreVersionOffsetX8664 = 0;
 
 		public int GetStoreVersion(UnrealArch? Architecture)
 		{
@@ -1067,7 +1065,7 @@ namespace UnrealBuildTool
 			string AppType = "";
 			if (CookFlavor.EndsWith("Client"))
 			{
-//				AppType = ".Client";		// should always be empty now; fix up the name in batch file instead
+				//				AppType = ".Client";		// should always be empty now; fix up the name in batch file instead
 			}
 
 			int StoreVersion = GetStoreVersion(UnrealArch);
@@ -1155,9 +1153,9 @@ namespace UnrealBuildTool
 			ShimFileContent.AppendFormat("import {0}.DownloaderActivity;\n", replacements["$$PackageName$$"]);
 
 			// Do OBB file checking without using DownloadActivity to avoid transit to another activity
-				ShimFileContent.Append("import android.app.Activity;\n");
-				ShimFileContent.Append("import com.google.android.vending.expansion.downloader.Helpers;\n");
-				ShimFileContent.AppendFormat("import {0}.OBBData;\n", replacements["$$PackageName$$"]);
+			ShimFileContent.Append("import android.app.Activity;\n");
+			ShimFileContent.Append("import com.google.android.vending.expansion.downloader.Helpers;\n");
+			ShimFileContent.AppendFormat("import {0}.OBBData;\n", replacements["$$PackageName$$"]);
 
 			ShimFileContent.Append("\n\npublic class DownloadShim\n{\n");
 			ShimFileContent.Append("\tpublic static OBBDownloaderService DownloaderService;\n");
@@ -1177,11 +1175,11 @@ namespace UnrealBuildTool
 			ShimFileContent.Append("\t\t\t}\n");
 			ShimFileContent.Append("\t\t\telse if (Helpers.doesFileExistDev(activity, fileName, xf.mFileSize, false)) {\n");
 			ShimFileContent.Append("\t\t\t\tGameActivity.Log.debug(\"Found OBB here: \" + fileForDevFile);\n");
-				ShimFileContent.Append("\t\t\t}\n");
+			ShimFileContent.Append("\t\t\t}\n");
 			ShimFileContent.Append("\t\t\telse return false;\n");
 			ShimFileContent.Append("\t\t}\n");
-				ShimFileContent.Append("\t\treturn true;\n");
-				ShimFileContent.Append("\t}\n");
+			ShimFileContent.Append("\t\treturn true;\n");
+			ShimFileContent.Append("\t}\n");
 
 			ShimFileContent.Append("}\n");
 			Logger.LogInformation("\n==== Writing to shim file {ShimFileName} ====", ShimFileName);
@@ -1300,7 +1298,7 @@ namespace UnrealBuildTool
 			{
 				return "x86_64";
 			}
-			
+
 			throw new BuildException("Unknown Unreal architecture '{0}'", UnrealArch);
 		}
 
@@ -1309,13 +1307,13 @@ namespace UnrealBuildTool
 			switch (NDKArch)
 			{
 				case "arm64-v8a":
-				case "arm64":       return UnrealArch.Arm64;
+				case "arm64": return UnrealArch.Arm64;
 				case "x86_64":
-				case "x64":			return UnrealArch.X64;
-					
-//				default: throw new BuildException("Unknown NDK architecture '{0}'", NDKArch);
+				case "x64": return UnrealArch.X64;
+
+				//				default: throw new BuildException("Unknown NDK architecture '{0}'", NDKArch);
 				// future-proof by returning arm64 for unknown
-				default:            return UnrealArch.Arm64;
+				default: return UnrealArch.Arm64;
 			}
 		}
 
@@ -1383,7 +1381,7 @@ namespace UnrealBuildTool
 		private static void CopySTL(AndroidToolChain ToolChain, string UnrealBuildPath, UnrealArch UnrealArch, string NDKArch, bool bForDistribution)
 		{
 			// copy it in!
-			string SourceSTLSOName = Environment.ExpandEnvironmentVariables("%NDKROOT%/sources/cxx-stl/llvm-libc++/libs/") +  NDKArch + "/libc++_shared.so";
+			string SourceSTLSOName = Environment.ExpandEnvironmentVariables("%NDKROOT%/sources/cxx-stl/llvm-libc++/libs/") + NDKArch + "/libc++_shared.so";
 			if (!File.Exists(SourceSTLSOName))
 			{
 				// NDK25 has changed a directory where it stores libs, check it instead
@@ -1469,7 +1467,7 @@ namespace UnrealBuildTool
 			Logger.LogInformation("bBuildForES31: {bBuildForES31}", (bBuildForES31 ? "true" : "false"));
 			Logger.LogInformation("bSupportsVulkan: {bSupportsVulkan}", (bSupportsVulkan ? "true" : "false"));
 		}
-		
+
 		void CopyVulkanValidationLayers(string UnrealBuildPath, UnrealArch UnrealArch, string NDKArch, string Configuration)
 		{
 			bool bSupportsVulkan = false;
@@ -1556,7 +1554,7 @@ namespace UnrealBuildTool
 					LibName = "tsan";
 					break;
 			}
-			
+
 			string SanitizerFullLibName = "libclang_rt." + LibName + Architecture + "-android.so";
 
 			string WrapSh = Path.Combine(Unreal.EngineDirectory.ToString(), "Build", "Android", "ClangSanitizers", "wrap.sh");
@@ -1939,7 +1937,7 @@ namespace UnrealBuildTool
 			return true;
 		}
 
-		private string GetAllBuildSettings(AndroidToolChain ToolChain, UnrealPluginLanguage UPL, bool bForDistribution, bool bMakeSeparateApks, bool bPackageDataInsideApk, 
+		private string GetAllBuildSettings(AndroidToolChain ToolChain, UnrealPluginLanguage UPL, bool bForDistribution, bool bMakeSeparateApks, bool bPackageDataInsideApk,
 			bool bDisableVerifyOBBOnStartUp, bool bUseExternalFilesDir, string TemplatesHashCode)
 		{
 			// make the settings string - this will be char by char compared against last time
@@ -2015,7 +2013,7 @@ namespace UnrealBuildTool
 			{
 				foreach (string Key in Section.KeyNames)
 				{
-					if (!Key.Equals("GameDefaultMap") && 
+					if (!Key.Equals("GameDefaultMap") &&
 						!Key.Equals("GlobalDefaultGameMode"))
 					{
 						continue;
@@ -2295,7 +2293,7 @@ namespace UnrealBuildTool
 				}
 			}
 		}
-	
+
 		private void PackageForDaydream(string UnrealBuildPath)
 		{
 			bool bPackageForDaydream = IsPackagingForDaydream();
@@ -2322,7 +2320,7 @@ namespace UnrealBuildTool
 			Ini.GetBool("/Script/AndroidRuntimeSettings.AndroidRuntimeSettings", "bShowLaunchImage", out bShowLaunchImage);
 			bool bPackageForOculusMobile = IsPackagingForOculusMobile(Ini); ;
 			bool bPackageForDaydream = IsPackagingForDaydream(Ini);
-			
+
 			//override the parameters if we are not showing a launch image or are packaging for Oculus Mobile and Daydream
 			if (bPackageForOculusMobile || bPackageForDaydream || !bShowLaunchImage)
 			{
@@ -2582,7 +2580,7 @@ namespace UnrealBuildTool
 			Ini.GetBool("/Script/AndroidRuntimeSettings.AndroidRuntimeSettings", "bSupportAdMob", out bSupportAdMob);
 			bool bValidateTextureFormats;
 			Ini.GetBool("/Script/AndroidRuntimeSettings.AndroidRuntimeSettings", "bValidateTextureFormats", out bValidateTextureFormats);
-			
+
 			bool bBuildForES31 = false;
 			Ini.GetBool("/Script/AndroidRuntimeSettings.AndroidRuntimeSettings", "bBuildForES31", out bBuildForES31);
 			bool bSupportsVulkan = false;
@@ -2623,7 +2621,7 @@ namespace UnrealBuildTool
 					break;
 			}
 
-			bool bEnableMulticastSupport = false; 
+			bool bEnableMulticastSupport = false;
 			Ini.GetBool("/Script/AndroidRuntimeSettings.AndroidRuntimeSettings", "bEnableMulticastSupport", out bEnableMulticastSupport);
 
 			// only apply density to configChanges if using android-24 or higher and minimum sdk is 17
@@ -2716,7 +2714,7 @@ namespace UnrealBuildTool
 			StringBuilder Text = new StringBuilder();
 			Text.AppendLine(XML_HEADER);
 			Text.AppendLine("<manifest xmlns:android=\"http://schemas.android.com/apk/res/android\" xmlns:tools=\"http://schemas.android.com/tools\"");
-//			Text.AppendLine(string.Format("          package=\"{0}\"", PackageName));
+			//			Text.AppendLine(string.Format("          package=\"{0}\"", PackageName));
 			if (ExtraManifestNodeTags != null)
 			{
 				foreach (string Line in ExtraManifestNodeTags)
@@ -2933,7 +2931,7 @@ namespace UnrealBuildTool
 			// Declare the 8 OpenGL program compiling services.
 			Text.AppendLine(@"		<service android:name=""com.epicgames.unreal.psoservices.OGLProgramService"" android:process="":psoprogramservice"" />");
 			// Declare the remaining 7 OpenGL program compiling services. (all derived from OGLProgramService)
-			for(int i = 1; i<8;i++)
+			for (int i = 1; i < 8; i++)
 			{
 				String serviceLine = String.Format("		<service android:name=\"com.epicgames.unreal.psoservices.OGLProgramService{0}\" android:process=\":psoprogramservice{0}\" />", i);
 				Text.AppendLine(serviceLine);
@@ -2975,7 +2973,7 @@ namespace UnrealBuildTool
 			// Max supported aspect ratio
 			string MaxAspectRatioString = MaxAspectRatioValue.ToString("f", System.Globalization.CultureInfo.InvariantCulture);
 			Text.AppendLine(string.Format("\t\t<meta-data android:name=\"android.max_aspect\" android:value=\"{0}\" />", MaxAspectRatioString));
-					
+
 			Text.AppendLine("\t</application>");
 
 			Text.AppendLine("");
@@ -3000,11 +2998,11 @@ namespace UnrealBuildTool
 				}
 				Text.AppendLine("\t<uses-permission android:name=\"android.permission.ACCESS_NETWORK_STATE\"/>");
 				Text.AppendLine("\t<uses-permission android:name=\"android.permission.WAKE_LOCK\"/>");
-			//	Text.AppendLine("\t<uses-permission android:name=\"android.permission.READ_PHONE_STATE\"/>");
+				//	Text.AppendLine("\t<uses-permission android:name=\"android.permission.READ_PHONE_STATE\"/>");
 				Text.AppendLine("\t<uses-permission android:name=\"com.android.vending.CHECK_LICENSE\"/>");
 				Text.AppendLine("\t<uses-permission android:name=\"android.permission.ACCESS_WIFI_STATE\"/>");
 
-				if(bEnableMulticastSupport)
+				if (bEnableMulticastSupport)
 				{
 					// This permission is needed to be able to acquire a WifiManager.MulticastLock so broadcast/multcast traffic is 
 					// not filtered out by the device network interface 
@@ -3021,7 +3019,7 @@ namespace UnrealBuildTool
 					Text.AppendLine("\t<uses-permission android:name=\"android.permission.GET_ACCOUNTS\"/>");
 				}
 
-				if(!bPackageForOculusMobile)
+				if (!bPackageForOculusMobile)
 				{
 					Text.AppendLine("\t<uses-permission android:name=\"android.permission.MODIFY_AUDIO_SETTINGS\"/>");
 					Text.AppendLine("\t<uses-permission android:name=\"android.permission.VIBRATE\"/>");
@@ -3175,24 +3173,24 @@ namespace UnrealBuildTool
 					string XmlAppId = Line.Substring(StartIndex, EndIndex - StartIndex);
 
 					//validate that the AppId matches the .ini value for the GooglePlay AppId, assuming it's valid
-					if (!bInvalidIniAppId &&  IniAppId.CompareTo(XmlAppId) != 0)
+					if (!bInvalidIniAppId && IniAppId.CompareTo(XmlAppId) != 0)
 					{
 						Logger.LogInformation("Replacing Google Play AppID in GooglePlayAppID.xml with AndroidRuntimeSettings .ini value");
 
 						bInvalid = true;
 						ReplacementId = IniAppId;
-						
-					}					
-					else if(XmlAppId.Length == 0 || !Int64.TryParse(XmlAppId, out Value))
+
+					}
+					else if (XmlAppId.Length == 0 || !Int64.TryParse(XmlAppId, out Value))
 					{
 						Logger.LogWarning("\nWARNING: GooglePlay Games App ID is invalid! Replacing it with \"1\"");
 
 						//write file with something which will fail but not cause an exception if executed
 						bInvalid = true;
 						ReplacementId = "1";
-					}	
+					}
 
-					if(bInvalid)
+					if (bInvalid)
 					{
 						// remove any read only flags if invalid so it can be replaced
 						FileInfo DestFileInfo = new FileInfo(Filename);
@@ -3285,8 +3283,8 @@ namespace UnrealBuildTool
 			using (FileStream SourceStream = new FileStream(SourceFilename, FileMode.Open, FileAccess.Read, FileShare.Read))
 			using (FileStream DestStream = new FileStream(DestFilename, FileMode.Open, FileAccess.Read, FileShare.Read))
 			{
-				using(BinaryReader SourceReader = new BinaryReader(SourceStream))
-				using(BinaryReader DestReader = new BinaryReader(DestStream))
+				using (BinaryReader SourceReader = new BinaryReader(SourceStream))
+				using (BinaryReader DestReader = new BinaryReader(DestStream))
 				{
 					bool bEOF = false;
 					while (!bEOF)
@@ -3307,7 +3305,7 @@ namespace UnrealBuildTool
 					return true;
 				}
 			}
-			
+
 		}
 
 		private bool RequiresOBB(bool bDisallowPackageInAPK, string OBBLocation)
@@ -3784,7 +3782,7 @@ namespace UnrealBuildTool
 		}
 
 		private void CreateGradlePropertiesFiles(AndroidToolChain ToolChain, UnrealArch Arch, int MinSDKVersion, int TargetSDKVersion, string CompileSDKVersion, string BuildToolsVersion, string PackageName,
-			string? DestApkName, string NDKArch,	string UnrealBuildFilesPath, string GameBuildFilesPath, string UnrealBuildGradleAppPath, string UnrealBuildPath, string UnrealBuildGradlePath,
+			string? DestApkName, string NDKArch, string UnrealBuildFilesPath, string GameBuildFilesPath, string UnrealBuildGradleAppPath, string UnrealBuildPath, string UnrealBuildGradlePath,
 			bool bForDistribution, bool bIsEmbedded, List<string> OBBFiles)
 		{
 			// Create gradle.properties
@@ -4026,12 +4024,12 @@ namespace UnrealBuildTool
 			File.WriteAllText(GradleBuildScriptAdditionsFilename, UPL.ProcessPluginNode(NDKArch, "buildscriptGradleAdditions", ""));
 		}
 
-		private void MakeApk(AndroidToolChain ToolChain, string ProjectName, TargetType InTargetType, string ProjectDirectory, string OutputPath, string EngineDirectory, bool bForDistribution, string CookFlavor, 
+		private void MakeApk(AndroidToolChain ToolChain, string ProjectName, TargetType InTargetType, string ProjectDirectory, string OutputPath, string EngineDirectory, bool bForDistribution, string CookFlavor,
 			UnrealTargetConfiguration Configuration, bool bMakeSeparateApks, bool bIncrementalPackage, bool bDisallowPackagingDataInApk, bool bDisallowExternalFilesDir, bool bSkipGradleBuild)
 		{
 			Logger.LogInformation("");
 			Logger.LogInformation("===={Time}====PREPARING TO MAKE APK=================================================================", DateTime.Now.ToString());
-			
+
 			if (Architectures == null)
 			{
 				throw new BuildException("Called MakeApk without first calling SetAndroidPluginData");
@@ -4180,7 +4178,7 @@ namespace UnrealBuildTool
 				DeleteDirectory(IntermediateAndroidPath, Logger);
 				Directory.CreateDirectory(IntermediateAndroidPath);
 			}
-			
+
 			if (!System.IO.Directory.Exists(IntermediateAndroidPath))
 			{
 				System.IO.Directory.CreateDirectory(IntermediateAndroidPath);
@@ -4209,8 +4207,8 @@ namespace UnrealBuildTool
 			// Template generated files
 			string JavaTemplateSourceDir = GetUnrealTemplateJavaSourceDir(EngineDirectory);
 			IEnumerable<TemplateFile> templates = from template in Directory.EnumerateFiles(JavaTemplateSourceDir, "*.template")
-							let RealName = Path.GetFileNameWithoutExtension(template)
-							select new TemplateFile(SourceFile: template, DestinationFile: GetUnrealTemplateJavaDestination(TemplateDestinationBase, RealName));
+												  let RealName = Path.GetFileNameWithoutExtension(template)
+												  select new TemplateFile(SourceFile: template, DestinationFile: GetUnrealTemplateJavaDestination(TemplateDestinationBase, RealName));
 
 			// Generate the OBB and Shim files here
 			string ObbFileLocation = ProjectDirectory + "/Saved/StagedBuilds/Android" + CookFlavor + ".obb";
@@ -4235,7 +4233,7 @@ namespace UnrealBuildTool
 
 			WriteJavaDownloadSupportFiles(UnrealDownloadShimFileName, templates, new Dictionary<string, string>{
 				{ "$$GameName$$", ProjectName },
-				{ "$$PublicKey$$", GetPublicKey() }, 
+				{ "$$PublicKey$$", GetPublicKey() },
 				{ "$$PackageName$$",PackageName }
 			});
 
@@ -4364,9 +4362,9 @@ namespace UnrealBuildTool
 			{
 				// check if so's are up to date against various inputs
 				List<string> JavaFiles = new List<string>{
-                                                    UnrealOBBDataFileName,
-                                                    UnrealDownloadShimFileName
-                                                };
+													UnrealOBBDataFileName,
+													UnrealDownloadShimFileName
+												};
 				// Add the generated files too
 				JavaFiles.AddRange(from t in templates select t.SourceFile);
 				JavaFiles.AddRange(from t in templates select t.DestinationFile);
@@ -4391,7 +4389,7 @@ namespace UnrealBuildTool
 			}
 
 			// Initialize UPL contexts for each architecture enabled
-			UPL.Init(NDKArches, bForDistribution, EngineDirectory, IntermediateAndroidPath, ProjectDirectory, Configuration.ToString(), bSkipGradleBuild, bPerArchBuildDir:true, ArchRemapping:ArchRemapping);
+			UPL.Init(NDKArches, bForDistribution, EngineDirectory, IntermediateAndroidPath, ProjectDirectory, Configuration.ToString(), bSkipGradleBuild, bPerArchBuildDir: true, ArchRemapping: ArchRemapping);
 
 			IEnumerable<Tuple<UnrealArch, string>>? BuildList = null;
 
@@ -4430,7 +4428,7 @@ namespace UnrealBuildTool
 					UnrealArch Arch = build.Item1;
 					string Manifest = build.Item2;
 					string NDKArch = GetNDKArch(Arch);
-					
+
 					string UnrealBuildPath = Path.Combine(IntermediateAndroidPath, Arch.ToString());
 
 					Logger.LogInformation("\n===={Time}====PREPARING NATIVE CODE====={Arch}============================================================", DateTime.Now.ToString(), Arch);
@@ -4794,7 +4792,7 @@ namespace UnrealBuildTool
 				CopyPSOService(UnrealBuildPath, UnrealPreBuiltFilesPath, Arch, NDKArch);
 				CopyGfxDebugger(UnrealBuildPath, Arch, NDKArch);
 				CopyVulkanValidationLayers(UnrealBuildPath, Arch, NDKArch, Configuration.ToString());
-				
+
 				if (Sanitizer != AndroidToolChain.ClangSanitizer.None)
 				{
 					CopyClangSanitizerLib(UnrealBuildPath, Arch, NDKArch, Sanitizer);
@@ -4877,7 +4875,7 @@ namespace UnrealBuildTool
 				int TargetSDKVersion = 0;
 				int NDKLevelInt = 0;
 				GetMinTargetSDKVersions(ToolChain, Arch, UPL, NDKArch, bEnableBundle, out MinSDKVersion, out TargetSDKVersion, out NDKLevelInt);
-					
+
 				// move JavaLibs into subprojects
 				string JavaLibsDir = Path.Combine(UnrealBuildPath, "JavaLibs");
 				PrepareJavaLibsForGradle(JavaLibsDir, UnrealBuildGradlePath, MinSDKVersion.ToString(), TargetSDKVersion.ToString(), CompileSDKVersion, BuildToolsVersion, NDKArch);
@@ -4885,7 +4883,7 @@ namespace UnrealBuildTool
 				// Create local.properties
 				String LocalPropertiesFilename = Path.Combine(UnrealBuildGradlePath, "local.properties");
 				StringBuilder LocalProperties = new StringBuilder();
-//				LocalProperties.AppendLine(string.Format("ndk.dir={0}", Environment.GetEnvironmentVariable("NDKROOT")!.Replace("\\", "/")));
+				//				LocalProperties.AppendLine(string.Format("ndk.dir={0}", Environment.GetEnvironmentVariable("NDKROOT")!.Replace("\\", "/")));
 				LocalProperties.AppendLine(string.Format("sdk.dir={0}", Environment.GetEnvironmentVariable("ANDROID_HOME")!.Replace("\\", "/")));
 				File.WriteAllText(LocalPropertiesFilename, LocalProperties.ToString());
 
@@ -5505,7 +5503,7 @@ namespace UnrealBuildTool
 			TargetType Type = TargetType.Game;
 			if (CookFlavor.EndsWith("Client"))
 			{
-				Type = TargetType.Client;			
+				Type = TargetType.Client;
 			}
 			else if (CookFlavor.EndsWith("Server"))
 			{
@@ -5531,7 +5529,7 @@ namespace UnrealBuildTool
 			SavePackageInfo(ProjectName, ProjectDirectory.FullName, Type, bSkipGradleBuild);
 
 			MakeApk(ToolChain, ProjectName, Type, ProjectDirectory.FullName, ExecutablePath, EngineDirectory, bForDistribution: bForDistribution, CookFlavor: CookFlavor, Configuration: Configuration,
-				bMakeSeparateApks: ShouldMakeSeparateApks(), bIncrementalPackage: false, bDisallowPackagingDataInApk: bIsDataDeploy, bDisallowExternalFilesDir: !bForDistribution || bIsDataDeploy, bSkipGradleBuild:bSkipGradleBuild);
+				bMakeSeparateApks: ShouldMakeSeparateApks(), bIncrementalPackage: false, bDisallowPackagingDataInApk: bIsDataDeploy, bDisallowExternalFilesDir: !bForDistribution || bIsDataDeploy, bSkipGradleBuild: bSkipGradleBuild);
 			return true;
 		}
 
@@ -5680,7 +5678,7 @@ namespace UnrealBuildTool
 				bool Changed = false;
 				foreach (KeyValuePair<string, string> KVP in Replacements)
 				{
-					if(SrcLine.Contains(KVP.Key))
+					if (SrcLine.Contains(KVP.Key))
 					{
 						SrcLine = SrcLine.Replace(KVP.Key, KVP.Value);
 						Changed = true;
@@ -5796,7 +5794,7 @@ namespace UnrealBuildTool
 			}
 		}
 
-		private AndroidAARHandler CreateAARHandler(string EngineDir, string UnrealBuildPath, List<string> NDKArches, bool HandleDependencies=true)
+		private AndroidAARHandler CreateAARHandler(string EngineDir, string UnrealBuildPath, List<string> NDKArches, bool HandleDependencies = true)
 		{
 			AndroidAARHandler AARHandler = new AndroidAARHandler();
 			string ImportList = "";
diff --git a/Engine/Source/Programs/UnrealBuildTool/Platform/Apple/AppleExports.cs b/Engine/Source/Programs/UnrealBuildTool/Platform/Apple/AppleExports.cs
--- a/Engine/Source/Programs/UnrealBuildTool/Platform/Apple/AppleExports.cs
+++ b/Engine/Source/Programs/UnrealBuildTool/Platform/Apple/AppleExports.cs
@@ -1,15 +1,7 @@
 ﻿// Copyright Epic Games, Inc. All Rights Reserved.
 
 using System;
-using System.IO;
-using System.Collections.Generic;
-using System.Diagnostics;
-using System.Linq;
-using System.Text;
-using System.Threading.Tasks;
-using System.Xml.Linq;
 using EpicGames.Core;
-using UnrealBuildBase;
 using Microsoft.Extensions.Logging;
 
 namespace UnrealBuildTool
