UnrealBuildTool: Automated code cleanup

diff --git a/Engine/Source/Programs/UnrealBuildTool/Platform/IOS/IOSExports.cs b/Engine/Source/Programs/UnrealBuildTool/Platform/IOS/IOSExports.cs
--- a/Engine/Source/Programs/UnrealBuildTool/Platform/IOS/IOSExports.cs
+++ b/Engine/Source/Programs/UnrealBuildTool/Platform/IOS/IOSExports.cs
@@ -1,15 +1,12 @@
 ﻿// Copyright Epic Games, Inc. All Rights Reserved.
 
 using System;
-using System.IO;
-using System.Collections.Generic;
 using System.Diagnostics;
+using System.IO;
 using System.Linq;
 using System.Text;
-using System.Threading.Tasks;
 using System.Xml.Linq;
 using EpicGames.Core;
-using UnrealBuildBase;
 using Microsoft.Extensions.Logging;
 
 namespace UnrealBuildTool
@@ -143,14 +140,14 @@ namespace UnrealBuildTool
 				return;
 			}
 
-            // Also don't attempt to use a remote Mac if packaging for TVOS on PC.
-            if (Platform == UnrealTargetPlatform.TVOS && BuildHostPlatform.Current.Platform != UnrealTargetPlatform.Mac)
-            {
-                return;
-            }
+			// Also don't attempt to use a remote Mac if packaging for TVOS on PC.
+			if (Platform == UnrealTargetPlatform.TVOS && BuildHostPlatform.Current.Platform != UnrealTargetPlatform.Mac)
+			{
+				return;
+			}
 
 			// Compile the asset catalog immediately
-			if(BuildHostPlatform.Current.Platform != UnrealTargetPlatform.Mac)
+			if (BuildHostPlatform.Current.Platform != UnrealTargetPlatform.Mac)
 			{
 				FileReference OutputFile = FileReference.Combine(StageDirectory, "Assets.car");
 
@@ -171,10 +168,10 @@ namespace UnrealBuildTool
 				}
 
 				// Run the process locally
-				using(Process Process = new Process())
+				using (Process Process = new Process())
 				{
 					Process.StartInfo.FileName = "/usr/bin/xcrun";
-					Process.StartInfo.Arguments = IOSToolChain.GetAssetCatalogArgs(Platform, ResourcesDir.FullName, OutputFile.Directory.FullName);; 
+					Process.StartInfo.Arguments = IOSToolChain.GetAssetCatalogArgs(Platform, ResourcesDir.FullName, OutputFile.Directory.FullName); ;
 					Process.StartInfo.UseShellExecute = false;
 					Utils.RunLocalProcess(Process);
 				}
@@ -185,7 +182,7 @@ namespace UnrealBuildTool
 		/// Set a secondary remote Mac to retrieve built data on a remote Mac.
 		/// </summary>
 		/// <returns></returns>
-		
+
 		public static void SetSecondaryRemoteMac(string ClientPlatform, FileReference ProjectFile, ILogger Logger)
 		{
 			RemoteMac Remote = new RemoteMac(ProjectFile, Logger, true, true);
@@ -203,7 +200,7 @@ namespace UnrealBuildTool
 		/// <param name="Logger">A logger</param>
 		/// <returns></returns>
 		public static void PrepareRemoteMacForDebugging(string ClientPlatform, FileReference ProjectFile, ILogger Logger)
-        {
+		{
 			RemoteMac Remote = new RemoteMac(ProjectFile, Logger);
 			Remote.PrepareToDebug(ClientPlatform, ProjectFile, Logger);
 		}
@@ -233,7 +230,7 @@ namespace UnrealBuildTool
 			if (MobileProvisionFile != null && File.Exists(MobileProvisionFile.FullName))
 			{
 				Console.WriteLine("Write entitlements from provisioning file {0}", MobileProvisionFile);
-				
+
 				MobileProvisionContents MobileProvisionContent = MobileProvisionContents.Read(MobileProvisionFile);
 
 				iCloudContainerIdentifier = MobileProvisionContent.GetNodeValueByName("com.apple.developer.icloud-container-identifiers");
@@ -341,7 +338,7 @@ namespace UnrealBuildTool
 					Text.AppendLine("\t<key>com.apple.developer.user-management</key>");
 					Text.AppendLine("\t<array><string>runs-as-current-user</string></array>");
 				}
-				
+
 				// End of entitlements
 				Text.AppendLine("</dict>");
 				Text.AppendLine("</plist>");
diff --git a/Engine/Source/Programs/UnrealBuildTool/Platform/IOS/IOSPostBuildSyncMode.cs b/Engine/Source/Programs/UnrealBuildTool/Platform/IOS/IOSPostBuildSyncMode.cs
--- a/Engine/Source/Programs/UnrealBuildTool/Platform/IOS/IOSPostBuildSyncMode.cs
+++ b/Engine/Source/Programs/UnrealBuildTool/Platform/IOS/IOSPostBuildSyncMode.cs
@@ -2,12 +2,10 @@
 
 using System;
 using System.Collections.Generic;
-using System.Linq;
-using System.Text;
 using System.Threading.Tasks;
 using EpicGames.Core;
-using UnrealBuildBase;
 using Microsoft.Extensions.Logging;
+using UnrealBuildBase;
 
 namespace UnrealBuildTool
 {
diff --git a/Engine/Source/Programs/UnrealBuildTool/Platform/IOS/IOSProjectGenerator.cs b/Engine/Source/Programs/UnrealBuildTool/Platform/IOS/IOSProjectGenerator.cs
--- a/Engine/Source/Programs/UnrealBuildTool/Platform/IOS/IOSProjectGenerator.cs
+++ b/Engine/Source/Programs/UnrealBuildTool/Platform/IOS/IOSProjectGenerator.cs
@@ -1,10 +1,6 @@
 // Copyright Epic Games, Inc. All Rights Reserved.
 
-using System;
 using System.Collections.Generic;
-using System.Text;
-using System.Diagnostics;
-using System.IO;
 using EpicGames.Core;
 using Microsoft.Extensions.Logging;
 
diff --git a/Engine/Source/Programs/UnrealBuildTool/Platform/IOS/IOSToolChain.cs b/Engine/Source/Programs/UnrealBuildTool/Platform/IOS/IOSToolChain.cs
--- a/Engine/Source/Programs/UnrealBuildTool/Platform/IOS/IOSToolChain.cs
+++ b/Engine/Source/Programs/UnrealBuildTool/Platform/IOS/IOSToolChain.cs
@@ -2,22 +2,17 @@
 
 
 using System;
-using System.Collections;
 using System.Collections.Generic;
-using System.Linq;
-using System.IO;
 using System.Diagnostics;
-using System.Security.AccessControl;
-using System.Xml;
-using System.Xml.Linq;
+using System.IO;
+using System.Linq;
+using System.Security.Cryptography.X509Certificates;
 using System.Text;
+using EpicGames.Core;
 using Ionic.Zip;
 using Ionic.Zlib;
-using EpicGames.Core;
-using System.Security.Cryptography.X509Certificates;
-using UnrealBuildBase;
 using Microsoft.Extensions.Logging;
-using System.Text.RegularExpressions;
+using UnrealBuildBase;
 
 namespace UnrealBuildTool
 {
@@ -490,7 +485,7 @@ namespace UnrealBuildTool
 					LinkCommandArguments += string.Format(" -l\\\"{0}\\\"", AdditionalLibrary);
 				}
 
-				foreach(FileReference Library in LinkEnvironment.Libraries)
+				foreach (FileReference Library in LinkEnvironment.Libraries)
 				{
 					// for absolute library paths, convert to remote filename
 					// add it to the prerequisites to make sure it's built first (this should be the case of non-system libraries)
@@ -1204,7 +1199,7 @@ namespace UnrealBuildTool
 
 				if (Log.OutputFile != null)
 				{
-					string LogFileName = Log.OutputFile.GetFileNameWithoutExtension()+"_PostBuildSync.txt";
+					string LogFileName = Log.OutputFile.GetFileNameWithoutExtension() + "_PostBuildSync.txt";
 					LogFileName = FileReference.Combine(Log.OutputFile.Directory, LogFileName).ToString();
 					PostBuildSyncArguments += " -log=\"" + LogFileName + "\"";
 				}
@@ -1308,8 +1303,8 @@ namespace UnrealBuildTool
 			}
 
 			string FabricPath = Unreal.EngineDirectory + "/Intermediate/UnzippedFrameworks/Crashlytics/Fabric.embeddedframework";
-            if (Directory.Exists(FabricPath) && Environment.GetEnvironmentVariable("IsBuildMachine") == "1")
-            {
+			if (Directory.Exists(FabricPath) && Environment.GetEnvironmentVariable("IsBuildMachine") == "1")
+			{
 				//string PlistFile = ProjectDir + "/Intermediate/IOS/" + ProjectName + "-Info.plist";
 				Process FabricProcess = new Process();
 				FabricProcess.StartInfo.WorkingDirectory = Path.GetDirectoryName(DsymZip);
@@ -1380,15 +1375,15 @@ namespace UnrealBuildTool
 			FileReference.WriteAllText(MissingPlistPath, Plist);
 		}
 
-        /// <summary>
-        /// If the project is a UnrealGame project, Target.ProjectDirectory refers to the engine dir, not the actual dir of the project. So this method gets the 
-        /// actual directory of the project whether it is a UnrealGame project or not.
-        /// </summary>
-        /// <returns>The actual project directory.</returns>
-        /// <param name="ProjectFile">The path to the project file</param>
-        private static DirectoryReference GetActualProjectDirectory(FileReference? ProjectFile)
+		/// <summary>
+		/// If the project is a UnrealGame project, Target.ProjectDirectory refers to the engine dir, not the actual dir of the project. So this method gets the 
+		/// actual directory of the project whether it is a UnrealGame project or not.
+		/// </summary>
+		/// <returns>The actual project directory.</returns>
+		/// <param name="ProjectFile">The path to the project file</param>
+		private static DirectoryReference GetActualProjectDirectory(FileReference? ProjectFile)
 		{
-			DirectoryReference ProjectDirectory = (ProjectFile == null ? DirectoryReference.FromString(UnrealBuildTool.GetRemoteIniPath())! : DirectoryReference.FromFile(ProjectFile)!); 
+			DirectoryReference ProjectDirectory = (ProjectFile == null ? DirectoryReference.FromString(UnrealBuildTool.GetRemoteIniPath())! : DirectoryReference.FromFile(ProjectFile)!);
 			return ProjectDirectory;
 		}
 
@@ -1401,7 +1396,7 @@ namespace UnrealBuildTool
 			string FrameworkName = Target.TargetName;
 			string BundleId = ProjectFiles.Xcode.XcodeFrameworkWrapperUtils.GetBundleID(Target.ProjectDirectory, Target.ProjectFile);
 			string EnginePath = Unreal.EngineDirectory.ToString();
-			string SrcFrameworkPath = DirectoryReference.Combine(Target.ProjectDirectory , "Binaries", "IOS", Target.Configuration.ToString()).ToString(); // We use Target.ProjectDirectory because if it is a UnrealGame we want the engine dir and not the actual project dir.
+			string SrcFrameworkPath = DirectoryReference.Combine(Target.ProjectDirectory, "Binaries", "IOS", Target.Configuration.ToString()).ToString(); // We use Target.ProjectDirectory because if it is a UnrealGame we want the engine dir and not the actual project dir.
 			string CookedDataPath = DirectoryReference.Combine(ProjectDirectory, "Saved", "StagedBuilds", "IOS", "cookeddata").ToString();
 			bool GenerateFrameworkWrapperProject = ProjectFiles.Xcode.XcodeFrameworkWrapperUtils.GetGenerateFrameworkWrapperProject(Target.ProjectDirectory);
 
@@ -1500,8 +1495,8 @@ namespace UnrealBuildTool
 
 				// generate the dummy project so signing works
 				DirectoryReference? XcodeWorkspaceDir = null;
-                if (!Target.bBuildAsFramework)
-                {
+				if (!Target.bBuildAsFramework)
+				{
 					if (AppName == "UnrealGame" || AppName == "UnrealClient" || Target.ProjectFile == null || Target.ProjectFile.IsUnderDirectory(Unreal.EngineDirectory))
 					{
 						XcodeWorkspaceDir = DirectoryReference.Combine(Unreal.EngineDirectory, "Intermediate/ProjectFiles", String.Format("UE5_{0}.xcworkspace", (Target.Platform == UnrealTargetPlatform.IOS ? "IOS" : "TVOS")));
@@ -1510,7 +1505,7 @@ namespace UnrealBuildTool
 					{
 						XcodeWorkspaceDir = DirectoryReference.Combine(Target.ProjectDirectory, "Intermediate/ProjectFiles", String.Format("{0}_{1}.xcworkspace", Target.ProjectFile.GetFileNameWithoutExtension(), (Target.Platform == UnrealTargetPlatform.IOS ? "IOS" : "TVOS")));
 					}
-                }
+				}
 
 				// Path to the temporary keychain. When -ImportCertificate is specified, we will temporarily add this to the list of keychains to search, and remove it later.
 				FileReference TempKeychain = FileReference.Combine(Target.ProjectIntermediateDirectory!, "TempKeychain.keychain");
@@ -1678,20 +1673,20 @@ namespace UnrealBuildTool
 				{
 					Logger.LogInformation("Generating project...");
 					DirectoryReference? GeneratedProjectFile;
-				    if (AppName == "UnrealGame" || AppName == "UnrealClient" || Target.ProjectFile == null || Target.ProjectFile.IsUnderDirectory(Unreal.EngineDirectory))
-				    {
-					    GenerateProjectFiles(Target.ProjectFile, new string[] { "-platforms=" + (Target.Platform == UnrealTargetPlatform.IOS ? "IOS" : "TVOS"), "-NoIntellIsense", (Target.Platform == UnrealTargetPlatform.IOS ? "-iosdeployonly" : "-tvosdeployonly"), "-ignorejunk", (Target.bForDistribution ? "-distribution" : "-development"), "-bundleID=" + BundleID, "-includetemptargets", "-appname=" + AppName }, Logger, out GeneratedProjectFile);
-				    }
-				    else
-				    {
-					    GenerateProjectFiles(Target.ProjectFile, new string[] { "-platforms=" + (Target.Platform == UnrealTargetPlatform.IOS ? "IOS" : "TVOS"), "-NoIntellIsense", (Target.Platform == UnrealTargetPlatform.IOS ? "-iosdeployonly" : "-tvosdeployonly"), "-ignorejunk", (Target.bForDistribution ? "-distribution" : "-development"), String.Format("-project={0}", Target.ProjectFile), "-game", "-bundleID=" + BundleID, "-includetemptargets" }, Logger, out GeneratedProjectFile);
-				    }
+					if (AppName == "UnrealGame" || AppName == "UnrealClient" || Target.ProjectFile == null || Target.ProjectFile.IsUnderDirectory(Unreal.EngineDirectory))
+					{
+						GenerateProjectFiles(Target.ProjectFile, new string[] { "-platforms=" + (Target.Platform == UnrealTargetPlatform.IOS ? "IOS" : "TVOS"), "-NoIntellIsense", (Target.Platform == UnrealTargetPlatform.IOS ? "-iosdeployonly" : "-tvosdeployonly"), "-ignorejunk", (Target.bForDistribution ? "-distribution" : "-development"), "-bundleID=" + BundleID, "-includetemptargets", "-appname=" + AppName }, Logger, out GeneratedProjectFile);
+					}
+					else
+					{
+						GenerateProjectFiles(Target.ProjectFile, new string[] { "-platforms=" + (Target.Platform == UnrealTargetPlatform.IOS ? "IOS" : "TVOS"), "-NoIntellIsense", (Target.Platform == UnrealTargetPlatform.IOS ? "-iosdeployonly" : "-tvosdeployonly"), "-ignorejunk", (Target.bForDistribution ? "-distribution" : "-development"), String.Format("-project={0}", Target.ProjectFile), "-game", "-bundleID=" + BundleID, "-includetemptargets" }, Logger, out GeneratedProjectFile);
+					}
 					Logger.LogInformation("Generated {0}!", GeneratedProjectFile?.FullName ?? "NONE");
 					// Make sure it exists
 					if (!DirectoryReference.Exists(XcodeWorkspaceDir!) || (GeneratedProjectFile != null && GeneratedProjectFile != XcodeWorkspaceDir))
-				    {
-					    throw new BuildException("Unable to create stub IPA; Xcode workspace not found at {0}", XcodeWorkspaceDir);
-				    }
+					{
+						throw new BuildException("Unable to create stub IPA; Xcode workspace not found at {0}", XcodeWorkspaceDir);
+					}
 				}
 
 				// ensure the plist, entitlements, and provision files are properly copied
diff --git a/Engine/Source/Programs/UnrealBuildTool/Platform/IOS/MobileProvisionContents.cs b/Engine/Source/Programs/UnrealBuildTool/Platform/IOS/MobileProvisionContents.cs
--- a/Engine/Source/Programs/UnrealBuildTool/Platform/IOS/MobileProvisionContents.cs
+++ b/Engine/Source/Programs/UnrealBuildTool/Platform/IOS/MobileProvisionContents.cs
@@ -4,11 +4,8 @@ using System;
 using System.Collections.Generic;
 using System.IO;
 using System.Linq;
-using System.Text;
-using System.Threading.Tasks;
 using System.Xml;
 using EpicGames.Core;
-using UnrealBuildBase;
 
 namespace UnrealBuildTool
 {
@@ -66,7 +63,7 @@ namespace UnrealBuildTool
 		public string GetUniqueId()
 		{
 			XmlElement? UniqueIdElement;
-			if(!NameToValue.TryGetValue("UUID", out UniqueIdElement))
+			if (!NameToValue.TryGetValue("UUID", out UniqueIdElement))
 			{
 				throw new BuildException("Missing UUID in MobileProvision");
 			}
@@ -122,14 +119,14 @@ namespace UnrealBuildTool
 		public bool TryGetTeamUniqueId(out string? UniqueId)
 		{
 			XmlElement? UniqueIdElement;
-			if(!NameToValue.TryGetValue("TeamIdentifier", out UniqueIdElement) || UniqueIdElement.Name != "array")
+			if (!NameToValue.TryGetValue("TeamIdentifier", out UniqueIdElement) || UniqueIdElement.Name != "array")
 			{
 				UniqueId = null;
 				return false;
 			}
 
 			XmlElement? ValueElement = UniqueIdElement.SelectSingleNode("string") as XmlElement;
-			if(ValueElement == null)
+			if (ValueElement == null)
 			{
 				UniqueId = null;
 				return false;
@@ -158,21 +155,21 @@ namespace UnrealBuildTool
 		public static XmlDocument ReadXml(FileReference Location)
 		{
 			// Provision data is stored as PKCS7-signed file in ASN.1 BER format
-			using(BinaryReader Reader = new BinaryReader(File.Open(Location.FullName, FileMode.Open, FileAccess.Read)))
+			using (BinaryReader Reader = new BinaryReader(File.Open(Location.FullName, FileMode.Open, FileAccess.Read)))
 			{
 				long Length = Reader.BaseStream.Length;
-				while(Reader.BaseStream.Position < Length)
+				while (Reader.BaseStream.Position < Length)
 				{
 					Asn.FieldInfo Field = Asn.ReadField(Reader);
-					if(Field.Tag == Asn.FieldTag.OBJECT_IDENTIFIER)
+					if (Field.Tag == Asn.FieldTag.OBJECT_IDENTIFIER)
 					{
 						int[] Identifier = Asn.ReadObjectIdentifier(Reader, Field.Length);
-						if(Enumerable.SequenceEqual(Identifier, Asn.ObjectIdentifier.Pkcs7_Data))
+						if (Enumerable.SequenceEqual(Identifier, Asn.ObjectIdentifier.Pkcs7_Data))
 						{
-							while(Reader.BaseStream.Position < Length)
+							while (Reader.BaseStream.Position < Length)
 							{
 								Asn.FieldInfo NextField = Asn.ReadField(Reader);
-								if(NextField.Tag == Asn.FieldTag.OCTET_STRING)
+								if (NextField.Tag == Asn.FieldTag.OCTET_STRING)
 								{
 									byte[] Data = Reader.ReadBytes(NextField.Length);
 
diff --git a/Engine/Source/Programs/UnrealBuildTool/Platform/IOS/UEBuildIOS.cs b/Engine/Source/Programs/UnrealBuildTool/Platform/IOS/UEBuildIOS.cs
--- a/Engine/Source/Programs/UnrealBuildTool/Platform/IOS/UEBuildIOS.cs
+++ b/Engine/Source/Programs/UnrealBuildTool/Platform/IOS/UEBuildIOS.cs
@@ -2,15 +2,13 @@
 
 using System;
 using System.Collections.Generic;
-using System.Text;
 using System.Diagnostics;
 using System.IO;
 using System.Linq;
+using System.Text;
 using System.Text.RegularExpressions;
 using EpicGames.Core;
-using UnrealBuildBase;
 using Microsoft.Extensions.Logging;
-using System.Security.Cryptography.X509Certificates;
 
 namespace UnrealBuildTool
 {
@@ -233,7 +231,7 @@ namespace UnrealBuildTool
 		[ConfigFile(ConfigHierarchyType.Engine, "/Script/IOSRuntimeSettings.IOSRuntimeSettings", "bGeneratedSYMFile")]
 		[CommandLine("-generatedsymfile")]
 		public readonly bool bGeneratedSYMFile = false;
-		
+
 		/// <summary>
 		/// Whether to generate a dSYM bundle (as opposed to single file dSYM)
 		/// </summary>
@@ -241,16 +239,16 @@ namespace UnrealBuildTool
 		[CommandLine("-generatedsymbundle")]
 		public readonly bool bGeneratedSYMBundle = false;
 
-        /// <summary>
-        /// Whether to generate a dSYM file or not.
-        /// </summary>
-        [ConfigFile(ConfigHierarchyType.Engine, "/Script/IOSRuntimeSettings.IOSRuntimeSettings", "bGenerateCrashReportSymbols")]
-        public readonly bool bGenerateCrashReportSymbols = false;
+		/// <summary>
+		/// Whether to generate a dSYM file or not.
+		/// </summary>
+		[ConfigFile(ConfigHierarchyType.Engine, "/Script/IOSRuntimeSettings.IOSRuntimeSettings", "bGenerateCrashReportSymbols")]
+		public readonly bool bGenerateCrashReportSymbols = false;
 
-        /// <summary>
-        /// The minimum supported version
-        /// </summary>
-        [ConfigFile(ConfigHierarchyType.Engine, "/Script/IOSRuntimeSettings.IOSRuntimeSettings", "MinimumiOSVersion")]
+		/// <summary>
+		/// The minimum supported version
+		/// </summary>
+		[ConfigFile(ConfigHierarchyType.Engine, "/Script/IOSRuntimeSettings.IOSRuntimeSettings", "MinimumiOSVersion")]
 		private readonly string? MinimumIOSVersion = null;
 
 		/// <summary>
@@ -283,23 +281,23 @@ namespace UnrealBuildTool
 		[ConfigFile(ConfigHierarchyType.Engine, "/Script/IOSRuntimeSettings.IOSRuntimeSettings", "MobileProvision")]
 		public readonly string MobileProvision = "";
 
-        /// <summary>
-        /// signing certificate to use for code signing
-        /// </summary>
+		/// <summary>
+		/// signing certificate to use for code signing
+		/// </summary>
 		[ConfigFile(ConfigHierarchyType.Engine, "/Script/IOSRuntimeSettings.IOSRuntimeSettings", "SigningCertificate")]
-        public readonly string SigningCertificate = "";
+		public readonly string SigningCertificate = "";
 
-        /// <summary>
-        /// true if notifications are enabled
-        /// </summary>
+		/// <summary>
+		/// true if notifications are enabled
+		/// </summary>
 		[ConfigFile(ConfigHierarchyType.Engine, "/Script/IOSRuntimeSettings.IOSRuntimeSettings", "bEnableRemoteNotificationsSupport")]
-        public readonly bool bNotificationsEnabled = false;
+		public readonly bool bNotificationsEnabled = false;
 
-        /// <summary>
-        /// true if notifications are enabled
-        /// </summary>
-        [ConfigFile(ConfigHierarchyType.Engine, "/Script/IOSRuntimeSettings.IOSRuntimeSettings", "bEnableBackgroundFetch")]
-        public readonly bool bBackgroundFetchEnabled = false;
+		/// <summary>
+		/// true if notifications are enabled
+		/// </summary>
+		[ConfigFile(ConfigHierarchyType.Engine, "/Script/IOSRuntimeSettings.IOSRuntimeSettings", "bEnableBackgroundFetch")]
+		public readonly bool bBackgroundFetchEnabled = false;
 
 		/// <summary>
 		/// true if iTunes file sharing support is enabled
@@ -330,7 +328,7 @@ namespace UnrealBuildTool
 		/// </summary>
 		[ConfigFile(ConfigHierarchyType.Engine, "/Script/IOSRuntimeSettings.IOSRuntimeSettings", "bDisableForceInline")]
 		public readonly bool bDisableForceInline = false;
-		
+
 		/// <summary>
 		/// true if IDFA are enabled
 		/// </summary>
@@ -402,7 +400,7 @@ namespace UnrealBuildTool
 		/// </summary>
 		/// <param name="ProjectFile">The project file to read settings for</param>
 		/// <param name="Bundle">Bundle identifier needed when project file is empty</param>
-		public IOSProjectSettings(FileReference? ProjectFile, string? Bundle) 
+		public IOSProjectSettings(FileReference? ProjectFile, string? Bundle)
 			: this(ProjectFile, UnrealTargetPlatform.IOS, Bundle)
 		{
 		}
@@ -428,13 +426,13 @@ namespace UnrealBuildTool
 	/// <summary>
 	/// IOS provisioning data
 	/// </summary>
-    class IOSProvisioningData
-    {
+	class IOSProvisioningData
+	{
 		public string? SigningCertificate;
 		public FileReference? MobileProvisionFile;
-        public string? MobileProvisionUUID;
-        public string? MobileProvisionName;
-        public string? TeamUUID;
+		public string? MobileProvisionUUID;
+		public string? MobileProvisionName;
+		public string? TeamUUID;
 		public string? BundleIdentifier;
 		public bool bHaveCertificate = false;
 
@@ -447,27 +445,27 @@ namespace UnrealBuildTool
 
 		protected IOSProvisioningData(IOSProjectSettings ProjectSettings, bool bIsTVOS, bool bForDistribution, ILogger Logger)
 		{
-            SigningCertificate = ProjectSettings.SigningCertificate;
-            string? MobileProvision = ProjectSettings.MobileProvision;
+			SigningCertificate = ProjectSettings.SigningCertificate;
+			string? MobileProvision = ProjectSettings.MobileProvision;
 
 			FileReference? ProjectFile = ProjectSettings.ProjectFile;
 			CodeSigningConfig.Initialize(ProjectFile, bIsTVOS);
 
 			if (!string.IsNullOrEmpty(SigningCertificate))
-            {
+			{
 				List<string> Certs = AppleCodeSign.FindCertificates();
 				List<FileReference> Provisions = AppleCodeSign.FindProvisions(ProjectSettings.BundleIdentifier, bForDistribution, out _);
-            }
-            else
-            {
-                SigningCertificate = bForDistribution ? "iPhone Distribution" : "iPhone Developer";
-                bHaveCertificate = true;
-            }
-
-			if(!string.IsNullOrEmpty(MobileProvision))
+			}
+			else
+			{
+				SigningCertificate = bForDistribution ? "iPhone Distribution" : "iPhone Developer";
+				bHaveCertificate = true;
+			}
+
+			if (!string.IsNullOrEmpty(MobileProvision))
 			{
 				DirectoryReference MobileProvisionDir;
-				if(BuildHostPlatform.Current.Platform == UnrealTargetPlatform.Mac)
+				if (BuildHostPlatform.Current.Platform == UnrealTargetPlatform.Mac)
 				{
 					MobileProvisionDir = DirectoryReference.Combine(new DirectoryReference(Environment.GetEnvironmentVariable("HOME")!), "Library", "MobileDevice", "Provisioning Profiles");
 				}
@@ -477,45 +475,45 @@ namespace UnrealBuildTool
 				}
 
 				FileReference PossibleMobileProvisionFile = FileReference.Combine(MobileProvisionDir, MobileProvision);
-				if(FileReference.Exists(PossibleMobileProvisionFile))
+				if (FileReference.Exists(PossibleMobileProvisionFile))
 				{
 					MobileProvisionFile = PossibleMobileProvisionFile;
 				}
 			}
 
-            if (MobileProvisionFile == null || !bHaveCertificate)
-            {
+			if (MobileProvisionFile == null || !bHaveCertificate)
+			{
 
-                SigningCertificate = "";
-                MobileProvision = "";
+				SigningCertificate = "";
+				MobileProvision = "";
 				MobileProvisionFile = null;
-                Logger.LogInformation("Provision not specified or not found for {Project}, searching for compatible match...", ((ProjectFile != null) ? ProjectFile.GetFileNameWithoutAnyExtensions() : "UnrealGame"));
+				Logger.LogInformation("Provision not specified or not found for {Project}, searching for compatible match...", ((ProjectFile != null) ? ProjectFile.GetFileNameWithoutAnyExtensions() : "UnrealGame"));
 
 				if (AppleCodeSign.FindCertAndProvision(ProjectSettings.BundleIdentifier, out MobileProvisionFile, out SigningCertificate))
 				{
 					MobileProvision = MobileProvisionFile!.FullName;
 				}
-				
+
 				if (MobileProvisionFile != null)
 				{
 					Logger.LogInformation("Provision found for {Project}, Provision: {Provision}, Certificate: {Certificate}", ((ProjectFile != null) ? ProjectFile.GetFileNameWithoutAnyExtensions() : "UnrealGame"), MobileProvisionFile, SigningCertificate);
 				}
 			}
 
-            // add to the dictionary
-            SigningCertificate = SigningCertificate.Replace("\"", "");
+			// add to the dictionary
+			SigningCertificate = SigningCertificate.Replace("\"", "");
 
-            // read the provision to get the UUID
-			if(MobileProvisionFile == null)
+			// read the provision to get the UUID
+			if (MobileProvisionFile == null)
 			{
 				Logger.LogInformation("No matching provision file was discovered for {ProjectFile}. Please ensure you have a compatible provision installed.", ProjectFile);
 			}
-			else if(!FileReference.Exists(MobileProvisionFile))
+			else if (!FileReference.Exists(MobileProvisionFile))
 			{
 				Logger.LogInformation("Selected mobile provision for {ProjectFile} ({MobileProvisionFile}) was not found. Please ensure you have a compatible provision installed.", ProjectFile, MobileProvisionFile);
 			}
 			else
-            {
+			{
 				byte[] AllBytes = FileReference.ReadAllBytes(MobileProvisionFile);
 
 				uint StartIndex = (uint)AllBytes.Length;
@@ -523,7 +521,7 @@ namespace UnrealBuildTool
 
 				for (uint i = 0; i + 4 < AllBytes.Length; i++)
 				{
-					if (AllBytes[i] == '<' && AllBytes[i+1] == '?' && AllBytes[i+ 2] == 'x' && AllBytes[i+ 3] == 'm' && AllBytes[i+ 4] == 'l')
+					if (AllBytes[i] == '<' && AllBytes[i + 1] == '?' && AllBytes[i + 2] == 'x' && AllBytes[i + 3] == 'm' && AllBytes[i + 4] == 'l')
 					{
 						StartIndex = i;
 						break;
@@ -534,9 +532,9 @@ namespace UnrealBuildTool
 				{
 					for (uint i = StartIndex; i + 7 < AllBytes.Length; i++)
 					{
-						if(AllBytes[i] == '<' && AllBytes[i + 1] == '/' && AllBytes[i + 2] == 'p' && AllBytes[i + 3] == 'l' && AllBytes[i + 4] == 'i' && AllBytes[i + 5] == 's' && AllBytes[i + 6] == 't' && AllBytes[i + 7] == '>')
+						if (AllBytes[i] == '<' && AllBytes[i + 1] == '/' && AllBytes[i + 2] == 'p' && AllBytes[i + 3] == 'l' && AllBytes[i + 4] == 'i' && AllBytes[i + 5] == 's' && AllBytes[i + 6] == 't' && AllBytes[i + 7] == '>')
 						{
-							EndIndex = i+7;
+							EndIndex = i + 7;
 							break;
 						}
 					}
@@ -580,15 +578,15 @@ namespace UnrealBuildTool
 						}
 					}
 					idx = AllText.IndexOf("<key>Name</key>");
-                    if (idx > 0)
-                    {
-                        idx = AllText.IndexOf("<string>", idx);
-                        if (idx > 0)
-                        {
-                            idx += "<string>".Length;
-                            MobileProvisionName = AllText.Substring(idx, AllText.IndexOf("</string>", idx) - idx);
-                        }
-                    }
+					if (idx > 0)
+					{
+						idx = AllText.IndexOf("<string>", idx);
+						if (idx > 0)
+						{
+							idx += "<string>".Length;
+							MobileProvisionName = AllText.Substring(idx, AllText.IndexOf("</string>", idx) - idx);
+						}
+					}
 				}
 
 				if (string.IsNullOrEmpty(MobileProvisionUUID) || string.IsNullOrEmpty(TeamUUID))
@@ -597,13 +595,13 @@ namespace UnrealBuildTool
 					SigningCertificate = null;
 					Logger.LogInformation("Failed to parse the mobile provisioning profile.");
 				}
-            }
+			}
 		}
 
-        void IPPDataReceivedHandler(Object Sender, DataReceivedEventArgs Line, ILogger Logger)
-        {
-            if ((Line != null) && (Line.Data != null))
-            {
+		void IPPDataReceivedHandler(Object Sender, DataReceivedEventArgs Line, ILogger Logger)
+		{
+			if ((Line != null) && (Line.Data != null))
+			{
 				if (Line.Data.StartsWith("IPP WARNING:"))
 				{
 					// Don't output IPP warnings to the console as they may not be warnings relevant to the build and could cause build failures.
@@ -615,35 +613,35 @@ namespace UnrealBuildTool
 				}
 
 				if (!string.IsNullOrEmpty(SigningCertificate))
-                {
-                    if (Line.Data.Contains("CERTIFICATE-") && Line.Data.Contains(SigningCertificate))
-                    {
-                        bHaveCertificate = true;
-                    }
-                }
-                else
-                {
-                    int cindex = Line.Data.IndexOf("CERTIFICATE-");
-                    int pindex = Line.Data.IndexOf("PROVISION-");
-                    if (cindex > -1 && pindex > -1)
-                    {
-                        cindex += "CERTIFICATE-".Length;
-                        SigningCertificate = Line.Data.Substring(cindex, pindex - cindex - 1);
-                        pindex += "PROVISION-".Length;
-						if(pindex < Line.Data.Length)
+				{
+					if (Line.Data.Contains("CERTIFICATE-") && Line.Data.Contains(SigningCertificate))
+					{
+						bHaveCertificate = true;
+					}
+				}
+				else
+				{
+					int cindex = Line.Data.IndexOf("CERTIFICATE-");
+					int pindex = Line.Data.IndexOf("PROVISION-");
+					if (cindex > -1 && pindex > -1)
+					{
+						cindex += "CERTIFICATE-".Length;
+						SigningCertificate = Line.Data.Substring(cindex, pindex - cindex - 1);
+						pindex += "PROVISION-".Length;
+						if (pindex < Line.Data.Length)
 						{
 							MobileProvisionFile = new FileReference(Line.Data.Substring(pindex));
 						}
-                    }
-                }
-            }
-        }
-    }
-	
+					}
+				}
+			}
+		}
+	}
+
 	class IOSArchitectureConfig : UnrealArchitectureConfig
 	{
 		public IOSArchitectureConfig()
-			: base(UnrealArchitectureMode.SingleTargetCompileSeparately,new[] { UnrealArch.Arm64, UnrealArch.IOSSimulator })
+			: base(UnrealArchitectureMode.SingleTargetCompileSeparately, new[] { UnrealArch.Arm64, UnrealArch.IOSSimulator })
 		{
 
 		}
@@ -675,7 +673,7 @@ namespace UnrealBuildTool
 		public override List<FileReference> FinalizeBinaryPaths(FileReference BinaryName, FileReference? ProjectFile, ReadOnlyTargetRules Target)
 		{
 			List<FileReference> BinaryPaths = new List<FileReference>();
-			if(Target.bShouldCompileAsDLL)
+			if (Target.bShouldCompileAsDLL)
 			{
 				BinaryPaths.Add(FileReference.Combine(BinaryName.Directory, Target.Configuration.ToString(), Target.Name + ".framework", Target.Name));
 			}
@@ -751,7 +749,7 @@ namespace UnrealBuildTool
 				}
 
 				Target.GlobalDefinitions.Add("BUILD_EMBEDDED_APP=1");
-				
+
 				if (Target.Platform == UnrealTargetPlatform.IOS)
 				{
 					Target.ExportPublicHeader = "Headers/PreIOSEmbeddedView.h";
@@ -891,12 +889,12 @@ namespace UnrealBuildTool
 		}
 
 		public IOSProvisioningData ReadProvisioningData(IOSProjectSettings ProjectSettings, bool bForDistribution = false)
-        {
+		{
 			string ProvisionKey = ProjectSettings.BundleIdentifier + " " + bForDistribution.ToString();
 
-            IOSProvisioningData? ProvisioningData;
- 
-			lock(ProvisionCache)
+			IOSProvisioningData? ProvisioningData;
+
+			lock (ProvisionCache)
 			{
 				if (!ProvisionCache.TryGetValue(ProvisionKey, out ProvisioningData))
 				{
@@ -905,7 +903,7 @@ namespace UnrealBuildTool
 				}
 			}
 			return ProvisioningData;
-        }
+		}
 
 		protected virtual IOSProvisioningData CreateProvisioningData(IOSProjectSettings ProjectSettings, bool bForDistribution)
 		{
@@ -929,7 +927,7 @@ namespace UnrealBuildTool
 				}
 			}
 
-            return new string [] {};
+			return new string[] { };
 		}
 
 		public override bool CanUseXGE()
@@ -945,7 +943,7 @@ namespace UnrealBuildTool
 		public bool HasCustomIcons(DirectoryReference ProjectDirectoryName, ILogger Logger)
 		{
 			string IconDir = Path.Combine(ProjectDirectoryName.FullName, "Build", "IOS", "Resources", "Graphics");
-			if(Directory.Exists(IconDir))
+			if (Directory.Exists(IconDir))
 			{
 				foreach (string f in Directory.EnumerateFiles(IconDir))
 				{
@@ -968,10 +966,10 @@ namespace UnrealBuildTool
 			string[] BoolKeys = new string[] {
 				"bShipForBitcode", "bGeneratedSYMFile",
 				"bGeneratedSYMBundle", "bEnableRemoteNotificationsSupport", "bEnableCloudKitSupport",
-                "bGenerateCrashReportSymbols", "bEnableBackgroundFetch"
-            };
+				"bGenerateCrashReportSymbols", "bEnableBackgroundFetch"
+			};
 			string[] StringKeys = new string[] {
-				"MinimumiOSVersion", 
+				"MinimumiOSVersion",
 				"AdditionalLinkerFlags",
 				"AdditionalShippingLinkerFlags"
 			};
@@ -1131,14 +1129,14 @@ namespace UnrealBuildTool
 			{
 				CompileEnvironment.Definitions.Add("NOTIFICATIONS_ENABLED=0");
 			}
-            if (ProjectSettings.bBackgroundFetchEnabled)
-            {
-                CompileEnvironment.Definitions.Add("BACKGROUNDFETCH_ENABLED=1");
-            }
-            else
-            {
-                CompileEnvironment.Definitions.Add("BACKGROUNDFETCH_ENABLED=0");
-            }
+			if (ProjectSettings.bBackgroundFetchEnabled)
+			{
+				CompileEnvironment.Definitions.Add("BACKGROUNDFETCH_ENABLED=1");
+			}
+			else
+			{
+				CompileEnvironment.Definitions.Add("BACKGROUNDFETCH_ENABLED=0");
+			}
 			if (ProjectSettings.bFileSharingEnabled)
 			{
 				CompileEnvironment.Definitions.Add("FILESHARING_ENABLED=1");
diff --git a/Engine/Source/Programs/UnrealBuildTool/Platform/IOS/UEDeployIOS.cs b/Engine/Source/Programs/UnrealBuildTool/Platform/IOS/UEDeployIOS.cs
--- a/Engine/Source/Programs/UnrealBuildTool/Platform/IOS/UEDeployIOS.cs
+++ b/Engine/Source/Programs/UnrealBuildTool/Platform/IOS/UEDeployIOS.cs
@@ -1,17 +1,14 @@
 // Copyright Epic Games, Inc. All Rights Reserved.
 
 using System;
-using System.Collections;
 using System.Collections.Generic;
+using System.IO;
 using System.Linq;
 using System.Text;
-using System.Xml;
-using System.IO;
-using System.Diagnostics;
 using System.Xml.Linq;
 using EpicGames.Core;
-using UnrealBuildBase;
 using Microsoft.Extensions.Logging;
+using UnrealBuildBase;
 
 namespace UnrealBuildTool
 {
@@ -169,8 +166,8 @@ namespace UnrealBuildTool
 				}
 				else
 				{
-                    // get the changelist
-                    CFBundleVersion = ReadOnlyBuildVersion.Current.Changelist.ToString();
+					// get the changelist
+					CFBundleVersion = ReadOnlyBuildVersion.Current.Changelist.ToString();
 
 				}
 
@@ -221,7 +218,7 @@ namespace UnrealBuildTool
 		{
 			ConfigHierarchy Ini = ConfigCache.ReadHierarchy(ConfigHierarchyType.Engine, ProjectLocation, UnrealTargetPlatform.IOS);
 			// required capabilities
-			List <string> RequiredCaps = new() { "arm64", "metal" };
+			List<string> RequiredCaps = new() { "arm64", "metal" };
 
 			// orientations
 			string InterfaceOrientation = "";
@@ -367,7 +364,7 @@ namespace UnrealBuildTool
 			Text.AppendLine("\t<key>UIRequiredDeviceCapabilities</key>");
 			Text.AppendLine("\t<array>");
 			foreach (string Cap in RequiredCaps)
-				{
+			{
 				Text.AppendLine($"\t\t<string>{Cap}</string>\n");
 			}
 			Text.AppendLine("\t</array>");
@@ -410,7 +407,7 @@ namespace UnrealBuildTool
 			if (!string.IsNullOrEmpty(ExtraData))
 			{
 				ExtraData = ExtraData.Replace("\\n", "\n");
-				foreach(string Line in ExtraData.Split("\r\n".ToCharArray()))
+				foreach (string Line in ExtraData.Split("\r\n".ToCharArray()))
 				{
 					if (!string.IsNullOrWhiteSpace(Line))
 					{
@@ -484,14 +481,14 @@ namespace UnrealBuildTool
 			// generate the Info.plist for future use
 			string BuildDirectory = ProjectDirectory + "/Build/IOS";
 			string IntermediateDirectory = (bIsUnrealGame ? InEngineDir : ProjectDirectory) + "/Intermediate/IOS";
-			string PListFile = IntermediateDirectory + "/" + GameName + "-Info.plist";;
+			string PListFile = IntermediateDirectory + "/" + GameName + "-Info.plist"; ;
 			ProjectName = !String.IsNullOrEmpty(ProjectName) ? ProjectName : GameName;
 			VersionUtilities.BuildDirectory = BuildDirectory;
 			VersionUtilities.GameName = GameName;
 
 			WritePlistFile(new FileReference(PListFile), DirRef, UPL, GameName, bIsUnrealGame, Logger);
 
-            if (BuildHostPlatform.Current.Platform == UnrealTargetPlatform.Mac && !bBuildAsFramework)
+			if (BuildHostPlatform.Current.Platform == UnrealTargetPlatform.Mac && !bBuildAsFramework)
 			{
 				FileReference FinalPlistFile;
 				FinalPlistFile = new FileReference($"{ProjectDirectory}/Build/IOS/UBTGenerated/Info.Template.plist");
@@ -904,18 +901,18 @@ namespace UnrealBuildTool
 		}
 
 		public static VersionNumber? GetSdkVersion(TargetReceipt Receipt)
-        {
+		{
 			VersionNumber? SdkVersion = null;
-            if (Receipt != null)
-            {
-                ReceiptProperty? SdkVersionProperty = Receipt.AdditionalProperties.FirstOrDefault(x => x.Name == "SDK");
-				if(SdkVersionProperty != null)
+			if (Receipt != null)
+			{
+				ReceiptProperty? SdkVersionProperty = Receipt.AdditionalProperties.FirstOrDefault(x => x.Name == "SDK");
+				if (SdkVersionProperty != null)
 				{
 					VersionNumber.TryParse(SdkVersionProperty.Value, out SdkVersion);
 				}
 			}
 			return SdkVersion;
-        }
+		}
 
 		public static bool GetCompileAsDll(TargetReceipt? Receipt)
 		{
@@ -930,7 +927,7 @@ namespace UnrealBuildTool
 			return false;
 		}
 
- 		public bool GeneratePList(FileReference ProjectFile, UnrealTargetConfiguration Config, string ProjectDirectory, bool bIsUnrealGame, string GameName, bool bIsClient, string ProjectName, string InEngineDir, string AppDirectory, TargetReceipt Receipt)
+		public bool GeneratePList(FileReference ProjectFile, UnrealTargetConfiguration Config, string ProjectDirectory, bool bIsUnrealGame, string GameName, bool bIsClient, string ProjectName, string InEngineDir, string AppDirectory, TargetReceipt Receipt)
 		{
 			List<string> UPLScripts = CollectPluginDataPaths(Receipt.AdditionalProperties, Logger);
 			VersionNumber? SdkVersion = GetSdkVersion(Receipt);
@@ -1037,7 +1034,7 @@ namespace UnrealBuildTool
 			try
 			{
 				FileInfo DestFileInfo;
-				if(!Directory.Exists(ProvisionDir))
+				if (!Directory.Exists(ProvisionDir))
 				{
 					throw new DirectoryNotFoundException(string.Format("Provision Directory {0} not found.", ProvisionDir), null);
 				}
@@ -1047,7 +1044,7 @@ namespace UnrealBuildTool
 				{
 					Logger.LogDebug("Local Provision Folder {LocalProvisionFolder} not found, attempting to create...", LocalProvisionFolder);
 					Directory.CreateDirectory(LocalProvisionFolder);
-					if(Directory.Exists(LocalProvisionFolder))
+					if (Directory.Exists(LocalProvisionFolder))
 					{
 						Logger.LogDebug("Local Provision Folder {LocalProvisionFolder} created successfully.", LocalProvisionFolder);
 					}
@@ -1057,13 +1054,13 @@ namespace UnrealBuildTool
 					}
 				}
 
-				foreach(string Provision in Directory.EnumerateFiles(ProvisionDir, "*.mobileprovision", SearchOption.AllDirectories))
+				foreach (string Provision in Directory.EnumerateFiles(ProvisionDir, "*.mobileprovision", SearchOption.AllDirectories))
 				{
 					string LocalProvisionFile = Path.Combine(LocalProvisionFolder, Path.GetFileName(Provision));
 					bool LocalFileExists = File.Exists(LocalProvisionFile);
-					if(!LocalFileExists || File.GetLastWriteTime(LocalProvisionFile) < File.GetLastWriteTime(Provision))
+					if (!LocalFileExists || File.GetLastWriteTime(LocalProvisionFile) < File.GetLastWriteTime(Provision))
 					{
-						if(LocalFileExists)
+						if (LocalFileExists)
 						{
 							DestFileInfo = new FileInfo(LocalProvisionFile);
 							DestFileInfo.Attributes = DestFileInfo.Attributes & ~FileAttributes.ReadOnly;
@@ -1200,7 +1197,7 @@ namespace UnrealBuildTool
 			// install the distribution provision
 			ProvisionWithPrefix = InEngineDir + "/Build/IOS/UnrealGame_Distro.mobileprovision";
 			string ProjectDistroProvision = InProjectName + "_Distro.mobileprovision";
-			if (File.Exists(Path.Combine(BuildDirectory, ProjectDistroProvision )))
+			if (File.Exists(Path.Combine(BuildDirectory, ProjectDistroProvision)))
 			{
 				ProvisionWithPrefix = Path.Combine(BuildDirectory, ProjectDistroProvision);
 			}
@@ -1312,7 +1309,7 @@ namespace UnrealBuildTool
 
 			// collect plugin extra data paths from target receipt
 			IEnumerable<ReceiptProperty> Results = ReceiptProperties.Where(x => x.Name == "IOSPlugin");
-			foreach(ReceiptProperty Property in Results)
+			foreach (ReceiptProperty Property in Results)
 			{
 				// Keep only unique paths
 				string PluginPath = Property.Value;
@@ -1350,7 +1347,7 @@ namespace UnrealBuildTool
 			if (SourceFolderInfo.Exists)
 			{
 				FileInfo[] SourceFiles = SourceFolderInfo.GetFiles(TargetFiles);
-				foreach(FileInfo SourceFile in SourceFiles)
+				foreach (FileInfo SourceFile in SourceFiles)
 				{
 					string DestinationPath = Path.Combine(DestinationDirectory, SourceFile.Name);
 					SafeFileCopy(SourceFile, DestinationPath, bOverwrite);
@@ -1366,7 +1363,7 @@ namespace UnrealBuildTool
 
 		static private void RecursiveFolderCopy(DirectoryInfo SourceFolderInfo, DirectoryInfo DestFolderInfo, bool bOverwrite = false, FilenameFilter? Filter = null)
 		{
-			foreach(FileInfo SourceFileInfo in SourceFolderInfo.GetFiles())
+			foreach (FileInfo SourceFileInfo in SourceFolderInfo.GetFiles())
 			{
 				string DestinationPath = Path.Combine(DestFolderInfo.FullName, SourceFileInfo.Name);
 				if (Filter != null && !Filter(DestinationPath))
@@ -1376,7 +1373,7 @@ namespace UnrealBuildTool
 				SafeFileCopy(SourceFileInfo, DestinationPath, bOverwrite);
 			}
 
-			foreach(DirectoryInfo SourceSubFolderInfo in SourceFolderInfo.GetDirectories())
+			foreach (DirectoryInfo SourceSubFolderInfo in SourceFolderInfo.GetDirectories())
 			{
 				string DestFolderName = Path.Combine(DestFolderInfo.FullName, SourceSubFolderInfo.Name);
 				Directory.CreateDirectory(DestFolderName);
diff --git a/Engine/Source/Programs/UnrealBuildTool/Platform/Linux/LinuxCommon.cs b/Engine/Source/Programs/UnrealBuildTool/Platform/Linux/LinuxCommon.cs
--- a/Engine/Source/Programs/UnrealBuildTool/Platform/Linux/LinuxCommon.cs
+++ b/Engine/Source/Programs/UnrealBuildTool/Platform/Linux/LinuxCommon.cs
@@ -1,13 +1,7 @@
 // Copyright Epic Games, Inc. All Rights Reserved.
 
-using System;
-using System.Collections.Generic;
-using System.Text.RegularExpressions;
 using System.Diagnostics;
 using System.IO;
-using System.Linq;
-using Microsoft.Win32;
-using EpicGames.Core;
 using Microsoft.Extensions.Logging;
 
 namespace UnrealBuildTool
diff --git a/Engine/Source/Programs/UnrealBuildTool/Platform/Linux/LinuxPlatformSDK.cs b/Engine/Source/Programs/UnrealBuildTool/Platform/Linux/LinuxPlatformSDK.cs
--- a/Engine/Source/Programs/UnrealBuildTool/Platform/Linux/LinuxPlatformSDK.cs
+++ b/Engine/Source/Programs/UnrealBuildTool/Platform/Linux/LinuxPlatformSDK.cs
@@ -1,12 +1,11 @@
 // Copyright Epic Games, Inc. All Rights Reserved.
 
 using System;
-using System.Collections.Generic;
 using System.IO;
-using EpicGames.Core;
 using System.Text.RegularExpressions;
-using UnrealBuildBase;
+using EpicGames.Core;
 using Microsoft.Extensions.Logging;
+using UnrealBuildBase;
 
 ///////////////////////////////////////////////////////////////////
 // If you are looking for supported version numbers, look in the
@@ -228,7 +227,7 @@ namespace UnrealBuildTool
 			DirectoryReference? MultiArchRoot = GetSDKLocation();
 			DirectoryReference? BaseLinuxPath;
 
-			if (MultiArchRoot != null )
+			if (MultiArchRoot != null)
 			{
 				BaseLinuxPath = DirectoryReference.Combine(MultiArchRoot, Architecture.LinuxName);
 			}
diff --git a/Engine/Source/Programs/UnrealBuildTool/Platform/Linux/LinuxProjectGenerator.cs b/Engine/Source/Programs/UnrealBuildTool/Platform/Linux/LinuxProjectGenerator.cs
--- a/Engine/Source/Programs/UnrealBuildTool/Platform/Linux/LinuxProjectGenerator.cs
+++ b/Engine/Source/Programs/UnrealBuildTool/Platform/Linux/LinuxProjectGenerator.cs
@@ -1,10 +1,6 @@
 // Copyright Epic Games, Inc. All Rights Reserved.
 
-using System;
 using System.Collections.Generic;
-using System.Text;
-using System.Diagnostics;
-using System.IO;
 using EpicGames.Core;
 using Microsoft.Extensions.Logging;
 
diff --git a/Engine/Source/Programs/UnrealBuildTool/Platform/Linux/LinuxToolChain.cs b/Engine/Source/Programs/UnrealBuildTool/Platform/Linux/LinuxToolChain.cs
--- a/Engine/Source/Programs/UnrealBuildTool/Platform/Linux/LinuxToolChain.cs
+++ b/Engine/Source/Programs/UnrealBuildTool/Platform/Linux/LinuxToolChain.cs
@@ -2,14 +2,12 @@
 
 using System;
 using System.Collections.Generic;
-using System.Text.RegularExpressions;
 using System.Diagnostics;
 using System.IO;
 using System.Linq;
-using Microsoft.Win32;
 using EpicGames.Core;
-using UnrealBuildBase;
 using Microsoft.Extensions.Logging;
+using UnrealBuildBase;
 
 namespace UnrealBuildTool
 {
diff --git a/Engine/Source/Programs/UnrealBuildTool/Platform/Linux/UEBuildLinux.cs b/Engine/Source/Programs/UnrealBuildTool/Platform/Linux/UEBuildLinux.cs
--- a/Engine/Source/Programs/UnrealBuildTool/Platform/Linux/UEBuildLinux.cs
+++ b/Engine/Source/Programs/UnrealBuildTool/Platform/Linux/UEBuildLinux.cs
@@ -2,13 +2,10 @@
 
 using System;
 using System.Collections.Generic;
-using System.Text;
-using System.Diagnostics;
 using System.IO;
 using EpicGames.Core;
-using System.Text.RegularExpressions;
-using UnrealBuildBase;
 using Microsoft.Extensions.Logging;
+using UnrealBuildBase;
 
 namespace UnrealBuildTool
 {
@@ -129,7 +126,7 @@ namespace UnrealBuildTool
 		/// Accessors for fields on the inner TargetRules instance
 		/// </summary>
 		#region Read-only accessor properties 
-		#pragma warning disable CS1591
+#pragma warning disable CS1591
 
 		public bool bPreservePSYM => Inner.bPreservePSYM;
 
@@ -149,7 +146,7 @@ namespace UnrealBuildTool
 
 		public bool bEnableRayTracing => Inner.bEnableRayTracing;
 
-		#pragma warning restore CS1591
+#pragma warning restore CS1591
 		#endregion
 	}
 
@@ -157,7 +154,7 @@ namespace UnrealBuildTool
 	class LinuxArchitectureConfig : UnrealArchitectureConfig
 	{
 		public LinuxArchitectureConfig(UnrealTargetPlatform Platform)
-			:  base(Platform == UnrealTargetPlatform.Linux ? UnrealArch.X64 : UnrealArch.Arm64)
+			: base(Platform == UnrealTargetPlatform.Linux ? UnrealArch.X64 : UnrealArch.Arm64)
 		{
 		}
 	}
@@ -178,7 +175,7 @@ namespace UnrealBuildTool
 		/// <summary>
 		/// Constructor
 		/// </summary>
-		public LinuxPlatform(LinuxPlatformSDK InSDK, ILogger Logger) 
+		public LinuxPlatform(LinuxPlatformSDK InSDK, ILogger Logger)
 			: this(UnrealTargetPlatform.Linux, InSDK, Logger)
 		{
 			SDK = InSDK;
@@ -262,7 +259,7 @@ namespace UnrealBuildTool
 			}
 
 			if (Target.GlobalDefinitions.Contains("USE_NULL_RHI=1"))
-			{				
+			{
 				Target.bCompileCEF3 = false;
 			}
 
@@ -350,14 +347,14 @@ namespace UnrealBuildTool
 				case UEBuildBinaryType.Executable:
 					if (InTarget.LinuxPlatform.bPreservePSYM)
 					{
-						return new string[] {".psym", ".sym", ".debug"};
+						return new string[] { ".psym", ".sym", ".debug" };
 					}
 					else
 					{
-						return new string[] {".sym", ".debug"};
+						return new string[] { ".sym", ".debug" };
 					}
 			}
-			return new string [] {};
+			return new string[] { };
 		}
 
 		/// <summary>
@@ -639,7 +636,7 @@ namespace UnrealBuildTool
 		}
 	}
 
-	class UEDeployLinux: UEBuildDeploy
+	class UEDeployLinux : UEBuildDeploy
 	{
 		public UEDeployLinux(ILogger InLogger)
 			: base(InLogger)
diff --git a/Engine/Source/Programs/UnrealBuildTool/Platform/Mac/ApplePlatformSDK.Versions.cs b/Engine/Source/Programs/UnrealBuildTool/Platform/Mac/ApplePlatformSDK.Versions.cs
--- a/Engine/Source/Programs/UnrealBuildTool/Platform/Mac/ApplePlatformSDK.Versions.cs
+++ b/Engine/Source/Programs/UnrealBuildTool/Platform/Mac/ApplePlatformSDK.Versions.cs
@@ -1,12 +1,7 @@
 // Copyright Epic Games, Inc. All Rights Reserved.
 
 using System;
-using System.Collections.Generic;
-using System.IO;
 using EpicGames.Core;
-using System.Text.RegularExpressions;
-using Microsoft.Win32;
-using System.Diagnostics;
 
 namespace UnrealBuildTool
 {
diff --git a/Engine/Source/Programs/UnrealBuildTool/Platform/Mac/ApplePlatformSDK.cs b/Engine/Source/Programs/UnrealBuildTool/Platform/Mac/ApplePlatformSDK.cs
--- a/Engine/Source/Programs/UnrealBuildTool/Platform/Mac/ApplePlatformSDK.cs
+++ b/Engine/Source/Programs/UnrealBuildTool/Platform/Mac/ApplePlatformSDK.cs
@@ -1,12 +1,7 @@
 // Copyright Epic Games, Inc. All Rights Reserved.
 
 using System;
-using System.Collections.Generic;
-using System.IO;
 using EpicGames.Core;
-using System.Text.RegularExpressions;
-using Microsoft.Win32;
-using System.Diagnostics;
 using Microsoft.Extensions.Logging;
 
 ///////////////////////////////////////////////////////////////////
@@ -41,9 +36,9 @@ namespace UnrealBuildTool
 			// iTunes is technically only need to deploy to and run on connected devices.
 			// This code removes requirement for Windows builders to have Xcode installed.
 			if (Status == SDKStatus.Invalid && !RuntimePlatform.IsMac && Environment.GetEnvironmentVariable("IsBuildMachine") == "1")
-            {
+			{
 				Status = SDKStatus.Valid;
-            }
+			}
 			return Status;
 		}
 	}
diff --git a/Engine/Source/Programs/UnrealBuildTool/Platform/Mac/MacExports.cs b/Engine/Source/Programs/UnrealBuildTool/Platform/Mac/MacExports.cs
--- a/Engine/Source/Programs/UnrealBuildTool/Platform/Mac/MacExports.cs
+++ b/Engine/Source/Programs/UnrealBuildTool/Platform/Mac/MacExports.cs
@@ -1,10 +1,5 @@
 ﻿// Copyright Epic Games, Inc. All Rights Reserved.
 
-using System;
-using System.Collections.Generic;
-using System.Linq;
-using System.Text;
-using System.Threading.Tasks;
 using EpicGames.Core;
 using Microsoft.Extensions.Logging;
 
@@ -78,6 +73,6 @@ namespace UnrealBuildTool
 		{
 			MacToolChain ToolChain = new MacToolChain(null, ClangToolChainOptions.None, Logger);
 			ToolChain.StripSymbols(SourceFile, TargetFile);
-		}		
+		}
 	}
 }
diff --git a/Engine/Source/Programs/UnrealBuildTool/Platform/Mac/MacPlatformSDK.Versions.cs b/Engine/Source/Programs/UnrealBuildTool/Platform/Mac/MacPlatformSDK.Versions.cs
--- a/Engine/Source/Programs/UnrealBuildTool/Platform/Mac/MacPlatformSDK.Versions.cs
+++ b/Engine/Source/Programs/UnrealBuildTool/Platform/Mac/MacPlatformSDK.Versions.cs
@@ -1,6 +1,5 @@
 // Copyright Epic Games, Inc. All Rights Reserved.
 
-using EpicGames.Core;
 using Microsoft.Extensions.Logging;
 
 namespace UnrealBuildTool
@@ -14,7 +13,7 @@ namespace UnrealBuildTool
 
 		protected override void GetValidSoftwareVersionRange(out string MinVersion, out string? MaxVersion)
 		{
-			MinVersion = "12.5.0";		// macOS Monterey 12.5 (Released 2022/07/22)
+			MinVersion = "12.5.0";      // macOS Monterey 12.5 (Released 2022/07/22)
 			MaxVersion = null;
 		}
 	}
diff --git a/Engine/Source/Programs/UnrealBuildTool/Platform/Mac/MacToolChain.cs b/Engine/Source/Programs/UnrealBuildTool/Platform/Mac/MacToolChain.cs
--- a/Engine/Source/Programs/UnrealBuildTool/Platform/Mac/MacToolChain.cs
+++ b/Engine/Source/Programs/UnrealBuildTool/Platform/Mac/MacToolChain.cs
@@ -1,19 +1,14 @@
 // Copyright Epic Games, Inc. All Rights Reserved.
 
 using System;
-using System.Collections;
 using System.Collections.Generic;
+using System.Globalization;
 using System.IO;
-using System.Diagnostics;
-using System.Security.AccessControl;
-using System.Text;
 using System.Linq;
-using Ionic.Zip;
-using EpicGames.Core;
-using System.Globalization;
 using System.Text.RegularExpressions;
-using UnrealBuildBase;
+using EpicGames.Core;
 using Microsoft.Extensions.Logging;
+using UnrealBuildBase;
 
 namespace UnrealBuildTool
 {
@@ -58,11 +53,11 @@ namespace UnrealBuildTool
 			SelectSDK(BaseSDKDir, "MacOSX", ref MacOSSDKVersion, bVerbose, Logger);
 
 			// convert to float for easy comparison
-			if(String.IsNullOrWhiteSpace(MacOSSDKVersion))
+			if (String.IsNullOrWhiteSpace(MacOSSDKVersion))
 			{
 				throw new BuildException("Unable to find installed MacOS SDK on remote agent.");
 			}
-			else if(!float.TryParse(MacOSSDKVersion, NumberStyles.AllowDecimalPoint, CultureInfo.InvariantCulture.NumberFormat, out MacOSSDKVersionFloat))
+			else if (!float.TryParse(MacOSSDKVersion, NumberStyles.AllowDecimalPoint, CultureInfo.InvariantCulture.NumberFormat, out MacOSSDKVersionFloat))
 			{
 				throw new BuildException("Unable to parse installed MacOS version (\"{0}\")", MacOSSDKVersion);
 			}
@@ -107,7 +102,7 @@ namespace UnrealBuildTool
 		}
 
 		public static DirectoryReference FindProductDirectory(FileReference? ProjectFile, DirectoryReference BinaryDir, string? NameIfProgram)
-		{ 
+		{
 			// a project file is always used if there is one
 			if (ProjectFile != null)
 			{
@@ -234,7 +229,7 @@ namespace UnrealBuildTool
 			// derived data builds.
 			// This is also added to the x64 builds in case someone ever adds -mfma.
 			Arguments.Add("-ffp-contract=off");
-			
+
 
 			if (CompileEnvironment.bEnableOSX109Support)
 			{
@@ -257,7 +252,7 @@ namespace UnrealBuildTool
 				}
 			}
 		}
-		
+
 		string AddFrameworkToLinkCommand(string FrameworkName, string Arg = "-framework")
 		{
 			string Result = "";
@@ -355,7 +350,7 @@ namespace UnrealBuildTool
 			BuildVersion? Version;
 			if (BuildVersion.TryRead(BuildVersion.GetDefaultFileName(), out Version))
 			{
-				return String.Format("{0}.{1}.{2}", Version.MajorVersion, Version.MinorVersion, bIgnorePatchVersion? 0 : Version.PatchVersion);
+				return String.Format("{0}.{1}.{2}", Version.MajorVersion, Version.MinorVersion, bIgnorePatchVersion ? 0 : Version.PatchVersion);
 			}
 			else
 			{
@@ -402,7 +397,7 @@ namespace UnrealBuildTool
 		private void AddLibraryPathToRPaths(string Library, string ExeAbsolutePath, ref List<string> RPaths, ref string LinkCommand, bool bIsBuildingAppBundle, ILogger Logger)
 		{
 			string LibraryFullPath = Path.GetFullPath(Library);
- 			string LibraryDir = Path.GetDirectoryName(LibraryFullPath)!;
+			string LibraryDir = Path.GetDirectoryName(LibraryFullPath)!;
 			string ExeDir = Path.GetDirectoryName(ExeAbsolutePath)!;
 
 			// Only dylibs and frameworks, and only those that are outside of Engine/Binaries/Mac and Engine/Source/ThirdParty, and outside of the folder where the executable is need an additional RPATH entry
@@ -753,7 +748,7 @@ namespace UnrealBuildTool
 			// Tell the action that we're building an import library here and it should conditionally be
 			// ignored as a prerequisite for other actions
 			LinkAction.bProducesImportLibrary = bBuildImportLibraryOnly || LinkEnvironment.bIsBuildingDLL;
-			
+
 			// Add the output file as a production of the link action.
 			FileItem OutputFile;
 			if (bBuildImportLibraryOnly)
@@ -981,33 +976,33 @@ namespace UnrealBuildTool
 
 			string ArgumentString = "-c \"";
 			ArgumentString += string.Format("for i in {{1..30}}; ");
-				ArgumentString += string.Format("do if [ -f \\\"{0}\\\" ] ; ", MachOBinary.AbsolutePath);
-				ArgumentString += string.Format("then ");
-					ArgumentString += string.Format("break; ");
-				ArgumentString += string.Format("else ");
-					ArgumentString += string.Format("sleep 1; ");
-				ArgumentString += string.Format("fi; ");
+			ArgumentString += string.Format("do if [ -f \\\"{0}\\\" ] ; ", MachOBinary.AbsolutePath);
+			ArgumentString += string.Format("then ");
+			ArgumentString += string.Format("break; ");
+			ArgumentString += string.Format("else ");
+			ArgumentString += string.Format("sleep 1; ");
+			ArgumentString += string.Format("fi; ");
 			ArgumentString += string.Format("done; ");
 
 			ArgumentString += string.Format("if [ ! -f \\\"{1}\\\" ] || [ \\\"{0}\\\" -nt \\\"{1}\\\" ] ; ", MachOBinary.AbsolutePath, OutputFile.AbsolutePath);
 			ArgumentString += string.Format("then ");
-				ArgumentString += string.Format("rm -rf \\\"{0}\\\"; ", OutputFile.AbsolutePath);
-				// use the new script for monolthic (ie large) targets
-				if (LinkEnvironment.LinkType == TargetLinkType.Monolithic)
-				{
-					ArgumentString += string.Format(" \\\"{0}\\\" \\\"{1}\\\" \\\"{2}\\\"; ",
-						UniversalDsymutilScriptPath,
-						MachOBinary.AbsolutePath,
-						OutputFile.AbsolutePath);
-				}
-				else
-				{
-					ArgumentString += string.Format(" \\\"{0}\\\" {3} -f \\\"{1}\\\" -o \\\"{2}\\\"; ",
-						DsymutilPath,
-						MachOBinary.AbsolutePath,
-						OutputFile.AbsolutePath,
-						ExtraOptions);
-				}
+			ArgumentString += string.Format("rm -rf \\\"{0}\\\"; ", OutputFile.AbsolutePath);
+			// use the new script for monolthic (ie large) targets
+			if (LinkEnvironment.LinkType == TargetLinkType.Monolithic)
+			{
+				ArgumentString += string.Format(" \\\"{0}\\\" \\\"{1}\\\" \\\"{2}\\\"; ",
+					UniversalDsymutilScriptPath,
+					MachOBinary.AbsolutePath,
+					OutputFile.AbsolutePath);
+			}
+			else
+			{
+				ArgumentString += string.Format(" \\\"{0}\\\" {3} -f \\\"{1}\\\" -o \\\"{2}\\\"; ",
+					DsymutilPath,
+					MachOBinary.AbsolutePath,
+					OutputFile.AbsolutePath,
+					ExtraOptions);
+			}
 			ArgumentString += string.Format("fi; ");
 			ArgumentString += "\"";
 
@@ -1086,7 +1081,7 @@ namespace UnrealBuildTool
 
 				foreach (KeyValuePair<FileReference, BuildProductType> BuildProductPair in BuildProductsArray)
 				{
-					string[] DebugExtensions = new string[] {};
+					string[] DebugExtensions = new string[] { };
 					switch (BuildProductPair.Value)
 					{
 						case BuildProductType.Executable:
@@ -1100,9 +1095,9 @@ namespace UnrealBuildTool
 					if (!string.IsNullOrEmpty(DSYMExtension))
 					{
 						string BinaryPath = BuildProductPair.Key.FullName;
-						if(BinaryPath.Contains(".app"))
+						if (BinaryPath.Contains(".app"))
 						{
-							while(BinaryPath.Contains(".app"))
+							while (BinaryPath.Contains(".app"))
 							{
 								BinaryPath = Path.GetDirectoryName(BinaryPath)!;
 							}
@@ -1112,11 +1107,11 @@ namespace UnrealBuildTool
 							BuildProducts[Ref] = BuildProductType.SymbolFile;
 						}
 					}
-					else if(BuildProductPair.Value == BuildProductType.SymbolFile && BuildProductPair.Key.FullName.Contains(".app"))
+					else if (BuildProductPair.Value == BuildProductType.SymbolFile && BuildProductPair.Key.FullName.Contains(".app"))
 					{
 						BuildProducts.Remove(BuildProductPair.Key);
 					}
-					if(BuildProductPair.Value == BuildProductType.DynamicLibrary && Target.bCreateMapFile)
+					if (BuildProductPair.Value == BuildProductType.DynamicLibrary && Target.bCreateMapFile)
 					{
 						BuildProducts.Add(new FileReference(BuildProductPair.Key.FullName + ".map"), BuildProductType.MapFile);
 					}
@@ -1137,7 +1132,7 @@ namespace UnrealBuildTool
 			// We need to know what third party dylibs would be copied to the bundle
 			if (Binary.Type != UEBuildBinaryType.StaticLibrary)
 			{
-			    foreach (UEBuildBundleResource Resource in BundleResources)
+				foreach (UEBuildBundleResource Resource in BundleResources)
 				{
 					if (Directory.Exists(Resource.ResourcePath))
 					{
@@ -1196,7 +1191,7 @@ namespace UnrealBuildTool
 			List<string> OutputFiles = Makefile.OutputItems.Select(Item => Path.ChangeExtension(Item.FullName, ".dSYM")).Distinct().ToList();
 			foreach (FileItem DebugItem in DebugInfoFiles)
 			{
-				if(OutputFiles.Any(Item => string.Equals(Item, DebugItem.FullName, StringComparison.InvariantCultureIgnoreCase)))
+				if (OutputFiles.Any(Item => string.Equals(Item, DebugItem.FullName, StringComparison.InvariantCultureIgnoreCase)))
 				{
 					Makefile.OutputItems.Add(DebugItem);
 				}
diff --git a/Engine/Source/Programs/UnrealBuildTool/Platform/Mac/UEBuildMac.cs b/Engine/Source/Programs/UnrealBuildTool/Platform/Mac/UEBuildMac.cs
--- a/Engine/Source/Programs/UnrealBuildTool/Platform/Mac/UEBuildMac.cs
+++ b/Engine/Source/Programs/UnrealBuildTool/Platform/Mac/UEBuildMac.cs
@@ -3,11 +3,6 @@
 
 using System;
 using System.Collections.Generic;
-using System.Configuration;
-using System.IO;
-using System.Linq;
-using System.Runtime.InteropServices;
-using System.Text.RegularExpressions;
 using EpicGames.Core;
 using Microsoft.Extensions.Logging;
 using UnrealBuildBase;
@@ -91,7 +86,7 @@ namespace UnrealBuildTool
 		/// Accessors for fields on the inner TargetRules instance
 		/// </summary>
 		#region Read-only accessor properties 
-		#pragma warning disable CS1591
+#pragma warning disable CS1591
 
 		public bool bEnableAddressSanitizer => Inner.bEnableAddressSanitizer;
 
@@ -253,7 +248,7 @@ namespace UnrealBuildTool
 
 	class MacPlatform : UEBuildPlatform
 	{
-		public MacPlatform(UEBuildPlatformSDK InSDK, ILogger InLogger) 
+		public MacPlatform(UEBuildPlatformSDK InSDK, ILogger InLogger)
 			: base(UnrealTargetPlatform.Mac, InSDK, new MacArchitectureConfig(), InLogger)
 		{
 		}
@@ -325,7 +320,7 @@ namespace UnrealBuildTool
 
 			// Force using the ANSI allocator if ASan is enabled
 			string? AddressSanitizer = Environment.GetEnvironmentVariable("ENABLE_ADDRESS_SANITIZER");
-			if(Target.MacPlatform.bEnableAddressSanitizer || (AddressSanitizer != null && AddressSanitizer == "YES"))
+			if (Target.MacPlatform.bEnableAddressSanitizer || (AddressSanitizer != null && AddressSanitizer == "YES"))
 			{
 				Target.GlobalDefinitions.Add("FORCE_ANSI_ALLOCATOR=1");
 			}
@@ -412,10 +407,10 @@ namespace UnrealBuildTool
 			{
 				case UEBuildBinaryType.DynamicLinkLibrary:
 				case UEBuildBinaryType.Executable:
-					return Target.bUsePDBFiles ? new string[] {".dSYM"} : new string[] {};
+					return Target.bUsePDBFiles ? new string[] { ".dSYM" } : new string[] { };
 				case UEBuildBinaryType.StaticLibrary:
 				default:
-					return new string [] {};
+					return new string[] { };
 			}
 		}
 
@@ -561,19 +556,19 @@ namespace UnrealBuildTool
 			string? ThreadSanitizer = Environment.GetEnvironmentVariable("ENABLE_THREAD_SANITIZER");
 			string? UndefSanitizerMode = Environment.GetEnvironmentVariable("ENABLE_UNDEFINED_BEHAVIOR_SANITIZER");
 
-			if(Target.MacPlatform.bEnableAddressSanitizer || (AddressSanitizer != null && AddressSanitizer == "YES"))
+			if (Target.MacPlatform.bEnableAddressSanitizer || (AddressSanitizer != null && AddressSanitizer == "YES"))
 			{
 				Options |= ClangToolChainOptions.EnableAddressSanitizer;
 			}
-			if(Target.MacPlatform.bEnableThreadSanitizer || (ThreadSanitizer != null && ThreadSanitizer == "YES"))
+			if (Target.MacPlatform.bEnableThreadSanitizer || (ThreadSanitizer != null && ThreadSanitizer == "YES"))
 			{
 				Options |= ClangToolChainOptions.EnableThreadSanitizer;
 			}
-			if(Target.MacPlatform.bEnableUndefinedBehaviorSanitizer || (UndefSanitizerMode != null && UndefSanitizerMode == "YES"))
+			if (Target.MacPlatform.bEnableUndefinedBehaviorSanitizer || (UndefSanitizerMode != null && UndefSanitizerMode == "YES"))
 			{
 				Options |= ClangToolChainOptions.EnableUndefinedBehaviorSanitizer;
 			}
-			if(Target.bShouldCompileAsDLL)
+			if (Target.bShouldCompileAsDLL)
 			{
 				Options |= ClangToolChainOptions.OutputDylib;
 			}
diff --git a/Engine/Source/Programs/UnrealBuildTool/Platform/Mac/UEDeployMac.cs b/Engine/Source/Programs/UnrealBuildTool/Platform/Mac/UEDeployMac.cs
--- a/Engine/Source/Programs/UnrealBuildTool/Platform/Mac/UEDeployMac.cs
+++ b/Engine/Source/Programs/UnrealBuildTool/Platform/Mac/UEDeployMac.cs
@@ -1,13 +1,6 @@
 // Copyright Epic Games, Inc. All Rights Reserved.
 
-using System;
-using System.Collections;
-using System.Collections.Generic;
-using System.Linq;
-using System.Text;
-using System.Xml;
 using System.IO;
-using Ionic.Zip;
 using EpicGames.Core;
 using Microsoft.Extensions.Logging;
 
@@ -15,7 +8,7 @@ namespace UnrealBuildTool
 {
 	class UEDeployMac : UEBuildDeploy
 	{
-		public UEDeployMac(ILogger InLogger) 
+		public UEDeployMac(ILogger InLogger)
 			: base(InLogger)
 		{
 		}
@@ -46,15 +39,15 @@ namespace UnrealBuildTool
 				return false;
 			}
 
-            // bundle identifier
-            // plist replacements
-            DirectoryReference? DirRef = bIsUnrealGame ? (!string.IsNullOrEmpty(UnrealBuildTool.GetRemoteIniPath()) ? new DirectoryReference(UnrealBuildTool.GetRemoteIniPath()!) : null) : new DirectoryReference(ProjectDirectory);
-            ConfigHierarchy Ini = ConfigCache.ReadHierarchy(ConfigHierarchyType.Engine, DirRef, UnrealTargetPlatform.IOS);
+			// bundle identifier
+			// plist replacements
+			DirectoryReference? DirRef = bIsUnrealGame ? (!string.IsNullOrEmpty(UnrealBuildTool.GetRemoteIniPath()) ? new DirectoryReference(UnrealBuildTool.GetRemoteIniPath()!) : null) : new DirectoryReference(ProjectDirectory);
+			ConfigHierarchy Ini = ConfigCache.ReadHierarchy(ConfigHierarchyType.Engine, DirRef, UnrealTargetPlatform.IOS);
 
-            string BundleIdentifier;
-            Ini.GetString("/Script/IOSRuntimeSettings.IOSRuntimeSettings", "BundleIdentifier", out BundleIdentifier);
+			string BundleIdentifier;
+			Ini.GetString("/Script/IOSRuntimeSettings.IOSRuntimeSettings", "BundleIdentifier", out BundleIdentifier);
 
-            string BundleVersion = MacToolChain.LoadEngineDisplayVersion();
+			string BundleVersion = MacToolChain.LoadEngineDisplayVersion();
 			// duplicating some logic in MacToolchain for the BundleID
 			string[] ExeNameParts = ExeName.Split('-');
 			bool bBuildingEditor = ExeNameParts[0].EndsWith("Editor");
diff --git a/Engine/Source/Programs/UnrealBuildTool/Platform/TVOS/TVOSProjectGenerator.cs b/Engine/Source/Programs/UnrealBuildTool/Platform/TVOS/TVOSProjectGenerator.cs
--- a/Engine/Source/Programs/UnrealBuildTool/Platform/TVOS/TVOSProjectGenerator.cs
+++ b/Engine/Source/Programs/UnrealBuildTool/Platform/TVOS/TVOSProjectGenerator.cs
@@ -1,17 +1,13 @@
 // Copyright Epic Games, Inc. All Rights Reserved.
 
-using System;
 using System.Collections.Generic;
-using System.Text;
-using System.Diagnostics;
-using System.IO;
 using EpicGames.Core;
 using Microsoft.Extensions.Logging;
 
 namespace UnrealBuildTool
 {
 	class TVOSProjectGenerator : IOSProjectGenerator
-    {
+	{
 		/// <summary>
 		/// Constructor
 		/// </summary>
@@ -29,5 +25,5 @@ namespace UnrealBuildTool
 		{
 			yield return UnrealTargetPlatform.TVOS;
 		}
-    }
+	}
 }
diff --git a/Engine/Source/Programs/UnrealBuildTool/Platform/TVOS/TVOSToolChain.cs b/Engine/Source/Programs/UnrealBuildTool/Platform/TVOS/TVOSToolChain.cs
--- a/Engine/Source/Programs/UnrealBuildTool/Platform/TVOS/TVOSToolChain.cs
+++ b/Engine/Source/Programs/UnrealBuildTool/Platform/TVOS/TVOSToolChain.cs
@@ -1,18 +1,6 @@
 // Copyright Epic Games, Inc. All Rights Reserved.
 
 
-using System;
-using System.Collections;
-using System.Collections.Generic;
-using System.Linq;
-using System.IO;
-using System.Diagnostics;
-using System.Security.AccessControl;
-using System.Xml;
-using System.Text;
-using Ionic.Zip;
-using Ionic.Zlib;
-using EpicGames.Core;
 using Microsoft.Extensions.Logging;
 
 namespace UnrealBuildTool
diff --git a/Engine/Source/Programs/UnrealBuildTool/Platform/TVOS/UEBuildTVOS.cs b/Engine/Source/Programs/UnrealBuildTool/Platform/TVOS/UEBuildTVOS.cs
--- a/Engine/Source/Programs/UnrealBuildTool/Platform/TVOS/UEBuildTVOS.cs
+++ b/Engine/Source/Programs/UnrealBuildTool/Platform/TVOS/UEBuildTVOS.cs
@@ -1,7 +1,6 @@
 // Copyright Epic Games, Inc. All Rights Reserved.
 
 using System;
-using System.Collections.Generic;
 using EpicGames.Core;
 using Microsoft.Extensions.Logging;
 
@@ -29,7 +28,7 @@ namespace UnrealBuildTool
 	}
 
 	class TVOSPlatform : IOSPlatform
-    {
+	{
 		public TVOSPlatform(UEBuildPlatformSDK InSDK, ILogger Logger)
 			: base(InSDK, UnrealTargetPlatform.TVOS, Logger)
 		{
@@ -59,7 +58,7 @@ namespace UnrealBuildTool
 		}
 
 		public TVOSProvisioningData ReadProvisioningData(TVOSProjectSettings ProjectSettings, bool bForDistribution = false)
-        {
+		{
 			return (TVOSProvisioningData)base.ReadProvisioningData(ProjectSettings, bForDistribution);
 		}
 
@@ -132,7 +131,7 @@ namespace UnrealBuildTool
 		public override void RegisterBuildPlatforms(ILogger Logger)
 		{
 			ApplePlatformSDK SDK = new IOSPlatformSDK(Logger);
-		
+
 			// Register this build platform for IOS
 			UEBuildPlatform.RegisterBuildPlatform(new TVOSPlatform(SDK, Logger), Logger);
 			UEBuildPlatform.RegisterPlatformWithGroup(UnrealTargetPlatform.TVOS, UnrealPlatformGroup.Apple);
diff --git a/Engine/Source/Programs/UnrealBuildTool/Platform/TVOS/UEDeployTVOS.cs b/Engine/Source/Programs/UnrealBuildTool/Platform/TVOS/UEDeployTVOS.cs
--- a/Engine/Source/Programs/UnrealBuildTool/Platform/TVOS/UEDeployTVOS.cs
+++ b/Engine/Source/Programs/UnrealBuildTool/Platform/TVOS/UEDeployTVOS.cs
@@ -2,14 +2,10 @@
 
 using System;
 using System.Collections.Generic;
-using System.Linq;
-using System.Text;
-using System.Xml;
 using System.IO;
-using System.Diagnostics;
+using System.Text;
 using System.Xml.Linq;
 using EpicGames.Core;
-using UnrealBuildBase;
 using Microsoft.Extensions.Logging;
 
 namespace UnrealBuildTool
@@ -17,10 +13,10 @@ namespace UnrealBuildTool
 	class UEDeployTVOS : UEDeployIOS
 	{
 
-        public UEDeployTVOS(ILogger InLogger)
+		public UEDeployTVOS(ILogger InLogger)
 			: base(InLogger)
-        {
-        }
+		{
+		}
 
 		protected override string GetTargetPlatformName()
 		{
@@ -82,16 +78,16 @@ namespace UnrealBuildTool
 			string ExtraData = "";
 			Ini.GetString("/Script/IOSRuntimeSettings.IOSRuntimeSettings", "AdditionalPlistData", out ExtraData);
 
-            // create the final display name, including converting all entities for XML use
-            string FinalDisplayName = BundleDisplayName.Replace("[PROJECT_NAME]", ProjectName).Replace("_", "");
-            FinalDisplayName = FinalDisplayName.Replace("&", "&amp;");
-            FinalDisplayName = FinalDisplayName.Replace("\"", "&quot;");
-            FinalDisplayName = FinalDisplayName.Replace("\'", "&apos;");
-            FinalDisplayName = FinalDisplayName.Replace("<", "&lt;");
-            FinalDisplayName = FinalDisplayName.Replace(">", "&gt;");
+			// create the final display name, including converting all entities for XML use
+			string FinalDisplayName = BundleDisplayName.Replace("[PROJECT_NAME]", ProjectName).Replace("_", "");
+			FinalDisplayName = FinalDisplayName.Replace("&", "&amp;");
+			FinalDisplayName = FinalDisplayName.Replace("\"", "&quot;");
+			FinalDisplayName = FinalDisplayName.Replace("\'", "&apos;");
+			FinalDisplayName = FinalDisplayName.Replace("<", "&lt;");
+			FinalDisplayName = FinalDisplayName.Replace(">", "&gt;");
 
-            // generate the plist file
-            StringBuilder Text = new StringBuilder();
+			// generate the plist file
+			StringBuilder Text = new StringBuilder();
 			Text.AppendLine("<?xml version=\"1.0\" encoding=\"UTF-8\"?>");
 			Text.AppendLine("<!DOCTYPE plist PUBLIC \"-//Apple//DTD PLIST 1.0//EN\" \"http://www.apple.com/DTDs/PropertyList-1.0.dtd\">");
 			Text.AppendLine("<plist version=\"1.0\">");
@@ -106,7 +102,7 @@ namespace UnrealBuildTool
 				(bIsClient ? GameName + "Client" : GameName);
 			Text.AppendLine(string.Format("\t<string>{0}</string>", BundleExecutable));
 			Text.AppendLine("\t<key>CFBundleIdentifier</key>");
-			Text.AppendLine(string.Format("\t<string>{0}</string>", BundleIdentifier.Replace("[PROJECT_NAME]", ProjectName).Replace("_","")));
+			Text.AppendLine(string.Format("\t<string>{0}</string>", BundleIdentifier.Replace("[PROJECT_NAME]", ProjectName).Replace("_", "")));
 			Text.AppendLine("\t<key>CFBundleInfoDictionaryVersion</key>");
 			Text.AppendLine("\t<string>6.0</string>");
 			Text.AppendLine("\t<key>CFBundleName</key>");
@@ -130,18 +126,18 @@ namespace UnrealBuildTool
 					Text.AppendLine(Line);
 				}
 			}
-            Text.AppendLine("\t</array>");
+			Text.AppendLine("\t</array>");
 
-            Text.AppendLine("\t<key>TVTopShelfImage</key>");
-            Text.AppendLine("\t<dict>");
+			Text.AppendLine("\t<key>TVTopShelfImage</key>");
+			Text.AppendLine("\t<dict>");
 			Text.AppendLine("\t\t<key>TVTopShelfPrimaryImageWide</key>");
 			Text.AppendLine("\t\t<string>Top Shelf Image Wide</string>");
 			Text.AppendLine("\t</dict>");
 			Text.AppendLine("\t<key>CFBundleIcons</key>");
-            Text.AppendLine("\t<dict>");
-            Text.AppendLine("\t\t<key>CFBundlePrimaryIcon</key>");
-            Text.AppendLine("\t\t<string>App Icon</string>");
-            Text.AppendLine("\t</dict>");
+			Text.AppendLine("\t<dict>");
+			Text.AppendLine("\t\t<key>CFBundlePrimaryIcon</key>");
+			Text.AppendLine("\t\t<string>App Icon</string>");
+			Text.AppendLine("\t</dict>");
 			Text.AppendLine("\t<key>UILaunchStoryboardName</key>");
 			Text.AppendLine("\t<string>LaunchScreen</string>");
 
@@ -159,7 +155,7 @@ namespace UnrealBuildTool
 				}
 			}
 
-            Text.AppendLine("</dict>");
+			Text.AppendLine("</dict>");
 			Text.AppendLine("</plist>");
 
 			// Create the intermediate directory if needed
@@ -168,7 +164,7 @@ namespace UnrealBuildTool
 				Directory.CreateDirectory(IntermediateDirectory);
 			}
 
-			if(UPL != null)
+			if (UPL != null)
 			{
 				// Allow UPL to modify the plist here
 				XDocument XDoc;
@@ -183,8 +179,8 @@ namespace UnrealBuildTool
 
 				XDoc.DocumentType!.InternalSubset = "";
 				UPL.ProcessPluginNode("None", "iosPListUpdates", "", ref XDoc);
-                string result = XDoc.Declaration?.ToString() + "\n" + XDoc.ToString().Replace("<!DOCTYPE plist PUBLIC \"-//Apple//DTD PLIST 1.0//EN\" \"http://www.apple.com/DTDs/PropertyList-1.0.dtd\"[]>", "<!DOCTYPE plist PUBLIC \"-//Apple//DTD PLIST 1.0//EN\" \"http://www.apple.com/DTDs/PropertyList-1.0.dtd\">");
-                File.WriteAllText(PListFile, result);
+				string result = XDoc.Declaration?.ToString() + "\n" + XDoc.ToString().Replace("<!DOCTYPE plist PUBLIC \"-//Apple//DTD PLIST 1.0//EN\" \"http://www.apple.com/DTDs/PropertyList-1.0.dtd\"[]>", "<!DOCTYPE plist PUBLIC \"-//Apple//DTD PLIST 1.0//EN\" \"http://www.apple.com/DTDs/PropertyList-1.0.dtd\">");
+				File.WriteAllText(PListFile, result);
 			}
 			else
 			{
diff --git a/Engine/Source/Programs/UnrealBuildTool/Platform/Windows/AppXManifestGeneratorBase.cs b/Engine/Source/Programs/UnrealBuildTool/Platform/Windows/AppXManifestGeneratorBase.cs
--- a/Engine/Source/Programs/UnrealBuildTool/Platform/Windows/AppXManifestGeneratorBase.cs
+++ b/Engine/Source/Programs/UnrealBuildTool/Platform/Windows/AppXManifestGeneratorBase.cs
@@ -2,17 +2,13 @@
 
 using System;
 using System.Collections.Generic;
+using System.Diagnostics.CodeAnalysis;
 using System.Linq;
-using System.IO;
 using System.Text.RegularExpressions;
-using System.Xml;
-using EpicGames.Core;
 using System.Xml.Linq;
-using System.Text;
-using System.Diagnostics;
-using UnrealBuildBase;
-using System.Diagnostics.CodeAnalysis;
+using EpicGames.Core;
 using Microsoft.Extensions.Logging;
+using UnrealBuildBase;
 
 namespace UnrealBuildTool
 {
@@ -33,8 +29,8 @@ namespace UnrealBuildTool
 		/// default subdirectory for engine resources
 		protected const string EngineResourceSubPath = "DefaultImages";
 
-        /// Manifest compliance values
-        protected const int MaxResourceEntries = 200;
+		/// Manifest compliance values
+		protected const int MaxResourceEntries = 200;
 
 		/// cached engine ini
 		protected ConfigHierarchy? EngineIni;
@@ -49,7 +45,7 @@ namespace UnrealBuildTool
 		protected string? DefaultAppXCultureId;
 
 		/// lookup table for UE's CultureToStage StageId to AppX CultureId
-		protected Dictionary<string,string> UEStageIdToAppXCultureId = new Dictionary<string, string>();
+		protected Dictionary<string, string> UEStageIdToAppXCultureId = new Dictionary<string, string>();
 
 		/// the platform to generate the manifest for
 		protected UnrealTargetPlatform Platform;
@@ -66,7 +62,7 @@ namespace UnrealBuildTool
 		/// <summary>
 		/// Create a manifest generator for the given platform variant.
 		/// </summary>
-		public AppXManifestGeneratorBase( UnrealTargetPlatform InPlatform, ILogger InLogger )
+		public AppXManifestGeneratorBase(UnrealTargetPlatform InPlatform, ILogger InLogger)
 		{
 			Platform = InPlatform;
 			Logger = InLogger;
@@ -148,7 +144,7 @@ namespace UnrealBuildTool
 		/// Reads an integer from the cached ini files
 		/// </summary>
 		[return: NotNullIfNotNull("DefaultValue")]
-        protected string? ReadIniString(string? Key, string Section, string? DefaultValue = null)
+		protected string? ReadIniString(string? Key, string Section, string? DefaultValue = null)
 		{
 			if (Key == null)
 				return DefaultValue;
@@ -167,7 +163,7 @@ namespace UnrealBuildTool
 		/// Reads a string from the cached ini files
 		/// </summary>
 		[return: NotNullIfNotNull("DefaultValue")]
-        protected string? GetConfigString(string PlatformKey, string? GenericKey, string? DefaultValue = null)
+		protected string? GetConfigString(string PlatformKey, string? GenericKey, string? DefaultValue = null)
 		{
 			string? GenericValue = ReadIniString(GenericKey, IniSection_GeneralProjectSettings, DefaultValue);
 			return ReadIniString(PlatformKey, IniSection_PlatformTargetSettings, GenericValue);
@@ -225,7 +221,7 @@ namespace UnrealBuildTool
 			AppXResources!.ClearStrings();
 
 			// add all default strings
-			if (EngineIni!.GetString(IniSection_PlatformTargetSettings, "CultureStringResources", out string DefaultCultureScratchValue) && ConfigHierarchy.TryParse(DefaultCultureScratchValue, out Dictionary<string,string>? DefaultStrings))
+			if (EngineIni!.GetString(IniSection_PlatformTargetSettings, "CultureStringResources", out string DefaultCultureScratchValue) && ConfigHierarchy.TryParse(DefaultCultureScratchValue, out Dictionary<string, string>? DefaultStrings))
 			{
 				AppXResources!.AddDefaultStrings(DefaultStrings);
 			}
@@ -286,12 +282,12 @@ namespace UnrealBuildTool
 		{
 			if (ProjectFile != null)
 			{
-				AppXResources!.ProjectBinaryResourceDirectories.Add( DirectoryReference.Combine(ProjectFile.Directory, "Build", Platform.ToString(), BuildResourceSubPath) );
-				AppXResources!.ProjectBinaryResourceDirectories.Add( DirectoryReference.Combine(ProjectFile.Directory, "Platforms", Platform.ToString(), "Build", BuildResourceSubPath) );
+				AppXResources!.ProjectBinaryResourceDirectories.Add(DirectoryReference.Combine(ProjectFile.Directory, "Build", Platform.ToString(), BuildResourceSubPath));
+				AppXResources!.ProjectBinaryResourceDirectories.Add(DirectoryReference.Combine(ProjectFile.Directory, "Platforms", Platform.ToString(), "Build", BuildResourceSubPath));
 			}
 
-			AppXResources!.EngineFallbackBinaryResourceDirectories.Add( DirectoryReference.Combine(Unreal.EngineDirectory, "Build", Platform.ToString(), EngineResourceSubPath) );
-			AppXResources!.EngineFallbackBinaryResourceDirectories.Add( DirectoryReference.Combine(Unreal.EngineDirectory, "Platforms", Platform.ToString(), "Build", EngineResourceSubPath) );
+			AppXResources!.EngineFallbackBinaryResourceDirectories.Add(DirectoryReference.Combine(Unreal.EngineDirectory, "Build", Platform.ToString(), EngineResourceSubPath));
+			AppXResources!.EngineFallbackBinaryResourceDirectories.Add(DirectoryReference.Combine(Unreal.EngineDirectory, "Platforms", Platform.ToString(), "Build", EngineResourceSubPath));
 		}
 
 
@@ -324,14 +320,14 @@ namespace UnrealBuildTool
 		/// </summary>
 		protected virtual string GetIdentityPackageName()
 		{
-            // Read the PackageName from config
+			// Read the PackageName from config
 			var DefaultName = (ProjectFile != null) ? ProjectFile.GetFileNameWithoutAnyExtensions() : (TargetName ?? "DefaultUEProject");
-            var PackageName = Regex.Replace(GetConfigString("PackageName", "ProjectName", DefaultName), "[^-.A-Za-z0-9]", "");
-            if (string.IsNullOrWhiteSpace(PackageName))
-            {
-                Logger.LogError("Invalid package name {Name}. Package names must only contain letters, numbers, dash, and period and must be at least one character long.", PackageName);
-                Logger.LogError("Consider using the setting [{IniSection}]:PackageName to provide a specific value.", IniSection_PlatformTargetSettings);
-            }
+			var PackageName = Regex.Replace(GetConfigString("PackageName", "ProjectName", DefaultName), "[^-.A-Za-z0-9]", "");
+			if (string.IsNullOrWhiteSpace(PackageName))
+			{
+				Logger.LogError("Invalid package name {Name}. Package names must only contain letters, numbers, dash, and period and must be at least one character long.", PackageName);
+				Logger.LogError("Consider using the setting [{IniSection}]:PackageName to provide a specific value.", IniSection_PlatformTargetSettings);
+			}
 
 			// If specified in the project settings append the users machine name onto the package name to allow sharing of devkits without stomping of deploys
 			bool bPackageNameUseMachineName;
@@ -349,7 +345,7 @@ namespace UnrealBuildTool
 		/// </summary>
 		protected virtual string GetIdentityPublisherName()
 		{
-            var PublisherName = GetConfigString("PublisherName", "CompanyDistinguishedName", "CN=NoPublisher");
+			var PublisherName = GetConfigString("PublisherName", "CompanyDistinguishedName", "CN=NoPublisher");
 			return PublisherName;
 		}
 
@@ -358,15 +354,15 @@ namespace UnrealBuildTool
 		/// </summary>
 		protected virtual string? GetIdentityVersionNumber()
 		{
-            var VersionNumber = GetConfigString("PackageVersion", "ProjectVersion", "1.0.0.0");
-            VersionNumber = ValidatePackageVersion(VersionNumber);
+			var VersionNumber = GetConfigString("PackageVersion", "ProjectVersion", "1.0.0.0");
+			VersionNumber = ValidatePackageVersion(VersionNumber);
 
-            // If specified in the project settings attempt to retrieve the current build number and increment the version number by that amount, accounting for overflows
-            bool bIncludeEngineVersionInPackageVersion;
-            if (EngineIni!.GetBool(IniSection_PlatformTargetSettings, "bIncludeEngineVersionInPackageVersion", out bIncludeEngineVersionInPackageVersion) && bIncludeEngineVersionInPackageVersion)
-            {
+			// If specified in the project settings attempt to retrieve the current build number and increment the version number by that amount, accounting for overflows
+			bool bIncludeEngineVersionInPackageVersion;
+			if (EngineIni!.GetBool(IniSection_PlatformTargetSettings, "bIncludeEngineVersionInPackageVersion", out bIncludeEngineVersionInPackageVersion) && bIncludeEngineVersionInPackageVersion)
+			{
 				VersionNumber = IncludeBuildVersionInPackageVersion(VersionNumber);
-            }
+			}
 
 			return VersionNumber;
 		}
@@ -375,18 +371,18 @@ namespace UnrealBuildTool
 		/// Get the package identity element
 		/// </summary>
 		protected XElement GetIdentity(out string IdentityName)
-        {
-            string PackageName = GetIdentityPackageName();
-            string PublisherName = GetIdentityPublisherName();
-            string? VersionNumber = GetIdentityVersionNumber();
+		{
+			string PackageName = GetIdentityPackageName();
+			string PublisherName = GetIdentityPublisherName();
+			string? VersionNumber = GetIdentityVersionNumber();
 
-            IdentityName = PackageName;
+			IdentityName = PackageName;
 
-            return new XElement("Identity",
-                new XAttribute("Name", PackageName),
-                new XAttribute("Publisher", PublisherName),
-                new XAttribute("Version", VersionNumber!));
-        }
+			return new XElement("Identity",
+				new XAttribute("Name", PackageName),
+				new XAttribute("Publisher", PublisherName),
+				new XAttribute("Version", VersionNumber!));
+		}
 
 		/// <summary>
 		/// Updates the given package version to include the engine build version, if requested
@@ -424,13 +420,13 @@ namespace UnrealBuildTool
 		/// <summary>
 		/// Return the entire manifest element
 		/// </summary>
-		protected abstract XElement GetManifest(Dictionary<UnrealTargetConfiguration,string> InExecutablePairs, out string IdentityName);
+		protected abstract XElement GetManifest(Dictionary<UnrealTargetConfiguration, string> InExecutablePairs, out string IdentityName);
 
 		/// <summary>
 		/// Perform any platform-specific processing on the manifest before it is saved
 		/// </summary>
-		protected virtual void ProcessManifest(Dictionary<UnrealTargetConfiguration,string> InExecutablePairs, string ManifestName, string ManifestTargetPath, string ManifestIntermediatePath)
-        {
+		protected virtual void ProcessManifest(Dictionary<UnrealTargetConfiguration, string> InExecutablePairs, string ManifestName, string ManifestTargetPath, string ManifestIntermediatePath)
+		{
 		}
 
 		/// <summary>
@@ -443,7 +439,7 @@ namespace UnrealBuildTool
 		/// <summary>
 		/// Create a manifest and return the list of modified files
 		/// </summary>
-		public List<string>? CreateManifest(string InManifestName, DirectoryReference OutputDirectory, string? InTargetName, FileReference? InProjectFile, Dictionary<UnrealTargetConfiguration,string> InExecutablePairs)
+		public List<string>? CreateManifest(string InManifestName, DirectoryReference OutputDirectory, string? InTargetName, FileReference? InProjectFile, Dictionary<UnrealTargetConfiguration, string> InExecutablePairs)
 		{
 			// Check parameter values are valid.
 			if (InExecutablePairs.Count < 1)
@@ -498,7 +494,7 @@ namespace UnrealBuildTool
 			// generate the list of AppX cultures to stage
 			foreach (string UEStageId in SelectedUECultureIds)
 			{
-				if (!UEStageIdToAppXCultureId.TryGetValue(UEStageId, out string? AppXCultureId) || string.IsNullOrEmpty(AppXCultureId) )
+				if (!UEStageIdToAppXCultureId.TryGetValue(UEStageId, out string? AppXCultureId) || string.IsNullOrEmpty(AppXCultureId))
 				{
 					// use the culture directly - no remapping required
 					AppXCultureId = UEStageId;
@@ -535,12 +531,12 @@ namespace UnrealBuildTool
 			ProcessManifest(InExecutablePairs, InManifestName, ManifestTargetFile.FullName, ManifestIntermediateFile.FullName);
 
 			// Generate the package resource index and copy all resource files to the output
-			FileReference ManifestTargetPath = FileReference.Combine( OutputDirectory, InManifestName);
-			List<FileReference> UpdatedFiles = AppXResources.GenerateAppXResources( OutputDirectory, IntermediateDirectory, ManifestTargetFile, DefaultAppXCultureId, IdentityName );
+			FileReference ManifestTargetPath = FileReference.Combine(OutputDirectory, InManifestName);
+			List<FileReference> UpdatedFiles = AppXResources.GenerateAppXResources(OutputDirectory, IntermediateDirectory, ManifestTargetFile, DefaultAppXCultureId, IdentityName);
 
 			// Clean up and reutrn the list of updated files
 			FileUtils.ForceDeleteDirectory(IntermediateDirectory);
-			return UpdatedFiles.ConvertAll( X => X.FullName );
+			return UpdatedFiles.ConvertAll(X => X.FullName);
 		}
 	}
 }
diff --git a/Engine/Source/Programs/UnrealBuildTool/Platform/Windows/MicrosoftPlatformSDK.Versions.cs b/Engine/Source/Programs/UnrealBuildTool/Platform/Windows/MicrosoftPlatformSDK.Versions.cs
--- a/Engine/Source/Programs/UnrealBuildTool/Platform/Windows/MicrosoftPlatformSDK.Versions.cs
+++ b/Engine/Source/Programs/UnrealBuildTool/Platform/Windows/MicrosoftPlatformSDK.Versions.cs
@@ -1,8 +1,7 @@
 // Copyright Epic Games, Inc. All Rights Reserved.
 
-using EpicGames.Core;
-using System.Collections.Generic;
 using System.Linq;
+using EpicGames.Core;
 
 namespace UnrealBuildTool
 {
diff --git a/Engine/Source/Programs/UnrealBuildTool/Platform/Windows/MicrosoftPlatformSDK.VisualStudio.cs b/Engine/Source/Programs/UnrealBuildTool/Platform/Windows/MicrosoftPlatformSDK.VisualStudio.cs
--- a/Engine/Source/Programs/UnrealBuildTool/Platform/Windows/MicrosoftPlatformSDK.VisualStudio.cs
+++ b/Engine/Source/Programs/UnrealBuildTool/Platform/Windows/MicrosoftPlatformSDK.VisualStudio.cs
@@ -1,11 +1,7 @@
 ﻿// Copyright Epic Games, Inc. All Rights Reserved.
 
-using EpicGames.Core;
-using System;
 using System.Collections.Generic;
-using System.Linq;
-using System.Text;
-using System.Threading.Tasks;
+using EpicGames.Core;
 
 namespace UnrealBuildTool
 {
diff --git a/Engine/Source/Programs/UnrealBuildTool/Platform/Windows/MicrosoftPlatformSDK.cs b/Engine/Source/Programs/UnrealBuildTool/Platform/Windows/MicrosoftPlatformSDK.cs
--- a/Engine/Source/Programs/UnrealBuildTool/Platform/Windows/MicrosoftPlatformSDK.cs
+++ b/Engine/Source/Programs/UnrealBuildTool/Platform/Windows/MicrosoftPlatformSDK.cs
@@ -1,21 +1,20 @@
 // Copyright Epic Games, Inc. All Rights Reserved.
 
 using System;
-using System.IO;
+using System.Buffers.Binary;
+using System.Collections.Concurrent;
 using System.Collections.Generic;
-using EpicGames.Core;
-using System.Linq;
-using System.Text.RegularExpressions;
-using Microsoft.Win32;
 using System.Diagnostics;
 using System.Diagnostics.CodeAnalysis;
-using Microsoft.VisualStudio.Setup.Configuration;
-using System.Runtime.InteropServices;
-using System.Buffers.Binary;
-using UnrealBuildBase;
+using System.IO;
+using System.Linq;
 using System.Runtime.Versioning;
+using System.Text.RegularExpressions;
+using EpicGames.Core;
 using Microsoft.Extensions.Logging;
-using System.Collections.Concurrent;
+using Microsoft.VisualStudio.Setup.Configuration;
+using Microsoft.Win32;
+using UnrealBuildBase;
 
 ///////////////////////////////////////////////////////////////////
 // If you are looking for supported version numbers, look in the
@@ -214,7 +213,7 @@ namespace UnrealBuildTool
 			}
 			else if (CachedWindowsSdkDirs!.Count > 0)
 			{
-				WindowsSdkVersion = CachedWindowsSdkDirs.OrderBy(x => x.Key).Where( 
+				WindowsSdkVersion = CachedWindowsSdkDirs.OrderBy(x => x.Key).Where(
 					x =>
 					(MinVersion == null || x.Key >= MinVersion) &&
 					(MaxVersion == null || x.Key <= MaxVersion)).Last().Key;
@@ -1016,7 +1015,7 @@ namespace UnrealBuildTool
 			{
 				Logger.LogDebug("Found Visual Studio toolchain: {ToolChainDir} (Family={Family}, FamilyRank={FamilyRank}, Version={Version}, Is64Bit={Is64Bit}, ReleaseChannel={ReleaseChannel}, Architecture={Arch}, Error={Error}, Redist={RedistDir})", ToolChainDir, Family, FamilyRank, Version, Is64Bit, ReleaseChannel, UnrealArch.Arm64.ToString(), Error != null, RedistDir);
 				ToolChains.Add(new ToolChainInstallation(Family, FamilyRank, Version, Is64Bit, ReleaseChannel, UnrealArch.Arm64, Error, ToolChainDir, RedistDir, IsAutoSdk));
-				
+
 				if (HasArm64ECToolChain(ToolChainDir))
 				{
 					Logger.LogDebug("Found Visual Studio toolchain: {ToolChainDir} (Family={Family}, FamilyRank={FamilyRank}, Version={Version}, Is64Bit={Is64Bit}, ReleaseChannel={ReleaseChannel}, Architecture={Arch}, Error={Error}, Redist={RedistDir})", ToolChainDir, Family, FamilyRank, Version, Is64Bit, ReleaseChannel, UnrealArch.Arm64ec.ToString(), Error != null, RedistDir);
