UnrealBuildTool: Automated code cleanup

diff --git a/Engine/Source/Programs/UnrealBuildTool/Platform/Windows/PVSToolChain.cs b/Engine/Source/Programs/UnrealBuildTool/Platform/Windows/PVSToolChain.cs
--- a/Engine/Source/Programs/UnrealBuildTool/Platform/Windows/PVSToolChain.cs
+++ b/Engine/Source/Programs/UnrealBuildTool/Platform/Windows/PVSToolChain.cs
@@ -7,9 +7,7 @@ using System.IO;
 using System.Linq;
 using System.Text;
 using System.Text.RegularExpressions;
-using System.Threading;
 using System.Threading.Tasks;
-using System.Xml;
 using System.Xml.Serialization;
 using EpicGames.Core;
 using Microsoft.Extensions.Logging;
@@ -48,7 +46,7 @@ namespace UnrealBuildTool
 		MISRA = 32,
 	}
 
-  	/// <summary>
+	/// <summary>
 	/// Flags for the PVS analyzer timeout
 	/// </summary>
 	public enum AnalysisTimeoutFlags
@@ -117,7 +115,7 @@ namespace UnrealBuildTool
 		/// </summary>
 		public bool DisableMISRAAnalysis;
 
-    		/// <summary>
+		/// <summary>
 		/// File analysis timeout
 		/// </summary>
 		public AnalysisTimeoutFlags AnalysisTimeout;
@@ -131,7 +129,7 @@ namespace UnrealBuildTool
 		/// Enable the display of analyzer rules exceptions which can be specified by comments and .pvsconfig files.
 		/// </summary>
 		public bool ReportDisabledRules;
-    
+
 		/// <summary>
 		/// Gets the analysis mode flags from the settings
 		/// </summary>
@@ -230,8 +228,8 @@ namespace UnrealBuildTool
 		/// Private storage for the mode flags
 		/// </summary>
 		PVSAnalysisModeFlags? ModePrivate;
-    
-    		/// <summary>
+
+		/// <summary>
 		/// Override for the analysis timeoutFlag to use
 		/// </summary>
 		public AnalysisTimeoutFlags AnalysisTimeoutFlag
@@ -410,7 +408,7 @@ namespace UnrealBuildTool
 				foreach (FileReference InputFile in InputFiles)
 				{
 					string[] Lines = File.ReadAllLines(InputFile.FullName);
-					for(int LineIdx = 0; LineIdx < Lines.Length; LineIdx++)
+					for (int LineIdx = 0; LineIdx < Lines.Length; LineIdx++)
 					{
 						string Line = Lines[LineIdx];
 						if (!String.IsNullOrWhiteSpace(Line) && UniqueItems.Add(Line))
@@ -418,7 +416,7 @@ namespace UnrealBuildTool
 							bool bCanParse = false;
 
 							string[] Tokens = Line.Split(new string[] { "<#~>" }, StringSplitOptions.None);
-							if(Tokens.Length >= 9)
+							if (Tokens.Length >= 9)
 							{
 								//string Trial = Tokens[1];
 								string LineNumberStr = Tokens[2];
@@ -431,12 +429,12 @@ namespace UnrealBuildTool
 								int LineNumber;
 								bool bFalseAlarm;
 								int Level;
-								if(int.TryParse(LineNumberStr, out LineNumber) && bool.TryParse(FalseAlarmStr, out bFalseAlarm) && int.TryParse(LevelStr, out Level))
+								if (int.TryParse(LineNumberStr, out LineNumber) && bool.TryParse(FalseAlarmStr, out bFalseAlarm) && int.TryParse(LevelStr, out Level))
 								{
 									bCanParse = true;
 
 									// Ignore anything in ThirdParty folders
-									if(FileName.Replace('/', '\\').IndexOf("\\ThirdParty\\", StringComparison.InvariantCultureIgnoreCase) == -1)
+									if (FileName.Replace('/', '\\').IndexOf("\\ThirdParty\\", StringComparison.InvariantCultureIgnoreCase) == -1)
 									{
 										// Output the line to the raw output file
 										RawWriter.WriteLine(Line);
@@ -450,7 +448,7 @@ namespace UnrealBuildTool
 								}
 							}
 
-							if(!bCanParse)
+							if (!bCanParse)
 							{
 								Log.WriteLine(LogEventType.Warning, LogFormatOptions.NoSeverityPrefix, "{0}({1}): warning: Unable to parse PVS output line '{2}' (tokens=|{3}|)", InputFile, LineIdx + 1, Line, String.Join("|", Tokens));
 							}
@@ -458,7 +456,7 @@ namespace UnrealBuildTool
 					}
 				}
 			}
-			Logger.LogInformation("Written {NumItems} {Noun} to {File}.", UniqueItems.Count, (UniqueItems.Count == 1)? "diagnostic" : "diagnostics", OutputFile.FullName);
+			Logger.LogInformation("Written {NumItems} {Noun} to {File}.", UniqueItems.Count, (UniqueItems.Count == 1) ? "diagnostic" : "diagnostics", OutputFile.FullName);
 			return Task.FromResult(0);
 		}
 	}
@@ -499,7 +497,7 @@ namespace UnrealBuildTool
 			Settings = Target.WindowsPlatform.PVS;
 			ApplicationSettings = Settings.ApplicationSettings;
 
-			if(ApplicationSettings != null)
+			if (ApplicationSettings != null)
 			{
 				if (Settings.ModeFlags == 0)
 				{
@@ -515,7 +513,7 @@ namespace UnrealBuildTool
 			else
 			{
 				FileReference DefaultLicenseFile = AnalyzerFile.ChangeExtension(".lic");
-				if(FileReference.Exists(DefaultLicenseFile))
+				if (FileReference.Exists(DefaultLicenseFile))
 				{
 					LicenseFile = DefaultLicenseFile;
 				}
@@ -605,7 +603,7 @@ namespace UnrealBuildTool
 				case CppStandardVersion.Cpp17:
 					cppCfgStandard = CPP_17;
 					break;
-        			case CppStandardVersion.Cpp20:
+				case CppStandardVersion.Cpp20:
 					cppCfgStandard = CPP_20;
 					break;
 				case CppStandardVersion.Latest:
@@ -670,7 +668,7 @@ namespace UnrealBuildTool
 			CPPOutput Result = InnerToolChain.CompileAllCPPFiles(PreprocessCompileEnvironment, InputFiles, OutputDir, ModuleName, new ActionGraphCapture(Graph, PreprocessActions));
 
 			// Run the source files through PVS-Studio
-			for(int Idx = 0; Idx < PreprocessActions.Count; Idx++)
+			for (int Idx = 0; Idx < PreprocessActions.Count; Idx++)
 			{
 				VCCompileAction? PreprocessAction = PreprocessActions[Idx] as VCCompileAction;
 				if (PreprocessAction == null)
@@ -694,17 +692,17 @@ namespace UnrealBuildTool
 
 				// Write the PVS studio config file
 				StringBuilder ConfigFileContents = new StringBuilder();
-				foreach(DirectoryReference IncludePath in Target.WindowsPlatform.Environment!.IncludePaths)
+				foreach (DirectoryReference IncludePath in Target.WindowsPlatform.Environment!.IncludePaths)
 				{
 					ConfigFileContents.AppendFormat("exclude-path={0}\n", IncludePath.FullName);
 				}
-				if(ApplicationSettings != null && ApplicationSettings.PathMasks != null)
+				if (ApplicationSettings != null && ApplicationSettings.PathMasks != null)
 				{
-					foreach(string PathMask in ApplicationSettings.PathMasks)
+					foreach (string PathMask in ApplicationSettings.PathMasks)
 					{
 						if (PathMask.Contains(":") || PathMask.Contains("\\") || PathMask.Contains("/"))
 						{
-							if(Path.IsPathRooted(PathMask) && !PathMask.Contains(":"))
+							if (Path.IsPathRooted(PathMask) && !PathMask.Contains(":"))
 							{
 								ConfigFileContents.AppendFormat("exclude-path=*{0}*\n", PathMask);
 							}
@@ -715,7 +713,7 @@ namespace UnrealBuildTool
 						}
 					}
 				}
-        			if (Platform == UnrealTargetPlatform.Win64)
+				if (Platform == UnrealTargetPlatform.Win64)
 				{
 					ConfigFileContents.Append("platform=x64\n");
 				}
@@ -723,7 +721,7 @@ namespace UnrealBuildTool
 				{
 					throw new BuildException("PVS-Studio does not support this platform");
 				}
-        		ConfigFileContents.Append("preprocessor=visualcpp\n");
+				ConfigFileContents.Append("preprocessor=visualcpp\n");
 
 				bool shouldCompileAsC = ShouldCompileAsC(String.Join(" ", PreprocessAction.Arguments), SourceFileItem.AbsolutePath);
 				ConfigFileContents.AppendFormat("language={0}\n", shouldCompileAsC ? "C" : "C++");
@@ -738,8 +736,8 @@ namespace UnrealBuildTool
 					string languageStandardForCfg = GetLangStandForCfgFile(PreprocessCompileEnvironment.CppStandard, compilerVersion);
 
 					ConfigFileContents.AppendFormat("std={0}\n", languageStandardForCfg);
-          
-          			bool disableMsExtentinsFromArgs = PreprocessAction.Arguments.Any(arg => arg.Equals("/Za") || arg.Equals("-Za") || arg.Equals("/permissive-"));
+
+					bool disableMsExtentinsFromArgs = PreprocessAction.Arguments.Any(arg => arg.Equals("/Za") || arg.Equals("-Za") || arg.Equals("/permissive-"));
 					bool disableMsExtentions = isVisualCppCompiler && (languageStandardForCfg == CPP_20 || disableMsExtentinsFromArgs);
 					ConfigFileContents.AppendFormat("disable-ms-extensions={0}\n", disableMsExtentions ? "yes" : "no");
 				}
@@ -759,7 +757,7 @@ namespace UnrealBuildTool
 					ConfigFileContents.Append("report-disabled-rules=yes\n");
 				}
 
-        		int Timeout = (int)(Settings.AnalysisTimeoutFlag == AnalysisTimeoutFlags.No_timeout ? 0 : Settings.AnalysisTimeoutFlag);
+				int Timeout = (int)(Settings.AnalysisTimeoutFlag == AnalysisTimeoutFlags.No_timeout ? 0 : Settings.AnalysisTimeoutFlag);
 				ConfigFileContents.AppendFormat("timeout={0}\n", Timeout);
 
 				string BaseFileName = PreprocessedFileItem.Location.GetFileName();
diff --git a/Engine/Source/Programs/UnrealBuildTool/Platform/Windows/ResXResource/UEAppXResources.cs b/Engine/Source/Programs/UnrealBuildTool/Platform/Windows/ResXResource/UEAppXResources.cs
--- a/Engine/Source/Programs/UnrealBuildTool/Platform/Windows/ResXResource/UEAppXResources.cs
+++ b/Engine/Source/Programs/UnrealBuildTool/Platform/Windows/ResXResource/UEAppXResources.cs
@@ -2,13 +2,13 @@
 
 using System;
 using System.Collections.Generic;
-using System.Linq;
+using System.Diagnostics;
+using System.Diagnostics.CodeAnalysis;
 using System.IO;
+using System.Linq;
 using System.Text;
-using EpicGames.Core;
-using System.Diagnostics;
 using System.Xml;
-using System.Diagnostics.CodeAnalysis;
+using EpicGames.Core;
 using Microsoft.Extensions.Logging;
 
 namespace UnrealBuildTool
@@ -30,11 +30,11 @@ namespace UnrealBuildTool
 		};
 
 		readonly Dictionary<string, Resources> PerCultureResources = new();
-		readonly Dictionary<FileReference,string> FilesToCopy = new(); // source file -> target file fragment (relative to OutputDir)
+		readonly Dictionary<FileReference, string> FilesToCopy = new(); // source file -> target file fragment (relative to OutputDir)
 
 		readonly ILogger Logger;
 		readonly FileReference MakePriExe; // path to makepri.exe
-		
+
 		/// <summary>
 		/// Folders where binary resource files are located, in decreasing priority order
 		/// </summary>
@@ -73,11 +73,11 @@ namespace UnrealBuildTool
 		/// <summary>
 		/// Adds a reference to the given culture
 		/// </summary>
-		public void AddCulture( string AppXCultureId )
+		public void AddCulture(string AppXCultureId)
 		{
 			if (!PerCultureResources.TryGetValue(AppXCultureId, out Resources? Result) || Result == null)
 			{
-				PerCultureResources.Add( AppXCultureId, new());
+				PerCultureResources.Add(AppXCultureId, new());
 			}
 		}
 
@@ -85,9 +85,9 @@ namespace UnrealBuildTool
 		/// Adds a reference to all of the given cultures
 		/// </summary>
 		/// <param name="AppXCultureIds"></param>
-		public void AddCultures( IEnumerable<string> AppXCultureIds )
+		public void AddCultures(IEnumerable<string> AppXCultureIds)
 		{
-			foreach (string AppXCultureId in AppXCultureIds )
+			foreach (string AppXCultureId in AppXCultureIds)
 			{
 				AddCulture(AppXCultureId);
 			}
@@ -100,14 +100,14 @@ namespace UnrealBuildTool
 		/// <returns></returns>
 		public IEnumerable<string> GetAllCultureIds()
 		{
-			return PerCultureResources.Keys.Where( X => X != DefaultResources );
+			return PerCultureResources.Keys.Where(X => X != DefaultResources);
 		}
 
 
 		/// <summary>
 		/// Add the given string to the default culture
 		/// </summary>
-		public void AddDefaultString( string ConfigKey, string Value )
+		public void AddDefaultString(string ConfigKey, string Value)
 		{
 			AddCultureString(DefaultResources, ConfigKey, Value);
 		}
@@ -115,15 +115,15 @@ namespace UnrealBuildTool
 		/// <summary>
 		/// Adds all of the given strings to the default culture
 		/// </summary>
-		public void AddDefaultStrings( Dictionary<string, string> Strings )
+		public void AddDefaultStrings(Dictionary<string, string> Strings)
 		{
-			AddCultureStrings(DefaultResources, Strings );
+			AddCultureStrings(DefaultResources, Strings);
 		}
 
 		/// <summary>
 		/// Add the given string to the given culture
 		/// </summary>
-		public void AddCultureString( string AppXCultureId, string ConfigKey, string Value )
+		public void AddCultureString(string AppXCultureId, string ConfigKey, string Value)
 		{
 			GetCultureResources(AppXCultureId).SourceStrings[ConfigKey] = Value;
 		}
@@ -131,10 +131,10 @@ namespace UnrealBuildTool
 		/// <summary>
 		/// Adds all of the given strings to the given culture
 		/// </summary>
-		public void AddCultureStrings( string AppXCultureId, Dictionary<string, string> Strings )
+		public void AddCultureStrings(string AppXCultureId, Dictionary<string, string> Strings)
 		{
 			Resources ThisCultureResources = GetCultureResources(AppXCultureId);
-			foreach (KeyValuePair<string,string> String in Strings)
+			foreach (KeyValuePair<string, string> String in Strings)
 			{
 				ThisCultureResources.SourceStrings[String.Key] = String.Value;
 			}
@@ -143,7 +143,7 @@ namespace UnrealBuildTool
 		/// <summary>
 		/// Returns whether the given string is known
 		/// </summary>
-		public bool HasString( string Key )
+		public bool HasString(string Key)
 		{
 			return GetDefaultResources().SourceStrings.ContainsKey(Key);
 		}
@@ -161,7 +161,7 @@ namespace UnrealBuildTool
 		}
 
 
-		Resources GetCultureResources( string AppXCultureId )
+		Resources GetCultureResources(string AppXCultureId)
 		{
 			if (PerCultureResources.TryGetValue(AppXCultureId, out Resources? Result) && Result != null)
 			{
@@ -181,8 +181,8 @@ namespace UnrealBuildTool
 
 
 
-        bool RunMakePri(string CommandLine)
-        {
+		bool RunMakePri(string CommandLine)
+		{
 			StringBuilder ProcessOutput = new StringBuilder();
 			void LocalProcessOutput(DataReceivedEventArgs Args)
 			{
@@ -192,7 +192,7 @@ namespace UnrealBuildTool
 				}
 			}
 
-			ProcessStartInfo StartInfo = new ProcessStartInfo(MakePriExe.FullName, CommandLine);			
+			ProcessStartInfo StartInfo = new ProcessStartInfo(MakePriExe.FullName, CommandLine);
 			StartInfo.UseShellExecute = false;
 			StartInfo.CreateNoWindow = true;
 			StartInfo.StandardOutputEncoding = Encoding.Unicode;
@@ -216,14 +216,14 @@ namespace UnrealBuildTool
 				Logger.LogError("Exit code: {Code}", ExitCode);
 				return false;
 			}
-        }
+		}
 
 
 
 
-		private bool RemoveStaleResourceFiles( IEnumerable<string> RequiredFileFragments, DirectoryReference OutputDirectory )
+		private bool RemoveStaleResourceFiles(IEnumerable<string> RequiredFileFragments, DirectoryReference OutputDirectory)
 		{
-			DirectoryReference ResourceDirectory = DirectoryReference.Combine( OutputDirectory, BuildResourceSubPath );
+			DirectoryReference ResourceDirectory = DirectoryReference.Combine(OutputDirectory, BuildResourceSubPath);
 			if (!DirectoryReference.Exists(ResourceDirectory))
 			{
 				return false;
@@ -232,8 +232,8 @@ namespace UnrealBuildTool
 			// remove all files in the Resources/ subfolder that should not be included
 			IEnumerable<FileReference> ExistingFiles = DirectoryReference.EnumerateFiles(ResourceDirectory, "*.*", SearchOption.AllDirectories);
 			IEnumerable<FileReference> RequiredFiles = RequiredFileFragments
-				.Select( X => FileReference.Combine( OutputDirectory, X ))
-				.Where( X => X.IsUnderDirectory(ResourceDirectory))
+				.Select(X => FileReference.Combine(OutputDirectory, X))
+				.Where(X => X.IsUnderDirectory(ResourceDirectory))
 				;
 			IEnumerable<FileReference> StaleResourceFiles = ExistingFiles.Except(RequiredFiles);
 			if (!StaleResourceFiles.Any())
@@ -267,12 +267,12 @@ namespace UnrealBuildTool
 		/// <summary>
 		/// Attempts to locate the given resource binary file in the previously-specified resource folders
 		/// </summary>
-		bool FindResourceBinaryFile( [NotNullWhen(true)] out FileReference? SourceFilePath, string ResourceFileName, string CultureId = "", bool AllowEngineFallback = true)
+		bool FindResourceBinaryFile([NotNullWhen(true)] out FileReference? SourceFilePath, string ResourceFileName, string CultureId = "", bool AllowEngineFallback = true)
 		{
 			// look in project binary resource directories
 			foreach (DirectoryReference BinaryResourceDirectory in ProjectBinaryResourceDirectories)
 			{
-				FileReference BinaryResourceFile = FileReference.Combine(BinaryResourceDirectory, CultureId, ResourceFileName );
+				FileReference BinaryResourceFile = FileReference.Combine(BinaryResourceDirectory, CultureId, ResourceFileName);
 				if (FileReference.Exists(BinaryResourceFile))
 				{
 					SourceFilePath = BinaryResourceFile;
@@ -281,11 +281,11 @@ namespace UnrealBuildTool
 			}
 
 			// look in Engine, if allowed
-			if (AllowEngineFallback )
+			if (AllowEngineFallback)
 			{
 				foreach (DirectoryReference BinaryResourceDirectory in EngineFallbackBinaryResourceDirectories)
 				{
-					FileReference BinaryResourceFile = FileReference.Combine(BinaryResourceDirectory, CultureId, ResourceFileName );
+					FileReference BinaryResourceFile = FileReference.Combine(BinaryResourceDirectory, CultureId, ResourceFileName);
 					if (FileReference.Exists(BinaryResourceFile))
 					{
 						SourceFilePath = BinaryResourceFile;
@@ -304,7 +304,7 @@ namespace UnrealBuildTool
 		/// </summary>
 		public bool DoesCultureResourceBinaryFileExist(string ResourceFileName, string CultureId, bool AllowEngineFallback = true)
 		{
-			return FindResourceBinaryFile( out FileReference? _, ResourceFileName, CultureId, AllowEngineFallback );
+			return FindResourceBinaryFile(out FileReference? _, ResourceFileName, CultureId, AllowEngineFallback);
 		}
 
 		/// <summary>
@@ -312,7 +312,7 @@ namespace UnrealBuildTool
 		/// </summary>
 		public bool DoesDefaultResourceBinaryFileExist(string ResourceFileName, bool AllowEngineFallback = true)
 		{
-			return DoesCultureResourceBinaryFileExist( ResourceFileName, DefaultResources, AllowEngineFallback );
+			return DoesCultureResourceBinaryFileExist(ResourceFileName, DefaultResources, AllowEngineFallback);
 		}
 
 
@@ -322,14 +322,14 @@ namespace UnrealBuildTool
 		public bool AddResourceBinaryFileReference(string ResourceFileName, bool AllowEngineFallback = true)
 		{
 			// At least the default culture entry for any resource binary must always exist
-			if (!FindResourceBinaryFile( out FileReference? SourceFilePath, ResourceFileName, DefaultResources, AllowEngineFallback))
+			if (!FindResourceBinaryFile(out FileReference? SourceFilePath, ResourceFileName, DefaultResources, AllowEngineFallback))
 			{
 				return false;
 			}
 			AddFileReference(SourceFilePath, Path.Combine(BuildResourceSubPath, ResourceFileName));
 
 			// Copy all per-culture resource files
-			foreach( string CultureId in GetAllCultureIds() )
+			foreach (string CultureId in GetAllCultureIds())
 			{
 				if (FindResourceBinaryFile(out SourceFilePath, ResourceFileName, CultureId, AllowEngineFallback))
 				{
@@ -349,7 +349,7 @@ namespace UnrealBuildTool
 			FilesToCopy.Add(SourcePath, TargetPathFragment);
 		}
 
-		private static bool AreFilesDifferent( FileReference File1, FileReference File2)
+		private static bool AreFilesDifferent(FileReference File1, FileReference File2)
 		{
 			FileInfo FileInfo1 = File1.ToFileInfo();
 			FileInfo FileInfo2 = File2.ToFileInfo();
@@ -377,7 +377,7 @@ namespace UnrealBuildTool
 			}
 
 			Logger.LogDebug("Updating manifest resource files...");
-			foreach ( var FileToCopy in FilesToCopy )
+			foreach (var FileToCopy in FilesToCopy)
 			{
 				FileReference SourceFile = FileToCopy.Key;
 				FileReference TargetFile = FileReference.Combine(OutputDirectory, FileToCopy.Value);
@@ -458,7 +458,7 @@ namespace UnrealBuildTool
 			foreach (string CultureId in GetAllCultureIds())
 			{
 				Resources ThisCultureResources = GetCultureResources(CultureId);
-				if (ThisCultureResources.SourceStrings.TryGetValue(ConfigKey, out string? CultureString) && !string.IsNullOrEmpty(CultureString) )
+				if (ThisCultureResources.SourceStrings.TryGetValue(ConfigKey, out string? CultureString) && !string.IsNullOrEmpty(CultureString))
 				{
 					ThisCultureResources.StringResources.Add(ResourceEntryName, CultureString + ValueSuffix);
 				}
@@ -473,7 +473,7 @@ namespace UnrealBuildTool
 		/// <summary>
 		/// Generate the package resource index and copies all resources files to the output
 		/// </summary>
-		public List<FileReference> GenerateAppXResources( DirectoryReference OutputDirectory, DirectoryReference IntermediateDirectory, FileReference ManifestFile, string DefaultAppXCultureId, string? PackageIdentityName )
+		public List<FileReference> GenerateAppXResources(DirectoryReference OutputDirectory, DirectoryReference IntermediateDirectory, FileReference ManifestFile, string DefaultAppXCultureId, string? PackageIdentityName)
 		{
 			// Clean out the resources intermediate path so that we know there are no stale binary files.
 			FileUtils.ForceDeleteDirectory(DirectoryReference.Combine(IntermediateDirectory, BuildResourceSubPath));
@@ -490,24 +490,24 @@ namespace UnrealBuildTool
 				FileUtils.CreateDirectoryTree(OutputCultureResourceDirectory);
 
 				// Export the resource tables & add them to the list of files to copy
-				FileReference IntermediateResWIndexerFile = FileReference.Combine( IntermediateDirectory, BuildResourceSubPath, AppXCultureId, "resources.resw" );
+				FileReference IntermediateResWIndexerFile = FileReference.Combine(IntermediateDirectory, BuildResourceSubPath, AppXCultureId, "resources.resw");
 				UEResXWriter ResourceWriter = new(IntermediateResWIndexerFile.FullName);
 				foreach (KeyValuePair<string, string> StringResourcePair in Itr.Value.StringResources)
 				{
-					ResourceWriter.AddResource( StringResourcePair.Key, StringResourcePair.Value );
+					ResourceWriter.AddResource(StringResourcePair.Key, StringResourcePair.Value);
 				}
 				ResourceWriter.Close();
-			
-				AddFileReference( IntermediateResWIndexerFile, Path.Combine(BuildResourceSubPath, AppXCultureId, "resources.resw"));
+
+				AddFileReference(IntermediateResWIndexerFile, Path.Combine(BuildResourceSubPath, AppXCultureId, "resources.resw"));
 			}
 
 			// The resource database is dependent on everything else calculated here (manifest, resource string tables, binary resources).
 			// So if any file has been updated we'll need to run the config.
 			bool bHadStaleResources = RemoveStaleResourceFiles(FilesToCopy.Values, OutputDirectory);
 			List<FileReference> UpdatedFilePaths = CopyFilesToOutput(OutputDirectory);
-			FileReference TargetResourceIndexFile = FileReference.Combine( OutputDirectory, "resources.pri" );
+			FileReference TargetResourceIndexFile = FileReference.Combine(OutputDirectory, "resources.pri");
 
-			if (bHadStaleResources || UpdatedFilePaths.Any() || !FileReference.Exists(TargetResourceIndexFile) )
+			if (bHadStaleResources || UpdatedFilePaths.Any() || !FileReference.Exists(TargetResourceIndexFile))
 			{
 				// Create resource index configuration
 				FileReference ResourceConfigFile = FileReference.Combine(IntermediateDirectory, "priconfig.xml");
@@ -536,11 +536,11 @@ namespace UnrealBuildTool
 				PriConfig.Save(ResourceConfigFile.FullName);
 
 				// generate resources.resfiles
-				IEnumerable<FileReference> Resources = DirectoryReference.EnumerateFiles( DirectoryReference.Combine(OutputDirectory, BuildResourceSubPath), "*.*", SearchOption.AllDirectories);
+				IEnumerable<FileReference> Resources = DirectoryReference.EnumerateFiles(DirectoryReference.Combine(OutputDirectory, BuildResourceSubPath), "*.*", SearchOption.AllDirectories);
 				System.Text.StringBuilder ResourcesList = new System.Text.StringBuilder();
 				foreach (FileReference Resource in Resources)
 				{
-					ResourcesList.AppendLine( Resource.MakeRelativeTo(OutputDirectory) );
+					ResourcesList.AppendLine(Resource.MakeRelativeTo(OutputDirectory));
 				}
 				File.WriteAllText(ResourcesResFile.FullName, ResourcesList.ToString());
 
@@ -560,7 +560,7 @@ namespace UnrealBuildTool
 				UpdatedFilePaths.Add(TargetResourceIndexFile);
 
 				// .resw files are not needed - the data is embedded in the resources.pri
-				foreach (FileReference ResW in DirectoryReference.EnumerateFiles( OutputDirectory, "*.resw", SearchOption.AllDirectories))
+				foreach (FileReference ResW in DirectoryReference.EnumerateFiles(OutputDirectory, "*.resw", SearchOption.AllDirectories))
 				{
 					FileUtils.ForceDeleteFile(ResW);
 				}
diff --git a/Engine/Source/Programs/UnrealBuildTool/Platform/Windows/ResXResource/UEResXReader.cs b/Engine/Source/Programs/UnrealBuildTool/Platform/Windows/ResXResource/UEResXReader.cs
--- a/Engine/Source/Programs/UnrealBuildTool/Platform/Windows/ResXResource/UEResXReader.cs
+++ b/Engine/Source/Programs/UnrealBuildTool/Platform/Windows/ResXResource/UEResXReader.cs
@@ -1,11 +1,5 @@
 ﻿// Copyright Epic Games, Inc. All Rights Reserved.
 
-using System;
-using System.Collections.Generic;
-using System.Linq;
-using System.Text;
-using System.Threading.Tasks;
-
 namespace UnrealBuildTool
 {
 	/// <summary>
diff --git a/Engine/Source/Programs/UnrealBuildTool/Platform/Windows/UEBuildWindows.cs b/Engine/Source/Programs/UnrealBuildTool/Platform/Windows/UEBuildWindows.cs
--- a/Engine/Source/Programs/UnrealBuildTool/Platform/Windows/UEBuildWindows.cs
+++ b/Engine/Source/Programs/UnrealBuildTool/Platform/Windows/UEBuildWindows.cs
@@ -2,15 +2,14 @@
 
 using System;
 using System.Collections.Generic;
-using System.Text;
 using System.Diagnostics;
-using System.IO;
-using System.Linq;
-using EpicGames.Core;
 using System.Diagnostics.CodeAnalysis;
-using UnrealBuildBase;
+using System.Linq;
 using System.Runtime.Versioning;
+using System.Text;
+using EpicGames.Core;
 using Microsoft.Extensions.Logging;
+using UnrealBuildBase;
 
 namespace UnrealBuildTool
 {
@@ -429,7 +428,7 @@ namespace UnrealBuildTool
 			set => bStripUnreferencedSymbolsPrivate = value;
 		}
 		private bool? bStripUnreferencedSymbolsPrivate;
-			
+
 		/// <summary>
 		/// Merge identical COMDAT sections together (/OPT:ICF)
 		/// </summary>
@@ -536,7 +535,7 @@ namespace UnrealBuildTool
 				{
 					return MicrosoftPlatformSDK.FindVisualStudioInstallations(Environment!.ToolChain, Target.Logger).Select(x => x.BaseDir.FullName).FirstOrDefault();
 				}
-				catch(Exception) // Find function will throw if there is no visual studio installed! This can happen w/ clang builds
+				catch (Exception) // Find function will throw if there is no visual studio installed! This can happen w/ clang builds
 				{
 					return null;
 				}
@@ -572,7 +571,7 @@ namespace UnrealBuildTool
 			this.Target = Target;
 
 			string Platform = Target.Platform.ToString();
-			if (Target.Platform== UnrealTargetPlatform.Win64 && !Target.Architecture.bIsX64)
+			if (Target.Platform == UnrealTargetPlatform.Win64 && !Target.Architecture.bIsX64)
 			{
 				Platform += "-arm64";
 			}
@@ -604,7 +603,7 @@ namespace UnrealBuildTool
 		/// Accessors for fields on the inner TargetRules instance
 		/// </summary>
 		#region Read-only accessor properties 
-		#pragma warning disable CS1591
+#pragma warning disable CS1591
 
 		public bool bUseFastGenProfile => Inner.bUseFastGenProfile;
 
@@ -716,7 +715,7 @@ namespace UnrealBuildTool
 
 		public string? IDEDir => Inner.IDEDir;
 
-		#pragma warning restore CS1591
+#pragma warning restore CS1591
 		#endregion
 	}
 
@@ -767,7 +766,7 @@ namespace UnrealBuildTool
 	class WindowsArchitectureConfig : UnrealArchitectureConfig
 	{
 		public WindowsArchitectureConfig()
-			: base(UnrealArchitectureMode.OneTargetPerArchitecture, new[] { UnrealArch.X64, UnrealArch.Arm64, UnrealArch.Arm64ec})
+			: base(UnrealArchitectureMode.OneTargetPerArchitecture, new[] { UnrealArch.X64, UnrealArch.Arm64, UnrealArch.Arm64ec })
 		{
 
 		}
@@ -785,7 +784,7 @@ namespace UnrealBuildTool
 	}
 
 	class WindowsPlatform : UEBuildPlatform
-	{		
+	{
 		MicrosoftPlatformSDK SDK;
 
 		/// <summary>
@@ -868,7 +867,7 @@ namespace UnrealBuildTool
 			{
 				Target.bCompileCEF3 = false;
 			}
-			
+
 			// If clang is selected for static analysis, switch compiler to clang and proceed
 			// as normal.
 			if (Target.StaticAnalyzer == StaticAnalyzer.Clang)
@@ -879,8 +878,8 @@ namespace UnrealBuildTool
 				}
 				Target.StaticAnalyzer = StaticAnalyzer.Default;
 			}
-			else if (Target.StaticAnalyzer != StaticAnalyzer.None && 
-			         Target.StaticAnalyzerOutputType != StaticAnalyzerOutputType.Text)
+			else if (Target.StaticAnalyzer != StaticAnalyzer.None &&
+					 Target.StaticAnalyzerOutputType != StaticAnalyzerOutputType.Text)
 			{
 				Logger.LogInformation("Defaulting static analyzer output type to text");
 			}
@@ -908,7 +907,7 @@ namespace UnrealBuildTool
 			{
 				Target.bUsePCHFiles = false;
 			}
-			
+
 			// Disable chaining PCH files if not using clang
 			if (!Target.WindowsPlatform.Compiler.IsClang())
 			{
@@ -965,29 +964,29 @@ namespace UnrealBuildTool
 				throw new BuildException("LibFuzzer MSVC support requires Visual Studio 2022 17.0 or later (14.30.0). The current compiler version was detected as: {0}", Target.WindowsPlatform.Environment.CompilerVersion.ToString());
 			}
 
-//			@Todo: Still getting reports of frequent OOM issues with this enabled as of 15.7.
-//			// Enable fast PDB linking if we're on VS2017 15.7 or later. Previous versions have OOM issues with large projects.
-//			if(!Target.bFormalBuild && !Target.bUseFastPDBLinking.HasValue && Target.WindowsPlatform.Compiler.IsMSVC())
-//			{
-//				VersionNumber Version;
-//				DirectoryReference ToolChainDir;
-//				if(TryGetVCToolChainDir(Target.WindowsPlatform.Compiler, Target.WindowsPlatform.CompilerVersion, out Version, out ToolChainDir) && Version >= new VersionNumber(14, 14, 26316))
-//				{
-//					Target.bUseFastPDBLinking = true;
-//				}
-//			}
+			//			@Todo: Still getting reports of frequent OOM issues with this enabled as of 15.7.
+			//			// Enable fast PDB linking if we're on VS2017 15.7 or later. Previous versions have OOM issues with large projects.
+			//			if(!Target.bFormalBuild && !Target.bUseFastPDBLinking.HasValue && Target.WindowsPlatform.Compiler.IsMSVC())
+			//			{
+			//				VersionNumber Version;
+			//				DirectoryReference ToolChainDir;
+			//				if(TryGetVCToolChainDir(Target.WindowsPlatform.Compiler, Target.WindowsPlatform.CompilerVersion, out Version, out ToolChainDir) && Version >= new VersionNumber(14, 14, 26316))
+			//				{
+			//					Target.bUseFastPDBLinking = true;
+			//				}
+			//			}
 		}
 
 		/// <summary>
 		/// Gets the default compiler which should be used, if it's not set explicitly by the target, command line, or config file.
 		/// </summary>
 		/// <returns>The default compiler version</returns>
-		internal static WindowsCompiler GetDefaultCompiler(FileReference? ProjectFile, UnrealArch Architecture, ILogger Logger, bool bSkipWarning=false)
+		internal static WindowsCompiler GetDefaultCompiler(FileReference? ProjectFile, UnrealArch Architecture, ILogger Logger, bool bSkipWarning = false)
 		{
 			// If there's no specific compiler set, try to pick the matching compiler for the selected IDE
 			if (ProjectFileGeneratorSettings.Format != null)
 			{
-				foreach(ProjectFileFormat Format in ProjectFileGeneratorSettings.ParseFormatList(ProjectFileGeneratorSettings.Format, Logger))
+				foreach (ProjectFileFormat Format in ProjectFileGeneratorSettings.ParseFormatList(ProjectFileGeneratorSettings.Format, Logger))
 				{
 					if (Format == ProjectFileFormat.VisualStudio2022)
 					{
@@ -997,7 +996,7 @@ namespace UnrealBuildTool
 					{
 						return WindowsCompiler.VisualStudio2019;
 					}
-				} 
+				}
 			}
 
 			// Also check the default format for the Visual Studio project generator
@@ -1017,12 +1016,12 @@ namespace UnrealBuildTool
 
 			// Check the editor settings too
 			ProjectFileFormat PreferredAccessor;
-			if(ProjectFileGenerator.GetPreferredSourceCodeAccessor(ProjectFile, out PreferredAccessor))
+			if (ProjectFileGenerator.GetPreferredSourceCodeAccessor(ProjectFile, out PreferredAccessor))
 			{
-				if(PreferredAccessor == ProjectFileFormat.VisualStudio2022)
-			    {
-				    return WindowsCompiler.VisualStudio2022;
-			    }
+				if (PreferredAccessor == ProjectFileFormat.VisualStudio2022)
+				{
+					return WindowsCompiler.VisualStudio2022;
+				}
 				else if (PreferredAccessor == ProjectFileFormat.VisualStudio2019)
 				{
 					return WindowsCompiler.VisualStudio2019;
@@ -1086,7 +1085,7 @@ namespace UnrealBuildTool
 		public static IEnumerable<DirectoryReference>? TryGetVSInstallDirs(WindowsCompiler Compiler, ILogger Logger)
 		{
 			List<VisualStudioInstallation> Installations = MicrosoftPlatformSDK.FindVisualStudioInstallations(Compiler, Logger);
-			if(Installations.Count == 0)
+			if (Installations.Count == 0)
 			{
 				return null;
 			}
@@ -1222,9 +1221,9 @@ namespace UnrealBuildTool
 			{
 				case UEBuildBinaryType.DynamicLinkLibrary:
 				case UEBuildBinaryType.Executable:
-					return new string[] {".pdb"};
+					return new string[] { ".pdb" };
 			}
-			return new string [] {};
+			return new string[] { };
 		}
 
 		public override bool HasDefaultBuildConfig(UnrealTargetPlatform Platform, DirectoryReference ProjectPath)
@@ -1252,10 +1251,10 @@ namespace UnrealBuildTool
 		public static FileReference GetWindowsApplicationIcon(FileReference? ProjectFile)
 		{
 			// Check if there's a custom icon
-			if(ProjectFile != null)
+			if (ProjectFile != null)
 			{
 				FileReference IconFile = FileReference.Combine(ProjectFile.Directory, "Build", "Windows", "Application.ico");
-				if(FileReference.Exists(IconFile))
+				if (FileReference.Exists(IconFile))
 				{
 					return IconFile;
 				}
@@ -1333,7 +1332,7 @@ namespace UnrealBuildTool
 
 			CompileEnvironment.Definitions.Add(String.Format("_WIN32_WINNT=0x{0:X4}", Target.WindowsPlatform.TargetWindowsVersion));
 			CompileEnvironment.Definitions.Add(String.Format("WINVER=0x{0:X4}", Target.WindowsPlatform.TargetWindowsVersion));
-			
+
 			CompileEnvironment.Definitions.Add("PLATFORM_WINDOWS=1");
 			CompileEnvironment.Definitions.Add("PLATFORM_MICROSOFT=1");
 
@@ -1520,7 +1519,7 @@ namespace UnrealBuildTool
 		{
 			base.GetExternalBuildMetadata(ProjectFile, Metadata);
 
-			if(ProjectFile != null)
+			if (ProjectFile != null)
 			{
 				Metadata.AppendLine("ICON: {0}", GetApplicationIcon(ProjectFile));
 			}
diff --git a/Engine/Source/Programs/UnrealBuildTool/Platform/Windows/VCEnvironment.cs b/Engine/Source/Programs/UnrealBuildTool/Platform/Windows/VCEnvironment.cs
--- a/Engine/Source/Programs/UnrealBuildTool/Platform/Windows/VCEnvironment.cs
+++ b/Engine/Source/Programs/UnrealBuildTool/Platform/Windows/VCEnvironment.cs
@@ -2,15 +2,10 @@
 
 using System;
 using System.Collections.Generic;
-using System.Text.RegularExpressions;
 using System.Diagnostics;
-using System.IO;
 using System.Linq;
-using Microsoft.Win32;
-using System.Text;
-using EpicGames.Core;
-using UnrealBuildBase;
 using System.Runtime.Versioning;
+using EpicGames.Core;
 using Microsoft.Extensions.Logging;
 
 namespace UnrealBuildTool
@@ -54,7 +49,7 @@ namespace UnrealBuildTool
 		/// The toolchain version number
 		/// </summary>
 		public readonly VersionNumber ToolChainVersion;
-		
+
 		/// <summary>
 		/// Root directory containing the Windows Sdk
 		/// </summary>
@@ -186,12 +181,12 @@ namespace UnrealBuildTool
 		/// <param name="ToolPath">The path to add</param>
 		static void AddDirectoryToPath(DirectoryReference ToolPath)
 		{
-            string PathEnvironmentVariable = Environment.GetEnvironmentVariable("PATH") ?? "";
-            if (!PathEnvironmentVariable.Split(';').Any(x => String.Compare(x, ToolPath.FullName, true) == 0))
-            {
-                PathEnvironmentVariable = ToolPath.FullName + ";" + PathEnvironmentVariable;
-                Environment.SetEnvironmentVariable("PATH", PathEnvironmentVariable);
-            }
+			string PathEnvironmentVariable = Environment.GetEnvironmentVariable("PATH") ?? "";
+			if (!PathEnvironmentVariable.Split(';').Any(x => String.Compare(x, ToolPath.FullName, true) == 0))
+			{
+				PathEnvironmentVariable = ToolPath.FullName + ";" + PathEnvironmentVariable;
+				Environment.SetEnvironmentVariable("PATH", PathEnvironmentVariable);
+			}
 		}
 
 		/// <summary>
@@ -271,13 +266,13 @@ namespace UnrealBuildTool
 		virtual protected FileReference GetResourceCompilerToolPath(DirectoryReference WindowsSdkDir, VersionNumber WindowsSdkVersion, ILogger Logger)
 		{
 			FileReference ResourceCompilerPath = FileReference.Combine(WindowsSdkDir, "bin", WindowsSdkVersion.ToString(), "x64", "rc.exe");
-			if(FileReference.Exists(ResourceCompilerPath))
+			if (FileReference.Exists(ResourceCompilerPath))
 			{
 				return ResourceCompilerPath;
 			}
 
 			ResourceCompilerPath = FileReference.Combine(WindowsSdkDir, "bin", "x64", "rc.exe");
-			if(FileReference.Exists(ResourceCompilerPath))
+			if (FileReference.Exists(ResourceCompilerPath))
 			{
 				return ResourceCompilerPath;
 			}
@@ -325,10 +320,10 @@ namespace UnrealBuildTool
 			LibraryPaths.Add(GetToolChainLibsDir());
 
 			// If we're on >= Visual Studio 2015 and using pre-Windows 10 SDK, we need to find a Windows 10 SDK and add the UCRT include paths
-			if(ToolChain.IsMSVC() && WindowsSdkVersion < new VersionNumber(10))
+			if (ToolChain.IsMSVC() && WindowsSdkVersion < new VersionNumber(10))
 			{
 				KeyValuePair<VersionNumber, DirectoryReference> Pair = MicrosoftPlatformSDK.FindUniversalCrtDirs(Logger).OrderByDescending(x => x.Key).FirstOrDefault();
-				if(Pair.Key == null || Pair.Key < new VersionNumber(10))
+				if (Pair.Key == null || Pair.Key < new VersionNumber(10))
 				{
 					throw new BuildException("{0} requires the Universal CRT to be installed.", WindowsPlatform.GetCompilerName(ToolChain));
 				}
@@ -394,14 +389,14 @@ namespace UnrealBuildTool
 		[SupportedOSPlatform("windows")]
 		public static VCEnvironment Create(WindowsCompiler Compiler, WindowsCompiler ToolChain, UnrealTargetPlatform Platform, UnrealArch Architecture, string? CompilerVersion, string? WindowsSdkVersion, string? SuppliedSdkDirectoryForVersion, bool bUseCPPWinRT, bool bAllowClangLinker, ILogger Logger)
 		{
-			return Create( new VCEnvironmentParameters(Compiler, ToolChain, Platform, Architecture, CompilerVersion, WindowsSdkVersion, SuppliedSdkDirectoryForVersion, bUseCPPWinRT, bAllowClangLinker, Logger), Logger );
+			return Create(new VCEnvironmentParameters(Compiler, ToolChain, Platform, Architecture, CompilerVersion, WindowsSdkVersion, SuppliedSdkDirectoryForVersion, bUseCPPWinRT, bAllowClangLinker, Logger), Logger);
 		}
 
 		/// <summary>
 		/// Creates an environment with the given parameters
 		/// </summary>
 		[SupportedOSPlatform("windows")]
-		public static VCEnvironment Create( VCEnvironmentParameters Params, ILogger Logger)
+		public static VCEnvironment Create(VCEnvironmentParameters Params, ILogger Logger)
 		{
 			return new VCEnvironment(Params, Logger);
 		}
@@ -468,15 +463,15 @@ namespace UnrealBuildTool
 		/// <param name="Logger">Logger for output</param>
 		/// <returns>Creation parameters for VC environment</returns>
 		[SupportedOSPlatform("windows")]
-		public VCEnvironmentParameters (WindowsCompiler Compiler, WindowsCompiler ToolChain, UnrealTargetPlatform Platform, UnrealArch Architecture, string? CompilerVersion, string? WindowsSdkVersion, string? SuppliedSdkDirectoryForVersion, bool bUseCPPWinRT, bool bAllowClangLinker, ILogger Logger)
+		public VCEnvironmentParameters(WindowsCompiler Compiler, WindowsCompiler ToolChain, UnrealTargetPlatform Platform, UnrealArch Architecture, string? CompilerVersion, string? WindowsSdkVersion, string? SuppliedSdkDirectoryForVersion, bool bUseCPPWinRT, bool bAllowClangLinker, ILogger Logger)
 		{
 			// Get the compiler version info
 			VersionNumber? SelectedCompilerVersion;
 			DirectoryReference? SelectedCompilerDir;
 			DirectoryReference? SelectedRedistDir;
-			if(!WindowsPlatform.TryGetToolChainDir(Compiler, CompilerVersion, Architecture, Logger, out SelectedCompilerVersion, out SelectedCompilerDir, out SelectedRedistDir))
+			if (!WindowsPlatform.TryGetToolChainDir(Compiler, CompilerVersion, Architecture, Logger, out SelectedCompilerVersion, out SelectedCompilerDir, out SelectedRedistDir))
 			{
-				throw new BuildException("{0}{1} {2} must be installed in order to build this target.", WindowsPlatform.GetCompilerName(Compiler), String.IsNullOrEmpty(CompilerVersion)? "" : String.Format(" ({0})", CompilerVersion), Architecture.ToString());
+				throw new BuildException("{0}{1} {2} must be installed in order to build this target.", WindowsPlatform.GetCompilerName(Compiler), String.IsNullOrEmpty(CompilerVersion) ? "" : String.Format(" ({0})", CompilerVersion), Architecture.ToString());
 			}
 
 			// Get the toolchain info
diff --git a/Engine/Source/Programs/UnrealBuildTool/Platform/Windows/VCSpecificFileAction.cs b/Engine/Source/Programs/UnrealBuildTool/Platform/Windows/VCSpecificFileAction.cs
--- a/Engine/Source/Programs/UnrealBuildTool/Platform/Windows/VCSpecificFileAction.cs
+++ b/Engine/Source/Programs/UnrealBuildTool/Platform/Windows/VCSpecificFileAction.cs
@@ -1,10 +1,9 @@
 // Copyright Epic Games, Inc. All Rights Reserved.
 
-using EpicGames.Core;
-using Microsoft.Extensions.Logging;
 using System.Collections.Generic;
-using System.Diagnostics;
 using System.IO;
+using EpicGames.Core;
+using Microsoft.Extensions.Logging;
 using UnrealBuildBase;
 
 namespace UnrealBuildTool
diff --git a/Engine/Source/Programs/UnrealBuildTool/Platform/Windows/VCToolChain.cs b/Engine/Source/Programs/UnrealBuildTool/Platform/Windows/VCToolChain.cs
--- a/Engine/Source/Programs/UnrealBuildTool/Platform/Windows/VCToolChain.cs
+++ b/Engine/Source/Programs/UnrealBuildTool/Platform/Windows/VCToolChain.cs
@@ -4,12 +4,12 @@ using System;
 using System.Collections.Generic;
 using System.IO;
 using System.Linq;
+using System.Runtime.Versioning;
 using System.Text;
 using EpicGames.Core;
-using UnrealBuildBase;
-using System.Runtime.Versioning;
-using Microsoft.Extensions.Logging;
 using Microsoft.CodeAnalysis;
+using Microsoft.Extensions.Logging;
+using UnrealBuildBase;
 
 namespace UnrealBuildTool
 {
@@ -52,7 +52,7 @@ namespace UnrealBuildTool
 				}
 			}
 		}
-		
+
 		public override FileReference? GetCppCompilerPath()
 		{
 			return EnvVars.CompilerPath;
@@ -84,7 +84,7 @@ namespace UnrealBuildTool
 		/// <returns>String describing the current toolchain</returns>
 		public override void GetVersionInfo(List<string> Lines)
 		{
-			if(EnvVars.Compiler == EnvVars.ToolChain)
+			if (EnvVars.Compiler == EnvVars.ToolChain)
 			{
 				Lines.Add($"Using {WindowsPlatform.GetCompilerName(EnvVars.Compiler)} {EnvVars.ToolChainVersion} toolchain ({EnvVars.ToolChainDir}) and Windows {EnvVars.WindowsSdkVersion} SDK ({EnvVars.WindowsSdkDir}).");
 			}
@@ -300,7 +300,7 @@ namespace UnrealBuildTool
 				// Sync the compatibility version with the MSVC toolchain version (14.xx which maps to advertised
 				// compiler version of 19.xx).
 				Arguments.Add($"-fms-compatibility-version=19.{EnvVars.ToolChainVersion.GetComponent(1)}");
-				
+
 				if (Target.StaticAnalyzer == StaticAnalyzer.Default && CompileEnvironment.PrecompiledHeaderAction != PrecompiledHeaderAction.Create && !CompileEnvironment.bDisableStaticAnalysis)
 				{
 					Arguments.Add("-Wno-unused-command-line-argument");
@@ -439,7 +439,7 @@ namespace UnrealBuildTool
 			Arguments.Add("/c");
 
 			// Put symbols into different sections so the linker can remove them.
-			if(Target.WindowsPlatform.bOptimizeGlobalData)
+			if (Target.WindowsPlatform.bOptimizeGlobalData)
 			{
 				Arguments.Add("/Gw");
 			}
@@ -509,7 +509,7 @@ namespace UnrealBuildTool
 			AddDefinition(Arguments, "_SILENCE_STDEXT_HASH_DEPRECATION_WARNINGS=1");
 
 			// Ignore secure CRT warnings on Clang
-			if(Target.WindowsPlatform.Compiler.IsClang())
+			if (Target.WindowsPlatform.Compiler.IsClang())
 			{
 				AddDefinition(Arguments, "_CRT_SECURE_NO_WARNINGS");
 			}
@@ -634,7 +634,7 @@ namespace UnrealBuildTool
 			//
 			else
 			{
-				if(!CompileEnvironment.bOptimizeCode)
+				if (!CompileEnvironment.bOptimizeCode)
 				{
 					// Disable compiler optimization.
 					Arguments.Add("/Od");
@@ -808,7 +808,7 @@ namespace UnrealBuildTool
 			{
 				Arguments.Add("/wd4996");
 			}
-			else if(CompileEnvironment.DeprecationWarningLevel == WarningLevel.Error)
+			else if (CompileEnvironment.DeprecationWarningLevel == WarningLevel.Error)
 			{
 				Arguments.Add("/we4996");
 			}
@@ -851,19 +851,19 @@ namespace UnrealBuildTool
 			//   4305: 'identifier' : truncation from 'type1' to 'type2'
 			WarningLevel EffectiveCastWarningLevel = (Target.Platform == UnrealTargetPlatform.Win64) ? CompileEnvironment.UnsafeTypeCastWarningLevel : WarningLevel.Off;
 			if (EffectiveCastWarningLevel == WarningLevel.Error)
- 			{
- 				Arguments.Add("/we4244");
+			{
+				Arguments.Add("/we4244");
 				Arguments.Add("/we4838");
- 			}
- 			else if (EffectiveCastWarningLevel == WarningLevel.Warning)
- 			{
+			}
+			else if (EffectiveCastWarningLevel == WarningLevel.Warning)
+			{
 				// Note: The extra 4 is not a typo, /wLXXXX sets warning XXXX to level L
- 				Arguments.Add("/w44244");
+				Arguments.Add("/w44244");
 				Arguments.Add("/w44838");
 			}
 			else
- 			{
- 				Arguments.Add("/wd4244");
+			{
+				Arguments.Add("/wd4244");
 				Arguments.Add("/wd4838");
 			}
 
@@ -1002,7 +1002,7 @@ namespace UnrealBuildTool
 				// Treat all warnings as errors by default
 				Arguments.Add("-Werror");                                   // https://clang.llvm.org/docs/UsersManual.html#cmdoption-werror
 
-				ClangWarnings.GetEnabledWarnings(Arguments); 
+				ClangWarnings.GetEnabledWarnings(Arguments);
 
 				ClangWarnings.GetDisabledWarnings(CompileEnvironment,
 					Target.StaticAnalyzer,
@@ -1238,8 +1238,8 @@ namespace UnrealBuildTool
 			}
 
 			// Enable incremental linking if wanted. ( avoid /INCREMENTAL getting ignored (LNK4075) due to /LTCG, /RELEASE, and /OPT:ICF )
-			if (LinkEnvironment.bUseIncrementalLinking && 
-				LinkEnvironment.Configuration != CppConfiguration.Shipping && 
+			if (LinkEnvironment.bUseIncrementalLinking &&
+				LinkEnvironment.Configuration != CppConfiguration.Shipping &&
 				!Target.WindowsPlatform.bMergeIdenticalCOMDATs &&
 				!LinkEnvironment.bAllowLTCG &&
 				Target.WindowsPlatform.Architecture == UnrealArch.X64)
@@ -1330,12 +1330,12 @@ namespace UnrealBuildTool
 			// Add include paths to the argument list.
 			BaseCompileAction.IncludePaths.AddRange(CompileEnvironment.UserIncludePaths);
 			BaseCompileAction.SystemIncludePaths.AddRange(CompileEnvironment.SystemIncludePaths);
-			
+
 			if (!CompileEnvironment.bHasSharedResponseFile)
 			{
 				BaseCompileAction.SystemIncludePaths.AddRange(EnvVars.IncludePaths);
 			}
-			
+
 			// Remember the architecture
 			BaseCompileAction.Architecture = CompileEnvironment.Architecture;
 
@@ -1630,7 +1630,7 @@ namespace UnrealBuildTool
 				}
 
 				// Allow derived toolchains to make further changes
-				ModifyFinalCompileAction(CompileAction, CompileEnvironment, SourceFile, OutputDir, ModuleName );
+				ModifyFinalCompileAction(CompileAction, CompileEnvironment, SourceFile, OutputDir, ModuleName);
 
 				if (!ProjectFileGenerator.bGenerateProjectFiles)
 				{
@@ -1842,11 +1842,11 @@ namespace UnrealBuildTool
 					"pgosweep.exe",
 				};
 				foreach (string PGOFile in PGOFiles)
-				{			
+				{
 					FileReference SrcFile = FileReference.Combine(EnvVars.ToolChainDir, "bin", "Hostx64", "x64", PGOFile);
 					FileReference DstFile = FileReference.Combine(ExeDir, PGOFile);
 					TargetFileToSourceFile[DstFile] = SrcFile;
-					RuntimeDependencies.Add( new RuntimeDependency(DstFile, StagedFileType.NonUFS));
+					RuntimeDependencies.Add(new RuntimeDependency(DstFile, StagedFileType.NonUFS));
 				}
 			}
 		}
@@ -1881,8 +1881,8 @@ namespace UnrealBuildTool
 
 				// Resource tool can run remotely if possible
 				CompileAction.bCanExecuteRemotely = true;
-				CompileAction.bCanExecuteRemotelyWithSNDBS = false;	// no tool template for SN-DBS results in warnings
-			
+				CompileAction.bCanExecuteRemotelyWithSNDBS = false; // no tool template for SN-DBS results in warnings
+
 				List<string> Arguments = new List<string>();
 
 				// Suppress header spew
@@ -1922,7 +1922,7 @@ namespace UnrealBuildTool
 
 				// Figure the icon to use. We can only use a custom icon when compiling to a project-specific intermediate directory (and not for the shared editor executable, for example).
 				FileReference IconFile;
-				if(Target.ProjectFile != null && !CompileEnvironment.bUseSharedBuildEnvironment)
+				if (Target.ProjectFile != null && !CompileEnvironment.bUseSharedBuildEnvironment)
 				{
 					IconFile = GetApplicationIcon(Target.ProjectFile);
 				}
@@ -1936,7 +1936,7 @@ namespace UnrealBuildTool
 				AddDefinition(Arguments, $"BUILD_ICON_FILE_NAME=\"\\\"{NormalizeCommandLinePath(IconFile).Replace("\\", "\\\\")}\\\"\"");
 
 				// Apply the target settings for the resources
-				if(!CompileEnvironment.bUseSharedBuildEnvironment)
+				if (!CompileEnvironment.bUseSharedBuildEnvironment)
 				{
 					if (!String.IsNullOrEmpty(Target.WindowsPlatform.CompanyName))
 					{
@@ -2141,9 +2141,9 @@ namespace UnrealBuildTool
 		static string SanitizeMacroValue(string Value)
 		{
 			StringBuilder Result = new StringBuilder(Value.Length);
-			for(int Idx = 0; Idx < Value.Length; Idx++)
+			for (int Idx = 0; Idx < Value.Length; Idx++)
 			{
-				if(Value[Idx] != '\'' && Value[Idx] != '\"')
+				if (Value[Idx] != '\'' && Value[Idx] != '\"')
 				{
 					Result.Append(Value[Idx]);
 				}
@@ -2216,7 +2216,7 @@ namespace UnrealBuildTool
 			}
 
 			// Set up the library paths for linking this binary
-			if(bBuildImportLibraryOnly)
+			if (bBuildImportLibraryOnly)
 			{
 				// When building an import library, ignore all the libraries included via embedded #pragma lib declarations.
 				// We shouldn't need them to generate exports.
@@ -2244,8 +2244,8 @@ namespace UnrealBuildTool
 			// If we're building either an executable or a DLL, make sure we link in the 
 			// correct address sanitizer helper libs.
 			// Note: As of MSVC 16.9, this is automatically done if the /fsanitize=address flag is used.
-			if (!bBuildImportLibraryOnly && !LinkEnvironment.bIsBuildingLibrary && Target.WindowsPlatform.bEnableAddressSanitizer && 
-			    EnvVars.CompilerVersion < new VersionNumber(14, 28, 0))
+			if (!bBuildImportLibraryOnly && !LinkEnvironment.bIsBuildingLibrary && Target.WindowsPlatform.bEnableAddressSanitizer &&
+				EnvVars.CompilerVersion < new VersionNumber(14, 28, 0))
 			{
 				String ASanArchSuffix = "";
 				if (EnvVars.Architecture == UnrealArch.X64)
@@ -2283,7 +2283,7 @@ namespace UnrealBuildTool
 			}
 
 			// Enable function level hot-patching
-			if(!bBuildImportLibraryOnly && Target.WindowsPlatform.bCreateHotpatchableImage)
+			if (!bBuildImportLibraryOnly && Target.WindowsPlatform.bCreateHotpatchableImage)
 			{
 				Arguments.Add("/FUNCTIONPADMIN:6"); // For some reason, not providing the number causes full linking to happen all the time
 			}
@@ -2298,7 +2298,7 @@ namespace UnrealBuildTool
 			{
 				ImportLibraryFilePath = FileReference.Combine(LinkEnvironment.IntermediateDirectory!, LinkEnvironment.OutputFilePath.GetFileNameWithoutExtension() + ".suppressed.lib");
 			}
-			else if(Target.bShouldCompileAsDLL)
+			else if (Target.bShouldCompileAsDLL)
 			{
 				ImportLibraryFilePath = FileReference.Combine(LinkEnvironment.OutputDirectory!, LinkEnvironment.OutputFilePath.GetFileNameWithoutExtension() + ".lib");
 			}
@@ -2341,11 +2341,11 @@ namespace UnrealBuildTool
 				{
 					InputFileNames.Add(string.Format("\"{0}\"", SystemLibrary));
 				}
-				
+
 				foreach (FileItem NatvisFile in LinkEnvironment.DebuggerVisualizerFiles)
 				{
 					PrerequisiteItems.Add(NatvisFile);
-					Arguments.Add(string.Format("/NATVIS:\"{0}\"", 
+					Arguments.Add(string.Format("/NATVIS:\"{0}\"",
 						NormalizeCommandLinePath(NatvisFile)));
 				}
 
@@ -2363,7 +2363,7 @@ namespace UnrealBuildTool
 			// For import libraries and exports generated by cross-referenced builds, we don't track output files. VS 15.3+ doesn't touch timestamps for libs
 			// and exp files with no modifications, breaking our dependency checking, but incremental linking will fall back to a full link if we delete it.
 			// Since all DLLs are typically marked as cross referenced now anyway, we can just ignore this file to allow incremental linking to work.
-			if(LinkEnvironment.bHasExports && !LinkEnvironment.bIsBuildingLibrary && !LinkEnvironment.bIsCrossReferenced)
+			if (LinkEnvironment.bHasExports && !LinkEnvironment.bIsBuildingLibrary && !LinkEnvironment.bIsCrossReferenced)
 			{
 				FileReference ExportFilePath = ImportLibraryFilePath.ChangeExtension(".exp");
 				FileItem ExportFile = FileItem.GetItemByFileReference(ExportFilePath);
@@ -2380,7 +2380,7 @@ namespace UnrealBuildTool
 					Arguments.Add($"/IMPLIB:\"{NormalizeCommandLinePath(ImportLibraryFilePath)}\"");
 
 					// Like the export file above, don't add the import library as a produced item when it's cross referenced.
-					if(!LinkEnvironment.bIsCrossReferenced)
+					if (!LinkEnvironment.bIsCrossReferenced)
 					{
 						ProducedItems.Add(ImportLibraryFile);
 					}
@@ -2410,14 +2410,14 @@ namespace UnrealBuildTool
 				}
 
 				// Add the additional arguments specified by the environment.
-				if(!String.IsNullOrEmpty(LinkEnvironment.AdditionalArguments))
+				if (!String.IsNullOrEmpty(LinkEnvironment.AdditionalArguments))
 				{
 					Arguments.Add(LinkEnvironment.AdditionalArguments.Trim());
 				}
 			}
 
 			// Add any forced references to functions
-			foreach(string IncludeFunction in LinkEnvironment.IncludeFunctions)
+			foreach (string IncludeFunction in LinkEnvironment.IncludeFunctions)
 			{
 				Arguments.Add($"/INCLUDE:{IncludeFunction}");
 			}
@@ -2438,7 +2438,7 @@ namespace UnrealBuildTool
 			string ReadableArch = UnrealArchitectureConfig.ForPlatform(LinkEnvironment.Platform).ConvertToReadableArchitecture(LinkEnvironment.Architecture);
 			LinkAction.CommandDescription += $"Link [{ReadableArch}]";
 			LinkAction.WorkingDirectory = Unreal.EngineSourceDirectory;
-			if(bIsBuildingLibraryOrImportLibrary)
+			if (bIsBuildingLibraryOrImportLibrary)
 			{
 				LinkAction.CommandPath = EnvVars.LibraryManagerPath;
 			}
@@ -2494,7 +2494,7 @@ namespace UnrealBuildTool
 				Logger.LogInformation("...copying the profile guided optimization files to output directory...");
 
 				// prefer a PGD file that matches the output file
-				string PGDFile = Path.Combine( LinkEnvironment.PGODirectory!, LinkEnvironment.PGOFilenamePrefix + ".pgd");
+				string PGDFile = Path.Combine(LinkEnvironment.PGODirectory!, LinkEnvironment.PGOFilenamePrefix + ".pgd");
 				if (!File.Exists(PGDFile))
 				{
 					string[] PGDFiles = Directory.GetFiles(LinkEnvironment.PGODirectory!, "*.pgd");
@@ -2576,7 +2576,7 @@ namespace UnrealBuildTool
 					Arguments.Add("/FASTGENPROFILE");
 					Log.TraceInformationOnce("Enabling Fastgen Profile Guided Optimization (PGO). Linking will take a while.");
 				}
-				else 
+				else
 				{
 					if (Target.WindowsPlatform.bPGONoExtraCounters)
 					{
@@ -2587,7 +2587,7 @@ namespace UnrealBuildTool
 					{
 						Log.TraceInformationOnce("Enabling Profile Guided Optimization (PGO). Linking will take a while.");
 						Arguments.Add("/GENPROFILE");
-					}	
+					}
 				}
 			}
 		}
@@ -2604,19 +2604,19 @@ namespace UnrealBuildTool
 		{
 			// Write the list of object file directories
 			HashSet<DirectoryReference> ObjectFileDirectories = new HashSet<DirectoryReference>();
-			foreach(FileItem InputFile in LinkEnvironment.InputFiles)
+			foreach (FileItem InputFile in LinkEnvironment.InputFiles)
 			{
 				ObjectFileDirectories.Add(InputFile.Location.Directory);
 			}
-			foreach(FileReference Library in LinkEnvironment.Libraries)
+			foreach (FileReference Library in LinkEnvironment.Libraries)
 			{
 				ObjectFileDirectories.Add(Library.Directory);
 			}
-			foreach(DirectoryReference LibraryPath in LinkEnvironment.SystemLibraryPaths)
+			foreach (DirectoryReference LibraryPath in LinkEnvironment.SystemLibraryPaths)
 			{
 				ObjectFileDirectories.Add(LibraryPath);
 			}
-			foreach(string LibraryPath in (Environment.GetEnvironmentVariable("LIB") ?? "").Split(new char[]{ ';' }, StringSplitOptions.RemoveEmptyEntries))
+			foreach (string LibraryPath in (Environment.GetEnvironmentVariable("LIB") ?? "").Split(new char[] { ';' }, StringSplitOptions.RemoveEmptyEntries))
 			{
 				ObjectFileDirectories.Add(new DirectoryReference(LibraryPath));
 			}
@@ -2642,7 +2642,7 @@ namespace UnrealBuildTool
 			// extracted to populate the project file's list of include paths
 			// @todo projectfiles: Should we only do this for VC++ platforms?
 			StringBuilder IncludePaths = new StringBuilder();
-			foreach(DirectoryReference IncludePath in EnvVars.IncludePaths)
+			foreach (DirectoryReference IncludePath in EnvVars.IncludePaths)
 			{
 				IncludePaths.AppendFormat("{0};", IncludePath);
 			}
@@ -2653,7 +2653,7 @@ namespace UnrealBuildTool
 		{
 			if (Binary.Type == UEBuildBinaryType.DynamicLinkLibrary)
 			{
-				if(Target.bShouldCompileAsDLL)
+				if (Target.bShouldCompileAsDLL)
 				{
 					BuildProducts.Add(FileReference.Combine(Binary.OutputDir, Binary.OutputFilePath.GetFileNameWithoutExtension() + ".lib"), BuildProductType.BuildResource);
 				}
@@ -2662,9 +2662,9 @@ namespace UnrealBuildTool
 					BuildProducts.Add(FileReference.Combine(Binary.IntermediateDirectory, Binary.OutputFilePath.GetFileNameWithoutExtension() + ".lib"), BuildProductType.BuildResource);
 				}
 			}
-			if(Binary.Type == UEBuildBinaryType.Executable && Target.bCreateMapFile)
+			if (Binary.Type == UEBuildBinaryType.Executable && Target.bCreateMapFile)
 			{
-				foreach(FileReference OutputFilePath in Binary.OutputFilePaths)
+				foreach (FileReference OutputFilePath in Binary.OutputFilePaths)
 				{
 					BuildProducts.Add(FileReference.Combine(OutputFilePath.Directory, OutputFilePath.GetFileNameWithoutExtension() + ".map"), BuildProductType.MapFile);
 					BuildProducts.Add(FileReference.Combine(OutputFilePath.Directory, OutputFilePath.GetFileNameWithoutExtension() + ".objpaths"), BuildProductType.MapFile);
diff --git a/Engine/Source/Programs/UnrealBuildTool/Platform/Windows/WindowsExports.cs b/Engine/Source/Programs/UnrealBuildTool/Platform/Windows/WindowsExports.cs
--- a/Engine/Source/Programs/UnrealBuildTool/Platform/Windows/WindowsExports.cs
+++ b/Engine/Source/Programs/UnrealBuildTool/Platform/Windows/WindowsExports.cs
@@ -5,8 +5,6 @@ using System.Collections.Generic;
 using System.Diagnostics.CodeAnalysis;
 using System.Linq;
 using System.Runtime.Versioning;
-using System.Text;
-using System.Threading.Tasks;
 using EpicGames.Core;
 using Microsoft.Extensions.Logging;
 
@@ -58,7 +56,7 @@ namespace UnrealBuildTool
 		public static bool TryGetWindowsSdkDir(string DesiredVersion, [NotNullWhen(true)] out Version? OutSdkVersion, [NotNullWhen(true)] out DirectoryReference? OutSdkDir)
 		{
 			VersionNumber? vn;
-			if(WindowsPlatform.TryGetWindowsSdkDir(DesiredVersion, Log.Logger, out vn, out OutSdkDir))
+			if (WindowsPlatform.TryGetWindowsSdkDir(DesiredVersion, Log.Logger, out vn, out OutSdkDir))
 			{
 				OutSdkVersion = new Version(vn.ToString());
 				return true;
@@ -86,9 +84,9 @@ namespace UnrealBuildTool
 
 			// Add all the other directories sorted in reverse order
 			IReadOnlyDictionary<VersionNumber, DirectoryReference> WindowsSdkDirPairs = MicrosoftPlatformSDK.FindWindowsSdkDirs(Log.Logger);
-			foreach(KeyValuePair<VersionNumber, DirectoryReference> Pair in WindowsSdkDirPairs.OrderByDescending(x => x.Key))
+			foreach (KeyValuePair<VersionNumber, DirectoryReference> Pair in WindowsSdkDirPairs.OrderByDescending(x => x.Key))
 			{
-				if(!WindowsSdkDirs.Any(x => x.Value == Pair.Value))
+				if (!WindowsSdkDirs.Any(x => x.Value == Pair.Value))
 				{
 					WindowsSdkDirs.Add(new KeyValuePair<string, DirectoryReference>(Pair.Key.ToString(), Pair.Value));
 				}
diff --git a/Engine/Source/Programs/UnrealBuildTool/Platform/Windows/WindowsProjectGenerator.cs b/Engine/Source/Programs/UnrealBuildTool/Platform/Windows/WindowsProjectGenerator.cs
--- a/Engine/Source/Programs/UnrealBuildTool/Platform/Windows/WindowsProjectGenerator.cs
+++ b/Engine/Source/Programs/UnrealBuildTool/Platform/Windows/WindowsProjectGenerator.cs
@@ -1,10 +1,7 @@
 // Copyright Epic Games, Inc. All Rights Reserved.
 
-using System;
 using System.Collections.Generic;
 using System.Text;
-using System.Diagnostics;
-using System.IO;
 using EpicGames.Core;
 using Microsoft.Extensions.Logging;
 
diff --git a/Engine/Source/Programs/UnrealBuildTool/ProjectFiles/CLion/CLionGenerator.cs b/Engine/Source/Programs/UnrealBuildTool/ProjectFiles/CLion/CLionGenerator.cs
--- a/Engine/Source/Programs/UnrealBuildTool/ProjectFiles/CLion/CLionGenerator.cs
+++ b/Engine/Source/Programs/UnrealBuildTool/ProjectFiles/CLion/CLionGenerator.cs
@@ -1,9 +1,5 @@
 // Copyright Epic Games, Inc. All Rights Reserved.
 
-using System;
-using System.Collections.Generic;
-using System.Text;
-using System.IO;
 using EpicGames.Core;
 
 namespace UnrealBuildTool
diff --git a/Engine/Source/Programs/UnrealBuildTool/ProjectFiles/CodeLite/CodeLiteGenerator.cs b/Engine/Source/Programs/UnrealBuildTool/ProjectFiles/CodeLite/CodeLiteGenerator.cs
--- a/Engine/Source/Programs/UnrealBuildTool/ProjectFiles/CodeLite/CodeLiteGenerator.cs
+++ b/Engine/Source/Programs/UnrealBuildTool/ProjectFiles/CodeLite/CodeLiteGenerator.cs
@@ -2,12 +2,11 @@
 
 using System;
 using System.Collections.Generic;
-using System.Text;
 using System.IO;
+using System.Linq;
 using System.Xml;
 using System.Xml.Linq;
 using EpicGames.Core;
-using System.Linq;
 using Microsoft.Extensions.Logging;
 
 namespace UnrealBuildTool
@@ -28,7 +27,7 @@ namespace UnrealBuildTool
 		public CodeLiteGenerator(FileReference? InOnlyGameProject, CommandLineArguments CommandLine)
 			: base(InOnlyGameProject)
 		{
-			if(CommandLine.HasOption("-cl10"))
+			if (CommandLine.HasOption("-cl10"))
 			{
 				ProjectFileFormat = CodeliteProjectFileFormat.CodeLite10;
 			}
@@ -170,9 +169,9 @@ namespace UnrealBuildTool
 					if (ProjectFileFormat == CodeliteProjectFileFormat.CodeLite10)
 					{
 						if ((CurrentTarget.TargetRules!.Type == TargetType.Client) ||
-						    (CurrentTarget.TargetRules.Type == TargetType.Server) ||
-						    (CurrentTarget.TargetRules.Type == TargetType.Editor) ||
-						    (CurrentTarget.TargetRules.Type == TargetType.Game))
+							(CurrentTarget.TargetRules.Type == TargetType.Server) ||
+							(CurrentTarget.TargetRules.Type == TargetType.Editor) ||
+							(CurrentTarget.TargetRules.Type == TargetType.Game))
 						{
 							if (ProjectName.Equals("UnrealClient") ||
 								ProjectName.Equals("UnrealServer") ||
@@ -217,7 +216,7 @@ namespace UnrealBuildTool
 			}
 			if (ProjectFileFormat == CodeliteProjectFileFormat.CodeLite10)
 			{
-				
+
 				if (CodeLiteWorkspaceTargetEngine != null)
 				{
 					CodeLiteWorkspace.Add(CodeLiteWorkspaceTargetEngine);
diff --git a/Engine/Source/Programs/UnrealBuildTool/ProjectFiles/CodeLite/CodeLiteProject.cs b/Engine/Source/Programs/UnrealBuildTool/ProjectFiles/CodeLite/CodeLiteProject.cs
--- a/Engine/Source/Programs/UnrealBuildTool/ProjectFiles/CodeLite/CodeLiteProject.cs
+++ b/Engine/Source/Programs/UnrealBuildTool/ProjectFiles/CodeLite/CodeLiteProject.cs
@@ -2,13 +2,11 @@
 
 using System;
 using System.Collections.Generic;
-using System.Text;
 using System.IO;
-using System.Xml;
 using System.Xml.Linq;
 using EpicGames.Core;
-using UnrealBuildBase;
 using Microsoft.Extensions.Logging;
+using UnrealBuildBase;
 
 namespace UnrealBuildTool
 {
@@ -20,11 +18,11 @@ namespace UnrealBuildTool
 		{
 			OnlyGameProject = InOnlyGameProject;
 		}
-		
+
 		// Check if the XElement is empty.
 		bool IsEmpty(IEnumerable<XElement> en)
 		{
-			foreach(XElement c in en) { return false; }
+			foreach (XElement c in en) { return false; }
 			return true;
 		}
 
@@ -65,7 +63,7 @@ namespace UnrealBuildTool
 			// Create the folder where the project files goes if it does not exist
 			//
 			String FilePath = Path.GetDirectoryName(ProjectFilePath.FullName)!;
-			if( (FilePath.Length > 0) && !Directory.Exists(FilePath))
+			if ((FilePath.Length > 0) && !Directory.Exists(FilePath))
 			{
 				Directory.CreateDirectory(FilePath);
 			}
@@ -79,11 +77,11 @@ namespace UnrealBuildTool
 			//
 			// Write all targets which will be separate projects.
 			//
-			foreach (Project target in ProjectTargets) 
+			foreach (Project target in ProjectTargets)
 			{
 				string[] tmp = target.ToString()!.Split('.');
-				string ProjectTargetFileName = Path.GetDirectoryName (ProjectFilePath.FullName) + "/" + tmp [0] +  ProjectExtension;
-				String TargetName = tmp [0];
+				string ProjectTargetFileName = Path.GetDirectoryName(ProjectFilePath.FullName) + "/" + tmp[0] + ProjectExtension;
+				String TargetName = tmp[0];
 				TargetType ProjectTargetType = target.TargetRules!.Type;
 
 				//
@@ -112,7 +110,7 @@ namespace UnrealBuildTool
 				//
 				// Find/Create the correct virtual folder and place the file into it.
 				//
-				foreach(SourceFile CurrentFile in FilterSourceFile)
+				foreach (SourceFile CurrentFile in FilterSourceFile)
 				{
 					//
 					// Try to get the correct relative folder representation for the project.
@@ -122,7 +120,7 @@ namespace UnrealBuildTool
 					if ((ProjectTargetType == TargetType.Client) ||
 						(ProjectTargetType == TargetType.Editor) ||
 						(ProjectTargetType == TargetType.Game) ||
-						(ProjectTargetType == TargetType.Server) )
+						(ProjectTargetType == TargetType.Server))
 					{
 						if (TargetName.Equals("UnrealClient") ||
 							TargetName.Equals("UnrealServer") ||
@@ -148,7 +146,7 @@ namespace UnrealBuildTool
 					}
 
 					char[] Delimiters = new char[] { '/', '\\' };
-					string [] SplitFolders = CurrentFilePath.Split(Delimiters, StringSplitOptions.RemoveEmptyEntries);
+					string[] SplitFolders = CurrentFilePath.Split(Delimiters, StringSplitOptions.RemoveEmptyEntries);
 					//
 					// Set the CodeLite root folder again.
 					//
@@ -164,7 +162,7 @@ namespace UnrealBuildTool
 						{
 							continue;
 						}
-							
+
 						//
 						// Let's look if there is a virtual folder withint the current XElement.
 						//
@@ -221,9 +219,9 @@ namespace UnrealBuildTool
 					// If we are at this point we found the correct XElement folder
 					//
 					XElement file = new XElement("File");
-					XAttribute fileAttribute = new XAttribute("Name",  CurrentFile.Reference.FullName);
+					XAttribute fileAttribute = new XAttribute("Name", CurrentFile.Reference.FullName);
 					file.Add(fileAttribute);
-					root.Add(file); 
+					root.Add(file);
 				}
 
 				XElement CodeLiteSettings = new XElement("Settings");
@@ -242,16 +240,16 @@ namespace UnrealBuildTool
 					// Create Configuration General part. 
 					//
 					XElement CodeLiteConfigurationGeneral = new XElement("General");
-						
+
 					//
 					// Create the executable filename.
 					//
 					string ExecutableToRun = "";
-					string PlatformConfiguration = "-" + ProjectPlatformName + "-" + CurConf.ToString ();
+					string PlatformConfiguration = "-" + ProjectPlatformName + "-" + CurConf.ToString();
 					if (BuildHostPlatform.Current.Platform == UnrealTargetPlatform.Linux)
 					{
 						ExecutableToRun = "./" + TargetName;
-						if ((ProjectTargetType == TargetType.Game) || 
+						if ((ProjectTargetType == TargetType.Game) ||
 							(ProjectTargetType == TargetType.Program))
 						{
 							if (CurConf != UnrealTargetConfiguration.Development)
@@ -262,8 +260,8 @@ namespace UnrealBuildTool
 						else if (ProjectTargetType == TargetType.Editor)
 						{
 							ExecutableToRun = "./UnrealEditor";
-							if ((CurConf == UnrealTargetConfiguration.Debug) || 
-								(CurConf == UnrealTargetConfiguration.Shipping) || 
+							if ((CurConf == UnrealTargetConfiguration.Debug) ||
+								(CurConf == UnrealTargetConfiguration.Shipping) ||
 								(CurConf == UnrealTargetConfiguration.Test))
 							{
 								ExecutableToRun += PlatformConfiguration;
@@ -275,7 +273,7 @@ namespace UnrealBuildTool
 					{
 						ExecutableToRun = "./" + TargetName;
 						if ((ProjectTargetType == TargetType.Game) || (ProjectTargetType == TargetType.Program))
-						{			
+						{
 							if (CurConf != UnrealTargetConfiguration.Development)
 							{
 								ExecutableToRun += PlatformConfiguration;
@@ -290,8 +288,8 @@ namespace UnrealBuildTool
 						else if (ProjectTargetType == TargetType.Editor)
 						{
 							ExecutableToRun = "./UnrealEditor";
-							if ((CurConf == UnrealTargetConfiguration.Debug) || 
-								(CurConf == UnrealTargetConfiguration.Shipping) || 
+							if ((CurConf == UnrealTargetConfiguration.Debug) ||
+								(CurConf == UnrealTargetConfiguration.Shipping) ||
 								(CurConf == UnrealTargetConfiguration.Test))
 							{
 								ExecutableToRun += PlatformConfiguration;
@@ -303,7 +301,7 @@ namespace UnrealBuildTool
 							}
 						}
 
-					} 
+					}
 					else if (BuildHostPlatform.Current.Platform == UnrealTargetPlatform.Win64)
 					{
 						ExecutableToRun = TargetName;
@@ -317,8 +315,8 @@ namespace UnrealBuildTool
 						else if (ProjectTargetType == TargetType.Editor)
 						{
 							ExecutableToRun = "UnrealEditor";
-							if ((CurConf == UnrealTargetConfiguration.Debug) || 
-								(CurConf == UnrealTargetConfiguration.Shipping) || 
+							if ((CurConf == UnrealTargetConfiguration.Debug) ||
+								(CurConf == UnrealTargetConfiguration.Shipping) ||
 								(CurConf == UnrealTargetConfiguration.Test))
 							{
 								ExecutableToRun += PlatformConfiguration;
@@ -326,33 +324,36 @@ namespace UnrealBuildTool
 						}
 
 						ExecutableToRun += ".exe";
-								
+
 					}
 					else
-					{ 
-							throw new BuildException("Unsupported platform.");
+					{
+						throw new BuildException("Unsupported platform.");
 					}
-						
-						
+
+
 					// Is this project a Game type?
 					XAttribute GeneralExecutableToRun = new XAttribute("Command", ExecutableToRun);
-					if (ProjectTargetType == TargetType.Game) 
+					if (ProjectTargetType == TargetType.Game)
 					{
-						if (CurConf.ToString ().Contains ("Debug")) 
+						if (CurConf.ToString().Contains("Debug"))
 						{
 							string commandArguments = " -debug";
 							XAttribute GeneralExecutableToRunArguments = new XAttribute("CommandArguments", commandArguments);
 							CodeLiteConfigurationGeneral.Add(GeneralExecutableToRunArguments);
 						}
-					if (TargetName.Equals ("UnrealGame")) {
+						if (TargetName.Equals("UnrealGame"))
+						{
 							XAttribute GeneralExecutableWorkingDirectory = new XAttribute("WorkingDirectory", UnrealEngineEditorWorkingDirectory);
 							CodeLiteConfigurationGeneral.Add(GeneralExecutableWorkingDirectory);
-						} else {
+						}
+						else
+						{
 							XAttribute GeneralExecutableWorkingDirectory = new XAttribute("WorkingDirectory", GameWorkingDirectory);
 							CodeLiteConfigurationGeneral.Add(GeneralExecutableWorkingDirectory);
 						}
-					} 
-					else if (ProjectTargetType == TargetType.Editor) 
+					}
+					else if (ProjectTargetType == TargetType.Editor)
 					{
 						if (TargetName != "UnrealEditor" && GameProjectFile != "")
 						{
@@ -362,18 +363,18 @@ namespace UnrealBuildTool
 						}
 						XAttribute WorkingDirectory = new XAttribute("WorkingDirectory", UnrealEngineEditorWorkingDirectory);
 						CodeLiteConfigurationGeneral.Add(WorkingDirectory);
-					} 
-					else if (ProjectTargetType == TargetType.Program) 
+					}
+					else if (ProjectTargetType == TargetType.Program)
 					{
 						XAttribute WorkingDirectory = new XAttribute("WorkingDirectory", UnrealEngineEditorWorkingDirectory);
 						CodeLiteConfigurationGeneral.Add(WorkingDirectory);
-					} 
-					else if (ProjectTargetType == TargetType.Client) 
+					}
+					else if (ProjectTargetType == TargetType.Client)
 					{
 						XAttribute WorkingDirectory = new XAttribute("WorkingDirectory", UnrealEngineEditorWorkingDirectory);
 						CodeLiteConfigurationGeneral.Add(WorkingDirectory);
 					}
-					else if (ProjectTargetType == TargetType.Server) 
+					else if (ProjectTargetType == TargetType.Server)
 					{
 						XAttribute WorkingDirectory = new XAttribute("WorkingDirectory", UnrealEngineEditorWorkingDirectory);
 						CodeLiteConfigurationGeneral.Add(WorkingDirectory);
@@ -393,7 +394,7 @@ namespace UnrealBuildTool
 					CodeLiteConfiguration.Add(CodeLiteConfigurationGeneral);
 					XAttribute CodeLiteConfigurationCustomBuildEnabled = new XAttribute("Enabled", "yes");
 					CodeLiteConfigurationCustomBuild.Add(CodeLiteConfigurationCustomBuildEnabled);
-						
+
 					//
 					// Add the working directory for the custom build commands.
 					//
@@ -477,7 +478,7 @@ namespace UnrealBuildTool
 					//
 					// Some other fun Custom Targets.
 					//
-					if (ProjectTargetType == TargetType.Game) 
+					if (ProjectTargetType == TargetType.Game)
 					{
 						string CookGameCommandLine = "mono AutomationTool.exe BuildCookRun ";
 
@@ -489,7 +490,7 @@ namespace UnrealBuildTool
 
 						// Disables Perforce functionality 
 						CookGameCommandLine += "-noP4 ";
-							
+
 						// Do not kill any spawned processes on exit
 						CookGameCommandLine += "-nokill ";
 						CookGameCommandLine += "-clientconfig=" + CurConf.ToString() + " ";
@@ -498,7 +499,7 @@ namespace UnrealBuildTool
 						CookGameCommandLine += "-targetplatform=" + ProjectPlatformName + " "; // TODO Maybe I can add all the supported one.
 						CookGameCommandLine += "-nocompile ";
 						CookGameCommandLine += "-compressed -stage -deploy";
-							
+
 						//
 						// Cook Game.
 						//
diff --git a/Engine/Source/Programs/UnrealBuildTool/ProjectFiles/Eddie/EddieProject.cs b/Engine/Source/Programs/UnrealBuildTool/ProjectFiles/Eddie/EddieProject.cs
--- a/Engine/Source/Programs/UnrealBuildTool/ProjectFiles/Eddie/EddieProject.cs
+++ b/Engine/Source/Programs/UnrealBuildTool/ProjectFiles/Eddie/EddieProject.cs
@@ -1,12 +1,7 @@
 // Copyright Epic Games, Inc. All Rights Reserved.
 
-using System;
 using System.Collections.Generic;
 using System.IO;
-using System.Xml;
-using System.Xml.XPath;
-using System.Xml.Linq;
-using System.Linq;
 using System.Text;
 using EpicGames.Core;
 using Microsoft.Extensions.Logging;
@@ -24,7 +19,7 @@ namespace UnrealBuildTool
 		{
 		}
 	}
-	
+
 	class EddieFolder
 	{
 		public EddieFolder(string InName, string InPath)
@@ -34,7 +29,7 @@ namespace UnrealBuildTool
 			FullPath = null;
 			bIsModuleFolder = false;
 		}
-		
+
 		public string Name;
 		public string Path;
 		public string? FullPath;
@@ -43,11 +38,11 @@ namespace UnrealBuildTool
 		public Dictionary<string, EddieFolder> Folders = new Dictionary<string, EddieFolder>();
 		public List<EddieSourceFile> Files = new List<EddieSourceFile>();
 	}
-	
+
 	class EddieProjectFile : ProjectFile
 	{
 		Dictionary<string, EddieFolder> Folders = new Dictionary<string, EddieFolder>();
-		
+
 		public EddieProjectFile(FileReference InitFilePath, DirectoryReference BaseDir)
 			: base(InitFilePath, BaseDir)
 		{
@@ -57,12 +52,12 @@ namespace UnrealBuildTool
 		{
 			return ProjectFilePath.GetFileNameWithoutExtension();
 		}
-	
+
 		private bool IsSourceCode(string Extension)
 		{
 			return Extension == ".c" || Extension == ".cc" || Extension == ".cpp" || Extension == ".m" || Extension == ".mm" || Extension == ".cs";
 		}
-		
+
 		public override SourceFile? AllocSourceFile(FileReference InitFilePath, DirectoryReference? InitProjectSubFolder)
 		{
 			if (InitFilePath.GetFileName().StartsWith("."))
@@ -71,7 +66,7 @@ namespace UnrealBuildTool
 			}
 			return new EddieSourceFile(InitFilePath, InitProjectSubFolder);
 		}
-		
+
 		public EddieFolder? FindFolderByRelativePath(ref Dictionary<string, EddieFolder> Groups, string RelativePath)
 		{
 			string[] Parts = RelativePath.Split(Path.DirectorySeparatorChar);
@@ -109,12 +104,12 @@ namespace UnrealBuildTool
 
 			return null;
 		}
-		
+
 		private static string ConvertPath(string InPath)
 		{
 			return InPath.Replace("\\", "/");
 		}
-		
+
 		private void ParseSourceFilesIntoGroups()
 		{
 			foreach (SourceFile CurSourceFile in SourceFiles)
@@ -124,11 +119,11 @@ namespace UnrealBuildTool
 				//string FileExtension = Path.GetExtension(FileName);
 				//string FilePath = SourceFile.Reference.MakeRelativeTo(ProjectFilePath.Directory);
 				//string FilePathMac = Utils.CleanDirectorySeparators(FilePath, '/');
-				
+
 				string ProjectRelativeSourceFile = CurSourceFile.Reference.MakeRelativeTo(ProjectFilePath.Directory);
 				string RelativeSourceDirectory = Path.GetDirectoryName(ProjectRelativeSourceFile)!;
 				// Use the specified relative base folder
-				if (CurSourceFile.BaseFolder != null)	// NOTE: We are looking for null strings, not empty strings!
+				if (CurSourceFile.BaseFolder != null)   // NOTE: We are looking for null strings, not empty strings!
 				{
 					RelativeSourceDirectory = Path.GetDirectoryName(CurSourceFile.Reference.MakeRelativeTo(CurSourceFile.BaseFolder))!;
 				}
@@ -140,65 +135,65 @@ namespace UnrealBuildTool
 						Group.bIsModuleFolder = true;
 						Group.WorksetPath = ProjectFilePath.FullName + "." + Group.Name + ".wkst";
 					}
-				
+
 					Group.FullPath = Path.GetDirectoryName(SourceFile.Reference.FullName);
 					Group.Files.Add(SourceFile);
 				}
 			}
 		}
-		
+
 		private void EmitProject(StringBuilder Content, Dictionary<string, EddieFolder> Folders, ILogger Logger)
 		{
 			foreach (KeyValuePair<string, EddieFolder> CurGroup in Folders)
 			{
-                if (Path.GetFileName(CurGroup.Key) != "Documentation")
-                {
-                    Content.Append("AddFileGroup \"" + Path.GetFileName(CurGroup.Key) + "\" \"" + (CurGroup.Value.FullPath != null ? CurGroup.Value.FullPath : CurGroup.Value.Path) + "\"" + ProjectFileGenerator.NewLine);
-                
-                    if (CurGroup.Value.bIsModuleFolder)
-                    {
-                        StringBuilder ProjectFileContent = new StringBuilder();
-                
-                        ProjectFileContent.Append("# @Eddie Workset@" + ProjectFileGenerator.NewLine);
-                        ProjectFileContent.Append("AddWorkset \"" + Path.GetFileName(CurGroup.Key) + ".wkst\" \"" + CurGroup.Value.WorksetPath + "\"" + ProjectFileGenerator.NewLine);
-                
-                        ProjectFileContent.Append("AddFileGroup \"" + Path.GetFileName(CurGroup.Key) + "\" \"" + (CurGroup.Value.FullPath != null ? CurGroup.Value.FullPath : CurGroup.Value.Path) + "\"" + ProjectFileGenerator.NewLine);
-                
-                        EmitProject(ProjectFileContent, CurGroup.Value.Folders, Logger);
-                        
-                        foreach (EddieSourceFile File in CurGroup.Value.Files)
-                        {
-                            ProjectFileContent.Append("AddFile \"" + File.Reference.GetFileName() + "\" \"" + File.Reference.FullName + "\"" + ProjectFileGenerator.NewLine);
-                        }
-                        
-                        
-                        ProjectFileContent.Append("EndFileGroup \"" + Path.GetFileName(CurGroup.Key) + "\"" + ProjectFileGenerator.NewLine);
-                
-                        ProjectFileGenerator.WriteFileIfChanged(CurGroup.Value.WorksetPath!, ProjectFileContent.ToString(), Logger, new UTF8Encoding());
-                    
-                        Content.Append("AddFile \"" + Path.GetFileName(CurGroup.Key) + "\" \"" + CurGroup.Value.WorksetPath + "\"" + ProjectFileGenerator.NewLine);
-                    }
-                    else
-                    {
-                        EmitProject(Content, CurGroup.Value.Folders, Logger);
-                        
-                        foreach (EddieSourceFile File in CurGroup.Value.Files)
-                        {
-                            Content.Append("AddFile \"" + File.Reference.GetFileName() + "\" \"" + File.Reference.FullName + "\"" + ProjectFileGenerator.NewLine);
-                        }
-                    }
-                    
-                    Content.Append("EndFileGroup \"" + Path.GetFileName(CurGroup.Key) + "\"" + ProjectFileGenerator.NewLine);
-                }
+				if (Path.GetFileName(CurGroup.Key) != "Documentation")
+				{
+					Content.Append("AddFileGroup \"" + Path.GetFileName(CurGroup.Key) + "\" \"" + (CurGroup.Value.FullPath != null ? CurGroup.Value.FullPath : CurGroup.Value.Path) + "\"" + ProjectFileGenerator.NewLine);
+
+					if (CurGroup.Value.bIsModuleFolder)
+					{
+						StringBuilder ProjectFileContent = new StringBuilder();
+
+						ProjectFileContent.Append("# @Eddie Workset@" + ProjectFileGenerator.NewLine);
+						ProjectFileContent.Append("AddWorkset \"" + Path.GetFileName(CurGroup.Key) + ".wkst\" \"" + CurGroup.Value.WorksetPath + "\"" + ProjectFileGenerator.NewLine);
+
+						ProjectFileContent.Append("AddFileGroup \"" + Path.GetFileName(CurGroup.Key) + "\" \"" + (CurGroup.Value.FullPath != null ? CurGroup.Value.FullPath : CurGroup.Value.Path) + "\"" + ProjectFileGenerator.NewLine);
+
+						EmitProject(ProjectFileContent, CurGroup.Value.Folders, Logger);
+
+						foreach (EddieSourceFile File in CurGroup.Value.Files)
+						{
+							ProjectFileContent.Append("AddFile \"" + File.Reference.GetFileName() + "\" \"" + File.Reference.FullName + "\"" + ProjectFileGenerator.NewLine);
+						}
+
+
+						ProjectFileContent.Append("EndFileGroup \"" + Path.GetFileName(CurGroup.Key) + "\"" + ProjectFileGenerator.NewLine);
+
+						ProjectFileGenerator.WriteFileIfChanged(CurGroup.Value.WorksetPath!, ProjectFileContent.ToString(), Logger, new UTF8Encoding());
+
+						Content.Append("AddFile \"" + Path.GetFileName(CurGroup.Key) + "\" \"" + CurGroup.Value.WorksetPath + "\"" + ProjectFileGenerator.NewLine);
+					}
+					else
+					{
+						EmitProject(Content, CurGroup.Value.Folders, Logger);
+
+						foreach (EddieSourceFile File in CurGroup.Value.Files)
+						{
+							Content.Append("AddFile \"" + File.Reference.GetFileName() + "\" \"" + File.Reference.FullName + "\"" + ProjectFileGenerator.NewLine);
+						}
+					}
+
+					Content.Append("EndFileGroup \"" + Path.GetFileName(CurGroup.Key) + "\"" + ProjectFileGenerator.NewLine);
+				}
 			}
 		}
-		
+
 		public override bool WriteProjectFile(List<UnrealTargetPlatform> InPlatforms, List<UnrealTargetConfiguration> InConfigurations, PlatformProjectGeneratorCollection PlatformProjectGenerators, ILogger Logger)
 		{
 			bool bSuccess = false;
-			
+
 			//string TargetName = ProjectFilePath.GetFileNameWithoutExtension();
-			
+
 			//FileReference GameProjectPath = null;
 			//foreach(ProjectTarget Target in ProjectTargets)
 			//{
@@ -208,17 +203,17 @@ namespace UnrealBuildTool
 			//		break;
 			//	}
 			//}
-			
+
 			StringBuilder ProjectFileContent = new StringBuilder();
-			
+
 			ProjectFileContent.Append("# @Eddie Workset@" + ProjectFileGenerator.NewLine);
 			ProjectFileContent.Append("AddWorkset \"" + this.ToString() + ".wkst\" \"" + ProjectFilePath.FullName + "\"" + ProjectFileGenerator.NewLine);
-			
+
 			ParseSourceFilesIntoGroups();
 			EmitProject(ProjectFileContent, Folders, Logger);
-			
+
 			bSuccess = ProjectFileGenerator.WriteFileIfChanged(ProjectFilePath.FullName, ProjectFileContent.ToString(), Logger, new UTF8Encoding());
-			
+
 			return bSuccess;
 		}
 	}
diff --git a/Engine/Source/Programs/UnrealBuildTool/ProjectFiles/Make/MakefileGenerator.cs b/Engine/Source/Programs/UnrealBuildTool/ProjectFiles/Make/MakefileGenerator.cs
--- a/Engine/Source/Programs/UnrealBuildTool/ProjectFiles/Make/MakefileGenerator.cs
+++ b/Engine/Source/Programs/UnrealBuildTool/ProjectFiles/Make/MakefileGenerator.cs
@@ -1,13 +1,11 @@
 // Copyright Epic Games, Inc. All Rights Reserved.
 
 using System;
-using System.Collections.Generic;
+using System.Linq;
 using System.Text;
-using System.IO;
 using EpicGames.Core;
-using UnrealBuildBase;
-using System.Linq;
 using Microsoft.Extensions.Logging;
+using UnrealBuildBase;
 
 namespace UnrealBuildTool
 {
diff --git a/Engine/Source/Programs/UnrealBuildTool/ProjectFiles/ProjectFileGenerator.cs b/Engine/Source/Programs/UnrealBuildTool/ProjectFiles/ProjectFileGenerator.cs
--- a/Engine/Source/Programs/UnrealBuildTool/ProjectFiles/ProjectFileGenerator.cs
+++ b/Engine/Source/Programs/UnrealBuildTool/ProjectFiles/ProjectFileGenerator.cs
@@ -2,19 +2,14 @@
 
 using System;
 using System.Collections.Generic;
-using System.Text;
+using System.Diagnostics;
 using System.IO;
-using System.Reflection;
-using Microsoft.Win32;
 using System.Linq;
-using System.Diagnostics;
-using EpicGames.Core;
+using System.Text;
 using System.Xml.Linq;
-using System.Xml;
-using UnrealBuildBase;
+using EpicGames.Core;
 using Microsoft.Extensions.Logging;
-using EpicGames.Serialization;
-using System.Threading;
+using UnrealBuildBase;
 
 namespace UnrealBuildTool
 {
@@ -25,14 +20,14 @@ namespace UnrealBuildTool
 			PluginFiles = InPluginFiles;
 			ReferencingProjects = new List<FileReference>();
 			ModuleToSourceFiles = new Dictionary<FileReference, List<FileReference>>();
-			
+
 		}
 
 		public List<FileReference> PluginFiles;
 		public List<FileReference> ReferencingProjects;
 		public Dictionary<FileReference, List<FileReference>> ModuleToSourceFiles;
 	}
-	
+
 	/// <summary>
 	/// Represents a folder within the primary project (e.g. Visual Studio solution)
 	/// </summary>
@@ -481,7 +476,7 @@ namespace UnrealBuildTool
 		/// <summary>
 		/// Adds all rules project files to the solution.
 		/// </summary>
-		void AddRulesModules(Rules.RulesFileType RulesFileType, string ProgramSubDirectory, List<ProjectFile> AddedProjectFiles, 
+		void AddRulesModules(Rules.RulesFileType RulesFileType, string ProgramSubDirectory, List<ProjectFile> AddedProjectFiles,
 			List<FileReference> UnrealProjectFiles, PrimaryProjectFolder RootFolder, PrimaryProjectFolder ProgramsFolder, ILogger Logger)
 		{
 			List<DirectoryReference> GameFolders = new List<DirectoryReference>();
@@ -612,7 +607,7 @@ namespace UnrealBuildTool
 				DirectoryReference EngineExtras = DirectoryReference.Combine(Unreal.EngineDirectory, "Extras");
 				ProjectDirs.Add(EngineExtras);
 				DiscoverCSharpProgramProjectsRecursively(EngineExtras, FoundProjects);
-			
+
 				ProjectDirs = ProjectDirs.Union(AllEngineDirectories).ToList();
 				foreach (DirectoryReference EngineDir in AllEngineDirectories)
 				{
@@ -702,8 +697,8 @@ namespace UnrealBuildTool
 				Writer.WriteLine("<!-- This file was generated by UnrealBuildTool.ProjectFileGenerator.CreateRulesAssemblyProject() -->");
 				Writer.WriteLine("<Project Sdk=\"Microsoft.NET.Sdk\">");
 
-				Writer.WriteLine("  <Import Project=\"" + 
-					FileReference.Combine(Unreal.EngineDirectory, "Source", "Programs", "Shared", "UnrealEngine.csproj.props").MakeRelativeTo(ProjectFilesDirectory) + 
+				Writer.WriteLine("  <Import Project=\"" +
+					FileReference.Combine(Unreal.EngineDirectory, "Source", "Programs", "Shared", "UnrealEngine.csproj.props").MakeRelativeTo(ProjectFilesDirectory) +
 					"\" />");
 
 				Writer.WriteLine("  <PropertyGroup>");
@@ -712,7 +707,7 @@ namespace UnrealBuildTool
 				Writer.WriteLine("    <Configurations>Debug;Release;Development</Configurations>"); // VCSharpProject requires at least Debug & Development configurations
 				Writer.WriteLine("    <DefineConstants>$(DefineConstants);" + String.Join(';', RulesAssembly.PreprocessorDefines!) + "</DefineConstants>");
 				Writer.WriteLine("  </PropertyGroup>");
-				
+
 				Writer.WriteLine("  <ItemGroup>");
 				foreach (FileReference Reference in ReferencedProjects)
 				{
@@ -986,7 +981,7 @@ namespace UnrealBuildTool
 
 			// Remove any game projects that don't have a target (or are under a directory "Programs" since they don't have Targets under the .uproject)
 			List<FileReference> CodeBasedGameProjects = new(AllGameProjects);
-			CodeBasedGameProjects.RemoveAll(x => 
+			CodeBasedGameProjects.RemoveAll(x =>
 			{
 				// anything with a target is valid
 				bool bHasTarget = AllTargetFiles.Any(y => y.IsUnderDirectory(x.Directory));
@@ -1018,7 +1013,7 @@ namespace UnrealBuildTool
 			Dictionary<FileReference, ProjectFile> ProgramProjects = new Dictionary<FileReference, ProjectFile>();
 			Dictionary<RulesAssembly, DirectoryReference> RulesAssemblies = new Dictionary<RulesAssembly, DirectoryReference>();
 			Dictionary<ProjectFile, FileReference> ProjectFileToUProjectFile = new Dictionary<ProjectFile, FileReference>();
-			if(IncludeCppSource)
+			if (IncludeCppSource)
 			{
 				// Setup buildable projects for all targets
 				AddProjectsForAllTargets(PlatformProjectGenerators, AllGameProjects, AllTargetFiles, Arguments, EngineProjects, GameProjects, ProjectFileToUProjectFile, ProgramProjects, RulesAssemblies, Logger);
@@ -1256,7 +1251,7 @@ namespace UnrealBuildTool
 						RootFolder.AddSubFolder("Rules").ChildProjects.Add(RulesProject);
 					}
 				}
-				
+
 				// Eliminate all redundant project folders.  E.g., folders which contain only one project and that project
 				// has the same name as the folder itself.  To the user, projects "feel like" folders already in the IDE, so we
 				// want to collapse them down where possible.
@@ -2141,7 +2136,7 @@ namespace UnrealBuildTool
 							Logger.LogWarning("Exception while generating include data for {Target}: {Ex}", CurTarget.Name, Ex.ToString());
 						}
 
-						lock(Progress)
+						lock (Progress)
 						{
 							Progress.Write(++TasksFinished, NumTasks);
 						}
@@ -2157,7 +2152,7 @@ namespace UnrealBuildTool
 							UEBuildTarget Target = CreatedTargets[TargetIndex];
 							TargetDescriptor TargetDesc = CreatedTargetDesc[TargetIndex];
 
-							if (TargetDesc == null || Target == null) 
+							if (TargetDesc == null || Target == null)
 							{
 								continue;
 							}
@@ -2359,7 +2354,7 @@ namespace UnrealBuildTool
 							PluginData = new AdditionalPluginData(PluginsBase.EnumeratePlugins(AdditionalPluginDirectory).ToList());
 							AdditionalPluginsByDirectory.Add(AdditionalPluginDirectory, PluginData);
 						}
-						
+
 						PluginData.ReferencingProjects.Add(ProjectPath);
 					}
 				}
@@ -2381,7 +2376,7 @@ namespace UnrealBuildTool
 					}
 				}
 			}
-			
+
 			HashSet<ProjectFile> ProjectsWithPlugins = new HashSet<ProjectFile>();
 			foreach (FileReference CurModuleFile in AllModuleFiles)
 			{
@@ -2463,7 +2458,7 @@ namespace UnrealBuildTool
 					}
 				}
 			}
-			
+
 			if (IncludeCppSource)
 			{
 				foreach ((DirectoryReference PluginDirectory, AdditionalPluginData PluginData) in AdditionalPluginsByDirectory)
@@ -2687,7 +2682,7 @@ namespace UnrealBuildTool
 					if (CheckProjectFile == null)
 					{
 						RulesAssembly = RulesCompiler.CreateEngineRulesAssembly(false, false, false, Logger);
-					
+
 						// Record the Engine assembly, and any parent assemblies (varies e.g. Rules, ProgramRules, MarketplaceRules)
 						if (!RulesAssemblies.ContainsKey(RulesAssembly))
 						{
@@ -2881,11 +2876,11 @@ namespace UnrealBuildTool
 							ProjectFilePath: ProjectFilePath,
 							UnrealProjectFilePath: CheckProjectFile,
 							SupportedPlatforms: TargetRulesObject.GetSupportedPlatforms().Where(
-													x => UEBuildPlatform.TryGetBuildPlatform(x, out _) && 
+													x => UEBuildPlatform.TryGetBuildPlatform(x, out _) &&
 													(TargetRulesObject.LinkType != TargetLinkType.Modular || !UEBuildPlatform.PlatformRequiresMonolithicBuilds(x, TargetRulesObject.Configuration))
 													).ToArray(),
 							CreateRulesDelegate: (Platform, Configuration) => RulesAssembly.CreateTargetRules(TargetName, Platform, Configuration, null, CheckProjectFile, new CommandLineArguments(GetTargetArguments(Arguments)), Logger)
-                        );
+						);
 
 					ProjectFile.ProjectTargets.Add(ProjectTarget);
 
@@ -3168,7 +3163,7 @@ namespace UnrealBuildTool
 						return;
 					}
 
-					lock(Progress)
+					lock (Progress)
 					{
 						Progress.Write(++ProjectsFinished, TotalProjectFileCount);
 					}
@@ -3215,7 +3210,7 @@ namespace UnrealBuildTool
 		protected virtual void AddAdditionalNativeTargetInformation(PlatformProjectGeneratorCollection PlatformProjectGenerators, List<Tuple<ProjectFile, ProjectTarget>> Targets, ILogger Logger)
 		{
 		}
-		
+
 
 
 		/// <summary>
@@ -3320,7 +3315,7 @@ namespace UnrealBuildTool
 				ProjectTarget.ExtraSupportedPlatforms.AddRange(UnrealTargetPlatform.GetValidPlatforms());
 
 				// Add all configurations
-				foreach (UnrealTargetConfiguration CurConfiguration in (UnrealTargetConfiguration[]) Enum.GetValues(typeof(UnrealTargetConfiguration)))
+				foreach (UnrealTargetConfiguration CurConfiguration in (UnrealTargetConfiguration[])Enum.GetValues(typeof(UnrealTargetConfiguration)))
 				{
 					ProjectTarget.ExtraSupportedConfigurations.Add(CurConfiguration);
 				}
@@ -3352,7 +3347,7 @@ namespace UnrealBuildTool
 				else
 				{
 					// Otherwise, add all configurations
-				foreach (UnrealTargetConfiguration CurConfiguration in (UnrealTargetConfiguration[]) Enum.GetValues(typeof(UnrealTargetConfiguration)))
+					foreach (UnrealTargetConfiguration CurConfiguration in (UnrealTargetConfiguration[])Enum.GetValues(typeof(UnrealTargetConfiguration)))
 					{
 						ProjectTarget.ExtraSupportedConfigurations.Add(CurConfiguration);
 					}
@@ -3445,7 +3440,7 @@ namespace UnrealBuildTool
 			{
 				return null;
 			}
-			
+
 			ConfigHierarchy Ini = ConfigCache.ReadHierarchy(ConfigHierarchyType.Engine, ProjectDirectory, BuildHostPlatform.Current.Platform);
 			string? Value;
 			Ini.TryGetValue("/Script/BuildSettings.BuildSettings", "DefaultEditorTarget", out Value);
diff --git a/Engine/Source/Programs/UnrealBuildTool/ProjectFiles/VisualStudio/VCProject.cs b/Engine/Source/Programs/UnrealBuildTool/ProjectFiles/VisualStudio/VCProject.cs
--- a/Engine/Source/Programs/UnrealBuildTool/ProjectFiles/VisualStudio/VCProject.cs
+++ b/Engine/Source/Programs/UnrealBuildTool/ProjectFiles/VisualStudio/VCProject.cs
@@ -1,20 +1,19 @@
 // Copyright Epic Games, Inc. All Rights Reserved.
 
 using System;
+using System.Collections.Concurrent;
 using System.Collections.Generic;
+using System.Diagnostics.CodeAnalysis;
 using System.IO;
-using System.Xml;
 using System.Linq;
+using System.Security;
 using System.Text;
 using System.Text.RegularExpressions;
-using System.Security;
+using System.Threading.Tasks;
+using System.Xml;
 using EpicGames.Core;
-using UnrealBuildBase;
-using System.Diagnostics.CodeAnalysis;
 using Microsoft.Extensions.Logging;
-using System.Reflection.Metadata;
-using System.Threading.Tasks;
-using System.Collections.Concurrent;
+using UnrealBuildBase;
 
 namespace UnrealBuildTool
 {
@@ -904,7 +903,7 @@ namespace UnrealBuildTool
 				{
 					SourceDirToCount.AddOrUpdate(SourceFile.Reference.Directory, _ => 1, (k, v) => v + 1);
 				});
-				
+
 				// Figure out the most common include paths
 				ConcurrentDictionary<DirectoryReference, int> IncludePathToCount = new();
 				Parallel.ForEach(SourceDirToCount, Pair =>
@@ -1699,14 +1698,14 @@ namespace UnrealBuildTool
 		private class ProjectConfigurationForGenerator : ProjectBuildConfiguration
 		{
 			public override string ConfigurationName => Combination.ProjectConfigurationName;
-			public override string BuildCommand      => $"{EscapePath(NormalizeProjectPath(CommandBuilder.BuildScript))} {CommandBuilder.GetBuildArguments()}";
+			public override string BuildCommand => $"{EscapePath(NormalizeProjectPath(CommandBuilder.BuildScript))} {CommandBuilder.GetBuildArguments()}";
 
 			private ProjectConfigAndTargetCombination Combination;
-			private BuildCommandBuilder               CommandBuilder;
+			private BuildCommandBuilder CommandBuilder;
 
 			public ProjectConfigurationForGenerator(ProjectConfigAndTargetCombination InCombination, BuildCommandBuilder InCommandBuilder)
 			{
-				Combination    = InCombination;
+				Combination = InCombination;
 				CommandBuilder = InCommandBuilder;
 			}
 		}
@@ -1855,11 +1854,11 @@ namespace UnrealBuildTool
 			public bool bIsFromMSBuild;
 
 			public PlatformProjectGenerator? ProjectGenerator;
-			
+
 			public FileReference BuildScript { get; }
 			public FileReference RebuildScript { get; }
 			public FileReference CleanScript { get; }
-			
+
 			private readonly string? BuildToolOverride;
 
 			private readonly string UProjectPath;
@@ -1876,20 +1875,20 @@ namespace UnrealBuildTool
 				ProjectTarget = InProjectTarget;
 				UProjectPath = InUProjectPath;
 				BuildToolOverride = InBuildToolOverride;
-				
+
 				DirectoryReference BatchFilesDirectory = DirectoryReference.Combine(Unreal.EngineDirectory, "Build", "BatchFiles");
 				BuildScript = FileReference.Combine(BatchFilesDirectory, "Build.bat");
 				RebuildScript = FileReference.Combine(BatchFilesDirectory, "Rebuild.bat");
 				CleanScript = FileReference.Combine(BatchFilesDirectory, "Clean.bat");
 			}
-			
+
 			public string GetBuildArguments()
 			{
 				TargetRules TargetRulesObject = ProjectTarget.TargetRules!;
 				string TargetName = ProjectTarget.TargetFilePath.GetFileNameWithoutAnyExtensions();
 
 				StringBuilder BuildArguments = new StringBuilder();
-				
+
 				BuildArguments.AppendFormat("{0} {1} {2}", TargetName, Platform.ToString(), Configuration.ToString());
 				if (UProjectPath.Length > 0)
 				{
@@ -1970,7 +1969,7 @@ namespace UnrealBuildTool
 
 			return Builder;
 		}
-		
+
 		// Anonymous function that writes project configuration data
 		private void WriteConfiguration(string ProjectName, ProjectConfigAndTargetCombination Combination, StringBuilder VCProjectFileContent, PlatformProjectGeneratorCollection PlatformProjectGenerators, StringBuilder? VCUserFileContent, ILogger Logger)
 		{
