UnrealBuildTool: Automated code cleanup

diff --git a/Engine/Source/Programs/UnrealBuildTool/ProjectFiles/VisualStudio/VCSolutionOptions.cs b/Engine/Source/Programs/UnrealBuildTool/ProjectFiles/VisualStudio/VCSolutionOptions.cs
--- a/Engine/Source/Programs/UnrealBuildTool/ProjectFiles/VisualStudio/VCSolutionOptions.cs
+++ b/Engine/Source/Programs/UnrealBuildTool/ProjectFiles/VisualStudio/VCSolutionOptions.cs
@@ -2,20 +2,15 @@
 
 using System;
 using System.Collections.Generic;
+using System.Diagnostics;
+using System.Diagnostics.CodeAnalysis;
 using System.IO;
-using System.Xml;
-using System.Xml.XPath;
-using System.Xml.Linq;
 using System.Linq;
 using System.Runtime.InteropServices;
+using System.Runtime.Versioning;
 using System.Text;
-using System.Diagnostics;
-
-using STATSTG = System.Runtime.InteropServices.ComTypes.STATSTG;
 using FILETIME = System.Runtime.InteropServices.ComTypes.FILETIME;
-using UnrealBuildBase;
-using System.Diagnostics.CodeAnalysis;
-using System.Runtime.Versioning;
+using STATSTG = System.Runtime.InteropServices.ComTypes.STATSTG;
 
 namespace UnrealBuildTool
 {
@@ -214,7 +209,7 @@ namespace UnrealBuildTool
 		}
 
 		[DllImport("ole32.dll")]
-		static extern int StgCreateDocfile([MarshalAs(UnmanagedType.LPWStr)]string pwcsName, STGM grfMode, uint reserved, out IOleStorage ppstgOpen);
+		static extern int StgCreateDocfile([MarshalAs(UnmanagedType.LPWStr)] string pwcsName, STGM grfMode, uint reserved, out IOleStorage ppstgOpen);
 
 		[DllImport("ole32.dll")]
 		static extern int StgOpenStorage([MarshalAs(UnmanagedType.LPWStr)] string pwcsName, IOleStorage? pstgPriority, STGM grfMode, IntPtr snbExclude, uint reserved, out IOleStorage ppstgOpen);
@@ -370,11 +365,13 @@ namespace UnrealBuildTool
 		{
 			long OpenFoldersEnd = Reader.BaseStream.Position + Reader.ReadInt32();
 
-			if(Format >= VCProjectFileFormat.VisualStudio2019)
+			if (Format >= VCProjectFileFormat.VisualStudio2019)
 			{
 				int Header1 = Reader.ReadInt32();
-				/*int Header2 =*/ Reader.ReadInt32();
-				/*int Header3 =*/ Reader.ReadByte();
+				/*int Header2 =*/
+				Reader.ReadInt32();
+				/*int Header3 =*/
+				Reader.ReadByte();
 				if (Header1 != 15)
 				{
 					throw new BuildException("Unexpected data in open projects section");
@@ -410,7 +407,7 @@ namespace UnrealBuildTool
 
 		void WriteOpenProjects(BinaryWriter Writer, VCProjectFileFormat Format)
 		{
-			if(Format >= VCProjectFileFormat.VisualStudio2019)
+			if (Format >= VCProjectFileFormat.VisualStudio2019)
 			{
 				Writer.Write(4 + (4 + 4 + 1) + 2 + OpenProjects.Sum(x => GetStringSize(x.Item1) + 2 + x.Item2.Sum(y => GetStringSize(y))));
 				Writer.Write(15);
diff --git a/Engine/Source/Programs/UnrealBuildTool/ProjectFiles/VisualStudioCode/VSCodeProjectFileGenerator.cs b/Engine/Source/Programs/UnrealBuildTool/ProjectFiles/VisualStudioCode/VSCodeProjectFileGenerator.cs
--- a/Engine/Source/Programs/UnrealBuildTool/ProjectFiles/VisualStudioCode/VSCodeProjectFileGenerator.cs
+++ b/Engine/Source/Programs/UnrealBuildTool/ProjectFiles/VisualStudioCode/VSCodeProjectFileGenerator.cs
@@ -2,13 +2,13 @@
 
 using System;
 using System.Collections.Generic;
-using System.Text;
-using System.IO;
-using EpicGames.Core;
 using System.Diagnostics;
+using System.IO;
 using System.Linq;
-using UnrealBuildBase;
+using System.Text;
+using EpicGames.Core;
 using Microsoft.Extensions.Logging;
+using UnrealBuildBase;
 
 namespace UnrealBuildTool
 {
@@ -76,7 +76,7 @@ namespace UnrealBuildTool
 
 		private enum EQuoteType
 		{
-			Single,	// can be ignored on platforms that don't need it (windows atm)
+			Single, // can be ignored on platforms that don't need it (windows atm)
 			Double,
 		}
 
@@ -88,23 +88,23 @@ namespace UnrealBuildTool
 			}
 
 			string Processed = InRef.ToString();
-			
+
 			switch (InPathType)
 			{
 				case EPathType.Relative:
-				{
-					if (InRef.IsUnderDirectory(InRelativeRoot))
 					{
-						Processed = InRef.MakeRelativeTo(InRelativeRoot).ToString();
-					}
+						if (InRef.IsUnderDirectory(InRelativeRoot))
+						{
+							Processed = InRef.MakeRelativeTo(InRelativeRoot).ToString();
+						}
 
-					break;
-				}
+						break;
+					}
 
 				default:
-				{
-					break;
-				}
+					{
+						break;
+					}
 			}
 
 			if (HostPlatform == UnrealTargetPlatform.Win64)
@@ -133,7 +133,7 @@ namespace UnrealBuildTool
 				{
 					Processed = "'" + Processed + "'";
 				}
- 			}
+			}
 
 			return Processed;
 		}
@@ -389,7 +389,7 @@ namespace UnrealBuildTool
 					}
 					foreach (DirectoryReference IncludePath in ModuleCompileEnvironment.UserIncludePaths)
 					{
-						CommandBuilder.AppendFormat("{0} \"{1}\" {2}",  UsingClang ? "-I" : "/I", IncludePath, Environment.NewLine);
+						CommandBuilder.AppendFormat("{0} \"{1}\" {2}", UsingClang ? "-I" : "/I", IncludePath, Environment.NewLine);
 					}
 					foreach (DirectoryReference IncludePath in ModuleCompileEnvironment.SystemIncludePaths)
 					{
@@ -420,15 +420,15 @@ namespace UnrealBuildTool
 				public UnrealTargetConfiguration Config { get; set; }
 				public UnrealTargetPlatform Platform { get; set; }
 				public EOutputType OutputType { get; set; }
-		
+
 				public CsProjectInfo? CSharpInfo { get; set; }
 
 				public override string ToString()
 				{
-					return Platform.ToString() + " " + Config.ToString(); 
+					return Platform.ToString() + " " + Config.ToString();
 				}
 
-				public BuildProduct(FileReference OutputFile) 
+				public BuildProduct(FileReference OutputFile)
 				{
 					this.OutputFile = OutputFile;
 				}
@@ -449,7 +449,7 @@ namespace UnrealBuildTool
 
 				public override string ToString()
 				{
-					return Name.ToString() + " " + Type.ToString(); 
+					return Name.ToString() + " " + Type.ToString();
 				}
 			}
 
@@ -532,62 +532,62 @@ namespace UnrealBuildTool
 				{
 					VCSharpProjectFile VCSharpProject = (VCSharpProjectFile)Project;
 
-	
-						string ProjectName = Project.ProjectFilePath.GetFileNameWithoutExtension();
 
-						ProjectData.Target Target = new ProjectData.Target(NewProject, ProjectName, TargetType.Program);
+					string ProjectName = Project.ProjectFilePath.GetFileNameWithoutExtension();
 
-						UnrealTargetConfiguration[] Configs = { UnrealTargetConfiguration.Debug, UnrealTargetConfiguration.Development };
+					ProjectData.Target Target = new ProjectData.Target(NewProject, ProjectName, TargetType.Program);
 
-						foreach (UnrealTargetConfiguration Config in Configs)
-						{
-							CsProjectInfo? Info = VCSharpProject.GetProjectInfo(Config)!;
+					UnrealTargetConfiguration[] Configs = { UnrealTargetConfiguration.Debug, UnrealTargetConfiguration.Development };
 
-							if (Info.Properties.ContainsKey("OutputPath"))
+					foreach (UnrealTargetConfiguration Config in Configs)
+					{
+						CsProjectInfo? Info = VCSharpProject.GetProjectInfo(Config)!;
+
+						if (Info.Properties.ContainsKey("OutputPath"))
+						{
+							ProjectData.EOutputType OutputType;
+							string? OutputTypeName;
+							if (Info.Properties.TryGetValue("OutputType", out OutputTypeName))
 							{
-								ProjectData.EOutputType OutputType;
-								string? OutputTypeName;
-								if (Info.Properties.TryGetValue("OutputType", out OutputTypeName))
-								{
-									OutputType = (ProjectData.EOutputType)Enum.Parse(typeof(ProjectData.EOutputType), OutputTypeName);
-								}
-								else
-								{
-									OutputType = ProjectData.EOutputType.Library;
-								}
+								OutputType = (ProjectData.EOutputType)Enum.Parse(typeof(ProjectData.EOutputType), OutputTypeName);
+							}
+							else
+							{
+								OutputType = ProjectData.EOutputType.Library;
+							}
 
-								if (OutputType == ProjectData.EOutputType.WinExe)
-								{
-									OutputType = ProjectData.EOutputType.Exe;
-								}
+							if (OutputType == ProjectData.EOutputType.WinExe)
+							{
+								OutputType = ProjectData.EOutputType.Exe;
+							}
 
-								FileReference? OutputFile = null;
-								HashSet<FileReference> ProjectBuildProducts = new HashSet<FileReference>();
-								Info.FindCompiledBuildProducts(DirectoryReference.Combine(VCSharpProject.ProjectFilePath.Directory, Info.Properties["OutputPath"]), ProjectBuildProducts);
-								foreach (FileReference ProjectBuildProduct in ProjectBuildProducts)
+							FileReference? OutputFile = null;
+							HashSet<FileReference> ProjectBuildProducts = new HashSet<FileReference>();
+							Info.FindCompiledBuildProducts(DirectoryReference.Combine(VCSharpProject.ProjectFilePath.Directory, Info.Properties["OutputPath"]), ProjectBuildProducts);
+							foreach (FileReference ProjectBuildProduct in ProjectBuildProducts)
+							{
+								if ((OutputType == ProjectData.EOutputType.Exe && ProjectBuildProduct.GetExtension() == FrameworkExecutableExtension) ||
+									(OutputType == ProjectData.EOutputType.Library && ProjectBuildProduct.GetExtension() == FrameworkLibraryExtension))
 								{
-									if ((OutputType == ProjectData.EOutputType.Exe && ProjectBuildProduct.GetExtension() == FrameworkExecutableExtension) ||
-										(OutputType == ProjectData.EOutputType.Library && ProjectBuildProduct.GetExtension() == FrameworkLibraryExtension))
-									{
-										OutputFile = ProjectBuildProduct;
-										break;
-									}
+									OutputFile = ProjectBuildProduct;
+									break;
 								}
+							}
 
-								if (OutputFile != null)
+							if (OutputFile != null)
+							{
+								Target.BuildProducts.Add(new ProjectData.BuildProduct(OutputFile)
 								{
-									Target.BuildProducts.Add(new ProjectData.BuildProduct(OutputFile)
-									{
-										Platform = HostPlatform,
-										Config = Config,
-										OutputType = OutputType,
-										CSharpInfo = Info
-									});
-								}
+									Platform = HostPlatform,
+									Config = Config,
+									OutputType = OutputType,
+									CSharpInfo = Info
+								});
 							}
 						}
+					}
 
-						ProjectData.CSharpProjects.Add(NewProject);
+					ProjectData.CSharpProjects.Add(NewProject);
 				}
 			}
 
@@ -609,7 +609,7 @@ namespace UnrealBuildTool
 					FileReference? CompilerPath = null;
 					DirectoryReference? SysRootPath = null;
 					CppStandardVersion CppStandard = CppStandardVersion.Default;
-					
+
 					foreach (ProjectData.Project Project in Projects.AllProjects)
 					{
 						AllSourceFiles.UnionWith(Project.SourceProject.SourceFiles.Select(x => x.Reference));
@@ -629,7 +629,7 @@ namespace UnrealBuildTool
 
 							foreach (KeyValuePair<DirectoryReference, string> Pair in BuildTarget.ModuleCommandLines)
 							{
-								if(!AllModuleCommandLines.ContainsKey(Pair.Key))
+								if (!AllModuleCommandLines.ContainsKey(Pair.Key))
 								{
 									AllModuleCommandLines[Pair.Key] = Pair.Value;
 								}
@@ -848,7 +848,7 @@ namespace UnrealBuildTool
 
 		}
 
-		private void WriteCompileCommands(FileReference CompileCommandsFile, IEnumerable<FileReference> SourceFiles, 
+		private void WriteCompileCommands(FileReference CompileCommandsFile, IEnumerable<FileReference> SourceFiles,
 			FileReference CompilerPath, Dictionary<DirectoryReference, string> ModuleCommandLines)
 		{
 			// this creates a compileCommands.json
@@ -863,7 +863,7 @@ namespace UnrealBuildTool
 				DirectoryReference.CreateDirectory(ResponseFileDir);
 
 				Dictionary<DirectoryReference, FileReference?> DirectoryToResponseFile = new Dictionary<DirectoryReference, FileReference?>();
-				foreach(KeyValuePair<DirectoryReference, string> Pair in ModuleCommandLines)
+				foreach (KeyValuePair<DirectoryReference, string> Pair in ModuleCommandLines)
 				{
 					FileReference ResponseFile = FileReference.Combine(ResponseFileDir, String.Format("{0}.{1}.rsp", Pair.Key.GetDirectoryName(), DirectoryToResponseFile.Count));
 					FileReference.WriteAllText(ResponseFile, Pair.Value);
@@ -941,7 +941,7 @@ namespace UnrealBuildTool
 									CleanParam = "-clean";
 								}
 							}
-							
+
 
 							OutFile.BeginArray("args");
 							{
@@ -1075,7 +1075,7 @@ namespace UnrealBuildTool
 				OutFile.EndArray();
 			}
 		}
-		
+
 		private FileReference GetExecutableFilename(ProjectFile Project, ProjectTarget Target, UnrealTargetPlatform Platform, UnrealTargetConfiguration Configuration)
 		{
 			TargetRules? TargetRulesObject = Target.TargetRules;
@@ -1106,7 +1106,7 @@ namespace UnrealBuildTool
 			DirectoryReference RootDirectory = Unreal.EngineDirectory;
 			if (TargetRulesType != TargetType.Program && (bShouldCompileMonolithic || TargetRulesObject!.BuildEnvironment == TargetBuildEnvironment.Unique))
 			{
-				if(Target.UnrealProjectFilePath != null)
+				if (Target.UnrealProjectFilePath != null)
 				{
 					RootDirectory = Target.UnrealProjectFilePath.Directory;
 				}
@@ -1114,7 +1114,7 @@ namespace UnrealBuildTool
 
 			if (TargetRulesType == TargetType.Program)
 			{
-				if(Target.UnrealProjectFilePath != null)
+				if (Target.UnrealProjectFilePath != null)
 				{
 					RootDirectory = Target.UnrealProjectFilePath.Directory;
 				}
@@ -1125,7 +1125,7 @@ namespace UnrealBuildTool
 
 			// Get the executable name (minus any platform or config suffixes)
 			string BinaryName;
-			if(Target.TargetRules!.BuildEnvironment == TargetBuildEnvironment.Shared && TargetRulesType != TargetType.Program)
+			if (Target.TargetRules!.BuildEnvironment == TargetBuildEnvironment.Shared && TargetRulesType != TargetType.Program)
 			{
 				BinaryName = UEBuildTarget.GetAppNameForTargetType(TargetRulesType);
 			}
@@ -1265,7 +1265,7 @@ namespace UnrealBuildTool
 			List<string> Types = new List<string>();
 			Types.Add("Launch");
 
-			if(bAddDebugAttachConfig && bIsLinux)
+			if (bAddDebugAttachConfig && bIsLinux)
 			{
 				Types.Add("Attach");
 			}
@@ -1277,12 +1277,12 @@ namespace UnrealBuildTool
 
 			string LaunchTaskName = String.Format("{0} {1} {2} Build", Target.Name, BuildProduct.Platform, BuildProduct.Config);
 
-			foreach(string Type in Types)
+			foreach (string Type in Types)
 			{
 				OutFile.BeginObject();
 				{
 					OutFile.AddField("name", Type + " " + Target.Name + " (" + BuildProduct.Config.ToString() + ")");
-					OutFile.AddField("request", (Type == "Attach")?"attach":"launch");
+					OutFile.AddField("request", (Type == "Attach") ? "attach" : "launch");
 					OutFile.AddField("program", MakeUnquotedPathString(BuildProduct.OutputFile, EPathType.Absolute));
 					switch (Type)
 					{
@@ -1292,12 +1292,12 @@ namespace UnrealBuildTool
 						case "Debug Core":
 							OutFile.AddField("coreDumpPath", "${input:coreFileName}");
 							break;
-						case "Attach": 
+						case "Attach":
 							OutFile.AddField("processId", "${command:pickProcess}");
 							break;
 					}
 
-					if(Type != "Attach")
+					if (Type != "Attach")
 					{
 						OutFile.BeginArray("args");
 						{
@@ -1397,7 +1397,7 @@ namespace UnrealBuildTool
 				{
 					OutFile.AddField("preLaunchTask", InBuildTaskName);
 				}
-				
+
 				DirectoryReference CWD = ProjectRoot;
 
 				OutFile.AddField("program", MakeUnquotedPathString(InExecutable, EPathType.Absolute));
@@ -1449,7 +1449,7 @@ namespace UnrealBuildTool
 		private void WriteLaunch(JsonFile OutFile, ProjectData ProjectData, ILogger Logger)
 		{
 			OutFile.AddField("version", "0.2.0");
-			if(bAddDebugCoreConfig)
+			if (bAddDebugCoreConfig)
 			{
 				OutFile.BeginArray("inputs");
 				OutFile.BeginObject();
@@ -1499,7 +1499,7 @@ namespace UnrealBuildTool
 			{
 				Args.Add("-engine");
 			}
-			
+
 			if (bIncludeDotNetPrograms)
 			{
 				Args.Add("-dotnet");
@@ -1531,7 +1531,7 @@ namespace UnrealBuildTool
 					{
 						DirectoryReference ProjDir = Target.TargetFilePath.Directory.GetDirectoryName() == "Source" ? Target.TargetFilePath.Directory.ParentDirectory! : Target.TargetFilePath.Directory;
 						GetExcludePathsCPP(ProjDir, PathsToExclude);
-						
+
 						DirectoryReference PluginRootDir = DirectoryReference.Combine(ProjDir, "Plugins");
 						WriteWorkspaceIgnoreFileForPlugins(PluginRootDir, PathsToExclude);
 
@@ -1587,7 +1587,7 @@ namespace UnrealBuildTool
 				}
 			}
 		}
-		
+
 		private void WriteWorkspaceFile(ProjectData ProjectData, ILogger Logger)
 		{
 			JsonFile WorkspaceFile = new JsonFile();
@@ -1656,7 +1656,7 @@ namespace UnrealBuildTool
 				}
 			}
 			WorkspaceFile.EndObject();
-			
+
 			WorkspaceFile.BeginObject("extensions");
 			{
 				// extensions is a set of recommended extensions that a user should install.
diff --git a/Engine/Source/Programs/UnrealBuildTool/ProjectFiles/VisualStudioMac/VisualStudioMacProjectFileGenerator.cs b/Engine/Source/Programs/UnrealBuildTool/ProjectFiles/VisualStudioMac/VisualStudioMacProjectFileGenerator.cs
--- a/Engine/Source/Programs/UnrealBuildTool/ProjectFiles/VisualStudioMac/VisualStudioMacProjectFileGenerator.cs
+++ b/Engine/Source/Programs/UnrealBuildTool/ProjectFiles/VisualStudioMac/VisualStudioMacProjectFileGenerator.cs
@@ -1,118 +1,114 @@
 ﻿// Copyright Epic Games, Inc. All Rights Reserved.
 
-using System;
-using System.Collections.Generic;
-using System.Text;
-using System.IO;
 using System.Linq;
 using System.Xml;
 using System.Xml.Linq;
 using EpicGames.Core;
-using UnrealBuildBase;
 using Microsoft.Extensions.Logging;
+using UnrealBuildBase;
 
 namespace UnrealBuildTool
 {
-    
-    /// <summary>
-    /// Visual Studio for Mac project file generator implementation
-    /// </summary>
-    class VCMacProjectFileGenerator : VCProjectFileGenerator
-    {
-       
-        /// <summary>
-        /// Default constructor
-        /// </summary>
-        /// <param name="InOnlyGameProject">The single project to generate project files for, or null</param>
-        /// <param name="InArguments">Additional command line arguments</param>
-        public VCMacProjectFileGenerator(FileReference? InOnlyGameProject, CommandLineArguments InArguments)
+
+	/// <summary>
+	/// Visual Studio for Mac project file generator implementation
+	/// </summary>
+	class VCMacProjectFileGenerator : VCProjectFileGenerator
+	{
+
+		/// <summary>
+		/// Default constructor
+		/// </summary>
+		/// <param name="InOnlyGameProject">The single project to generate project files for, or null</param>
+		/// <param name="InArguments">Additional command line arguments</param>
+		public VCMacProjectFileGenerator(FileReference? InOnlyGameProject, CommandLineArguments InArguments)
 			: base(InOnlyGameProject, VCProjectFileFormat.Default, InArguments)
-        {
+		{
 			// no suo file, requires ole32
 			Settings.bWriteSolutionOptionFile = false;
-        }
-
-        /// True if we should include IntelliSense data in the generated project files when possible
-        override public bool ShouldGenerateIntelliSenseData()
-        {
-            return false;
-        }
-
-        protected bool IsValidProject(ProjectFile ProjectFile)
-        {
-            XmlDocument Doc = new XmlDocument();
-            Doc.Load(ProjectFile.ProjectFilePath.FullName);
-
-            XmlNodeList Elements = Doc.GetElementsByTagName("TargetFramework");
-            foreach (XmlElement? Element in Elements)
-            {
-                if (Element == null)
-                {
-                    continue;
-                }
-
-                // some projects can have TargetFramework's nested in conditionals for IsWindows, etc. Rather than try to handle that,
-                // we look if _any_ TargetFramework doesn't contain Windows and assume that means the conditionals are set up as expected
-                if (!Element.InnerText.Contains("windows"))
+		}
+
+		/// True if we should include IntelliSense data in the generated project files when possible
+		override public bool ShouldGenerateIntelliSenseData()
+		{
+			return false;
+		}
+
+		protected bool IsValidProject(ProjectFile ProjectFile)
+		{
+			XmlDocument Doc = new XmlDocument();
+			Doc.Load(ProjectFile.ProjectFilePath.FullName);
+
+			XmlNodeList Elements = Doc.GetElementsByTagName("TargetFramework");
+			foreach (XmlElement? Element in Elements)
+			{
+				if (Element == null)
+				{
+					continue;
+				}
+
+				// some projects can have TargetFramework's nested in conditionals for IsWindows, etc. Rather than try to handle that,
+				// we look if _any_ TargetFramework doesn't contain Windows and assume that means the conditionals are set up as expected
+				if (!Element.InnerText.Contains("windows"))
 				{
-                    return true;
+					return true;
 				}
-            }
-
-            // if there was no framework detected without windows, then this is not valid
-            return false;
-        }
-
-        /// <summary>
-        /// Writes the project files to disk
-        /// </summary>
-        /// <returns>True if successful</returns>
-        protected override bool WriteProjectFiles(PlatformProjectGeneratorCollection PlatformProjectGenerators, ILogger Logger)
-        {
+			}
+
+			// if there was no framework detected without windows, then this is not valid
+			return false;
+		}
+
+		/// <summary>
+		/// Writes the project files to disk
+		/// </summary>
+		/// <returns>True if successful</returns>
+		protected override bool WriteProjectFiles(PlatformProjectGeneratorCollection PlatformProjectGenerators, ILogger Logger)
+		{
 			// This can be reset by higher level code when it detects that we don't have
 			// VS2019 installed (TODO - add custom format for Mac?)
 			Settings.ProjectFileFormat = VCProjectFileFormat.VisualStudio2019;
 
-            // we can't generate native projects so clear them here, we will just
-            // write out OtherProjectFiles and AutomationProjectFiles
-            GeneratedProjectFiles.Clear();
-
-            // remove C# projects that require windows (see BuildAllScriptPlugins() looking at TargetFramework)
-            OtherProjectFiles.RemoveAll(x => !IsValidProject(x));
-            AutomationProjectFiles.RemoveAll(x => !IsValidProject(x));
-
-            if (!base.WriteProjectFiles(PlatformProjectGenerators, Logger))
-            {
-                return false;
-            }
-
-
-            // Write AutomationReferences file
-            if (AutomationProjectFiles.Any())
-            {
-                XNamespace NS = XNamespace.Get("http://schemas.microsoft.com/developer/msbuild/2003");
-
-                DirectoryReference AutomationToolDir = DirectoryReference.Combine(Unreal.EngineSourceDirectory, "Programs", "AutomationTool");
-                new XDocument(
-                    new XElement(NS + "Project",
-                        new XAttribute("ToolsVersion", VCProjectFileGenerator.GetProjectFileToolVersionString(Settings.ProjectFileFormat)),
-                        new XAttribute("DefaultTargets", "Build"),
-                        new XElement(NS + "ItemGroup",
-                            from AutomationProject in AutomationProjectFiles
-                            select new XElement(NS + "ProjectReference",
-                                new XAttribute("Include", AutomationProject.ProjectFilePath.MakeRelativeTo(AutomationToolDir)),
-                                new XElement(NS + "Project", (AutomationProject as VCSharpProjectFile)!.ProjectGUID.ToString("B")),
-                                new XElement(NS + "Name", AutomationProject.ProjectFilePath.GetFileNameWithoutExtension()),
-                                new XElement(NS + "Private", "false")
-                            )
-                        )
-                    )
-                ).Save(FileReference.Combine(IntermediateProjectFilesPath, "AutomationTool.csproj.References").FullName);
-            }
-
-            return true;
-        }
-
-    }
+			// we can't generate native projects so clear them here, we will just
+			// write out OtherProjectFiles and AutomationProjectFiles
+			GeneratedProjectFiles.Clear();
+
+			// remove C# projects that require windows (see BuildAllScriptPlugins() looking at TargetFramework)
+			OtherProjectFiles.RemoveAll(x => !IsValidProject(x));
+			AutomationProjectFiles.RemoveAll(x => !IsValidProject(x));
+
+			if (!base.WriteProjectFiles(PlatformProjectGenerators, Logger))
+			{
+				return false;
+			}
+
+
+			// Write AutomationReferences file
+			if (AutomationProjectFiles.Any())
+			{
+				XNamespace NS = XNamespace.Get("http://schemas.microsoft.com/developer/msbuild/2003");
+
+				DirectoryReference AutomationToolDir = DirectoryReference.Combine(Unreal.EngineSourceDirectory, "Programs", "AutomationTool");
+				new XDocument(
+					new XElement(NS + "Project",
+						new XAttribute("ToolsVersion", VCProjectFileGenerator.GetProjectFileToolVersionString(Settings.ProjectFileFormat)),
+						new XAttribute("DefaultTargets", "Build"),
+						new XElement(NS + "ItemGroup",
+							from AutomationProject in AutomationProjectFiles
+							select new XElement(NS + "ProjectReference",
+								new XAttribute("Include", AutomationProject.ProjectFilePath.MakeRelativeTo(AutomationToolDir)),
+								new XElement(NS + "Project", (AutomationProject as VCSharpProjectFile)!.ProjectGUID.ToString("B")),
+								new XElement(NS + "Name", AutomationProject.ProjectFilePath.GetFileNameWithoutExtension()),
+								new XElement(NS + "Private", "false")
+							)
+						)
+					)
+				).Save(FileReference.Combine(IntermediateProjectFilesPath, "AutomationTool.csproj.References").FullName);
+			}
+
+			return true;
+		}
+
+	}
 
 }
diff --git a/Engine/Source/Programs/UnrealBuildTool/ProjectFiles/Xcode/XcodeFrameworkWrapperProject.cs b/Engine/Source/Programs/UnrealBuildTool/ProjectFiles/Xcode/XcodeFrameworkWrapperProject.cs
--- a/Engine/Source/Programs/UnrealBuildTool/ProjectFiles/Xcode/XcodeFrameworkWrapperProject.cs
+++ b/Engine/Source/Programs/UnrealBuildTool/ProjectFiles/Xcode/XcodeFrameworkWrapperProject.cs
@@ -1,13 +1,12 @@
 ﻿// Copyright Epic Games, Inc. All Rights Reserved.
 
 using System;
+using System.Collections.Generic;
 using System.IO;
-using System.Text.RegularExpressions;
 using System.Linq;
-using System.Collections.Generic;
+using System.Text.RegularExpressions;
 using EpicGames.Core;
 using UnrealBuildBase;
-using Microsoft.Extensions.Logging;
 
 namespace UnrealBuildTool.ProjectFiles.Xcode
 {
@@ -56,7 +55,7 @@ namespace UnrealBuildTool.ProjectFiles.Xcode
 		/// An enumeration specifying the type of a filesystem entry, either directory, file, or something else.
 		/// </summary>
 		private enum EntryType { None, Directory, File }
-		
+
 		/// <summary>
 		/// Gets the type of filesystem entry pointed to by <paramref name="Path"/>.
 		/// </summary>
@@ -115,13 +114,13 @@ namespace UnrealBuildTool.ProjectFiles.Xcode
 				}
 			}
 		}
-		
+
 		/// <summary>
 		/// Opens each file in <paramref name="RootDirectory"/> and replaces all occurrences of <paramref name="OldValue"/>
 		/// with <paramref name="NewValue"/>.
 		/// </summary>
 		/// <param name="RootDirectory">The directory in which all files should be subject to replacements.</param>
-        /// <param name="SearchPattern">Only replace text in files that match this pattern. Default is all files.</param>
+		/// <param name="SearchPattern">Only replace text in files that match this pattern. Default is all files.</param>
 		/// <param name="OldValue">The value that should be replaced in all files.</param>
 		/// <param name="NewValue">The replacement value.</param>
 		private static void ReplaceTextInFiles(string RootDirectory, string OldValue, string NewValue, string SearchPattern = "*")
@@ -174,15 +173,15 @@ namespace UnrealBuildTool.ProjectFiles.Xcode
 					ProjectContents = ChangeProjectSetting(ProjectContents, Setting.Key, Setting.Value);
 				}
 
-                File.WriteAllText(ProjectFiles[0], ProjectContents);
+				File.WriteAllText(ProjectFiles[0], ProjectContents);
 			}
 		}
 
-        /// <summary>
-        /// Removes the readonly attribute from all files in a directory file while retaining all other attributes, thus making them writeable.
-        /// </summary>
-        /// <param name="RootDirectory">The path to the directory that will be make writeable.</param>
-        private static void MakeAllFilesWriteable(string RootDirectory)
+		/// <summary>
+		/// Removes the readonly attribute from all files in a directory file while retaining all other attributes, thus making them writeable.
+		/// </summary>
+		/// <param name="RootDirectory">The path to the directory that will be make writeable.</param>
+		private static void MakeAllFilesWriteable(string RootDirectory)
 		{
 			IEnumerable<string> FileNames = Directory.EnumerateFiles(RootDirectory, "*", SearchOption.AllDirectories);
 			foreach (string FileName in FileNames)
@@ -201,7 +200,7 @@ namespace UnrealBuildTool.ProjectFiles.Xcode
 		private static string ChangeProjectSetting(string ProjectContents, string SettingName, string? SettingValue)
 		{
 			string SettingNameRegexString = String.Format("(\\s+{0}\\s=\\s)\"?(.+)\"?;", SettingName);
-			
+
 			string SettingValueReplaceString = String.Format("$1\"{0}\";", SettingValue);
 
 			Regex SettingNameRegex = new Regex(SettingNameRegexString);
@@ -220,7 +219,7 @@ namespace UnrealBuildTool.ProjectFiles.Xcode
 				"Build",
 				"xcuserdata"
 			};
-			
+
 			IEnumerable<string> Directories = Directory.EnumerateDirectories(RootDirectory, "*", SearchOption.AllDirectories);
 			foreach (string Dir in Directories)
 			{
@@ -240,7 +239,7 @@ namespace UnrealBuildTool.ProjectFiles.Xcode
 		/// settings in the project to accommodate the framework, and replacing text in all the files to match the framework.
 		/// </summary>
 		/// <param name="OutputDirectory">The directory in which to place the framework. The framework will be placed in 'outputDirectory/frameworkName/'.</param>
-        /// <param name="ProjectName">The name of the project. If blueprint-only, use the actual name of the project, not just UnrealGame.</param>
+		/// <param name="ProjectName">The name of the project. If blueprint-only, use the actual name of the project, not just UnrealGame.</param>
 		/// <param name="FrameworkName">The name of the framework that this project is wrapping.</param>
 		/// <param name="BundleId">The Bundle ID to give to the wrapper project.</param>
 		/// <param name="SrcFrameworkPath">The path to the directory containing the framework to be wrapped.</param>
@@ -269,10 +268,10 @@ namespace UnrealBuildTool.ProjectFiles.Xcode
 		}
 
 	}
-	
+
 	class XcodeFrameworkWrapperUtils
 	{
-        private static ConfigHierarchy GetIni(DirectoryReference ProjectDirectory)
+		private static ConfigHierarchy GetIni(DirectoryReference ProjectDirectory)
 		{
 			return ConfigCache.ReadHierarchy(ConfigHierarchyType.Engine, ProjectDirectory, UnrealTargetPlatform.IOS);
 			//return ConfigCache.ReadHierarchy(ConfigHierarchyType.Engine, DirectoryReference.FromFile(ProjectFile), UnrealTargetPlatform.IOS);
@@ -305,7 +304,7 @@ namespace UnrealBuildTool.ProjectFiles.Xcode
 			Ini.GetBool("/Script/IOSRuntimeSettings.IOSRuntimeSettings", "bBuildAsFramework", out bBuildAsFramework);
 			return bBuildAsFramework;
 		}
-		
+
 		public static bool GetGenerateFrameworkWrapperProject(DirectoryReference ProjectDirectory)
 		{
 			ConfigHierarchy Ini = GetIni(ProjectDirectory);
@@ -314,4 +313,4 @@ namespace UnrealBuildTool.ProjectFiles.Xcode
 			return bGenerateFrameworkWrapperProject;
 		}
 	}
-} 
\ No newline at end of file
+}
\ No newline at end of file
diff --git a/Engine/Source/Programs/UnrealBuildTool/ProjectFiles/Xcode/XcodeProject.cs b/Engine/Source/Programs/UnrealBuildTool/ProjectFiles/Xcode/XcodeProject.cs
--- a/Engine/Source/Programs/UnrealBuildTool/ProjectFiles/Xcode/XcodeProject.cs
+++ b/Engine/Source/Programs/UnrealBuildTool/ProjectFiles/Xcode/XcodeProject.cs
@@ -6,12 +6,9 @@ using System.IO;
 using System.Linq;
 using System.Text;
 using EpicGames.Core;
-using UnrealBuildBase;
-using Microsoft.Extensions.Logging;
-using System.Text.RegularExpressions;
 using Microsoft.CodeAnalysis;
-using System.Configuration;
-using System.Xml.Linq;
+using Microsoft.Extensions.Logging;
+using UnrealBuildBase;
 
 /*****
 
@@ -158,7 +155,7 @@ namespace UnrealBuildTool.XcodeProjectXcconfig
 		//	XcodeProjectRelative = File.MakeRelativeTo(XcodeProjectFile.ParentDirectory!);
 		//}
 
-		
+
 		// Location: Can be key of setting entry in .ini, or the full file path
 		// CopyFromFolderIfNotFound: If the file at "Location" does not exist, try to copy the same named file from this folder
 		public MetadataItem(DirectoryReference ProductDirectory, DirectoryReference XcodeProject, ConfigHierarchy Ini, string Location, MetadataMode InMode, DirectoryReference CopyFromFolderIfNotFound)
@@ -445,7 +442,7 @@ namespace UnrealBuildTool.XcodeProjectXcconfig
 				}
 			}
 
-			if (ProjectFile.IsContentOnlyProject && XcodeProjectFileGenerator.Current?.OnlyGameProject != null && 
+			if (ProjectFile.IsContentOnlyProject && XcodeProjectFileGenerator.Current?.OnlyGameProject != null &&
 				XcodeProjectFileGenerator.Current.OnlyGameProject.IsUnderDirectory(ProjectFile.BaseDir))
 			{
 				UProjectFileLocation = XcodeProjectFileGenerator.Current.OnlyGameProject;
@@ -519,7 +516,7 @@ namespace UnrealBuildTool.XcodeProjectXcconfig
 			}
 
 			_TargetRules = AllConfigs[0].ProjectTarget!.TargetRules;
-			
+
 			// this project makes an app bundle (.app directory instead of a raw executable or dylib) if none of the fings make a non-appbundle
 			bIsAppBundle = !AllConfigs.Any(x => x.ProjectTarget!.TargetRules!.bIsBuildingConsoleApplication || x.ProjectTarget.TargetRules.bShouldCompileAsDLL);
 
@@ -565,10 +562,10 @@ namespace UnrealBuildTool.XcodeProjectXcconfig
 				FinalPath = FinalPath.Replace("$(Project)", ProductDirectory.FullName);
 				FinalPath = FinalPath.Replace("$(Platform)", Platform.ToString());
 
-//				Console.WriteLine($"Looking for {FinalPath}");
+				//				Console.WriteLine($"Looking for {FinalPath}");
 				if (File.Exists(FinalPath) || Directory.Exists(FinalPath))
 				{
-//					Console.WriteLine($"  Found it!");
+					//					Console.WriteLine($"  Found it!");
 					if (bMakeRelative)
 					{
 						FinalPath = new FileReference(FinalPath).MakeRelativeTo(XcodeProjectFileLocation.ParentDirectory!);
@@ -685,8 +682,8 @@ namespace UnrealBuildTool.XcodeProjectXcconfig
 
 											// Get the output directory
 											DirectoryReference RootDirectory;
-											if (UProjectDirectory != null && 
-												(bShouldCompileMonolithic || ProjectTarget.TargetRules!.BuildEnvironment == TargetBuildEnvironment.Unique) && 
+											if (UProjectDirectory != null &&
+												(bShouldCompileMonolithic || ProjectTarget.TargetRules!.BuildEnvironment == TargetBuildEnvironment.Unique) &&
 												ProjectTarget.TargetRules!.File!.IsUnderDirectory(UProjectDirectory))
 											{
 												RootDirectory = UEBuildTarget.GetOutputDirectoryForExecutable(UProjectDirectory, ProjectTarget.TargetRules.File!);
@@ -789,7 +786,7 @@ namespace UnrealBuildTool.XcodeProjectXcconfig
 
 		public virtual void WriteXcconfigFile()
 		{
-			
+
 		}
 
 
@@ -817,7 +814,7 @@ namespace UnrealBuildTool.XcodeProjectXcconfig
 					WrittenNodes.Add(Reference);
 					WriteNodeAndReferences(Content, Reference, WrittenNodes);
 				}
-			}	
+			}
 		}
 	}
 
@@ -879,12 +876,12 @@ namespace UnrealBuildTool.XcodeProjectXcconfig
 		{
 			Content.WriteLine($"/* Begin {IsAType} section */");
 			Content.WriteLine(2, $"{Guid} = {{");
-			Content.WriteLine(3,	$"isa = {IsAType};");
-			Content.WriteLine(3,	"buildActionMask = 2147483647;");
-			Content.WriteLine(3,	"files = (");
+			Content.WriteLine(3, $"isa = {IsAType};");
+			Content.WriteLine(3, "buildActionMask = 2147483647;");
+			Content.WriteLine(3, "files = (");
 			foreach (XcodeSourceFile File in FileItems)
 			{
-				Content.WriteLine(4,		$"{File.FileGuid} /* {File.Reference.GetFileName()} in {Name} */,");
+				Content.WriteLine(4, $"{File.FileGuid} /* {File.Reference.GetFileName()} in {Name} */,");
 			}
 			Content.WriteLine(3, ");");
 
@@ -982,7 +979,7 @@ namespace UnrealBuildTool.XcodeProjectXcconfig
 
 	class XcodeShellScriptBuildPhase : XcodeBuildPhase
 	{
-		public XcodeShellScriptBuildPhase(string Name, IEnumerable<string> ScriptLines, IEnumerable<string> Inputs, IEnumerable<string> Outputs, bool bInstallOnly=false)
+		public XcodeShellScriptBuildPhase(string Name, IEnumerable<string> ScriptLines, IEnumerable<string> Inputs, IEnumerable<string> Outputs, bool bInstallOnly = false)
 			: base(Name, "PBXShellScriptBuildPhase")
 		{
 			MiscItems.Add($"name = \"{Name}\";");
@@ -1031,25 +1028,25 @@ namespace UnrealBuildTool.XcodeProjectXcconfig
 		public override void Write(StringBuilder Content)
 		{
 			Content.WriteLine(2, $"{Guid} /* {Info.DisplayName} */ = {{");
-			Content.WriteLine(3,	"isa = XCBuildConfiguration;");
+			Content.WriteLine(3, "isa = XCBuildConfiguration;");
 			if (Xcconfig != null)
 			{
-				Content.WriteLine(3,	$"baseConfigurationReference = {Xcconfig.Guid} /* {Xcconfig.Name}.xcconfig */;");
+				Content.WriteLine(3, $"baseConfigurationReference = {Xcconfig.Guid} /* {Xcconfig.Name}.xcconfig */;");
 			}
-			Content.WriteLine(3,	"buildSettings = {");
+			Content.WriteLine(3, "buildSettings = {");
 			if (bIncludeAllPlatforms)
 			{
-				Content.WriteLine(4,		$"SUPPORTED_PLATFORMS = \"macosx iphonesimulator iphoneos appletvsimulator appletvos\";");
-				Content.WriteLine(4,		$"ONLY_ACTIVE_ARCH = YES;");
+				Content.WriteLine(4, $"SUPPORTED_PLATFORMS = \"macosx iphonesimulator iphoneos appletvsimulator appletvos\";");
+				Content.WriteLine(4, $"ONLY_ACTIVE_ARCH = YES;");
 			}
-			Content.WriteLine(3,	"};");
-			Content.WriteLine(3,	$"name = \"{Info.DisplayName}\";");
+			Content.WriteLine(3, "};");
+			Content.WriteLine(3, $"name = \"{Info.DisplayName}\";");
 			Content.WriteLine(2, "};");
 		}
 	}
-	
+
 	class XcodeBuildConfigList : XcodeProjectNode
-	{ 
+	{
 		public string Guid = XcodeProjectFileGenerator.MakeXcodeGuid();
 		public string TargetName;
 		private UnrealData UnrealData;
@@ -1095,15 +1092,15 @@ namespace UnrealBuildTool.XcodeProjectXcconfig
 			}
 
 			Content.WriteLine(2, $"{Guid} /* Build configuration list for target {TargetName} */ = {{");
-			Content.WriteLine(3,	"isa = XCConfigurationList;");
-			Content.WriteLine(3,	"buildConfigurations = (");
+			Content.WriteLine(3, "isa = XCConfigurationList;");
+			Content.WriteLine(3, "buildConfigurations = (");
 			foreach (XcodeBuildConfig Config in BuildConfigs)
 			{
-				Content.WriteLine(4,		$"{Config.Guid} /* {Config.Info.DisplayName} */,");
+				Content.WriteLine(4, $"{Config.Guid} /* {Config.Info.DisplayName} */,");
 			}
-			Content.WriteLine(3,	");");
-			Content.WriteLine(3,	"defaultConfigurationIsVisible = 0;");
-			Content.WriteLine(3,	$"defaultConfigurationName = \"{Default}\";");
+			Content.WriteLine(3, ");");
+			Content.WriteLine(3, "defaultConfigurationIsVisible = 0;");
+			Content.WriteLine(3, $"defaultConfigurationName = \"{Default}\";");
 			Content.WriteLine(2, "};");
 		}
 	}
@@ -1120,7 +1117,7 @@ namespace UnrealBuildTool.XcodeProjectXcconfig
 
 		// Guid for this target
 		public string Guid = XcodeProjectFileGenerator.MakeXcodeGuid();
-//		string TargetAppGuid = XcodeProjectFileGenerator.MakeXcodeGuid();
+		//		string TargetAppGuid = XcodeProjectFileGenerator.MakeXcodeGuid();
 
 		// com.apple.product-type.application, etc
 		string ProductType;
@@ -1146,7 +1143,7 @@ namespace UnrealBuildTool.XcodeProjectXcconfig
 
 		private FileReference? GameProject;
 
-		public XcodeTarget(Type Type, UnrealData UnrealData, string? OverrideName=null)
+		public XcodeTarget(Type Type, UnrealData UnrealData, string? OverrideName = null)
 		{
 			// when we are content only, we do not want to build with the uproject on the commandline, we will be building UnrealGame 
 			GameProject = UnrealData.bIsContentOnlyProject ? null : UnrealData.UProjectFileLocation;
@@ -1196,9 +1193,9 @@ namespace UnrealBuildTool.XcodeProjectXcconfig
 			Content.WriteLine($"/* Begin {TargetType} section */");
 
 			Content.WriteLine(2, $"{Guid} /* {Name} */ = {{");
-			Content.WriteLine(3,	$"isa = {TargetTypeName};");
-				
-			Content.WriteLine(3,	$"buildConfigurationList = {BuildConfigList!.Guid} /* Build configuration list for {TargetTypeName} \"{Name}\" */;");
+			Content.WriteLine(3, $"isa = {TargetTypeName};");
+
+			Content.WriteLine(3, $"buildConfigurationList = {BuildConfigList!.Guid} /* Build configuration list for {TargetTypeName} \"{Name}\" */;");
 
 			if (TargetType == Type.Build)
 			{
@@ -1208,25 +1205,25 @@ namespace UnrealBuildTool.XcodeProjectXcconfig
 				string BuildToolPath = UEDir + "/Engine/Build/BatchFiles/Mac/XcodeBuild.sh";
 
 				// insert elements to call UBT when building
-				Content.WriteLine(3,	$"buildArgumentsString = \"$(ACTION) {UnrealTargetName} $(PLATFORM_NAME) $(CONFIGURATION) {UProjectParam}\";");
-				Content.WriteLine(3,	$"buildToolPath = \"{BuildToolPath}\";");
-				Content.WriteLine(3,	$"buildWorkingDirectory = \"{UEDir}\";");
+				Content.WriteLine(3, $"buildArgumentsString = \"$(ACTION) {UnrealTargetName} $(PLATFORM_NAME) $(CONFIGURATION) {UProjectParam}\";");
+				Content.WriteLine(3, $"buildToolPath = \"{BuildToolPath}\";");
+				Content.WriteLine(3, $"buildWorkingDirectory = \"{UEDir}\";");
 			}
-			Content.WriteLine(3,	"buildPhases = (");
+			Content.WriteLine(3, "buildPhases = (");
 			foreach (XcodeBuildPhase BuildPhase in BuildPhases)
 			{
-				Content.WriteLine(4,		$"{BuildPhase.Guid} /* {BuildPhase.Name} */,");
+				Content.WriteLine(4, $"{BuildPhase.Guid} /* {BuildPhase.Name} */,");
 			}
-			Content.WriteLine(3,	");");
-			Content.WriteLine(3,	"dependencies = (");
+			Content.WriteLine(3, ");");
+			Content.WriteLine(3, "dependencies = (");
 			foreach (XcodeDependency Dependency in Dependencies)
 			{
-				Content.WriteLine(4,		$"{Dependency.Guid} /* {Dependency.Target.Name} */,");
+				Content.WriteLine(4, $"{Dependency.Guid} /* {Dependency.Target.Name} */,");
 			}
-			Content.WriteLine(3,	");");
-			Content.WriteLine(3,	$"name = \"{Name}\";");
-			Content.WriteLine(3,	"passBuildSettingsInEnvironment = 1;");
-			Content.WriteLine(3,	$"productType = \"{ProductType}\";");
+			Content.WriteLine(3, ");");
+			Content.WriteLine(3, $"name = \"{Name}\";");
+			Content.WriteLine(3, "passBuildSettingsInEnvironment = 1;");
+			Content.WriteLine(3, $"productType = \"{ProductType}\";");
 			WriteExtraTargetProperties(Content);
 			Content.WriteLine(2, "};");
 
@@ -1257,7 +1254,7 @@ namespace UnrealBuildTool.XcodeProjectXcconfig
 			this.Platform = Platform;
 			this.UnrealData = Project.UnrealData;
 
-			BuildConfigList = new XcodeBuildConfigList(Platform, Name, Project.UnrealData, bIncludeAllPlatformsInConfig:false);
+			BuildConfigList = new XcodeBuildConfigList(Platform, Name, Project.UnrealData, bIncludeAllPlatformsInConfig: false);
 			References.Add(BuildConfigList);
 
 			// add the Product item to the project to be visible in left pane
@@ -1434,7 +1431,7 @@ namespace UnrealBuildTool.XcodeProjectXcconfig
 			});
 			string DsymScriptInput = $"\\\"$(CONFIGURATION_BUILD_DIR)/$(EXECUTABLE_PATH)\\\"";
 			string DsymScriptOutput = $"\\\"$(DWARF_DSYM_FOLDER_PATH)/$(DWARF_DSYM_FILE_NAME)\\\"";
-			XcodeShellScriptBuildPhase DsymScriptPhase = new("Generate dsym for archive, and strip", DsymScript, new string[] { DsymScriptInput }, new string[] { DsymScriptOutput }, bInstallOnly:true);
+			XcodeShellScriptBuildPhase DsymScriptPhase = new("Generate dsym for archive, and strip", DsymScript, new string[] { DsymScriptInput }, new string[] { DsymScriptOutput }, bInstallOnly: true);
 			BuildPhases.Add(DsymScriptPhase);
 			References.Add(DsymScriptPhase);
 		}
@@ -1583,7 +1580,7 @@ namespace UnrealBuildTool.XcodeProjectXcconfig
 			UnrealBuildConfig BuildConfig = BuildConfigList!.BuildConfigs[0].Info;
 			TargetRules TargetRules = BuildConfig.ProjectTarget!.TargetRules!;
 			DirectoryReference ProjectOrEngineDir = UnrealData.UProjectFileLocation?.Directory ?? Unreal.EngineDirectory;
-			
+
 			// point to the shader Engine/Binaries for content only project (sadly, the TargetRules.OutputFile is not filled out)
 			DirectoryReference ConfigBuildDir = (TargetRules.Type == TargetType.Editor || UnrealData.bIsContentOnlyProject) ? Unreal.EngineDirectory : ProjectOrEngineDir;
 			if (TargetRules.Type == TargetType.Program && TargetRules.File!.IsUnderDirectory(Unreal.EngineDirectory))
@@ -1593,7 +1590,7 @@ namespace UnrealBuildTool.XcodeProjectXcconfig
 			string BinariesBaseDir = ConfigBuildDir.GetDirectoryName();
 
 			ConfigBuildDir = DirectoryReference.Combine(ConfigBuildDir, "Binaries", Platform.ToString(), TargetRules.ExeBinariesSubFolder);
-			
+
 			MetadataPlatform MetadataPlatform;
 			string TargetPlatformName = Platform.ToString();
 			if (TargetRules.Type != TargetType.Game && TargetRules.Type != TargetType.Program)
@@ -1828,7 +1825,7 @@ namespace UnrealBuildTool.XcodeProjectXcconfig
 
 
 				MetadataItem EntitlementsMetadata = UnrealData.Metadata!.EntitlementsFiles[MetadataPlatform];
-				
+
 				// if we have shipping specified, use it instead
 				if (Config.BuildConfig == UnrealTargetConfiguration.Shipping)
 				{
@@ -1969,10 +1966,10 @@ namespace UnrealBuildTool.XcodeProjectXcconfig
 
 			ProjectBuildConfigs = new XcodeBuildConfigList(bIsStubEditor ? null : Platform, UnrealData.ProductName, UnrealData, bIncludeAllPlatformsInConfig: true);
 			References.Add(ProjectBuildConfigs);
-			
+
 			// make an indexing target if we aren't just a run-only project, and it has buildable source files
 			if (!XcodeProjectFileGenerator.bGenerateRunOnlyProject && UnrealData.BatchedFiles.Count != 0)
-			{				
+			{
 				// index isn't a dependency of run, it's simply a target that xcode will find to index from
 				XcodeIndexTarget IndexTarget = new XcodeIndexTarget(this);
 				References.Add(IndexTarget);
@@ -1984,35 +1981,35 @@ namespace UnrealBuildTool.XcodeProjectXcconfig
 			Content.WriteLine("/* Begin PBXProject section */");
 
 			Content.WriteLine(2, $"{Guid} /* Project object */ = {{");
-			Content.WriteLine(3,	"isa = PBXProject;");
-			Content.WriteLine(3,	"attributes = {");
-			Content.WriteLine(4,		"LastUpgradeCheck = 2000;");
-			Content.WriteLine(4,		"ORGANIZATIONNAME = \"Epic Games, Inc.\";");
-			Content.WriteLine(4,		"TargetAttributes = {");
+			Content.WriteLine(3, "isa = PBXProject;");
+			Content.WriteLine(3, "attributes = {");
+			Content.WriteLine(4, "LastUpgradeCheck = 2000;");
+			Content.WriteLine(4, "ORGANIZATIONNAME = \"Epic Games, Inc.\";");
+			Content.WriteLine(4, "TargetAttributes = {");
 			foreach (XcodeRunTarget RunTarget in RunTargets)
-			{ 
-				Content.WriteLine(5,			$"{RunTarget.Guid} = {{");
-				Content.WriteLine(6,				$"ProvisioningStyle = {ProvisioningStyle};");
-				Content.WriteLine(5,			"};");
-			}
-			Content.WriteLine(4,		"};");
-			Content.WriteLine(3,	"};");
-			Content.WriteLine(3,	$"buildConfigurationList = {ProjectBuildConfigs.Guid} /* Build configuration list for PBXProject \"{ProjectBuildConfigs.TargetName}\" */;");
-			Content.WriteLine(3,	"compatibilityVersion = \"Xcode 8.0\";");
-			Content.WriteLine(3,	"developmentRegion = English;");
-			Content.WriteLine(3,	"hasScannedForEncodings = 0;");
-			Content.WriteLine(3,	"knownRegions = (");
-			Content.WriteLine(4,		"en");
-			Content.WriteLine(3,	");");
-			Content.WriteLine(3,	$"mainGroup = {FileCollection.MainGroupGuid};");
+			{
+				Content.WriteLine(5, $"{RunTarget.Guid} = {{");
+				Content.WriteLine(6, $"ProvisioningStyle = {ProvisioningStyle};");
+				Content.WriteLine(5, "};");
+			}
+			Content.WriteLine(4, "};");
+			Content.WriteLine(3, "};");
+			Content.WriteLine(3, $"buildConfigurationList = {ProjectBuildConfigs.Guid} /* Build configuration list for PBXProject \"{ProjectBuildConfigs.TargetName}\" */;");
+			Content.WriteLine(3, "compatibilityVersion = \"Xcode 8.0\";");
+			Content.WriteLine(3, "developmentRegion = English;");
+			Content.WriteLine(3, "hasScannedForEncodings = 0;");
+			Content.WriteLine(3, "knownRegions = (");
+			Content.WriteLine(4, "en");
+			Content.WriteLine(3, ");");
+			Content.WriteLine(3, $"mainGroup = {FileCollection.MainGroupGuid};");
 			// for stub editor projects, we don't have a run target, so we don't have a product folder
 			if (RunTargets.Count > 0)
 			{
 				Content.WriteLine(3, $"productRefGroup = {FileCollection.GetProductGroupGuid()};");
 			}
-			Content.WriteLine(3,	"projectDirPath = \"\";");
-			Content.WriteLine(3,	"projectRoot = \"\";");
-			Content.WriteLine(3,	"targets = (");
+			Content.WriteLine(3, "projectDirPath = \"\";");
+			Content.WriteLine(3, "projectRoot = \"\";");
+			Content.WriteLine(3, "targets = (");
 			foreach (XcodeTarget Target in XcodeProjectNode.GetNodesOfType<XcodeTarget>(this))
 			{
 				Content.WriteLine(4, $"{Target.Guid} /* {Target.Name} */,");
@@ -2131,7 +2128,7 @@ namespace UnrealBuildTool.XcodeProjectXcconfig
 			Dictionary<DirectoryReference, int> BuildFileMap = new();
 			foreach (XcodeSourceFile SourceFile in SourceFiles.OfType<XcodeSourceFile>())
 			{
-				SharedFileCollection.ProcessFile(SourceFile, bIsForBuild:IsGeneratedProject, bIsFolder:false, SourceToBuildFileMap:BuildFileMap);
+				SharedFileCollection.ProcessFile(SourceFile, bIsForBuild: IsGeneratedProject, bIsFolder: false, SourceToBuildFileMap: BuildFileMap);
 			}
 
 			// cache the main group
@@ -2350,7 +2347,7 @@ namespace UnrealBuildTool.XcodeProjectXcconfig
 			string TemplateMainGroupGuid = Plist($"Print :objects:{ProjectGuid}:mainGroup");
 
 			// fixup paths that were relative to original project to be relative to merged project
-//			List<string> ObjectGuids = PlistObjects();
+			//			List<string> ObjectGuids = PlistObjects();
 			IEnumerable<string> MainGroupChildrenGuids = XcodeUtils.PlistArray($":objects:{TemplateMainGroupGuid}:children");
 
 			string RelativeFromMergedToTemplate = TemplateProject.Directory.ParentDirectory!.MakeRelativeTo(PBXProjFilePath.Directory.ParentDirectory!);
diff --git a/Engine/Source/Programs/UnrealBuildTool/ProjectFiles/Xcode/XcodeProjectFileGenerator.cs b/Engine/Source/Programs/UnrealBuildTool/ProjectFiles/Xcode/XcodeProjectFileGenerator.cs
--- a/Engine/Source/Programs/UnrealBuildTool/ProjectFiles/Xcode/XcodeProjectFileGenerator.cs
+++ b/Engine/Source/Programs/UnrealBuildTool/ProjectFiles/Xcode/XcodeProjectFileGenerator.cs
@@ -2,13 +2,13 @@
 
 using System;
 using System.Collections.Generic;
-using System.Text;
 using System.IO;
-using EpicGames.Core;
 using System.Linq;
+using System.Text;
+using System.Threading.Tasks;
+using EpicGames.Core;
 using Microsoft.Extensions.Logging;
 using UnrealBuildBase;
-using System.Threading.Tasks;
 
 namespace UnrealBuildTool
 {
@@ -56,7 +56,7 @@ namespace UnrealBuildTool
 		/// Store the single game project (when using -game -project=...) to a place that XcodeProjectLegacy can easily retrieve it
 		/// </summary>
 		public static FileReference? SingleGameProject = null;
-		
+
 		/// <summary>
 		/// If set, only write out one target/project 
 		/// </summary>
diff --git a/Engine/Source/Programs/UnrealBuildTool/Storage/IStorageProvider.cs b/Engine/Source/Programs/UnrealBuildTool/Storage/IStorageProvider.cs
--- a/Engine/Source/Programs/UnrealBuildTool/Storage/IStorageProvider.cs
+++ b/Engine/Source/Programs/UnrealBuildTool/Storage/IStorageProvider.cs
@@ -1,14 +1,7 @@
 // Copyright Epic Games, Inc. All Rights Reserved.
 
 using System;
-using System.Collections.Generic;
-using System.Diagnostics;
 using System.IO;
-using System.Linq;
-using System.Security.Cryptography;
-using System.Text;
-using System.Threading;
-using System.Threading.Tasks;
 using EpicGames.Core;
 
 namespace UnrealBuildTool.Storage
diff --git a/Engine/Source/Programs/UnrealBuildTool/Storage/Impl/FileSystemStorageProvider.cs b/Engine/Source/Programs/UnrealBuildTool/Storage/Impl/FileSystemStorageProvider.cs
--- a/Engine/Source/Programs/UnrealBuildTool/Storage/Impl/FileSystemStorageProvider.cs
+++ b/Engine/Source/Programs/UnrealBuildTool/Storage/Impl/FileSystemStorageProvider.cs
@@ -1,10 +1,8 @@
 // Copyright Epic Games, Inc. All Rights Reserved.
 
 using System;
-using System.Collections.Generic;
 using System.Diagnostics;
 using System.IO;
-using System.Text;
 using EpicGames.Core;
 using UnrealBuildBase;
 
diff --git a/Engine/Source/Programs/UnrealBuildTool/System/ActionGraph.cs b/Engine/Source/Programs/UnrealBuildTool/System/ActionGraph.cs
--- a/Engine/Source/Programs/UnrealBuildTool/System/ActionGraph.cs
+++ b/Engine/Source/Programs/UnrealBuildTool/System/ActionGraph.cs
@@ -1,8 +1,9 @@
 // Copyright Epic Games, Inc. All Rights Reserved.
 
-using EpicGames.Core;
 using System;
+using System.Collections.Concurrent;
 using System.Collections.Generic;
+using System.Collections.Immutable;
 using System.Diagnostics;
 using System.Globalization;
 using System.IO;
@@ -10,15 +11,11 @@ using System.Linq;
 using System.Text;
 using System.Text.Json;
 using System.Text.Json.Serialization;
-using System.Threading;
 using System.Threading.Tasks;
+using EpicGames.Core;
+using Microsoft.Extensions.Logging;
 using OpenTracing.Util;
 using UnrealBuildBase;
-using Microsoft.Extensions.Logging;
-using System.Collections.Concurrent;
-using System.Collections.Immutable;
-using System.Reflection;
-using System.Runtime.CompilerServices;
 using UnrealBuildTool.Artifacts;
 
 namespace UnrealBuildTool
@@ -312,7 +309,7 @@ namespace UnrealBuildTool
 						}
 
 						// don't look in Intermediate directories - these aren't portable between machines, so don't need to be cbecked for length underneath the root
-						if (ProducedItem.Location.FullName.Length > Unreal.RootDirectory.FullName.Length + BuildConfiguration.MaxNestedPathLength && 
+						if (ProducedItem.Location.FullName.Length > Unreal.RootDirectory.FullName.Length + BuildConfiguration.MaxNestedPathLength &&
 							ProducedItem.Location.IsUnderDirectory(Unreal.RootDirectory) &&
 							!ProducedItem.Location.ContainsName("Intermediate", Unreal.RootDirectory)
 							)
@@ -690,7 +687,7 @@ namespace UnrealBuildTool
 				// where aborting an earlier compile produced invalid zero-sized obj files, but that may cause actions where that's
 				// legitimate output to always be considered outdated.
 				if (ProducedItem.Exists && (RootAction.ActionType != ActionType.Compile || ProducedItem.Length > 0 ||
-				                            (!ProducedItem.Location.HasExtension(".obj") && !ProducedItem.Location.HasExtension(".o"))))
+											(!ProducedItem.Location.HasExtension(".obj") && !ProducedItem.Location.HasExtension(".o"))))
 				{
 					// Find the newer of LastWriteTime and CreationTime, as copied files (such as from a build cache) can have a newer creation time in some cases
 					DateTime ExecutionTimeUtc = ProducedItem.LastWriteTimeUtc > ProducedItem.CreationTimeUtc ? ProducedItem.LastWriteTimeUtc : ProducedItem.CreationTimeUtc;
@@ -723,7 +720,7 @@ namespace UnrealBuildTool
 							// Need to check for import libraries here too
 							if (!bIgnoreOutdatedImportLibraries || !IsImportLibraryDependency(RootAction, PrerequisiteItem))
 							{
-								Logger.LogDebug("{StatusDescription}: Prerequisite {PrerequisiteItem} is newer than the last execution of the action: {Message}", RootAction.StatusDescription, PrerequisiteItem.Location, 
+								Logger.LogDebug("{StatusDescription}: Prerequisite {PrerequisiteItem} is newer than the last execution of the action: {Message}", RootAction.StatusDescription, PrerequisiteItem.Location,
 									$"{PrerequisiteItem.LastWriteTimeUtc.ToLocalTime().ToString(CultureInfo.CurrentCulture)} vs {LastExecutionTimeUtc.LocalDateTime.ToString(CultureInfo.CurrentCulture)}");
 								bIsOutdated = true;
 								break;
@@ -983,7 +980,7 @@ namespace UnrealBuildTool
 			{
 				System.Collections.DictionaryEntry Pair = (System.Collections.DictionaryEntry)Object!;
 				if (!UnrealBuildTool.InitialEnvironment!.Contains(Pair.Key) ||
-				    (string)(UnrealBuildTool.InitialEnvironment[Pair.Key]!) != (string)(Pair.Value!))
+					(string)(UnrealBuildTool.InitialEnvironment[Pair.Key]!) != (string)(Pair.Value!))
 				{
 					Writer.WriteValue((string)Pair.Key, (string)Pair.Value!);
 				}
diff --git a/Engine/Source/Programs/UnrealBuildTool/System/ActionHistory.cs b/Engine/Source/Programs/UnrealBuildTool/System/ActionHistory.cs
--- a/Engine/Source/Programs/UnrealBuildTool/System/ActionHistory.cs
+++ b/Engine/Source/Programs/UnrealBuildTool/System/ActionHistory.cs
@@ -6,12 +6,11 @@ using System.Collections.Generic;
 using System.Diagnostics;
 using System.IO;
 using System.Linq;
-using System.Runtime.Serialization.Formatters.Binary;
 using System.Security.Cryptography;
 using System.Text;
 using EpicGames.Core;
-using UnrealBuildBase;
 using Microsoft.Extensions.Logging;
+using UnrealBuildBase;
 
 namespace UnrealBuildTool
 {
@@ -58,7 +57,7 @@ namespace UnrealBuildTool
 		{
 			this.Location = Location;
 
-			if(FileReference.Exists(Location))
+			if (FileReference.Exists(Location))
 			{
 				Load(Logger);
 			}
@@ -71,10 +70,10 @@ namespace UnrealBuildTool
 		{
 			try
 			{
-				using(BinaryArchiveReader Reader = new BinaryArchiveReader(Location))
+				using (BinaryArchiveReader Reader = new BinaryArchiveReader(Location))
 				{
 					int Version = Reader.ReadInt();
-					if(Version != CurrentVersion)
+					if (Version != CurrentVersion)
 					{
 						Logger.LogDebug("Unable to read action history from {Location}; version {Version} vs current {CurrentVersion}", Location, Version, CurrentVersion);
 						return;
@@ -83,7 +82,7 @@ namespace UnrealBuildTool
 					OutputItemToAttributeHash = new ConcurrentDictionary<FileItem, byte[]>(Reader.ReadDictionary(() => Reader.ReadFileItem()!, () => Reader.ReadFixedSizeByteArray(HashLength))!);
 				}
 			}
-			catch(Exception Ex)
+			catch (Exception Ex)
 			{
 				Logger.LogWarning("Unable to read {Location}. See log for additional information.", Location);
 				Logger.LogDebug("{Ex}", ExceptionUtils.FormatExceptionDetails(Ex));
@@ -127,9 +126,9 @@ namespace UnrealBuildTool
 		/// <returns>True if the hashes are equal</returns>
 		static bool CompareHashes(byte[] A, byte[] B)
 		{
-			for(int Idx = 0; Idx < HashLength; Idx++)
+			for (int Idx = 0; Idx < HashLength; Idx++)
 			{
-				if(A[Idx] != B[Idx])
+				if (A[Idx] != B[Idx])
 				{
 					return false;
 				}
@@ -147,7 +146,7 @@ namespace UnrealBuildTool
 		{
 			byte[] NewHash = ComputeHash(Attributes);
 
-			for (;;)
+			for (; ; )
 			{
 				if (OutputItemToAttributeHash.TryAdd(File, NewHash))
 				{
@@ -175,8 +174,8 @@ namespace UnrealBuildTool
 							}
 						}
 					}
-				}				
-			}		
+				}
+			}
 		}
 
 		/// <summary>
@@ -190,7 +189,7 @@ namespace UnrealBuildTool
 		public static FileReference GetEngineLocation(string TargetName, UnrealTargetPlatform Platform, TargetType TargetType, UnrealArchitectures Architectures)
 		{
 			string AppName;
-			if(TargetType == TargetType.Program)
+			if (TargetType == TargetType.Program)
 			{
 				AppName = TargetName;
 			}
@@ -226,11 +225,11 @@ namespace UnrealBuildTool
 		/// <returns>Dependency cache hierarchy for the given project</returns>
 		public static IEnumerable<FileReference> GetFilesToClean(FileReference ProjectFile, string TargetName, UnrealTargetPlatform Platform, TargetType TargetType, UnrealArchitectures Architectures)
 		{
-			if(ProjectFile == null || !Unreal.IsEngineInstalled())
+			if (ProjectFile == null || !Unreal.IsEngineInstalled())
 			{
 				yield return GetEngineLocation(TargetName, Platform, TargetType, Architectures);
 			}
-			if(ProjectFile != null)
+			if (ProjectFile != null)
 			{
 				yield return GetProjectLocation(ProjectFile, TargetName, Platform, Architectures);
 			}
@@ -284,7 +283,7 @@ namespace UnrealBuildTool
 				lock (LockObject)
 				{
 					Layer = Layers.FirstOrDefault(x => x.Location == LayerLocation);
-					if(Layer == null)
+					if (Layer == null)
 					{
 						Layer = new ActionHistoryLayer(LayerLocation, Logger);
 
@@ -376,7 +375,7 @@ namespace UnrealBuildTool
 		/// </summary>
 		public void Save()
 		{
-			foreach(ActionHistoryLayer Layer in Layers)
+			foreach (ActionHistoryLayer Layer in Layers)
 			{
 				Layer.Save();
 			}
@@ -416,7 +415,7 @@ namespace UnrealBuildTool
 			lock (LockObject)
 			{
 				ActionHistoryPartition? Partition = Partitions.FirstOrDefault(x => x.BaseDir == BaseDir);
-				if(Partition == null)
+				if (Partition == null)
 				{
 					Partition = new ActionHistoryPartition(BaseDir);
 					Partitions.Add(Partition);
diff --git a/Engine/Source/Programs/UnrealBuildTool/System/ActionSerializer.cs b/Engine/Source/Programs/UnrealBuildTool/System/ActionSerializer.cs
--- a/Engine/Source/Programs/UnrealBuildTool/System/ActionSerializer.cs
+++ b/Engine/Source/Programs/UnrealBuildTool/System/ActionSerializer.cs
@@ -2,12 +2,8 @@
 
 using System;
 using System.Collections.Generic;
-using System.Linq;
 using System.Reflection;
-using System.Text;
-using System.Threading.Tasks;
 using EpicGames.Core;
-using UnrealBuildBase;
 
 namespace UnrealBuildTool
 {
diff --git a/Engine/Source/Programs/UnrealBuildTool/System/Asn.cs b/Engine/Source/Programs/UnrealBuildTool/System/Asn.cs
--- a/Engine/Source/Programs/UnrealBuildTool/System/Asn.cs
+++ b/Engine/Source/Programs/UnrealBuildTool/System/Asn.cs
@@ -6,7 +6,6 @@ using System.Diagnostics.CodeAnalysis;
 using System.IO;
 using System.Linq;
 using System.Text;
-using System.Threading.Tasks;
 
 namespace UnrealBuildTool
 {
@@ -38,10 +37,10 @@ namespace UnrealBuildTool
 
 		public static class ObjectIdentifier
 		{
-			public static readonly int[] CountryName = new int[]{ 2, 5, 4, 6 };
-			public static readonly int[] OrganizationName = new int[]{ 2, 5, 4, 10 };
-			public static readonly int[] Pkcs7_SignedData = new int[]{ 1, 2, 840, 113549, 1, 7, 2 };
-			public static readonly int[] Pkcs7_Data = new int[]{ 1, 2, 840, 113549, 1, 7, 1 };
+			public static readonly int[] CountryName = new int[] { 2, 5, 4, 6 };
+			public static readonly int[] OrganizationName = new int[] { 2, 5, 4, 10 };
+			public static readonly int[] Pkcs7_SignedData = new int[] { 1, 2, 840, 113549, 1, 7, 2 };
+			public static readonly int[] Pkcs7_Data = new int[] { 1, 2, 840, 113549, 1, 7, 1 };
 		}
 
 		static readonly List<KeyValuePair<int[], string>> KnownObjectIdentifiers = new List<KeyValuePair<int[], string>>
@@ -69,7 +68,7 @@ namespace UnrealBuildTool
 
 		public static void SkipValue(BinaryReader Reader, FieldInfo Field)
 		{
-			if(Field.bPrimitive)
+			if (Field.bPrimitive)
 			{
 				Reader.BaseStream.Seek(Field.Length, SeekOrigin.Current);
 			}
@@ -78,7 +77,7 @@ namespace UnrealBuildTool
 		public static int ReadInteger(BinaryReader Reader, int Length)
 		{
 			int Value = 0;
-			for(int Idx = 0; Idx < Length; Idx++)
+			for (int Idx = 0; Idx < Length; Idx++)
 			{
 				Value = (Value << 8) | (int)Reader.ReadByte();
 			}
@@ -88,7 +87,7 @@ namespace UnrealBuildTool
 		public static int ReadLength(BinaryReader Reader)
 		{
 			int Count = (int)Reader.ReadByte();
-			if(Count <= 0x7f)
+			if (Count <= 0x7f)
 			{
 				return Count;
 			}
@@ -105,10 +104,10 @@ namespace UnrealBuildTool
 			List<int> Values = new List<int>();
 			Values.Add((int)Data[0] / 40);
 			Values.Add((int)Data[0] % 40);
-			for(int Idx = 1; Idx < Data.Length; Idx++)
+			for (int Idx = 1; Idx < Data.Length; Idx++)
 			{
 				int Value = (int)Data[Idx] & 0x7f;
-				while(((int)Data[Idx] & 0x80) != 0)
+				while (((int)Data[Idx] & 0x80) != 0)
 				{
 					Value = (Value << 7) | ((int)Data[++Idx] & 0x7f);
 				}
@@ -120,7 +119,7 @@ namespace UnrealBuildTool
 		public static string GetObjectIdentifierName(int[] Values)
 		{
 			string? Description;
-			if(!TryGetObjectIdentifierName(Values, out Description))
+			if (!TryGetObjectIdentifierName(Values, out Description))
 			{
 				Description = String.Format("Unknown ({0})", String.Join(", ", Values.Select(x => x.ToString())));
 			}
@@ -129,9 +128,9 @@ namespace UnrealBuildTool
 
 		public static bool TryGetObjectIdentifierName(int[] Values, [NotNullWhen(true)] out string? Name)
 		{
-			foreach(KeyValuePair<int[], string> KnownObjectIdentifier in KnownObjectIdentifiers)
+			foreach (KeyValuePair<int[], string> KnownObjectIdentifier in KnownObjectIdentifiers)
 			{
-				if(Enumerable.SequenceEqual(Values, KnownObjectIdentifier.Key))
+				if (Enumerable.SequenceEqual(Values, KnownObjectIdentifier.Key))
 				{
 					Name = KnownObjectIdentifier.Value;
 					return true;
@@ -148,12 +147,12 @@ namespace UnrealBuildTool
 
 			// If it's a primitive type, unpack the value
 			string DescriptionSuffix = "";
-			if(Field.bPrimitive)
+			if (Field.bPrimitive)
 			{
-				if(Field.Length > 0)
+				if (Field.Length > 0)
 				{
 					string Description;
-					switch(Field.Tag)
+					switch (Field.Tag)
 					{
 						case FieldTag.INTEGER:
 							Description = String.Format("{0}", ReadInteger(Reader, Field.Length));
@@ -177,10 +176,10 @@ namespace UnrealBuildTool
 			Console.WriteLine("{0}{1}{2}", Indent, Field.Tag, DescriptionSuffix);
 
 			// If it's not a primitive, print the contents
-			if(!Field.bPrimitive)
+			if (!Field.bPrimitive)
 			{
 				long EndOffset = Reader.BaseStream.Position + Field.Length;
-				while(Reader.BaseStream.Position < EndOffset)
+				while (Reader.BaseStream.Position < EndOffset)
 				{
 					Dump(Reader, Indent + "  ");
 				}
diff --git a/Engine/Source/Programs/UnrealBuildTool/System/BuildException.cs b/Engine/Source/Programs/UnrealBuildTool/System/BuildException.cs
--- a/Engine/Source/Programs/UnrealBuildTool/System/BuildException.cs
+++ b/Engine/Source/Programs/UnrealBuildTool/System/BuildException.cs
@@ -1,10 +1,8 @@
 // Copyright Epic Games, Inc. All Rights Reserved.
 
+using System;
 using EpicGames.Core;
 using Microsoft.Extensions.Logging;
-using System;
-using System.Collections.Generic;
-using System.Text;
 
 namespace UnrealBuildTool
 {
diff --git a/Engine/Source/Programs/UnrealBuildTool/System/BuildHostPlatform.cs b/Engine/Source/Programs/UnrealBuildTool/System/BuildHostPlatform.cs
--- a/Engine/Source/Programs/UnrealBuildTool/System/BuildHostPlatform.cs
+++ b/Engine/Source/Programs/UnrealBuildTool/System/BuildHostPlatform.cs
@@ -2,10 +2,9 @@
 
 using System;
 using System.Collections.Generic;
-using System.Linq;
-using System.Text;
-using System.IO;
 using System.Diagnostics;
+using System.IO;
+using System.Linq;
 using System.Runtime.InteropServices;
 using EpicGames.Core;
 
diff --git a/Engine/Source/Programs/UnrealBuildTool/System/BuildVersion.cs b/Engine/Source/Programs/UnrealBuildTool/System/BuildVersion.cs
--- a/Engine/Source/Programs/UnrealBuildTool/System/BuildVersion.cs
+++ b/Engine/Source/Programs/UnrealBuildTool/System/BuildVersion.cs
@@ -1,12 +1,8 @@
 // Copyright Epic Games, Inc. All Rights Reserved.
 
 using System;
-using System.Collections.Generic;
 using System.Diagnostics.CodeAnalysis;
 using System.IO;
-using System.Linq;
-using System.Text;
-using System.Threading.Tasks;
 using EpicGames.Core;
 using Microsoft.Extensions.Logging;
 using UnrealBuildBase;
@@ -114,11 +110,11 @@ namespace UnrealBuildTool
 		{
 			// Get the architecture suffix. Platforms have the option of overriding whether to include this string in filenames.
 			string ArchitectureSuffix = "";
-			if(UnrealArchitectureConfig.ForPlatform(Platform).RequiresArchitectureFilenames(Architectures))
+			if (UnrealArchitectureConfig.ForPlatform(Platform).RequiresArchitectureFilenames(Architectures))
 			{
 				ArchitectureSuffix = Architectures.ToString();
 			}
-		
+
 			// Build the output filename
 			if (String.IsNullOrEmpty(ArchitectureSuffix) && Configuration == UnrealTargetConfiguration.Development)
 			{
@@ -214,8 +210,8 @@ namespace UnrealBuildTool
 			Writer.WriteValue("PatchVersion", PatchVersion);
 			Writer.WriteValue("Changelist", Changelist);
 			Writer.WriteValue("CompatibleChangelist", CompatibleChangelist);
-			Writer.WriteValue("IsLicenseeVersion", IsLicenseeVersion? 1 : 0);
-			Writer.WriteValue("IsPromotedBuild", IsPromotedBuild? 1 : 0);
+			Writer.WriteValue("IsLicenseeVersion", IsLicenseeVersion ? 1 : 0);
+			Writer.WriteValue("IsPromotedBuild", IsPromotedBuild ? 1 : 0);
 			Writer.WriteValue("BranchName", BranchName);
 			if (!String.IsNullOrEmpty(BuildId))
 			{
@@ -259,16 +255,16 @@ namespace UnrealBuildTool
 		{
 			get
 			{
-				if(CurrentCached == null)
+				if (CurrentCached == null)
 				{
 					FileReference File = BuildVersion.GetDefaultFileName();
-					if(!FileReference.Exists(File))
+					if (!FileReference.Exists(File))
 					{
 						throw new BuildException("Version file is missing ({0})", File);
 					}
 
 					BuildVersion? Version;
-					if(!BuildVersion.TryRead(File, out Version))
+					if (!BuildVersion.TryRead(File, out Version))
 					{
 						throw new BuildException("Unable to read version file ({0}). Check that this file is present and well-formed JSON.", File);
 					}
@@ -283,7 +279,7 @@ namespace UnrealBuildTool
 		/// Accessors for fields on the inner BuildVersion instance
 		/// </summary>
 		#region Read-only accessor properties 
-		#pragma warning disable CS1591
+#pragma warning disable CS1591
 
 		public int MajorVersion => Inner.MajorVersion;
 
@@ -305,7 +301,7 @@ namespace UnrealBuildTool
 
 		public string? BuildVersionString => Inner.BuildVersionString;
 
-		#pragma warning restore C1591
+#pragma warning restore C1591
 		#endregion
 	}
 }
diff --git a/Engine/Source/Programs/UnrealBuildTool/System/CommandLine.cs b/Engine/Source/Programs/UnrealBuildTool/System/CommandLine.cs
--- a/Engine/Source/Programs/UnrealBuildTool/System/CommandLine.cs
+++ b/Engine/Source/Programs/UnrealBuildTool/System/CommandLine.cs
@@ -5,10 +5,7 @@ using System.Collections.Generic;
 using System.Diagnostics.CodeAnalysis;
 using System.Linq;
 using System.Reflection;
-using System.Text;
-using System.Threading.Tasks;
 using EpicGames.Core;
-using UnrealBuildBase;
 using Microsoft.Extensions.Logging;
 
 namespace UnrealBuildTool
@@ -90,17 +87,17 @@ namespace UnrealBuildTool
 		{
 			// Build a mapping from name to field and attribute for this object
 			Dictionary<string, Parameter> PrefixToParameter = new Dictionary<string, Parameter>(StringComparer.InvariantCultureIgnoreCase);
-			for(Type TargetType = TargetObject.GetType(); TargetType != typeof(object); TargetType = TargetType.BaseType!)
+			for (Type TargetType = TargetObject.GetType(); TargetType != typeof(object); TargetType = TargetType.BaseType!)
 			{
-				foreach(FieldInfo FieldInfo in TargetType.GetFields(BindingFlags.Instance | BindingFlags.GetField | BindingFlags.GetProperty | BindingFlags.Public | BindingFlags.NonPublic | BindingFlags.DeclaredOnly))
+				foreach (FieldInfo FieldInfo in TargetType.GetFields(BindingFlags.Instance | BindingFlags.GetField | BindingFlags.GetProperty | BindingFlags.Public | BindingFlags.NonPublic | BindingFlags.DeclaredOnly))
 				{
 					IEnumerable<CommandLineAttribute> Attributes = FieldInfo.GetCustomAttributes<CommandLineAttribute>();
-					foreach(CommandLineAttribute Attribute in Attributes)
+					foreach (CommandLineAttribute Attribute in Attributes)
 					{
 						string? Prefix = Attribute.Prefix;
-						if(Prefix == null)
+						if (Prefix == null)
 						{
-							if(NonNullableType(FieldInfo.FieldType) == typeof(bool))
+							if (NonNullableType(FieldInfo.FieldType) == typeof(bool))
 							{
 								Prefix = String.Format("-{0}", FieldInfo.Name);
 							}
@@ -111,7 +108,7 @@ namespace UnrealBuildTool
 						}
 						else
 						{
-							if(NonNullableType(FieldInfo.FieldType) != typeof(bool) && Attribute.Value == null && !Prefix.EndsWith("=") && !Prefix.EndsWith(":"))
+							if (NonNullableType(FieldInfo.FieldType) != typeof(bool) && Attribute.Value == null && !Prefix.EndsWith("=") && !Prefix.EndsWith(":"))
 							{
 								Prefix = Prefix + "=";
 							}
@@ -123,25 +120,25 @@ namespace UnrealBuildTool
 
 			// Step through the arguments, and remove those that we can parse
 			Dictionary<FieldInfo, Parameter> AssignedFieldToParameter = new Dictionary<FieldInfo, Parameter>();
-			for(int Idx = 0; Idx < Arguments.Count; Idx++)
+			for (int Idx = 0; Idx < Arguments.Count; Idx++)
 			{
 				string Argument = Arguments[Idx];
-				if(Argument.Length > 0 && Argument[0] == '-')
+				if (Argument.Length > 0 && Argument[0] == '-')
 				{
 					// Get the length of the argument prefix
-					int EqualsIdx = Argument.IndexOfAny(new char[]{ '=', ':' });
-					string Prefix = (EqualsIdx == -1)? Argument : Argument.Substring(0, EqualsIdx + 1);
+					int EqualsIdx = Argument.IndexOfAny(new char[] { '=', ':' });
+					string Prefix = (EqualsIdx == -1) ? Argument : Argument.Substring(0, EqualsIdx + 1);
 
 					// Check if there's a matching argument registered
 					Parameter? Parameter;
-					if(PrefixToParameter.TryGetValue(Prefix, out Parameter))
+					if (PrefixToParameter.TryGetValue(Prefix, out Parameter))
 					{
 						int NextIdx = Idx + 1;
 
 						// Parse the value
-						if(Parameter.Attribute.Value != null)
+						if (Parameter.Attribute.Value != null)
 						{
-							if(EqualsIdx != -1)
+							if (EqualsIdx != -1)
 							{
 								Logger.LogWarning("Cannot specify a value for {ParameterPrefix}", Parameter.Prefix);
 							}
@@ -150,11 +147,11 @@ namespace UnrealBuildTool
 								AssignValue(Parameter, Parameter.Attribute.Value, TargetObject, AssignedFieldToParameter, Logger);
 							}
 						}
-						else if(EqualsIdx != -1)
+						else if (EqualsIdx != -1)
 						{
 							AssignValue(Parameter, Argument.Substring(EqualsIdx + 1), TargetObject, AssignedFieldToParameter, Logger);
 						}
-						else if(NonNullableType(Parameter.FieldInfo.FieldType) == typeof(bool))
+						else if (NonNullableType(Parameter.FieldInfo.FieldType) == typeof(bool))
 						{
 							AssignValue(Parameter, "true", TargetObject, AssignedFieldToParameter, Logger);
 						}
@@ -172,16 +169,16 @@ namespace UnrealBuildTool
 
 			// Make sure there are no required parameters that are missing
 			Dictionary<FieldInfo, Parameter> MissingFieldToParameter = new Dictionary<FieldInfo, Parameter>();
-			foreach(Parameter Parameter in PrefixToParameter.Values)
+			foreach (Parameter Parameter in PrefixToParameter.Values)
 			{
-				if(Parameter.Attribute.Required && !AssignedFieldToParameter.ContainsKey(Parameter.FieldInfo) && !MissingFieldToParameter.ContainsKey(Parameter.FieldInfo))
+				if (Parameter.Attribute.Required && !AssignedFieldToParameter.ContainsKey(Parameter.FieldInfo) && !MissingFieldToParameter.ContainsKey(Parameter.FieldInfo))
 				{
 					MissingFieldToParameter.Add(Parameter.FieldInfo, Parameter);
 				}
 			}
-			if(MissingFieldToParameter.Count > 0)
+			if (MissingFieldToParameter.Count > 0)
 			{
-				if(MissingFieldToParameter.Count == 1)
+				if (MissingFieldToParameter.Count == 1)
 				{
 					throw new BuildException("Missing {0} argument", MissingFieldToParameter.First().Value.Prefix.Replace("=", "=..."));
 				}
@@ -199,9 +196,9 @@ namespace UnrealBuildTool
 		/// <param name="Logger">Logger for output</param>
 		public static void CheckNoRemainingArguments(List<string> RemainingArguments, ILogger Logger)
 		{
-			if(RemainingArguments.Count > 0)
+			if (RemainingArguments.Count > 0)
 			{
-				if(RemainingArguments.Count == 1)
+				if (RemainingArguments.Count == 1)
 				{
 					Logger.LogWarning("Invalid argument: {Arg}", RemainingArguments[0]);
 				}
@@ -238,7 +235,7 @@ namespace UnrealBuildTool
 			{
 				// Try to parse the value
 				object? Value;
-				if(!TryParseValue(Parameter.FieldInfo.FieldType, Text, out Value))
+				if (!TryParseValue(Parameter.FieldInfo.FieldType, Text, out Value))
 				{
 					Logger.LogWarning("Invalid value for {ParameterPrefix}... - ignoring {Text}", Parameter.Prefix, Text);
 					return;
@@ -246,12 +243,12 @@ namespace UnrealBuildTool
 
 				// Check if this field has already been assigned to. Output a warning if the previous value is in conflict with the new one.
 				Parameter? PreviousParameter;
-				if(AssignedFieldToParameter.TryGetValue(Parameter.FieldInfo, out PreviousParameter))
+				if (AssignedFieldToParameter.TryGetValue(Parameter.FieldInfo, out PreviousParameter))
 				{
 					object? PreviousValue = Parameter.FieldInfo.GetValue(TargetObject);
-					if(!Object.Equals(PreviousValue, Value))
+					if (!Object.Equals(PreviousValue, Value))
 					{
-						if(PreviousParameter.Prefix == Parameter.Prefix)
+						if (PreviousParameter.Prefix == Parameter.Prefix)
 						{
 							Logger.LogWarning("Conflicting {ParameterPrefix} arguments - ignoring", Parameter.Prefix);
 						}
@@ -271,7 +268,7 @@ namespace UnrealBuildTool
 			{
 				// Split the text into an array of values if necessary
 				string[] ItemArray;
-				if(Parameter.Attribute.ListSeparator == 0)
+				if (Parameter.Attribute.ListSeparator == 0)
 				{
 					ItemArray = new string[] { Text };
 				}
@@ -281,10 +278,10 @@ namespace UnrealBuildTool
 				}
 
 				// Parse each of the argument values separately
-				foreach(string Item in ItemArray)
+				foreach (string Item in ItemArray)
 				{
 					object? Value;
-					if(TryParseValue(CollectionType.GenericTypeArguments[0], Item, out Value))
+					if (TryParseValue(CollectionType.GenericTypeArguments[0], Item, out Value))
 					{
 						CollectionType.InvokeMember("Add", BindingFlags.InvokeMethod, null, Parameter.FieldInfo.GetValue(TargetObject), new object[] { Value });
 					}
@@ -306,7 +303,7 @@ namespace UnrealBuildTool
 		static bool TryParseValue(Type FieldType, string Text, [NotNullWhen(true)] out object? Value)
 		{
 			FieldType = NonNullableType(FieldType);
-			if(FieldType.IsEnum)
+			if (FieldType.IsEnum)
 			{
 				// Special handling for enums; parse the value ignoring case.
 				try
@@ -314,13 +311,13 @@ namespace UnrealBuildTool
 					Value = Enum.Parse(FieldType, Text, true);
 					return true;
 				}
-				catch(ArgumentException)
+				catch (ArgumentException)
 				{
 					Value = null;
 					return false;
 				}
 			}
-			else if(FieldType == typeof(FileReference))
+			else if (FieldType == typeof(FileReference))
 			{
 				// Construct a file reference from the string
 				try
@@ -342,7 +339,7 @@ namespace UnrealBuildTool
 					Value = Convert.ChangeType(Text, FieldType);
 					return true;
 				}
-				catch(InvalidCastException)
+				catch (InvalidCastException)
 				{
 					Value = null;
 					return false;
diff --git a/Engine/Source/Programs/UnrealBuildTool/System/ConfigCache.cs b/Engine/Source/Programs/UnrealBuildTool/System/ConfigCache.cs
--- a/Engine/Source/Programs/UnrealBuildTool/System/ConfigCache.cs
+++ b/Engine/Source/Programs/UnrealBuildTool/System/ConfigCache.cs
@@ -1,16 +1,15 @@
 // Copyright Epic Games, Inc. All Rights Reserved.
 
 using System;
+using System.Collections.Concurrent;
 using System.Collections.Generic;
 using System.Diagnostics.CodeAnalysis;
+using System.IO;
 using System.Linq;
 using System.Reflection;
-using System.IO;
 using EpicGames.Core;
 using Microsoft.Extensions.Logging;
 using Microsoft.Extensions.Logging.Abstractions;
-using System.Collections.Concurrent;
-using Microsoft.Build.Logging;
 
 namespace UnrealBuildTool
 {
@@ -24,7 +23,7 @@ namespace UnrealBuildTool
 		/// </summary>
 		/// <param name="TargetObject">The object containing the field to be modified</param>
 		/// <param name="ValueObject">The value to add</param>
-		delegate void AddElementDelegate(object TargetObject, object? ValueObject); 
+		delegate void AddElementDelegate(object TargetObject, object? ValueObject);
 
 		/// <summary>
 		/// Caches information about a member with a [ConfigFile] attribute in a type
@@ -58,22 +57,22 @@ namespace UnrealBuildTool
 			/// <summary>
 			/// Returns Reflection.MemberInfo describing the target class member.
 			/// </summary>
-			public abstract MemberInfo               MemberInfo { get; }
+			public abstract MemberInfo MemberInfo { get; }
 
 			/// <summary>
 			/// Returns Reflection.Type of the target class member.
 			/// </summary>
-			public abstract Type                     Type       { get; }
+			public abstract Type Type { get; }
 
 			/// <summary>
 			/// Returns the value setter of the target class member.
 			/// </summary>
-			public abstract Action<object?, object?> SetValue   { get; }
+			public abstract Action<object?, object?> SetValue { get; }
 
 			/// <summary>
 			/// Returns the value getter of the target class member.
 			/// </summary>
-			public abstract Func<object?, object?>   GetValue   { get; }
+			public abstract Func<object?, object?> GetValue { get; }
 		}
 
 		/// <summary>
@@ -97,10 +96,10 @@ namespace UnrealBuildTool
 				this.FieldInfo = FieldInfo;
 			}
 
-			public override MemberInfo               MemberInfo => FieldInfo;
-			public override Type                     Type       => FieldInfo.FieldType;
-			public override Action<object?, object?> SetValue   => FieldInfo.SetValue;
-			public override Func<object?, object?>   GetValue   => FieldInfo.GetValue;
+			public override MemberInfo MemberInfo => FieldInfo;
+			public override Type Type => FieldInfo.FieldType;
+			public override Action<object?, object?> SetValue => FieldInfo.SetValue;
+			public override Func<object?, object?> GetValue => FieldInfo.GetValue;
 		}
 
 		/// <summary>
@@ -124,10 +123,10 @@ namespace UnrealBuildTool
 				this.PropertyInfo = PropertyInfo;
 			}
 
-			public override MemberInfo               MemberInfo => PropertyInfo;
-			public override Type                     Type       => PropertyInfo.PropertyType;
-			public override Action<object?, object?> SetValue   => PropertyInfo.SetValue;
-			public override Func<object?, object?>   GetValue   => PropertyInfo.GetValue;
+			public override MemberInfo MemberInfo => PropertyInfo;
+			public override Type Type => PropertyInfo.PropertyType;
+			public override Action<object?, object?> SetValue => PropertyInfo.SetValue;
+			public override Func<object?, object?> GetValue => PropertyInfo.GetValue;
 		}
 
 		/// <summary>
@@ -303,7 +302,7 @@ namespace UnrealBuildTool
 		/// <param name="CustomConfig">Optional override config directory to search, for support of multiple target types</param>
 		/// <param name="CmdLineArgs">The command line arguments to parse</param>
 		/// <returns>The requested config hierarchy</returns>
-		public static ConfigHierarchy ReadHierarchy(ConfigHierarchyType Type, DirectoryReference? ProjectDir, UnrealTargetPlatform Platform, 
+		public static ConfigHierarchy ReadHierarchy(ConfigHierarchyType Type, DirectoryReference? ProjectDir, UnrealTargetPlatform Platform,
 			string CustomConfig, CommandLineArguments CmdLineArgs)
 		{
 			// Handle command line overrides
@@ -337,8 +336,8 @@ namespace UnrealBuildTool
 			ConfigHierarchyKey Key = new ConfigHierarchyKey(Type, ProjectDir, Platform, CustomConfig, OverrideStrings);
 
 			ILogger Logger = NullLogger.Instance;
-			
-			if(bDumpIniLoads)
+
+			if (bDumpIniLoads)
 			{
 				Logger = Log.Logger;
 			};
@@ -371,7 +370,7 @@ namespace UnrealBuildTool
 				return new ConfigHierarchy(Files);
 			});
 		}
-	
+
 		/// <summary>
 		/// Gets a list of ConfigFields for the given type
 		/// </summary>
@@ -471,10 +470,10 @@ namespace UnrealBuildTool
 				// Parse the values from the config files and update the target object
 				if (Member.AddElement == null)
 				{
-					if(Values != null && Values.Count == 1)
+					if (Values != null && Values.Count == 1)
 					{
 						object? Value;
-						if(TryParseValue(Values[0], Member.Type, out Value))
+						if (TryParseValue(Values[0], Member.Type, out Value))
 						{
 							Member.SetValue(TargetObject, Value);
 						}
@@ -482,12 +481,12 @@ namespace UnrealBuildTool
 				}
 				else
 				{
-					if(Values != null)
+					if (Values != null)
 					{
-						foreach(string Item in Values)
+						foreach (string Item in Values)
 						{
 							object? Value;
-							if(TryParseValue(Item, Member.ElementType!, out Value))
+							if (TryParseValue(Item, Member.ElementType!, out Value))
 							{
 								Member.AddElement(TargetObject, Value);
 							}
@@ -513,15 +512,15 @@ namespace UnrealBuildTool
 		/// <returns>True if the value could be parsed, false otherwise</returns>
 		public static bool TryParseValue(string Text, Type FieldType, out object? Value)
 		{
-			if(FieldType == typeof(string))
+			if (FieldType == typeof(string))
 			{
 				Value = Text;
 				return true;
 			}
-			else if(FieldType == typeof(bool))
+			else if (FieldType == typeof(bool))
 			{
 				bool BoolValue;
-				if(ConfigHierarchy.TryParse(Text, out BoolValue))
+				if (ConfigHierarchy.TryParse(Text, out BoolValue))
 				{
 					Value = BoolValue;
 					return true;
@@ -532,10 +531,10 @@ namespace UnrealBuildTool
 					return false;
 				}
 			}
-			else if(FieldType == typeof(int))
+			else if (FieldType == typeof(int))
 			{
 				int IntValue;
-				if(ConfigHierarchy.TryParse(Text, out IntValue))
+				if (ConfigHierarchy.TryParse(Text, out IntValue))
 				{
 					Value = IntValue;
 					return true;
@@ -546,10 +545,10 @@ namespace UnrealBuildTool
 					return false;
 				}
 			}
-			else if(FieldType == typeof(float))
+			else if (FieldType == typeof(float))
 			{
 				float FloatValue;
-				if(ConfigHierarchy.TryParse(Text, out FloatValue))
+				if (ConfigHierarchy.TryParse(Text, out FloatValue))
 				{
 					Value = FloatValue;
 					return true;
@@ -560,10 +559,10 @@ namespace UnrealBuildTool
 					return false;
 				}
 			}
-			else if(FieldType == typeof(double))
+			else if (FieldType == typeof(double))
 			{
 				double DoubleValue;
-				if(ConfigHierarchy.TryParse(Text, out DoubleValue))
+				if (ConfigHierarchy.TryParse(Text, out DoubleValue))
 				{
 					Value = DoubleValue;
 					return true;
@@ -574,10 +573,10 @@ namespace UnrealBuildTool
 					return false;
 				}
 			}
-			else if(FieldType == typeof(Guid))
+			else if (FieldType == typeof(Guid))
 			{
 				Guid GuidValue;
-				if(ConfigHierarchy.TryParse(Text, out GuidValue))
+				if (ConfigHierarchy.TryParse(Text, out GuidValue))
 				{
 					Value = GuidValue;
 					return true;
@@ -588,7 +587,7 @@ namespace UnrealBuildTool
 					return false;
 				}
 			}
-			else if(FieldType.IsEnum)
+			else if (FieldType.IsEnum)
 			{
 				try
 				{
@@ -601,7 +600,7 @@ namespace UnrealBuildTool
 					return false;
 				}
 			}
-			else if(FieldType.GetGenericTypeDefinition() == typeof(Nullable<>))
+			else if (FieldType.GetGenericTypeDefinition() == typeof(Nullable<>))
 			{
 				return TryParseValue(Text, FieldType.GetGenericArguments()[0], out Value);
 			}
diff --git a/Engine/Source/Programs/UnrealBuildTool/System/ConfigFile.cs b/Engine/Source/Programs/UnrealBuildTool/System/ConfigFile.cs
--- a/Engine/Source/Programs/UnrealBuildTool/System/ConfigFile.cs
+++ b/Engine/Source/Programs/UnrealBuildTool/System/ConfigFile.cs
@@ -4,11 +4,7 @@ using System;
 using System.Collections.Generic;
 using System.Diagnostics.CodeAnalysis;
 using System.IO;
-using System.Linq;
-using System.Text;
-using System.Threading.Tasks;
 using EpicGames.Core;
-using UnrealBuildTool;
 using Microsoft.Extensions.Logging;
 
 namespace UnrealBuildTool
@@ -28,22 +24,22 @@ namespace UnrealBuildTool
 		/// </summary>
 		Add,
 
-        /// <summary>
-        /// Remove the key without having to match value (denoted with !X in config files)
-        /// </summary>
-        RemoveKey,
+		/// <summary>
+		/// Remove the key without having to match value (denoted with !X in config files)
+		/// </summary>
+		RemoveKey,
 
-        /// <summary>
-        /// Remove the matching key and value (denoted with -X=Y in config files)
-        /// </summary>
-        RemoveKeyValue
+		/// <summary>
+		/// Remove the matching key and value (denoted with -X=Y in config files)
+		/// </summary>
+		RemoveKeyValue
 
-    }
+	}
 
-    /// <summary>
-    /// Contains a pre-parsed raw config line, consisting of action, key and value components.
-    /// </summary>
-    public class ConfigLine
+	/// <summary>
+	/// Contains a pre-parsed raw config line, consisting of action, key and value components.
+	/// </summary>
+	public class ConfigLine
 	{
 		/// <summary>
 		/// The action to take when merging this key/value pair with an existing value
@@ -79,7 +75,7 @@ namespace UnrealBuildTool
 		/// <returns>The original config line</returns>
 		public override string ToString()
 		{
-			string Prefix = (Action == ConfigLineAction.Add)? "+" : (Action == ConfigLineAction.RemoveKey)? "!" : (Action == ConfigLineAction.RemoveKeyValue) ? "-" : "";
+			string Prefix = (Action == ConfigLineAction.Add) ? "+" : (Action == ConfigLineAction.RemoveKey) ? "!" : (Action == ConfigLineAction.RemoveKeyValue) ? "-" : "";
 			return String.Format("{0}{1}={2}", Prefix, Key, Value);
 		}
 	}
@@ -116,7 +112,7 @@ namespace UnrealBuildTool
 		/// <returns>return true if the line is retrieved return false and null OutLine if Name isn't found in this section</returns>
 		public bool TryGetLine(string Name, [NotNullWhen(true)] out ConfigLine? OutLine)
 		{
-			foreach ( ConfigLine Line in Lines)
+			foreach (ConfigLine Line in Lines)
 			{
 				if (Line.Key.Equals(Name))
 				{
@@ -184,7 +180,7 @@ namespace UnrealBuildTool
 				// Make ConfigFile when EngineDirectory is unknown a warning since ConfigRemapFile cannot be read in this case; e.g. Assemblies outside Engine that depend on ConfigFile
 				Log.Logger.LogWarning("Failed to read ConfigRemapFile into Sections");
 			}
-			
+
 			// walk over the sections, where all but the special SectionNameRemap section is a section of keys to remap in that same section
 			foreach (KeyValuePair<string, ConfigFileSection> Pair in Sections)
 			{
@@ -234,7 +230,7 @@ namespace UnrealBuildTool
 				ConfigFileSection? CurrentSection = null;
 				Dictionary<string, string>? CurrentRemap = null;
 
-				for (; ;)
+				for (; ; )
 				{
 					string? Line = Reader.ReadLine();
 					if (Line == null)
@@ -251,10 +247,10 @@ namespace UnrealBuildTool
 							int EndIdx = Line.Length;
 							while (EndIdx > StartIdx)
 							{
-								if(Line[EndIdx - 1] == '\\')
+								if (Line[EndIdx - 1] == '\\')
 								{
 									string? NextLine = Reader.ReadLine();
-									if(NextLine == null)
+									if (NextLine == null)
 									{
 										break;
 									}
@@ -262,7 +258,7 @@ namespace UnrealBuildTool
 									EndIdx = Line.Length;
 									continue;
 								}
-								if(Line[EndIdx - 1] != ' ' && Line[EndIdx - 1] != '\t')
+								if (Line[EndIdx - 1] != ' ' && Line[EndIdx - 1] != '\t')
 								{
 									break;
 								}
@@ -271,23 +267,23 @@ namespace UnrealBuildTool
 							}
 
 							// Break out if we've got a comment
-							if(Line[StartIdx] == ';')
+							if (Line[StartIdx] == ';')
 							{
 								break;
 							}
-							if(Line[StartIdx] == '/' && StartIdx + 1 < Line.Length && Line[StartIdx + 1] == '/')
+							if (Line[StartIdx] == '/' && StartIdx + 1 < Line.Length && Line[StartIdx + 1] == '/')
 							{
 								break;
 							}
 
 							// Check if it's the start of a new section
-							if(Line[StartIdx] == '[')
+							if (Line[StartIdx] == '[')
 							{
 								CurrentSection = null;
-								if(Line[EndIdx - 1] == ']')
+								if (Line[EndIdx - 1] == ']')
 								{
 									string SectionName = Line.Substring(StartIdx + 1, EndIdx - StartIdx - 2);
-									
+
 									// lookup remaps
 									SectionName = RemapSectionOrKey(SectionNameRemap, SectionName, $"which is a config section in '{Location.FullName}'");
 									SectionKeyRemap.TryGetValue(SectionName, out CurrentRemap);
@@ -302,7 +298,7 @@ namespace UnrealBuildTool
 							}
 
 							// Otherwise add it to the current section or add a new one
-							if(CurrentSection != null)
+							if (CurrentSection != null)
 							{
 								TryAddConfigLine(CurrentSection, CurrentRemap, Location.FullName, Line, StartIdx, EndIdx, DefaultAction, Sections);
 								break;
@@ -369,7 +365,7 @@ namespace UnrealBuildTool
 		{
 			// Find the '=' character separating key and value
 			int EqualsIdx = Line.IndexOf('=', StartIdx, EndIdx - StartIdx);
-			if(EqualsIdx == -1 && Line[StartIdx] != '!')
+			if (EqualsIdx == -1 && Line[StartIdx] != '!')
 			{
 				return false;
 			}
@@ -379,44 +375,44 @@ namespace UnrealBuildTool
 
 			// Remove the +/-/! prefix, if present
 			ConfigLineAction Action = DefaultAction;
-			if(Line[KeyStartIdx] == '+' || Line[KeyStartIdx] == '-' || Line[KeyStartIdx] == '!')
+			if (Line[KeyStartIdx] == '+' || Line[KeyStartIdx] == '-' || Line[KeyStartIdx] == '!')
 			{
-				Action = (Line[KeyStartIdx] == '+')? ConfigLineAction.Add : (Line[KeyStartIdx] == '!') ? ConfigLineAction.RemoveKey : ConfigLineAction.RemoveKeyValue;
+				Action = (Line[KeyStartIdx] == '+') ? ConfigLineAction.Add : (Line[KeyStartIdx] == '!') ? ConfigLineAction.RemoveKey : ConfigLineAction.RemoveKeyValue;
 				KeyStartIdx++;
-				while(Line[KeyStartIdx] == ' ' || Line[KeyStartIdx] == '\t')
+				while (Line[KeyStartIdx] == ' ' || Line[KeyStartIdx] == '\t')
 				{
 					KeyStartIdx++;
 				}
 			}
 
-            // RemoveKey actions do not require a value
-            if (Action == ConfigLineAction.RemoveKey && EqualsIdx == -1)
-            {
-                Section.Lines.Add(new ConfigLine(Action, Line.Substring(KeyStartIdx).Trim(), ""));
-                return true;
-            }
+			// RemoveKey actions do not require a value
+			if (Action == ConfigLineAction.RemoveKey && EqualsIdx == -1)
+			{
+				Section.Lines.Add(new ConfigLine(Action, Line.Substring(KeyStartIdx).Trim(), ""));
+				return true;
+			}
 
-            // Remove trailing spaces after the name of the key
-            int KeyEndIdx = EqualsIdx;
-			for(; KeyEndIdx > KeyStartIdx; KeyEndIdx--)
+			// Remove trailing spaces after the name of the key
+			int KeyEndIdx = EqualsIdx;
+			for (; KeyEndIdx > KeyStartIdx; KeyEndIdx--)
 			{
-				if(Line[KeyEndIdx - 1] != ' ' && Line[KeyEndIdx - 1] != '\t')
+				if (Line[KeyEndIdx - 1] != ' ' && Line[KeyEndIdx - 1] != '\t')
 				{
 					break;
 				}
 			}
 
 			// Make sure there's a non-empty key name
-			if(KeyStartIdx == EqualsIdx)
+			if (KeyStartIdx == EqualsIdx)
 			{
 				return false;
 			}
 
 			// Skip whitespace between the '=' and the start of the value
 			int ValueStartIdx = EqualsIdx + 1;
-			for(; ValueStartIdx < EndIdx; ValueStartIdx++)
+			for (; ValueStartIdx < EndIdx; ValueStartIdx++)
 			{
-				if(Line[ValueStartIdx] != ' ' && Line[ValueStartIdx] != '\t')
+				if (Line[ValueStartIdx] != ' ' && Line[ValueStartIdx] != '\t')
 				{
 					break;
 				}
@@ -424,7 +420,7 @@ namespace UnrealBuildTool
 
 			// Strip quotes around the value if present
 			int ValueEndIdx = EndIdx;
-			if(ValueEndIdx >= ValueStartIdx + 2 && Line[ValueStartIdx] == '"' && Line[ValueEndIdx - 1] == '"')
+			if (ValueEndIdx >= ValueStartIdx + 2 && Line[ValueStartIdx] == '"' && Line[ValueEndIdx - 1] == '"')
 			{
 				ValueStartIdx++;
 				ValueEndIdx--;
@@ -471,7 +467,7 @@ namespace UnrealBuildTool
 		public ConfigFileSection FindOrAddSection(string SectionName)
 		{
 			ConfigFileSection? Section;
-			if(!Sections.TryGetValue(SectionName, out Section))
+			if (!Sections.TryGetValue(SectionName, out Section))
 			{
 				Section = new ConfigFileSection(SectionName);
 				Sections.Add(SectionName, Section);
diff --git a/Engine/Source/Programs/UnrealBuildTool/System/ConfigFileAttribute.cs b/Engine/Source/Programs/UnrealBuildTool/System/ConfigFileAttribute.cs
--- a/Engine/Source/Programs/UnrealBuildTool/System/ConfigFileAttribute.cs
+++ b/Engine/Source/Programs/UnrealBuildTool/System/ConfigFileAttribute.cs
@@ -1,10 +1,6 @@
 // Copyright Epic Games, Inc. All Rights Reserved.
 
 using System;
-using System.Collections.Generic;
-using System.Linq;
-using System.Text;
-using System.Threading.Tasks;
 
 namespace UnrealBuildTool
 {
