UnrealBuildTool: Automated code cleanup

diff --git a/Engine/Source/Programs/UnrealBuildTool/System/ConfigHierarchy.cs b/Engine/Source/Programs/UnrealBuildTool/System/ConfigHierarchy.cs
--- a/Engine/Source/Programs/UnrealBuildTool/System/ConfigHierarchy.cs
+++ b/Engine/Source/Programs/UnrealBuildTool/System/ConfigHierarchy.cs
@@ -2,13 +2,11 @@
 
 using System;
 using System.Collections.Generic;
+using System.Diagnostics.CodeAnalysis;
 using System.Linq;
 using System.Text;
-using System.Threading.Tasks;
-using EpicGames.Core;
-using System.IO;
-using System.Diagnostics.CodeAnalysis;
 using System.Text.RegularExpressions;
+using EpicGames.Core;
 using UnrealBuildBase;
 
 namespace UnrealBuildTool
@@ -80,32 +78,32 @@ namespace UnrealBuildTool
 		/// <param name="FileSections">Config sections from individual files</param>
 		public ConfigHierarchySection(IEnumerable<ConfigFileSection> FileSections)
 		{
-			foreach(ConfigFileSection FileSection in FileSections)
+			foreach (ConfigFileSection FileSection in FileSections)
 			{
-				foreach(ConfigLine Line in FileSection.Lines)
+				foreach (ConfigLine Line in FileSection.Lines)
 				{
-                    if (Line.Action == ConfigLineAction.RemoveKey)
-                    {
-                        KeyToValue.Remove(Line.Key);
-                        continue;
-                    }
+					if (Line.Action == ConfigLineAction.RemoveKey)
+					{
+						KeyToValue.Remove(Line.Key);
+						continue;
+					}
 
-                    // Find or create the values for this key
-                    List<string>? Values;
+					// Find or create the values for this key
+					List<string>? Values;
 
-					if(KeyToValue.TryGetValue(Line.Key, out Values))
+					if (KeyToValue.TryGetValue(Line.Key, out Values))
 					{
 						// Update the existing list
-						if(Line.Action == ConfigLineAction.Set)
+						if (Line.Action == ConfigLineAction.Set)
 						{
 							Values.Clear();
 							Values.Add(Line.Value);
 						}
-						else if(Line.Action == ConfigLineAction.Add)
+						else if (Line.Action == ConfigLineAction.Add)
 						{
 							Values.Add(Line.Value);
 						}
-                        else if (Line.Action == ConfigLineAction.RemoveKeyValue)
+						else if (Line.Action == ConfigLineAction.RemoveKeyValue)
 						{
 							Values.RemoveAll(x => x.Equals(Line.Value, StringComparison.InvariantCultureIgnoreCase));
 						}
@@ -113,7 +111,7 @@ namespace UnrealBuildTool
 					else
 					{
 						// If it's a set or add action, create and add a new list
-						if(Line.Action == ConfigLineAction.Set || Line.Action == ConfigLineAction.Add)
+						if (Line.Action == ConfigLineAction.Set || Line.Action == ConfigLineAction.Add)
 						{
 							Values = new List<string>();
 							Values.Add(Line.Value);
@@ -138,7 +136,7 @@ namespace UnrealBuildTool
 		public bool TryGetValue(string KeyName, [NotNullWhen(true)] out string? Value)
 		{
 			List<string>? ValuesList;
-			if(KeyToValue.TryGetValue(KeyName, out ValuesList) && ValuesList.Count > 0)
+			if (KeyToValue.TryGetValue(KeyName, out ValuesList) && ValuesList.Count > 0)
 			{
 				Value = ValuesList[0];
 				return true;
@@ -159,7 +157,7 @@ namespace UnrealBuildTool
 		public bool TryGetValues(string KeyName, [NotNullWhen(true)] out IReadOnlyList<string>? Values)
 		{
 			List<string>? ValuesList;
-			if(KeyToValue.TryGetValue(KeyName, out ValuesList))
+			if (KeyToValue.TryGetValue(KeyName, out ValuesList))
 			{
 				Values = ValuesList;
 				return true;
@@ -187,10 +185,10 @@ namespace UnrealBuildTool
 		/// </summary>
 		Dictionary<string, ConfigHierarchySection> NameToSection = new Dictionary<string, ConfigHierarchySection>(StringComparer.InvariantCultureIgnoreCase);
 
-        /// <summary>
-        /// Lock for NameToSection
-        /// </summary>
-        System.Threading.ReaderWriterLockSlim NameToSectionLock = new System.Threading.ReaderWriterLockSlim();
+		/// <summary>
+		/// Lock for NameToSection
+		/// </summary>
+		System.Threading.ReaderWriterLockSlim NameToSectionLock = new System.Threading.ReaderWriterLockSlim();
 
 		/// <summary>
 		/// Construct a config hierarchy from the given files
@@ -214,7 +212,7 @@ namespace UnrealBuildTool
 				{
 					foreach (string SectionName in File.SectionNames)
 					{
-						if ( !Result.Contains(SectionName) )
+						if (!Result.Contains(SectionName))
 						{
 							Result.Add(SectionName);
 						}
@@ -231,47 +229,47 @@ namespace UnrealBuildTool
 		/// <returns>The merged config section</returns>
 		public ConfigHierarchySection FindSection(string SectionName)
 		{
-            ConfigHierarchySection? Section;
-            try
-            {
-                // Acquire a read lock and do a quick check for the config section
-                NameToSectionLock.EnterUpgradeableReadLock();
-                if (!NameToSection.TryGetValue(SectionName, out Section))
-                {
-                    try
-                    {
-                        // Acquire a write lock and add the config section if another thread didn't just complete it
-                        NameToSectionLock.EnterWriteLock();
-                        if (!NameToSection.TryGetValue(SectionName, out Section))
-                        {
-                            // Find all the raw sections from the file hierarchy
-                            List<ConfigFileSection> RawSections = new List<ConfigFileSection>();
-                            foreach (ConfigFile File in Files)
-                            {
-                                ConfigFileSection? RawSection;
-                                if (File.TryGetSection(SectionName, out RawSection))
-                                {
-                                    RawSections.Add(RawSection);
-                                }
-                            }
-
-                            // Merge them together and add it to the cache
-                            Section = new ConfigHierarchySection(RawSections);
-                            NameToSection.Add(SectionName, Section);
-                        }                        
-                    }
-                    finally
-                    {
-                        NameToSectionLock.ExitWriteLock();
-                    }
-                }
-            }
-            finally
-            {
-                NameToSectionLock.ExitUpgradeableReadLock();
-            }
-            return Section;
-        }
+			ConfigHierarchySection? Section;
+			try
+			{
+				// Acquire a read lock and do a quick check for the config section
+				NameToSectionLock.EnterUpgradeableReadLock();
+				if (!NameToSection.TryGetValue(SectionName, out Section))
+				{
+					try
+					{
+						// Acquire a write lock and add the config section if another thread didn't just complete it
+						NameToSectionLock.EnterWriteLock();
+						if (!NameToSection.TryGetValue(SectionName, out Section))
+						{
+							// Find all the raw sections from the file hierarchy
+							List<ConfigFileSection> RawSections = new List<ConfigFileSection>();
+							foreach (ConfigFile File in Files)
+							{
+								ConfigFileSection? RawSection;
+								if (File.TryGetSection(SectionName, out RawSection))
+								{
+									RawSections.Add(RawSection);
+								}
+							}
+
+							// Merge them together and add it to the cache
+							Section = new ConfigHierarchySection(RawSections);
+							NameToSection.Add(SectionName, Section);
+						}
+					}
+					finally
+					{
+						NameToSectionLock.ExitWriteLock();
+					}
+				}
+			}
+			finally
+			{
+				NameToSectionLock.ExitUpgradeableReadLock();
+			}
+			return Section;
+		}
 
 		/// <summary>
 		/// Legacy function for ease of transition from ConfigCacheIni to ConfigHierarchy. Gets a bool with the given key name.
@@ -295,7 +293,7 @@ namespace UnrealBuildTool
 		public bool GetArray(string SectionName, string KeyName, [NotNullWhen(true)] out List<string>? Values)
 		{
 			IReadOnlyList<string>? ValuesEnumerable;
-			if(TryGetValues(SectionName, KeyName, out ValuesEnumerable))
+			if (TryGetValues(SectionName, KeyName, out ValuesEnumerable))
 			{
 				Values = ValuesEnumerable.ToList();
 				return true;
@@ -317,7 +315,7 @@ namespace UnrealBuildTool
 		public bool GetString(string SectionName, string KeyName, out string Value)
 		{
 			string? RetrievedValue;
-			if(TryGetValue(SectionName, KeyName, out RetrievedValue))
+			if (TryGetValue(SectionName, KeyName, out RetrievedValue))
 			{
 				Value = RetrievedValue;
 				return true;
@@ -363,7 +361,7 @@ namespace UnrealBuildTool
 		public bool TryGetValue(string SectionName, string KeyName, out bool Value)
 		{
 			string? Text;
-			if(!TryGetValue(SectionName, KeyName, out Text))
+			if (!TryGetValue(SectionName, KeyName, out Text))
 			{
 				Value = false;
 				return false;
@@ -381,7 +379,7 @@ namespace UnrealBuildTool
 		public bool TryGetValue(string SectionName, string KeyName, out int Value)
 		{
 			string? Text;
-			if(!TryGetValue(SectionName, KeyName, out Text))
+			if (!TryGetValue(SectionName, KeyName, out Text))
 			{
 				Value = 0;
 				return false;
@@ -399,7 +397,7 @@ namespace UnrealBuildTool
 		public bool TryGetValue(string SectionName, string KeyName, out Guid Value)
 		{
 			string? Text;
-			if(!TryGetValue(SectionName, KeyName, out Text))
+			if (!TryGetValue(SectionName, KeyName, out Text))
 			{
 				Value = Guid.Empty;
 				return false;
@@ -417,7 +415,7 @@ namespace UnrealBuildTool
 		public bool TryGetValue(string SectionName, string KeyName, out float Value)
 		{
 			string? Text;
-			if(!TryGetValue(SectionName, KeyName, out Text))
+			if (!TryGetValue(SectionName, KeyName, out Text))
 			{
 				Value = 0;
 				return false;
@@ -435,7 +433,7 @@ namespace UnrealBuildTool
 		public bool TryGetValue(string SectionName, string KeyName, out double Value)
 		{
 			string? Text;
-			if(!TryGetValue(SectionName, KeyName, out Text))
+			if (!TryGetValue(SectionName, KeyName, out Text))
 			{
 				Value = 0;
 				return false;
@@ -529,7 +527,7 @@ namespace UnrealBuildTool
 					for (int ComponentIndex = 0; ComponentIndex < 4; ComponentIndex++)
 					{
 						int IntegerValue;
-						if(!Int32.TryParse(ComponentValues[ComponentIndex], out IntegerValue))
+						if (!Int32.TryParse(ComponentValues[ComponentIndex], out IntegerValue))
 						{
 							Value = Guid.Empty;
 							return false;
@@ -550,7 +548,7 @@ namespace UnrealBuildTool
 		/// <returns>True if the text was parsed, false otherwise</returns>
 		public static bool TryParse(string Text, out float Value)
 		{
-			if(Text.EndsWith("f"))
+			if (Text.EndsWith("f"))
 			{
 				return Single.TryParse(Text.Substring(0, Text.Length - 1), out Value);
 			}
@@ -568,7 +566,7 @@ namespace UnrealBuildTool
 		/// <returns>True if the text was parsed, false otherwise</returns>
 		public static bool TryParse(string Text, out double Value)
 		{
-			if(Text.EndsWith("f"))
+			if (Text.EndsWith("f"))
 			{
 				return Double.TryParse(Text.Substring(0, Text.Length - 1), out Value);
 			}
@@ -592,11 +590,11 @@ namespace UnrealBuildTool
 
 			// Get the opening paren
 			int Idx = 0;
-			while(Char.IsWhiteSpace(Chars[Idx]))
+			while (Char.IsWhiteSpace(Chars[Idx]))
 			{
 				Idx++;
 			}
-			if(Chars[Idx] != '(')
+			if (Chars[Idx] != '(')
 			{
 				Properties = null;
 				return false;
@@ -604,7 +602,7 @@ namespace UnrealBuildTool
 
 			// Read to the next token
 			Idx++;
-			while(Char.IsWhiteSpace(Chars[Idx]))
+			while (Char.IsWhiteSpace(Chars[Idx]))
 			{
 				Idx++;
 			}
@@ -614,17 +612,17 @@ namespace UnrealBuildTool
 
 			// Read a sequence of key/value pairs
 			StringBuilder Value = new StringBuilder();
-			if(Chars[Idx] != ')')
+			if (Chars[Idx] != ')')
 			{
-				for (;;)
+				for (; ; )
 				{
 					// Find the end of the name
 					int NameIdx = Idx;
-					while(Char.IsLetterOrDigit(Chars[Idx]) || Chars[Idx] == '_')
+					while (Char.IsLetterOrDigit(Chars[Idx]) || Chars[Idx] == '_')
 					{
 						Idx++;
 					}
-					if(Idx == NameIdx)
+					if (Idx == NameIdx)
 					{
 						Properties = null;
 						return false;
@@ -632,18 +630,18 @@ namespace UnrealBuildTool
 
 					// Extract the key string, and make sure it hasn't already been added
 					string Key = new string(Chars, NameIdx, Idx - NameIdx);
-					if(NewProperties.ContainsKey(Key))
+					if (NewProperties.ContainsKey(Key))
 					{
 						Properties = null;
 						return false;
 					}
 
 					// Consume the equals character
-					while(Char.IsWhiteSpace(Chars[Idx]))
+					while (Char.IsWhiteSpace(Chars[Idx]))
 					{
 						Idx++;
 					}
-					if(Chars[Idx] != '=')
+					if (Chars[Idx] != '=')
 					{
 						Properties = null;
 						return false;
@@ -669,7 +667,7 @@ namespace UnrealBuildTool
 					else if (Chars[Idx] == '\"')
 					{
 						Idx++;
-						for(; Chars[Idx] != '\"'; Idx++)
+						for (; Chars[Idx] != '\"'; Idx++)
 						{
 							if (Chars[Idx] == '\0')
 							{
@@ -720,15 +718,15 @@ namespace UnrealBuildTool
 					NewProperties[Key] = Value.ToString();
 
 					// Move to the separator
-					while(Char.IsWhiteSpace(Chars[Idx]))
+					while (Char.IsWhiteSpace(Chars[Idx]))
 					{
 						Idx++;
 					}
-					if(Chars[Idx] == ')')
+					if (Chars[Idx] == ')')
 					{
 						break;
 					}
-					if(Chars[Idx] != ',')
+					if (Chars[Idx] != ',')
 					{
 						Properties = null;
 						return false;
@@ -745,11 +743,11 @@ namespace UnrealBuildTool
 
 			// Make sure we're at the end of the string
 			Idx++;
-			while(Char.IsWhiteSpace(Chars[Idx]))
+			while (Char.IsWhiteSpace(Chars[Idx]))
 			{
 				Idx++;
 			}
-			if(Chars[Idx] != '\0')
+			if (Chars[Idx] != '\0')
 			{
 				Properties = null;
 				return false;
@@ -774,11 +772,11 @@ namespace UnrealBuildTool
 
 			// Get the opening paren
 			int Idx = 0;
-			while(Char.IsWhiteSpace(Chars[Idx]))
+			while (Char.IsWhiteSpace(Chars[Idx]))
 			{
 				Idx++;
 			}
-			if(Chars[Idx] != '(')
+			if (Chars[Idx] != '(')
 			{
 				Array = null;
 				return false;
@@ -786,7 +784,7 @@ namespace UnrealBuildTool
 
 			// Read to the next token
 			Idx++;
-			while(Char.IsWhiteSpace(Chars[Idx]))
+			while (Char.IsWhiteSpace(Chars[Idx]))
 			{
 				Idx++;
 			}
@@ -796,9 +794,9 @@ namespace UnrealBuildTool
 
 			// Read a sequence items
 			StringBuilder Value = new StringBuilder();
-			if(Chars[Idx] != ')')
+			if (Chars[Idx] != ')')
 			{
-				for (;;)
+				for (; ; )
 				{
 					// Skip whitespace
 					while (Char.IsWhiteSpace(Chars[Idx]))
@@ -819,7 +817,7 @@ namespace UnrealBuildTool
 					else if (Chars[Idx] == '\"')
 					{
 						Idx++;
-						for(; Chars[Idx] != '\"'; Idx++)
+						for (; Chars[Idx] != '\"'; Idx++)
 						{
 							if (Chars[Idx] == '\0')
 							{
@@ -859,7 +857,7 @@ namespace UnrealBuildTool
 							}
 							Value.Append(Chars[Idx]);
 						}
-					}					
+					}
 					else
 					{
 						Array = null;
@@ -870,15 +868,15 @@ namespace UnrealBuildTool
 					NewArray.Add(Value.ToString());
 
 					// Move to the separator
-					while(Char.IsWhiteSpace(Chars[Idx]))
+					while (Char.IsWhiteSpace(Chars[Idx]))
 					{
 						Idx++;
 					}
-					if(Chars[Idx] == ')')
+					if (Chars[Idx] == ')')
 					{
 						break;
 					}
-					if(Chars[Idx] != ',')
+					if (Chars[Idx] != ',')
 					{
 						Array = null;
 						return false;
@@ -895,11 +893,11 @@ namespace UnrealBuildTool
 
 			// Make sure we're at the end of the string
 			Idx++;
-			while(Char.IsWhiteSpace(Chars[Idx]))
+			while (Char.IsWhiteSpace(Chars[Idx]))
 			{
 				Idx++;
 			}
-			if(Chars[Idx] != '\0')
+			if (Chars[Idx] != '\0')
 			{
 				Array = null;
 				return false;
@@ -918,7 +916,7 @@ namespace UnrealBuildTool
 			public string? After2 = null;
 		};
 
-		static string [] ConfigLayers =
+		static string[] ConfigLayers =
 		{
 			// Engine/Base.ini
 			"{ENGINE}/Config/Base.ini",
@@ -953,10 +951,10 @@ namespace UnrealBuildTool
 			// The base expansion (ie, no expansion)
 			new ConfigLayerExpansion { }, 
 			// Restricted Locations
-			new ConfigLayerExpansion { Before1 = "{ENGINE}/", After1 = "{ENGINE}/Restricted/NotForLicensees/",	Before2 = "{PROJECT}/Config/", After2 = "{RESTRICTEDPROJECT_NFL}/Config/" },
-			new ConfigLayerExpansion { Before1 = "{ENGINE}/", After1 = "{ENGINE}/Restricted/NoRedist/",			Before2 = "{PROJECT}/Config/", After2 = "{RESTRICTEDPROJECT_NR}/Config/" },
+			new ConfigLayerExpansion { Before1 = "{ENGINE}/", After1 = "{ENGINE}/Restricted/NotForLicensees/", Before2 = "{PROJECT}/Config/", After2 = "{RESTRICTEDPROJECT_NFL}/Config/" },
+			new ConfigLayerExpansion { Before1 = "{ENGINE}/", After1 = "{ENGINE}/Restricted/NoRedist/",         Before2 = "{PROJECT}/Config/", After2 = "{RESTRICTEDPROJECT_NR}/Config/" },
 			// Platform Extensions
-			new ConfigLayerExpansion { Before1 = "{ENGINE}/Config/{PLATFORM}/", After1 = "{EXTENGINE}/Config/",	Before2 = "{PROJECT}/Config/{PLATFORM}/", After2 = "{EXTPROJECT}/Config/" },
+			new ConfigLayerExpansion { Before1 = "{ENGINE}/Config/{PLATFORM}/", After1 = "{EXTENGINE}/Config/",    Before2 = "{PROJECT}/Config/{PLATFORM}/", After2 = "{EXTPROJECT}/Config/" },
 			// Platform Extensions in Restricted Locations
 			new ConfigLayerExpansion { Before1 = "{ENGINE}/Config/{PLATFORM}/", After1 = "{ENGINE}/Restricted/NotForLicensees/Platforms/{PLATFORM}/Config/",   Before2 = "{PROJECT}/Config/{PLATFORM}/", After2 = "{RESTRICTEDPROJECT_NFL}/Platforms/{PLATFORM}/{OPT_SUBDIR}Config/" },
 			new ConfigLayerExpansion { Before1 = "{ENGINE}/Config/{PLATFORM}/", After1 = "{ENGINE}/Restricted/NoRedist/Platforms/{PLATFORM}/Config/",          Before2 = "{PROJECT}/Config/{PLATFORM}/", After2 = "{RESTRICTEDPROJECT_NR}/Platforms/{PLATFORM}/{OPT_SUBDIR}Config/" },
@@ -977,7 +975,7 @@ namespace UnrealBuildTool
 			{
 				OutString = OutString.Replace("{USER}", UserDir.FullName);
 			}
-			
+
 			OutString = OutString.Replace("{CUSTOMCONFIG}", CustomConfig);
 
 			return OutString;
@@ -1127,7 +1125,7 @@ namespace UnrealBuildTool
 			}
 
 			// Find all the generated config files
-			foreach(FileReference GeneratedConfigFile in EnumerateGeneratedConfigFileLocations(Type, ProjectDir, Platform))
+			foreach (FileReference GeneratedConfigFile in EnumerateGeneratedConfigFileLocations(Type, ProjectDir, Platform))
 			{
 				yield return GeneratedConfigFile;
 			}
@@ -1162,7 +1160,7 @@ namespace UnrealBuildTool
 		/// <returns></returns>
 		public static DirectoryReference GetGeneratedConfigDir(DirectoryReference? ProjectDir)
 		{
-			if(ProjectDir == null)
+			if (ProjectDir == null)
 			{
 				return DirectoryReference.Combine(Unreal.EngineDirectory, "Saved", "Config");
 			}
@@ -1178,7 +1176,7 @@ namespace UnrealBuildTool
 		/// <returns></returns>
 		public static DirectoryReference GetGameAgnosticSavedDir()
 		{
-			if(Unreal.IsEngineInstalled())
+			if (Unreal.IsEngineInstalled())
 			{
 				DirectoryReference? UserSettingDir = Unreal.UserSettingDirectory;
 				if (UserSettingDir != null)
diff --git a/Engine/Source/Programs/UnrealBuildTool/System/ConfigValueTracker.cs b/Engine/Source/Programs/UnrealBuildTool/System/ConfigValueTracker.cs
--- a/Engine/Source/Programs/UnrealBuildTool/System/ConfigValueTracker.cs
+++ b/Engine/Source/Programs/UnrealBuildTool/System/ConfigValueTracker.cs
@@ -4,10 +4,7 @@ using System;
 using System.Collections.Generic;
 using System.Diagnostics;
 using System.Linq;
-using System.Text;
-using System.Threading.Tasks;
 using EpicGames.Core;
-using Microsoft.Extensions.Logging;
 
 namespace UnrealBuildTool
 {
@@ -26,7 +23,7 @@ namespace UnrealBuildTool
 		/// Project directory to read config files from
 		/// </summary>
 		public DirectoryReference? ProjectDir;
-		
+
 		/// <summary>
 		/// The platform being built
 		/// </summary>
@@ -163,7 +160,7 @@ namespace UnrealBuildTool
 		/// <returns></returns>
 		public bool IsValid()
 		{
-			foreach(KeyValuePair<ConfigDependencyKey, IReadOnlyList<string>?> Pair in Dependencies)
+			foreach (KeyValuePair<ConfigDependencyKey, IReadOnlyList<string>?> Pair in Dependencies)
 			{
 				// Read the appropriate hierarchy
 				ConfigHierarchy Hierarchy = ConfigCache.ReadHierarchy(Pair.Key.Type, Pair.Key.ProjectDir, Pair.Key.Platform);
@@ -173,16 +170,16 @@ namespace UnrealBuildTool
 				Hierarchy.TryGetValues(Pair.Key.SectionName, Pair.Key.KeyName, out NewValues);
 
 				// Check if they're different
-				if(Pair.Value == null)
+				if (Pair.Value == null)
 				{
-					if(NewValues != null)
+					if (NewValues != null)
 					{
 						return false;
 					}
 				}
 				else
 				{
-					if(NewValues == null || !Enumerable.SequenceEqual(Pair.Value, NewValues, StringComparer.Ordinal))
+					if (NewValues == null || !Enumerable.SequenceEqual(Pair.Value, NewValues, StringComparer.Ordinal))
 					{
 						return false;
 					}
diff --git a/Engine/Source/Programs/UnrealBuildTool/System/CppCompileEnvironment.cs b/Engine/Source/Programs/UnrealBuildTool/System/CppCompileEnvironment.cs
--- a/Engine/Source/Programs/UnrealBuildTool/System/CppCompileEnvironment.cs
+++ b/Engine/Source/Programs/UnrealBuildTool/System/CppCompileEnvironment.cs
@@ -1,14 +1,8 @@
 // Copyright Epic Games, Inc. All Rights Reserved.
 
-using System;
 using System.Collections.Generic;
-using System.Diagnostics;
-using System.Text;
-using System.Text.RegularExpressions;
-using System.IO;
 using EpicGames.Core;
 using UnrealBuildBase;
-using System.Runtime.InteropServices;
 
 namespace UnrealBuildTool
 {
@@ -43,12 +37,12 @@ namespace UnrealBuildTool
 		/// Supports C++20
 		/// </summary>
 		Cpp20,
-		
+
 		/// <summary>
 		/// Latest standard supported by the compiler
 		/// </summary>
 		Latest,
-		
+
 		/// <summary>
 		/// Use the default standard version
 		/// </summary>
@@ -59,7 +53,7 @@ namespace UnrealBuildTool
 		/// </summary>
 		EngineDefault = Cpp20,
 	}
-	
+
 	/// <summary>
 	/// Specifies which C language standard to use. This enum should be kept in order, so that toolchains can check whether the requested setting is >= values that they support.
 	/// </summary>
@@ -614,7 +608,7 @@ namespace UnrealBuildTool
 		/// Which C++ standard to support. May not be compatible with all platforms.
 		/// </summary>
 		public CppStandardVersion CppStandard = CppStandardVersion.Default;
-		
+
 		/// <summary>
 		/// Which C standard to support. May not be compatible with all platforms.
 		/// </summary>
diff --git a/Engine/Source/Programs/UnrealBuildTool/System/CppDependencyCache.cs b/Engine/Source/Programs/UnrealBuildTool/System/CppDependencyCache.cs
--- a/Engine/Source/Programs/UnrealBuildTool/System/CppDependencyCache.cs
+++ b/Engine/Source/Programs/UnrealBuildTool/System/CppDependencyCache.cs
@@ -3,15 +3,14 @@
 using System;
 using System.Collections.Concurrent;
 using System.Collections.Generic;
-using System.Diagnostics;
 using System.Diagnostics.CodeAnalysis;
 using System.IO;
 using System.Linq;
 using System.Text;
 using System.Threading.Tasks;
 using EpicGames.Core;
-using UnrealBuildBase;
 using Microsoft.Extensions.Logging;
+using UnrealBuildBase;
 
 namespace UnrealBuildTool
 {
@@ -294,7 +293,7 @@ namespace UnrealBuildTool
 		/// <returns>True if the input file exists and the dependencies were read</returns>
 		private bool TryGetDependencyInfoInternal(FileItem InputFile, [NotNullWhen(true)] out DependencyInfo? OutInfo)
 		{
-			foreach(CachePartition Partition in Partitions)
+			foreach (CachePartition Partition in Partitions)
 			{
 				if (InputFile.Location.IsUnderDirectory(Partition.BaseDir))
 				{
@@ -524,7 +523,7 @@ namespace UnrealBuildTool
 					}
 				}
 				else
-				{ 
+				{
 					if (Data.TryGetStringArrayField("ImportedModules", out string[]? ImportedModuleArray) && ImportedModuleArray.Length > 0)
 					{
 						ImportedModules =
diff --git a/Engine/Source/Programs/UnrealBuildTool/System/CppIncludeLookup.cs b/Engine/Source/Programs/UnrealBuildTool/System/CppIncludeLookup.cs
--- a/Engine/Source/Programs/UnrealBuildTool/System/CppIncludeLookup.cs
+++ b/Engine/Source/Programs/UnrealBuildTool/System/CppIncludeLookup.cs
@@ -1,14 +1,13 @@
 // Copyright Epic Games, Inc. All Rights Reserved.
 
-using EpicGames.Core;
 using System;
 using System.Collections.Concurrent;
 using System.Collections.Generic;
 using System.Diagnostics;
-using System.IO;
 using System.Linq;
 using System.Text;
 using System.Text.RegularExpressions;
+using EpicGames.Core;
 using UnrealBuildBase;
 
 namespace UnrealBuildTool
@@ -62,7 +61,7 @@ namespace UnrealBuildTool
 					if (Version == CurrentVersion)
 					{
 						CppIncludeNameInfo[] NameInfos = Reader.ReadArray(() => new CppIncludeNameInfo(Reader.ReadString()!))!;
-						foreach(CppIncludeNameInfo NameInfo in NameInfos)
+						foreach (CppIncludeNameInfo NameInfo in NameInfos)
 						{
 							int NumFiles = Reader.ReadInt();
 							for (int Idx = 0; Idx < NumFiles; Idx++)
@@ -73,7 +72,7 @@ namespace UnrealBuildTool
 								if (NameIndexes != null)
 								{
 									FileInfo.IncludedNames = NameIndexes.Select(x => NameInfos[x]).ToList();
-									FileInfo.LastWriteTimeUtc = new DateTime(Reader.ReadLong(), DateTimeKind.Utc); 
+									FileInfo.LastWriteTimeUtc = new DateTime(Reader.ReadLong(), DateTimeKind.Utc);
 								}
 
 								NameInfo.Files.Add(FileInfo);
@@ -105,7 +104,7 @@ namespace UnrealBuildTool
 					foreach (CppIncludeFileInfo FileInfo in NameInfo.Files)
 					{
 						Writer.WriteCompactFileItem(FileInfo.File);
-						if(FileInfo.IncludedNames == null)
+						if (FileInfo.IncludedNames == null)
 						{
 							Writer.WriteIntArray(null);
 						}
@@ -328,7 +327,7 @@ namespace UnrealBuildTool
 				}
 
 				byte[]? Name = ParseIncludeName(Span, Idx);
-				if(Name != null)
+				if (Name != null)
 				{
 					Names.Add(Encoding.UTF8.GetString(Name));
 				}
@@ -347,10 +346,10 @@ namespace UnrealBuildTool
 
 			// Check it's an include directive
 			int MinIdx = Idx;
-			for (; ;)
+			for (; ; )
 			{
 				byte Char = Span[--MinIdx];
-				if(Char == '#')
+				if (Char == '#')
 				{
 					break;
 				}
@@ -363,7 +362,7 @@ namespace UnrealBuildTool
 					return null;
 				}
 			}
-			for (; MinIdx > 0; )
+			for (; MinIdx > 0;)
 			{
 				byte Char = Span[--MinIdx];
 				if (Char == '\n')
@@ -382,7 +381,7 @@ namespace UnrealBuildTool
 			{
 				return null;
 			}
-			for (;;)
+			for (; ; )
 			{
 				byte Char = Span[MinNameIdx++];
 				if (Char == '\"' || Char == '<')
diff --git a/Engine/Source/Programs/UnrealBuildTool/System/CppIncludeParser.cs b/Engine/Source/Programs/UnrealBuildTool/System/CppIncludeParser.cs
--- a/Engine/Source/Programs/UnrealBuildTool/System/CppIncludeParser.cs
+++ b/Engine/Source/Programs/UnrealBuildTool/System/CppIncludeParser.cs
@@ -1,11 +1,8 @@
 // Copyright Epic Games, Inc. All Rights Reserved.
 
 using System;
-using System.Collections.Generic;
 using System.IO;
-using System.Linq;
 using System.Text;
-using System.Threading.Tasks;
 
 namespace UnrealBuildTool
 {
@@ -22,49 +19,49 @@ namespace UnrealBuildTool
 		public static void CopyIncludeDirectives(StreamReader Reader, StringBuilder Writer)
 		{
 			StringBuilder Token = new StringBuilder();
-			while(TryReadToken(Reader, Token))
+			while (TryReadToken(Reader, Token))
 			{
-				if(Token.Length > 1 || Token[0] != '\n')
+				if (Token.Length > 1 || Token[0] != '\n')
 				{
-					if(Token.Length != 1 || Token[0] != '#')
+					if (Token.Length != 1 || Token[0] != '#')
 					{
 						break;
 					}
-					if(!TryReadToken(Reader, Token))
+					if (!TryReadToken(Reader, Token))
 					{
 						break;
 					}
 
 					string Directive = Token.ToString();
-					if(Directive == "pragma")
+					if (Directive == "pragma")
 					{
-						if(!TryReadToken(Reader, Token) || Token.ToString() != "once")
+						if (!TryReadToken(Reader, Token) || Token.ToString() != "once")
 						{
 							break;
 						}
-						if(!TryReadToken(Reader, Token) || Token.ToString() != "\n")
+						if (!TryReadToken(Reader, Token) || Token.ToString() != "\n")
 						{
 							break;
 						}
 					}
-					else if(Directive == "include")
+					else if (Directive == "include")
 					{
-						if(!TryReadToken(Reader, Token) || Token[0] != '\"')
+						if (!TryReadToken(Reader, Token) || Token[0] != '\"')
 						{
 							break;
 						}
 
 						string IncludeFile = Token.ToString();
 
-						if(!IncludeFile.EndsWith(".h\"") && !IncludeFile.EndsWith(".h>"))
+						if (!IncludeFile.EndsWith(".h\"") && !IncludeFile.EndsWith(".h>"))
 						{
 							break;
 						}
-						if(IncludeFile.Equals("\"RequiredProgramMainCPPInclude.h\"", StringComparison.OrdinalIgnoreCase))
+						if (IncludeFile.Equals("\"RequiredProgramMainCPPInclude.h\"", StringComparison.OrdinalIgnoreCase))
 						{
 							break;
 						}
-						if(!TryReadToken(Reader, Token) || Token.ToString() != "\n")
+						if (!TryReadToken(Reader, Token) || Token.ToString() != "\n")
 						{
 							break;
 						}
@@ -90,46 +87,46 @@ namespace UnrealBuildTool
 			Token.Clear();
 
 			int NextChar;
-			for(;;)
+			for (; ; )
 			{
 				NextChar = Reader.Read();
-				if(NextChar == -1)
+				if (NextChar == -1)
 				{
 					return false;
 				}
-				if(NextChar != ' ' && NextChar != '\t' && NextChar != '\r')
+				if (NextChar != ' ' && NextChar != '\t' && NextChar != '\r')
 				{
-					if(NextChar != '/')
+					if (NextChar != '/')
 					{
 						break;
 					}
-					else if(Reader.Peek() == '/')
+					else if (Reader.Peek() == '/')
 					{
 						Reader.Read();
-						for(;;)
+						for (; ; )
 						{
 							NextChar = Reader.Read();
-							if(NextChar == -1)
+							if (NextChar == -1)
 							{
 								return false;
 							}
-							if(NextChar == '\n')
+							if (NextChar == '\n')
 							{
 								break;
 							}
 						}
 					}
-					else if(Reader.Peek() == '*')
+					else if (Reader.Peek() == '*')
 					{
 						Reader.Read();
-						for(;;)
+						for (; ; )
 						{
 							NextChar = Reader.Read();
-							if(NextChar == -1)
+							if (NextChar == -1)
 							{
 								return false;
 							}
-							if(NextChar == '*' && Reader.Peek() == '/')
+							if (NextChar == '*' && Reader.Peek() == '/')
 							{
 								break;
 							}
@@ -145,29 +142,29 @@ namespace UnrealBuildTool
 
 			Token.Append((char)NextChar);
 
-			if(Char.IsLetterOrDigit((char)NextChar))
+			if (Char.IsLetterOrDigit((char)NextChar))
 			{
-				for(;;)
+				for (; ; )
 				{
 					NextChar = Reader.Read();
-					if(NextChar == -1 || !Char.IsLetterOrDigit((char)NextChar))
+					if (NextChar == -1 || !Char.IsLetterOrDigit((char)NextChar))
 					{
 						break;
 					}
 					Token.Append((char)NextChar);
 				}
 			}
-			else if(NextChar == '\"' || NextChar == '<')
+			else if (NextChar == '\"' || NextChar == '<')
 			{
-				for(;;)
+				for (; ; )
 				{
 					NextChar = Reader.Read();
-					if(NextChar == -1)
+					if (NextChar == -1)
 					{
 						break;
 					}
 					Token.Append((char)NextChar);
-					if(NextChar == '\"' || NextChar == '>')
+					if (NextChar == '\"' || NextChar == '>')
 					{
 						break;
 					}
diff --git a/Engine/Source/Programs/UnrealBuildTool/System/CustomBuildSteps.cs b/Engine/Source/Programs/UnrealBuildTool/System/CustomBuildSteps.cs
--- a/Engine/Source/Programs/UnrealBuildTool/System/CustomBuildSteps.cs
+++ b/Engine/Source/Programs/UnrealBuildTool/System/CustomBuildSteps.cs
@@ -1,11 +1,8 @@
 // Copyright Epic Games, Inc. All Rights Reserved.
 
-using System;
 using System.Collections.Generic;
 using System.Diagnostics.CodeAnalysis;
 using System.Linq;
-using System.Text;
-using System.Threading.Tasks;
 using EpicGames.Core;
 
 namespace UnrealBuildTool
@@ -15,14 +12,14 @@ namespace UnrealBuildTool
 	/// </summary>
 	public class CustomBuildSteps
 	{
-		Dictionary<UnrealTargetPlatform, string[]> HostPlatformToCommands = new Dictionary<UnrealTargetPlatform,string[]>();
+		Dictionary<UnrealTargetPlatform, string[]> HostPlatformToCommands = new Dictionary<UnrealTargetPlatform, string[]>();
 
 		/// <summary>
 		/// Construct a custom build steps object from a Json object.
 		/// </summary>
 		public CustomBuildSteps(JsonObject RawObject)
 		{
-			foreach(string HostPlatformName in RawObject.KeyNames)
+			foreach (string HostPlatformName in RawObject.KeyNames)
 			{
 				UnrealTargetPlatform Platform;
 				if (UnrealTargetPlatform.TryParse(HostPlatformName, out Platform))
@@ -42,7 +39,7 @@ namespace UnrealBuildTool
 		public static bool TryRead(JsonObject RawObject, string FieldName, [NotNullWhen(true)] out CustomBuildSteps? OutBuildSteps)
 		{
 			JsonObject? BuildStepsObject;
-			if(RawObject.TryGetObjectField(FieldName, out BuildStepsObject))
+			if (RawObject.TryGetObjectField(FieldName, out BuildStepsObject))
 			{
 				OutBuildSteps = new CustomBuildSteps(BuildStepsObject);
 				return true;
@@ -63,10 +60,10 @@ namespace UnrealBuildTool
 		public void Write(JsonWriter Writer, string FieldName)
 		{
 			Writer.WriteObjectStart(FieldName);
-			foreach(KeyValuePair<UnrealTargetPlatform, string[]> Pair in HostPlatformToCommands.OrderBy(x => x.Key.ToString()))
+			foreach (KeyValuePair<UnrealTargetPlatform, string[]> Pair in HostPlatformToCommands.OrderBy(x => x.Key.ToString()))
 			{
 				Writer.WriteArrayStart(Pair.Key.ToString());
-				foreach(string Line in Pair.Value)
+				foreach (string Line in Pair.Value)
 				{
 					Writer.WriteValue(Line);
 				}
@@ -84,7 +81,7 @@ namespace UnrealBuildTool
 		public bool TryGetCommands(UnrealTargetPlatform HostPlatform, [NotNullWhen(true)] out string[]? OutCommands)
 		{
 			string[]? Commands;
-			if(HostPlatformToCommands.TryGetValue(HostPlatform, out Commands) && Commands.Length > 0)
+			if (HostPlatformToCommands.TryGetValue(HostPlatform, out Commands) && Commands.Length > 0)
 			{
 				OutCommands = Commands;
 				return true;
diff --git a/Engine/Source/Programs/UnrealBuildTool/System/DataDrivenPlatformInfo.cs b/Engine/Source/Programs/UnrealBuildTool/System/DataDrivenPlatformInfo.cs
--- a/Engine/Source/Programs/UnrealBuildTool/System/DataDrivenPlatformInfo.cs
+++ b/Engine/Source/Programs/UnrealBuildTool/System/DataDrivenPlatformInfo.cs
@@ -2,12 +2,10 @@
 
 using System;
 using System.Collections.Generic;
+using System.IO;
 using System.Linq;
-using System.Text;
 using EpicGames.Core;
-using System.IO;
 using UnrealBuildBase;
-using Microsoft.Extensions.Logging;
 
 namespace UnrealBuildTool
 {
diff --git a/Engine/Source/Programs/UnrealBuildTool/System/DeployExports.cs b/Engine/Source/Programs/UnrealBuildTool/System/DeployExports.cs
--- a/Engine/Source/Programs/UnrealBuildTool/System/DeployExports.cs
+++ b/Engine/Source/Programs/UnrealBuildTool/System/DeployExports.cs
@@ -1,11 +1,6 @@
 // Copyright Epic Games, Inc. All Rights Reserved.
 
 using System;
-using System.Collections.Generic;
-using System.Linq;
-using System.Text;
-using System.Threading.Tasks;
-using EpicGames.Core;
 using UnrealBuildBase;
 
 namespace UnrealBuildTool
@@ -18,7 +13,7 @@ namespace UnrealBuildTool
 		/// <summary>
 		/// The max length of the deploy folder name.
 		/// </summary>
-		public const int DeployFolderMaxLength = 50; 
+		public const int DeployFolderMaxLength = 50;
 
 		/// <summary>
 		/// Gets the default deploy folder name on the format ShortProjectName-RootDirectoryName-UserName,
@@ -39,7 +34,7 @@ namespace UnrealBuildTool
 			if (DeployFolder.Length > DeployFolderMaxLength)
 			{
 				DeployFolder = DeployFolder.Substring(0, DeployFolderMaxLength);
-				char[] CharsToTrim = {'-'};
+				char[] CharsToTrim = { '-' };
 				DeployFolder = DeployFolder.Trim(CharsToTrim);
 			}
 
diff --git a/Engine/Source/Programs/UnrealBuildTool/System/DynamicCompilation.cs b/Engine/Source/Programs/UnrealBuildTool/System/DynamicCompilation.cs
--- a/Engine/Source/Programs/UnrealBuildTool/System/DynamicCompilation.cs
+++ b/Engine/Source/Programs/UnrealBuildTool/System/DynamicCompilation.cs
@@ -4,20 +4,14 @@ using System;
 using System.Collections.Generic;
 using System.IO;
 using System.Linq;
-using System.CodeDom.Compiler;
-using Microsoft.CSharp;
 using System.Reflection;
-using System.Diagnostics;
+using System.Runtime.InteropServices;
 using EpicGames.Core;
-
 using Microsoft.CodeAnalysis;
 using Microsoft.CodeAnalysis.CSharp;
-using Microsoft.CodeAnalysis.CSharp.Syntax;
 using Microsoft.CodeAnalysis.Emit;
-using Microsoft.Extensions.Logging;
-using System.Reflection.Metadata;
 using Microsoft.CodeAnalysis.Text;
-using System.Runtime.InteropServices;
+using Microsoft.Extensions.Logging;
 using OpenTracing.Util;
 using UnrealBuildBase;
 
@@ -89,30 +83,30 @@ namespace UnrealBuildTool
 			// Make sure the source files we're compiling are the same as the source files that were compiled
 			// for the assembly that we want to load
 			HashSet<FileItem> CurrentSourceFileItems = new HashSet<FileItem>();
-			foreach(string Line in Manifest.GetStringArrayField("SourceFiles"))
+			foreach (string Line in Manifest.GetStringArrayField("SourceFiles"))
 			{
 				CurrentSourceFileItems.Add(FileItem.GetItemByPath(Line));
 			}
 
 			// Get the new source files
 			HashSet<FileItem> SourceFileItems = new HashSet<FileItem>();
-			foreach(FileReference SourceFile in SourceFiles)
+			foreach (FileReference SourceFile in SourceFiles)
 			{
 				SourceFileItems.Add(FileItem.GetItemByFileReference(SourceFile));
 			}
 
 			// Check if there are any differences between the sets
-			foreach(FileItem CurrentSourceFileItem in CurrentSourceFileItems)
+			foreach (FileItem CurrentSourceFileItem in CurrentSourceFileItems)
 			{
-				if(!SourceFileItems.Contains(CurrentSourceFileItem))
+				if (!SourceFileItems.Contains(CurrentSourceFileItem))
 				{
 					Logger.LogDebug("Compiling {OutputAssemblyPath}: Removed source file ({CurrentSourceFileItem})", OutputAssemblyPath, CurrentSourceFileItem);
 					return true;
 				}
 			}
-			foreach(FileItem SourceFileItem in SourceFileItems)
+			foreach (FileItem SourceFileItem in SourceFileItems)
 			{
-				if(!CurrentSourceFileItems.Contains(SourceFileItem))
+				if (!CurrentSourceFileItems.Contains(SourceFileItem))
 				{
 					Logger.LogDebug("Compiling {OutputAssemblyPath}: Added source file ({SourceFileItem})", OutputAssemblyPath, SourceFileItem);
 					return true;
@@ -120,9 +114,9 @@ namespace UnrealBuildTool
 			}
 
 			// Check if any of the timestamps are newer
-			foreach(FileItem SourceFileItem in SourceFileItems)
+			foreach (FileItem SourceFileItem in SourceFileItems)
 			{
-				if(SourceFileItem.LastWriteTimeUtc > OutputAssemblyInfo.LastWriteTimeUtc)
+				if (SourceFileItem.LastWriteTimeUtc > OutputAssemblyInfo.LastWriteTimeUtc)
 				{
 					Logger.LogDebug("Compiling {OutputAssemblyPath}: {SourceFileItem} is newer", OutputAssemblyPath, SourceFileItem);
 					return true;
@@ -168,9 +162,9 @@ namespace UnrealBuildTool
 		private static Assembly? CompileAssembly(FileReference OutputAssemblyPath, IEnumerable<FileReference> SourceFileNames, ILogger Logger, IEnumerable<string>? ReferencedAssembies, IEnumerable<string>? PreprocessorDefines = null, bool TreatWarningsAsErrors = false)
 		{
 			CSharpParseOptions ParseOptions = new CSharpParseOptions(
-				languageVersion:LanguageVersion.Latest, 
-				kind:SourceCodeKind.Regular,
-				preprocessorSymbols:PreprocessorDefines
+				languageVersion: LanguageVersion.Latest,
+				kind: SourceCodeKind.Regular,
+				preprocessorSymbols: PreprocessorDefines
 			);
 
 			List<SyntaxTree> SyntaxTrees = new List<SyntaxTree>();
@@ -227,11 +221,11 @@ namespace UnrealBuildTool
 			MetadataReferences.Add(MetadataReference.CreateFromFile(Assembly.Load("System.Runtime.Extensions").Location));
 			MetadataReferences.Add(MetadataReference.CreateFromFile(Assembly.Load("Microsoft.Extensions.Logging.Abstractions").Location));
 			MetadataReferences.Add(MetadataReference.CreateFromFile(Assembly.Load("netstandard").Location));
-			
+
 			// process start dependencies
 			MetadataReferences.Add(MetadataReference.CreateFromFile(Assembly.Load("System.ComponentModel.Primitives").Location));
 			MetadataReferences.Add(MetadataReference.CreateFromFile(Assembly.Load("System.Diagnostics.Process").Location));
-			
+
 			// registry access
 			MetadataReferences.Add(MetadataReference.CreateFromFile(Assembly.Load("Microsoft.Win32.Registry").Location));
 
@@ -244,28 +238,28 @@ namespace UnrealBuildTool
 			MetadataReferences.Add(MetadataReference.CreateFromFile(typeof(UEBuildPlatformSDK).Assembly.Location));
 
 			CSharpCompilationOptions CompilationOptions = new CSharpCompilationOptions(
-				outputKind:OutputKind.DynamicallyLinkedLibrary,
+				outputKind: OutputKind.DynamicallyLinkedLibrary,
 #if DEBUG
 				optimizationLevel: OptimizationLevel.Debug,
 #else
 				// Optimize the managed code in Development
 				optimizationLevel: OptimizationLevel.Release,
 #endif
-				warningLevel:4,
-				assemblyIdentityComparer:DesktopAssemblyIdentityComparer.Default,
-				reportSuppressedDiagnostics:true
+				warningLevel: 4,
+				assemblyIdentityComparer: DesktopAssemblyIdentityComparer.Default,
+				reportSuppressedDiagnostics: true
 			);
 
 			CSharpCompilation Compilation = CSharpCompilation.Create(
-				assemblyName:OutputAssemblyPath.GetFileNameWithoutAnyExtensions(),
-				syntaxTrees:SyntaxTrees,
-				references:MetadataReferences,
-				options:CompilationOptions
+				assemblyName: OutputAssemblyPath.GetFileNameWithoutAnyExtensions(),
+				syntaxTrees: SyntaxTrees,
+				references: MetadataReferences,
+				options: CompilationOptions
 				);
 
 			using (FileStream AssemblyStream = FileReference.Open(OutputAssemblyPath, FileMode.Create))
 			{
-				using (FileStream? PdbStream = RuntimeInformation.IsOSPlatform(OSPlatform.Windows)? FileReference.Open(OutputAssemblyPath.ChangeExtension(".pdb"), FileMode.Create) : null)
+				using (FileStream? PdbStream = RuntimeInformation.IsOSPlatform(OSPlatform.Windows) ? FileReference.Open(OutputAssemblyPath.ChangeExtension(".pdb"), FileMode.Create) : null)
 				{
 					EmitOptions EmitOptions = new EmitOptions(
 						includePrivateMembers: true
@@ -332,9 +326,9 @@ namespace UnrealBuildTool
 					bNeedsCompilation = true;
 				}
 				catch (FileNotFoundException)
-			    {
-				    throw new BuildException("Precompiled rules assembly '{0}' does not exist.", OutputAssemblyPath);
-			    }
+				{
+					throw new BuildException("Precompiled rules assembly '{0}' does not exist.", OutputAssemblyPath);
+				}
 				catch (Exception Ex)
 				{
 					throw new BuildException(Ex, "Error while loading previously-compiled assembly file '{0}'.  (Exception: {1})", OutputAssemblyPath, Ex.Message);
diff --git a/Engine/Source/Programs/UnrealBuildTool/System/EncryptionAndSigning.cs b/Engine/Source/Programs/UnrealBuildTool/System/EncryptionAndSigning.cs
--- a/Engine/Source/Programs/UnrealBuildTool/System/EncryptionAndSigning.cs
+++ b/Engine/Source/Programs/UnrealBuildTool/System/EncryptionAndSigning.cs
@@ -2,9 +2,7 @@
 
 using System;
 using System.Collections.Generic;
-using System.IO;
 using System.Linq;
-using System.Text;
 using System.Text.Json;
 using System.Text.RegularExpressions;
 using EpicGames.Core;
@@ -158,7 +156,7 @@ namespace UnrealBuildTool
 			/// Config setting to enable pak encryption
 			/// </summary>
 			public bool PakSigningRequired { get; set; } = true;
-			
+
 			/// <summary>
 			/// A set of named encryption keys that can be used to encrypt different sets of data with a different key that is delivered dynamically (i.e. not embedded within the game executable)
 			/// </summary>
@@ -252,7 +250,7 @@ namespace UnrealBuildTool
 		public static CryptoSettings ParseCryptoSettings(DirectoryReference? InProjectDirectory, UnrealTargetPlatform InTargetPlatform, ILogger Logger)
 		{
 			CryptoSettings Settings = new CryptoSettings();
-			
+
 			ConfigHierarchy Ini = ConfigCache.ReadHierarchy(ConfigHierarchyType.Engine, InProjectDirectory, InTargetPlatform);
 			Ini.GetBool("PlatformCrypto", "PlatformRequiresDataCrypto", out bool bDataCryptoRequired);
 			Settings.bDataCryptoRequired = bDataCryptoRequired;
diff --git a/Engine/Source/Programs/UnrealBuildTool/System/ExternalExecution.cs b/Engine/Source/Programs/UnrealBuildTool/System/ExternalExecution.cs
--- a/Engine/Source/Programs/UnrealBuildTool/System/ExternalExecution.cs
+++ b/Engine/Source/Programs/UnrealBuildTool/System/ExternalExecution.cs
@@ -2,25 +2,20 @@
 
 using System;
 using System.Collections.Generic;
-using System.Linq;
-using System.Text;
 using System.Diagnostics;
 using System.IO;
-using System.Text.Json;
-using System.Text.RegularExpressions;
+using System.Linq;
 using System.Reflection;
-using System.Runtime.InteropServices;
-using System.Runtime.Serialization;
-using EpicGames.Core;
-using System.Text.Json.Serialization;
-using UnrealBuildBase;
-using System.Collections.Concurrent;
+using System.Text;
+using System.Text.Json;
 using System.Threading.Tasks;
+using EpicGames.Core;
+using EpicGames.UHT.Utils;
+using Microsoft.Extensions.Logging;
 using OpenTracing;
 using OpenTracing.Util;
+using UnrealBuildBase;
 using UnrealBuildTool.Modes;
-using EpicGames.UHT.Utils;
-using Microsoft.Extensions.Logging;
 
 namespace UnrealBuildTool
 {
@@ -311,15 +306,15 @@ namespace UnrealBuildTool
 				if (RulesObject.Plugin != null)
 				{
 					ModuleDescriptor? Module = RulesObject.Plugin.Descriptor.Modules?.FirstOrDefault(x => x.Name == RulesObject.Name);
-					if(Module != null)
+					if (Module != null)
 					{
 						return GetGameModuleTypeFromDescriptor(Module);
 					}
 				}
-				if(ProjectDescriptor != null && ProjectDescriptor.Modules != null)
+				if (ProjectDescriptor != null && ProjectDescriptor.Modules != null)
 				{
 					ModuleDescriptor? Module = ProjectDescriptor.Modules.FirstOrDefault(x => x.Name == RulesObject.Name);
-					if(Module != null)
+					if (Module != null)
 					{
 						return UHTModuleTypeExtensions.GameModuleTypeFromHostType(Module.Type) ?? UHTModuleType.GameRuntime;
 					}
@@ -366,7 +361,7 @@ namespace UnrealBuildTool
 
 			foreach (DirectoryItem SubDirectory in BaseDirectory.EnumerateDirectories())
 			{
-				if(!ExcludeFolders.Contains(SubDirectory.Name))
+				if (!ExcludeFolders.Contains(SubDirectory.Name))
 				{
 					FindHeaders(SubDirectory, ExcludeFolders, Headers);
 				}
@@ -414,7 +409,7 @@ namespace UnrealBuildTool
 		{
 			// Find the type of each module
 			Dictionary<UEBuildModuleCPP, UHTModuleType> ModuleToType = new Dictionary<UEBuildModuleCPP, UHTModuleType>();
-			foreach(UEBuildModuleCPP Module in ModulesToGenerateHeadersFor)
+			foreach (UEBuildModuleCPP Module in ModulesToGenerateHeadersFor)
 			{
 				ModuleToType[Module] = GetModuleType(Module.Rules, ProjectDescriptor);
 			}
@@ -425,10 +420,10 @@ namespace UnrealBuildTool
 
 			// Create the info for each module in parallel
 			UHTModuleInfo[] ModuleInfoArray = new UHTModuleInfo[ModulesSortedByType.Count];
-			using(ThreadPoolWorkQueue Queue = new ThreadPoolWorkQueue())
+			using (ThreadPoolWorkQueue Queue = new ThreadPoolWorkQueue())
 			{
 				IReadOnlySet<string> ExcludedFolders = UEBuildPlatform.GetBuildPlatform(Platform).GetExcludedFolderNames();
-				for(int Idx = 0; Idx < ModulesSortedByType.Count; Idx++)
+				for (int Idx = 0; Idx < ModulesSortedByType.Count; Idx++)
 				{
 					UEBuildModuleCPP Module = ModulesSortedByType[Idx];
 
@@ -443,7 +438,7 @@ namespace UnrealBuildTool
 			}
 
 			// Filter out all the modules with reflection data
-			for(int Idx = 0; Idx < ModulesSortedByType.Count; Idx++)
+			for (int Idx = 0; Idx < ModulesSortedByType.Count; Idx++)
 			{
 				UEBuildModuleCPP Module = ModulesSortedByType[Idx];
 				UHTModuleInfo Info = ModuleInfoArray[Idx];
@@ -550,7 +545,7 @@ namespace UnrealBuildTool
 			// UObjects or not, we can skip doing a test here.
 			if (MetadataCache.ContainsReflectionMarkup(HeaderFile))
 			{
-				lock(ModuleInfo)
+				lock (ModuleInfo)
 				{
 					bool bFoundHeaderLocation = false;
 					foreach (DirectoryReference ModuleDirectory in ModuleInfo.ModuleDirectories)
@@ -775,7 +770,7 @@ namespace UnrealBuildTool
 					HashSet<string> ObjectHeadersSet = new HashSet<string>(AllUObjectHeaders.Select(x => x.AbsolutePath), FileReference.Comparer);
 					foreach (string FileName in UObjectFilesFromPreviousRun)
 					{
-						if(!ObjectHeadersSet.Contains(FileName))
+						if (!ObjectHeadersSet.Contains(FileName))
 						{
 							Logger.LogDebug("UnrealHeaderTool needs to run because the set of UObject source files in module {ModuleModuleName} has changed ({FileName})", Module.ModuleName, FileName);
 							return true;
@@ -1020,7 +1015,7 @@ namespace UnrealBuildTool
 		/// Builds and runs the header tool and touches the header directories.
 		/// Performs any early outs if headers need no changes, given the UObject modules, tool path, game name, and configuration
 		/// </summary>
-		private static async Task ExecuteHeaderToolIfNecessaryInternalAsync(BuildConfiguration BuildConfiguration, FileReference? ProjectFile, 
+		private static async Task ExecuteHeaderToolIfNecessaryInternalAsync(BuildConfiguration BuildConfiguration, FileReference? ProjectFile,
 			TargetMakefile Makefile, string TargetName, ISourceFileWorkingSet WorkingSet, ILogger Logger)
 		{
 			if (ProgressWriter.bWriteMarkup)
@@ -1200,7 +1195,7 @@ namespace UnrealBuildTool
 
 		static void ResetCachedHeaderInfo(List<UHTModuleInfo> UObjectModules)
 		{
-			foreach(UHTModuleInfo ModuleInfo in UObjectModules)
+			foreach (UHTModuleInfo ModuleInfo in UObjectModules)
 			{
 				ModuleInfo.GeneratedCodeDirectory.ResetCachedInfo();
 			}
diff --git a/Engine/Source/Programs/UnrealBuildTool/System/FileMetadataPrefetch.cs b/Engine/Source/Programs/UnrealBuildTool/System/FileMetadataPrefetch.cs
--- a/Engine/Source/Programs/UnrealBuildTool/System/FileMetadataPrefetch.cs
+++ b/Engine/Source/Programs/UnrealBuildTool/System/FileMetadataPrefetch.cs
@@ -1,11 +1,8 @@
 // Copyright Epic Games, Inc. All Rights Reserved.
 
-using System;
 using System.Collections.Generic;
 using System.Linq;
-using System.Text;
 using System.Threading;
-using System.Threading.Tasks;
 using EpicGames.Core;
 using UnrealBuildBase;
 
@@ -41,9 +38,9 @@ namespace UnrealBuildTool
 		/// </summary>
 		public static void QueueEngineDirectory()
 		{
-			lock(QueuedDirectories)
+			lock (QueuedDirectories)
 			{
-				if(QueuedDirectories.Add(Unreal.EngineDirectory))
+				if (QueuedDirectories.Add(Unreal.EngineDirectory))
 				{
 					Enqueue(() => ScanEngineDirectory());
 				}
@@ -56,9 +53,9 @@ namespace UnrealBuildTool
 		/// <param name="ProjectDirectory">The project directory to prefetch</param>
 		public static void QueueProjectDirectory(DirectoryReference ProjectDirectory)
 		{
-			lock(QueuedDirectories)
+			lock (QueuedDirectories)
 			{
-				if(QueuedDirectories.Add(ProjectDirectory))
+				if (QueuedDirectories.Add(ProjectDirectory))
 				{
 					Enqueue(() => ScanProjectDirectory(DirectoryItem.GetItemByDirectoryReference(ProjectDirectory)));
 				}
@@ -71,9 +68,9 @@ namespace UnrealBuildTool
 		/// <param name="Directory">Directory to start searching from</param>
 		public static void QueueDirectoryTree(DirectoryReference Directory)
 		{
-			lock(QueuedDirectories)
+			lock (QueuedDirectories)
 			{
-				if(QueuedDirectories.Add(Directory))
+				if (QueuedDirectories.Add(Directory))
 				{
 					Enqueue(() => ScanDirectoryTree(DirectoryItem.GetItemByDirectoryReference(Directory)));
 				}
@@ -103,7 +100,7 @@ namespace UnrealBuildTool
 		/// <param name="Action">Action to enqueue</param>
 		static void Enqueue(System.Action Action)
 		{
-			Queue.Enqueue(() => { if(!CancelToken.IsCancellationRequested){ Action(); } });
+			Queue.Enqueue(() => { if (!CancelToken.IsCancellationRequested) { Action(); } });
 		}
 
 		/// <summary>
diff --git a/Engine/Source/Programs/UnrealBuildTool/System/GraphVisualization.cs b/Engine/Source/Programs/UnrealBuildTool/System/GraphVisualization.cs
--- a/Engine/Source/Programs/UnrealBuildTool/System/GraphVisualization.cs
+++ b/Engine/Source/Programs/UnrealBuildTool/System/GraphVisualization.cs
@@ -1,13 +1,11 @@
 // Copyright Epic Games, Inc. All Rights Reserved.
 
-using EpicGames.Core;
 using System;
 using System.Collections.Generic;
 using System.Drawing;
 using System.Globalization;
-using System.Linq;
-using System.Text;
 using System.Xml;
+using EpicGames.Core;
 
 namespace UnrealBuildTool
 {
diff --git a/Engine/Source/Programs/UnrealBuildTool/System/InstalledPlatformInfo.cs b/Engine/Source/Programs/UnrealBuildTool/System/InstalledPlatformInfo.cs
--- a/Engine/Source/Programs/UnrealBuildTool/System/InstalledPlatformInfo.cs
+++ b/Engine/Source/Programs/UnrealBuildTool/System/InstalledPlatformInfo.cs
@@ -2,13 +2,9 @@
 
 using System;
 using System.Collections.Generic;
-using System.Linq;
-using System.Text;
-using System.Threading.Tasks;
 using System.IO;
 using EpicGames.Core;
 using UnrealBuildBase;
-using Microsoft.Extensions.Logging;
 
 namespace UnrealBuildTool
 {
@@ -130,7 +126,7 @@ namespace UnrealBuildTool
 			ConfigHierarchy Ini = ConfigCache.ReadHierarchy(ConfigHierarchyType.Engine, (DirectoryReference?)null, BuildHostPlatform.Current.Platform);
 
 			bool bHasInstalledPlatformInfo;
-			if(Ini.TryGetValue("InstalledPlatforms", "HasInstalledPlatformInfo", out bHasInstalledPlatformInfo) && bHasInstalledPlatformInfo)
+			if (Ini.TryGetValue("InstalledPlatforms", "HasInstalledPlatformInfo", out bHasInstalledPlatformInfo) && bHasInstalledPlatformInfo)
 			{
 				InstalledPlatformConfigurations = new List<InstalledPlatformConfiguration>();
 				if (Ini.GetArray("InstalledPlatforms", "InstalledPlatformConfigurations", out InstalledPlatforms))
@@ -169,7 +165,7 @@ namespace UnrealBuildTool
 			}
 			if (Configuration == UnrealTargetConfiguration.Unknown)
 			{
-//				Logger.LogWarning("Unable to read configuration from {PlatformConfiguration}", PlatformConfiguration);
+				//				Logger.LogWarning("Unable to read configuration from {PlatformConfiguration}", PlatformConfiguration);
 				bCanCreateEntry = false;
 			}
 
@@ -178,7 +174,7 @@ namespace UnrealBuildTool
 			{
 				if (!UnrealTargetPlatform.IsValidName(PlatformName))
 				{
-//					Logger.LogWarning("Unable to read platform from {PlatformConfiguration}", PlatformConfiguration);
+					//					Logger.LogWarning("Unable to read platform from {PlatformConfiguration}", PlatformConfiguration);
 					bCanCreateEntry = false;
 				}
 			}
@@ -189,13 +185,13 @@ namespace UnrealBuildTool
 			{
 				if (!Enum.TryParse(PlatformTypeName, out PlatformType))
 				{
-//					Logger.LogWarning("Unable to read Platform Type from {PlatformConfiguration}, defaulting to Game", PlatformConfiguration);
+					//					Logger.LogWarning("Unable to read Platform Type from {PlatformConfiguration}, defaulting to Game", PlatformConfiguration);
 					PlatformType = TargetType.Game;
 				}
 			}
 			if (PlatformType == TargetType.Program)
 			{
-//				Logger.LogWarning("Program is not a valid PlatformType for an Installed Platform, defaulting to Game");
+				//				Logger.LogWarning("Program is not a valid PlatformType for an Installed Platform, defaulting to Game");
 				PlatformType = TargetType.Game;
 			}
 
@@ -210,14 +206,14 @@ namespace UnrealBuildTool
 			}
 
 			string ProjectTypeName;
-			EProjectType ProjectType =  EProjectType.Any;
+			EProjectType ProjectType = EProjectType.Any;
 			if (ParseSubValue(PlatformConfiguration, "ProjectType=", out ProjectTypeName))
 			{
 				Enum.TryParse(ProjectTypeName, out ProjectType);
 			}
 			if (ProjectType == EProjectType.Unknown)
 			{
-//				Logger.LogWarning("Unable to read project type from {PlatformConfiguration}", PlatformConfiguration);
+				//				Logger.LogWarning("Unable to read project type from {PlatformConfiguration}", PlatformConfiguration);
 				bCanCreateEntry = false;
 			}
 
@@ -291,7 +287,7 @@ namespace UnrealBuildTool
 		public static bool IsValidPlatform(UnrealTargetPlatform Platform, EProjectType ProjectType = EProjectType.Any)
 		{
 			// HACK: For installed builds, we always need to treat Mac as a valid platform for generating project files. When remote building from PC, we won't have all the libraries to do this, so we need to fake it.
-			if(Platform == UnrealTargetPlatform.Mac && ProjectType == EProjectType.Any && BuildHostPlatform.Current.Platform == UnrealTargetPlatform.Mac && Unreal.IsEngineInstalled())
+			if (Platform == UnrealTargetPlatform.Mac && ProjectType == EProjectType.Any && BuildHostPlatform.Current.Platform == UnrealTargetPlatform.Mac && Unreal.IsEngineInstalled())
 			{
 				return true;
 			}
@@ -336,31 +332,31 @@ namespace UnrealBuildTool
 		/// <returns>True if the target can be built</returns>
 		public static bool IsValid(TargetType? TargetType, UnrealTargetPlatform? Platform, UnrealTargetConfiguration? Configuration, EProjectType ProjectType, InstalledPlatformState State)
 		{
-			if(!Unreal.IsEngineInstalled() || InstalledPlatformConfigurations == null)
+			if (!Unreal.IsEngineInstalled() || InstalledPlatformConfigurations == null)
 			{
 				return true;
 			}
 
-			foreach(InstalledPlatformConfiguration Config in InstalledPlatformConfigurations)
+			foreach (InstalledPlatformConfiguration Config in InstalledPlatformConfigurations)
 			{
 				// Check whether this configuration matches all the criteria
-				if(TargetType.HasValue && Config.PlatformType != TargetType.Value)
+				if (TargetType.HasValue && Config.PlatformType != TargetType.Value)
 				{
 					continue;
 				}
-				if(Platform.HasValue && Config.Platform != Platform.Value)
+				if (Platform.HasValue && Config.Platform != Platform.Value)
 				{
 					continue;
 				}
-				if(Configuration.HasValue && Config.Configuration != Configuration.Value)
+				if (Configuration.HasValue && Config.Configuration != Configuration.Value)
 				{
 					continue;
 				}
-				if(ProjectType != EProjectType.Any && Config.ProjectType != EProjectType.Any && Config.ProjectType != ProjectType)
+				if (ProjectType != EProjectType.Any && Config.ProjectType != EProjectType.Any && Config.ProjectType != ProjectType)
 				{
 					continue;
 				}
-				if(State == InstalledPlatformState.Downloaded && !String.IsNullOrEmpty(Config.RequiredFile) && !File.Exists(Config.RequiredFile))
+				if (State == InstalledPlatformState.Downloaded && !String.IsNullOrEmpty(Config.RequiredFile) && !File.Exists(Config.RequiredFile))
 				{
 					continue;
 				}
diff --git a/Engine/Source/Programs/UnrealBuildTool/System/JsonExtensions.cs b/Engine/Source/Programs/UnrealBuildTool/System/JsonExtensions.cs
--- a/Engine/Source/Programs/UnrealBuildTool/System/JsonExtensions.cs
+++ b/Engine/Source/Programs/UnrealBuildTool/System/JsonExtensions.cs
@@ -1,12 +1,7 @@
 // Copyright Epic Games, Inc. All Rights Reserved.
 
-using System;
-using System.Collections.Generic;
-using System.IO;
-using System.Diagnostics;
-using System.Linq;
-using EpicGames.Core;
 using System.Diagnostics.CodeAnalysis;
+using EpicGames.Core;
 
 namespace JsonExtensions
 {
diff --git a/Engine/Source/Programs/UnrealBuildTool/System/JunkDeleter.cs b/Engine/Source/Programs/UnrealBuildTool/System/JunkDeleter.cs
--- a/Engine/Source/Programs/UnrealBuildTool/System/JunkDeleter.cs
+++ b/Engine/Source/Programs/UnrealBuildTool/System/JunkDeleter.cs
@@ -2,10 +2,7 @@
 
 using System;
 using System.Collections.Generic;
-using System.Linq;
-using System.Text;
 using System.IO;
-using EpicGames.Core;
 using Microsoft.Extensions.Logging;
 
 namespace UnrealBuildTool
diff --git a/Engine/Source/Programs/UnrealBuildTool/System/LinkEnvironment.cs b/Engine/Source/Programs/UnrealBuildTool/System/LinkEnvironment.cs
--- a/Engine/Source/Programs/UnrealBuildTool/System/LinkEnvironment.cs
+++ b/Engine/Source/Programs/UnrealBuildTool/System/LinkEnvironment.cs
@@ -1,10 +1,8 @@
 // Copyright Epic Games, Inc. All Rights Reserved.
 
 using System;
-using System.Collections;
 using System.Collections.Generic;
 using System.Linq;
-using System.Text;
 using EpicGames.Core;
 using UnrealBuildBase;
 
@@ -162,7 +160,7 @@ namespace UnrealBuildTool
 		/// The method of linking the target uses
 		/// </summary>
 		public TargetLinkType LinkType;
-		
+
 		/// <summary>
 		/// Whether we should compile using the statically-linked CRT. This is not widely supported for the whole engine, but is required for programs that need to run without dependencies.
 		/// </summary>
@@ -337,7 +335,7 @@ namespace UnrealBuildTool
 		/// List of functions that should be exported from this module
 		/// </summary>
 		public List<string> IncludeFunctions = new List<string>();
-		
+
 		/// <summary>
 		/// Debugger visualizer files to build into debug info for this binary. 
 		/// </summary>
@@ -469,4 +467,3 @@ namespace UnrealBuildTool
 		}
 	}
 }
-	
\ No newline at end of file
diff --git a/Engine/Source/Programs/UnrealBuildTool/System/LocalizationDescriptor.cs b/Engine/Source/Programs/UnrealBuildTool/System/LocalizationDescriptor.cs
--- a/Engine/Source/Programs/UnrealBuildTool/System/LocalizationDescriptor.cs
+++ b/Engine/Source/Programs/UnrealBuildTool/System/LocalizationDescriptor.cs
@@ -1,10 +1,6 @@
 // Copyright Epic Games, Inc. All Rights Reserved.
 
-using System;
-using System.Collections.Generic;
-using System.IO;
 using System.Diagnostics;
-using System.Linq;
 using EpicGames.Core;
 
 namespace UnrealBuildTool
diff --git a/Engine/Source/Programs/UnrealBuildTool/System/ModuleDescriptor.cs b/Engine/Source/Programs/UnrealBuildTool/System/ModuleDescriptor.cs
--- a/Engine/Source/Programs/UnrealBuildTool/System/ModuleDescriptor.cs
+++ b/Engine/Source/Programs/UnrealBuildTool/System/ModuleDescriptor.cs
@@ -2,7 +2,6 @@
 
 using System;
 using System.Collections.Generic;
-using System.IO;
 using System.Diagnostics;
 using System.Linq;
 using EpicGames.Core;
@@ -30,10 +29,10 @@ namespace UnrealBuildTool
 		/// </summary>
 		RuntimeNoCommandlet,
 
-        /// <summary>
-        /// Any target or program
-        /// </summary>
-        RuntimeAndProgram,
+		/// <summary>
+		/// Any target or program
+		/// </summary>
+		RuntimeAndProgram,
 
 		/// <summary>
 		/// Loaded only in cooked builds
@@ -78,12 +77,12 @@ namespace UnrealBuildTool
 		/// <summary>
 		/// Loaded only by servers
 		/// </summary>
-        ServerOnly,
+		ServerOnly,
 
 		/// <summary>
 		/// Loaded only by clients, and commandlets, and editor....
 		/// </summary>
-        ClientOnly,
+		ClientOnly,
 
 		/// <summary>
 		/// Loaded only by clients and editor (editor can run PIE which is kinda a commandlet)
@@ -410,11 +409,11 @@ namespace UnrealBuildTool
 				}
 				Writer.WriteArrayEnd();
 			}
-			if(ProgramAllowList != null && ProgramAllowList.Length > 0)
+			if (ProgramAllowList != null && ProgramAllowList.Length > 0)
 			{
 				Writer.WriteStringArrayField("ProgramAllowList", ProgramAllowList);
 			}
-			if(ProgramDenyList != null && ProgramDenyList.Length > 0)
+			if (ProgramDenyList != null && ProgramDenyList.Length > 0)
 			{
 				Writer.WriteStringArrayField("ProgramDenyList", ProgramDenyList);
 			}
@@ -459,7 +458,7 @@ namespace UnrealBuildTool
 		/// <param name="File">File containing the module declaration</param>
 		public void Validate(FileReference File)
 		{
-			if(Type == ModuleHostType.Developer)
+			if (Type == ModuleHostType.Developer)
 			{
 				Log.TraceWarningOnce("The 'Developer' module type has been deprecated in 4.24. Use 'DeveloperTool' for modules that can be loaded by game/client/server targets in non-shipping configurations, or 'UncookedOnly' for modules that should only be loaded by uncooked editor and program targets (eg. modules containing blueprint nodes)");
 				Log.TraceWarningOnce(File, "The 'Developer' module type has been deprecated in 4.24.");
@@ -481,7 +480,7 @@ namespace UnrealBuildTool
 			// important note: we don't check the length of the platform allow list, because if an unknown platform was read in, but was not valid, the 
 			// list will exist but be empty. In this case, we need to disallow all platforms from building, otherwise, build errors will occur when
 			// it starts compiling for _all_ platforms. This means we don't need to check bHasExplicitPlatforms either
-			
+
 			if (PlatformAllowList != null && !PlatformAllowList.Contains(Platform))
 			{
 				return false;
@@ -518,16 +517,16 @@ namespace UnrealBuildTool
 			}
 
 			// Special checks just for programs
-			if(TargetType == TargetType.Program)
+			if (TargetType == TargetType.Program)
 			{
 				// Check the program name is on the allow list. Note that this behavior is slightly different to other allow/deny checks; we will allow a module of any type if it's explicitly allowed for this program.
-				if(ProgramAllowList != null && ProgramAllowList.Length > 0)
+				if (ProgramAllowList != null && ProgramAllowList.Length > 0)
 				{
 					return ProgramAllowList.Contains(TargetName);
 				}
-				
+
 				// Check the program name is not denied
-				if(ProgramDenyList != null && ProgramDenyList.Contains(TargetName))
+				if (ProgramDenyList != null && ProgramDenyList.Contains(TargetName))
 				{
 					return false;
 				}
@@ -538,11 +537,11 @@ namespace UnrealBuildTool
 			{
 				case ModuleHostType.Runtime:
 				case ModuleHostType.RuntimeNoCommandlet:
-                    return TargetType != TargetType.Program;
+					return TargetType != TargetType.Program;
 				case ModuleHostType.RuntimeAndProgram:
 					return true;
 				case ModuleHostType.CookedOnly:
-                    return bBuildRequiresCookedData;
+					return bBuildRequiresCookedData;
 				case ModuleHostType.UncookedOnly:
 					return !bBuildRequiresCookedData;
 				case ModuleHostType.Developer:
@@ -556,12 +555,12 @@ namespace UnrealBuildTool
 					return TargetType == TargetType.Editor || TargetType == TargetType.Program;
 				case ModuleHostType.Program:
 					return TargetType == TargetType.Program;
-                case ModuleHostType.ServerOnly:
-                    return TargetType != TargetType.Program && TargetType != TargetType.Client;
-                case ModuleHostType.ClientOnly:
-                case ModuleHostType.ClientOnlyNoCommandlet:
-                    return TargetType != TargetType.Program && TargetType != TargetType.Server;
-            }
+				case ModuleHostType.ServerOnly:
+					return TargetType != TargetType.Program && TargetType != TargetType.Client;
+				case ModuleHostType.ClientOnly:
+				case ModuleHostType.ClientOnlyNoCommandlet:
+					return TargetType != TargetType.Program && TargetType != TargetType.Server;
+			}
 
 			return false;
 		}
diff --git a/Engine/Source/Programs/UnrealBuildTool/System/ModuleManifest.cs b/Engine/Source/Programs/UnrealBuildTool/System/ModuleManifest.cs
--- a/Engine/Source/Programs/UnrealBuildTool/System/ModuleManifest.cs
+++ b/Engine/Source/Programs/UnrealBuildTool/System/ModuleManifest.cs
@@ -5,8 +5,6 @@ using System.Collections.Generic;
 using System.Diagnostics.CodeAnalysis;
 using System.IO;
 using System.Linq;
-using System.Text;
-using System.Threading.Tasks;
 using EpicGames.Core;
 
 namespace UnrealBuildTool
@@ -64,7 +62,7 @@ namespace UnrealBuildTool
 			{
 				BaseName += String.Format("-{0}-{1}", Platform.ToString(), Configuration.ToString());
 			}
-			if(UnrealArchitectureConfig.ForPlatform(Platform).RequiresArchitectureFilenames(BuildArchitectures))
+			if (UnrealArchitectureConfig.ForPlatform(Platform).RequiresArchitectureFilenames(BuildArchitectures))
 			{
 				BaseName += BuildArchitectures.ToString();
 			}
@@ -122,7 +120,7 @@ namespace UnrealBuildTool
 		public void Write(FileReference FileName)
 		{
 			DirectoryReference.CreateDirectory(FileName.Directory);
-			using(StreamWriter Writer = new StreamWriter(FileName.FullName))
+			using (StreamWriter Writer = new StreamWriter(FileName.FullName))
 			{
 				Write(Writer);
 			}
diff --git a/Engine/Source/Programs/UnrealBuildTool/System/NativeProjects.cs b/Engine/Source/Programs/UnrealBuildTool/System/NativeProjects.cs
--- a/Engine/Source/Programs/UnrealBuildTool/System/NativeProjects.cs
+++ b/Engine/Source/Programs/UnrealBuildTool/System/NativeProjects.cs
@@ -5,11 +5,9 @@ using System.Collections.Generic;
 using System.Diagnostics.CodeAnalysis;
 using System.IO;
 using System.Linq;
-using System.Text;
-using System.Threading.Tasks;
 using EpicGames.Core;
-using UnrealBuildBase;
 using Microsoft.Extensions.Logging;
+using UnrealBuildBase;
 
 namespace UnrealBuildTool
 {
@@ -41,14 +39,14 @@ namespace UnrealBuildTool
 		/// <returns>True if the target was found</returns>
 		public static bool TryGetProjectForTarget(string InTargetName, ILogger Logger, [NotNullWhen(true)] out FileReference? OutProjectFileName)
 		{
-			if(CachedTargetNameToProjectFile == null)
+			if (CachedTargetNameToProjectFile == null)
 			{
-				lock(LockObject)
+				lock (LockObject)
 				{
-					if(CachedTargetNameToProjectFile == null)
+					if (CachedTargetNameToProjectFile == null)
 					{
 						Dictionary<string, FileReference> TargetNameToProjectFile = new Dictionary<string, FileReference>();
-						foreach(FileReference ProjectFile in EnumerateProjectFiles(Logger))
+						foreach (FileReference ProjectFile in EnumerateProjectFiles(Logger))
 						{
 							foreach (DirectoryReference ExtensionDir in Unreal.GetExtensionDirs(ProjectFile.Directory))
 							{
@@ -92,9 +90,9 @@ namespace UnrealBuildTool
 		public static bool IsHybridContentOnlyProject(FileReference UProjectFile, ILogger Logger)
 		{
 			return RequiresTempTarget(
-				UProjectFile, 
+				UProjectFile,
 				new List<UnrealTargetPlatform>() { BuildHostPlatform.Current.Platform },
-				new List<UnrealTargetConfiguration>() { UnrealTargetConfiguration.Development, UnrealTargetConfiguration.Shipping }, 
+				new List<UnrealTargetConfiguration>() { UnrealTargetConfiguration.Development, UnrealTargetConfiguration.Shipping },
 				out _,
 				Logger);
 		}
@@ -153,7 +151,7 @@ namespace UnrealBuildTool
 				}
 				return false;
 			}
-			
+
 			// if the files existed, just leave them be
 			if (bIsHybrid && bWasHybrid)
 			{
diff --git a/Engine/Source/Programs/UnrealBuildTool/System/PlatformExports.cs b/Engine/Source/Programs/UnrealBuildTool/System/PlatformExports.cs
--- a/Engine/Source/Programs/UnrealBuildTool/System/PlatformExports.cs
+++ b/Engine/Source/Programs/UnrealBuildTool/System/PlatformExports.cs
@@ -3,9 +3,6 @@
 using System;
 using System.Collections.Generic;
 using System.Linq;
-using System.Runtime.Versioning;
-using System.Text;
-using System.Threading.Tasks;
 using EpicGames.Core;
 using Microsoft.Extensions.Logging;
 
@@ -44,7 +41,7 @@ namespace UnrealBuildTool
 		public static bool HasDefaultBuildConfig(FileReference ProjectFile, UnrealTargetPlatform Platform)
 		{
 			UEBuildPlatform.TryGetBuildPlatform(Platform, out UEBuildPlatform? BuildPlat);
-			return (BuildPlat == null)? true : BuildPlat.HasDefaultBuildConfig(Platform, ProjectFile.Directory);
+			return (BuildPlat == null) ? true : BuildPlat.HasDefaultBuildConfig(Platform, ProjectFile.Directory);
 		}
 
 		/// <summary>
diff --git a/Engine/Source/Programs/UnrealBuildTool/System/PluginDescriptor.cs b/Engine/Source/Programs/UnrealBuildTool/System/PluginDescriptor.cs
--- a/Engine/Source/Programs/UnrealBuildTool/System/PluginDescriptor.cs
+++ b/Engine/Source/Programs/UnrealBuildTool/System/PluginDescriptor.cs
@@ -2,10 +2,7 @@
 
 using System;
 using System.Collections.Generic;
-using System.Diagnostics;
-using System.IO;
 using System.Linq;
-using System.Text;
 using EpicGames.Core;
 
 namespace UnrealBuildTool
@@ -36,9 +33,9 @@ namespace UnrealBuildTool
 		ProjectPluginUnification = 3,
 
 		/// <summary>
-        /// This needs to be the last line, so we can calculate the value of Latest below
+		/// This needs to be the last line, so we can calculate the value of Latest below
 		/// </summary>
-        LatestPlusOne,
+		LatestPlusOne,
 
 		/// <summary>
 		/// The latest plugin descriptor version
@@ -312,7 +309,7 @@ namespace UnrealBuildTool
 			}
 
 			bool bEnabledByDefaultValue;
-			if(RawObject.TryGetBoolField("EnabledByDefault", out bEnabledByDefaultValue))
+			if (RawObject.TryGetBoolField("EnabledByDefault", out bEnabledByDefaultValue))
 			{
 				bEnabledByDefault = bEnabledByDefaultValue;
 			}
@@ -324,9 +321,9 @@ namespace UnrealBuildTool
 			RawObject.TryGetBoolField("Installed", out bInstalled);
 
 			bool bCanBeUsedWithUnrealHeaderTool;
-			if(RawObject.TryGetBoolField("CanBeUsedWithUnrealHeaderTool", out bCanBeUsedWithUnrealHeaderTool) && bCanBeUsedWithUnrealHeaderTool)
+			if (RawObject.TryGetBoolField("CanBeUsedWithUnrealHeaderTool", out bCanBeUsedWithUnrealHeaderTool) && bCanBeUsedWithUnrealHeaderTool)
 			{
-				Array.Resize(ref SupportedPrograms, (SupportedPrograms == null)? 1 : SupportedPrograms.Length + 1);
+				Array.Resize(ref SupportedPrograms, (SupportedPrograms == null) ? 1 : SupportedPrograms.Length + 1);
 				SupportedPrograms[SupportedPrograms.Length - 1] = "UnrealHeaderTool";
 			}
 
@@ -339,7 +336,7 @@ namespace UnrealBuildTool
 			CustomBuildSteps.TryRead(RawObject, "PostBuildSteps", out PostBuildSteps);
 
 			JsonObject[]? PluginsArray;
-			if(RawObject.TryGetObjectArrayField("Plugins", out PluginsArray))
+			if (RawObject.TryGetObjectArrayField("Plugins", out PluginsArray))
 			{
 				Plugins = Array.ConvertAll(PluginsArray, x => PluginReferenceDescriptor.FromJsonObject(x)).ToList();
 			}
@@ -396,7 +393,7 @@ namespace UnrealBuildTool
 			Writer.WriteValue("DocsURL", DocsURL);
 			Writer.WriteValue("MarketplaceURL", MarketplaceURL);
 			Writer.WriteValue("SupportURL", SupportURL);
-			if(!String.IsNullOrEmpty(EngineVersion))
+			if (!String.IsNullOrEmpty(EngineVersion))
 			{
 				Writer.WriteValue("EngineVersion", EngineVersion);
 			}
@@ -430,7 +427,7 @@ namespace UnrealBuildTool
 				Writer.WriteValue("Installed", bInstalled);
 			}
 
-			if(bRequiresBuildPlatform)
+			if (bRequiresBuildPlatform)
 			{
 				Writer.WriteValue("RequiresBuildPlatform", bRequiresBuildPlatform);
 			}
@@ -450,7 +447,7 @@ namespace UnrealBuildTool
 				Writer.WriteValue("HasExplicitPlatforms", bHasExplicitPlatforms);
 			}
 
-			if(SupportedTargetPlatforms != null && SupportedTargetPlatforms.Count > 0)
+			if (SupportedTargetPlatforms != null && SupportedTargetPlatforms.Count > 0)
 			{
 				Writer.WriteStringArrayField("SupportedTargetPlatforms", SupportedTargetPlatforms.Select<UnrealTargetPlatform, string>(x => x.ToString()).ToArray());
 			}
@@ -471,12 +468,12 @@ namespace UnrealBuildTool
 
 			LocalizationTargetDescriptor.WriteArray(Writer, "LocalizationTargets", LocalizationTargets);
 
-			if(PreBuildSteps != null)
+			if (PreBuildSteps != null)
 			{
 				PreBuildSteps.Write(Writer, "PreBuildSteps");
 			}
 
-			if(PostBuildSteps != null)
+			if (PostBuildSteps != null)
 			{
 				PostBuildSteps.Write(Writer, "PostBuildSteps");
 			}
@@ -538,7 +535,7 @@ namespace UnrealBuildTool
 		/// <returns>String list of supported target platforms</returns>
 		public string[]? GetSupportedTargetPlatformNames()
 		{
-			return SupportedTargetPlatforms?.Select( P => P.ToString() ).ToArray();
+			return SupportedTargetPlatforms?.Select(P => P.ToString()).ToArray();
 		}
 	}
 }
diff --git a/Engine/Source/Programs/UnrealBuildTool/System/PluginReferenceDescriptor.cs b/Engine/Source/Programs/UnrealBuildTool/System/PluginReferenceDescriptor.cs
--- a/Engine/Source/Programs/UnrealBuildTool/System/PluginReferenceDescriptor.cs
+++ b/Engine/Source/Programs/UnrealBuildTool/System/PluginReferenceDescriptor.cs
@@ -1,11 +1,8 @@
 // Copyright Epic Games, Inc. All Rights Reserved.
 
 using System;
-using System.Collections.Generic;
 using System.Diagnostics;
 using System.Linq;
-using System.Text;
-using System.Threading.Tasks;
 using EpicGames.Core;
 using JsonExtensions;
 
@@ -110,31 +107,31 @@ namespace UnrealBuildTool
 			Writer.WriteObjectStart();
 			Writer.WriteValue("Name", Name);
 			Writer.WriteValue("Enabled", bEnabled);
-			if(bEnabled && bOptional)
+			if (bEnabled && bOptional)
 			{
 				Writer.WriteValue("Optional", bOptional);
 			}
-			if(!String.IsNullOrEmpty(Description))
+			if (!String.IsNullOrEmpty(Description))
 			{
 				Writer.WriteValue("Description", Description);
 			}
-			if(!String.IsNullOrEmpty(MarketplaceURL))
+			if (!String.IsNullOrEmpty(MarketplaceURL))
 			{
 				Writer.WriteValue("MarketplaceURL", MarketplaceURL);
 			}
-			if(PlatformAllowList != null && PlatformAllowList.Length > 0)
+			if (PlatformAllowList != null && PlatformAllowList.Length > 0)
 			{
 				Writer.WriteStringArrayField("PlatformAllowList", PlatformAllowList.Select(x => x.ToString()).ToArray());
 			}
-			if(PlatformDenyList != null && PlatformDenyList.Length > 0)
+			if (PlatformDenyList != null && PlatformDenyList.Length > 0)
 			{
 				Writer.WriteStringArrayField("PlatformDenyList", PlatformDenyList.Select(x => x.ToString()).ToArray());
 			}
-			if(TargetConfigurationAllowList != null && TargetConfigurationAllowList.Length > 0)
+			if (TargetConfigurationAllowList != null && TargetConfigurationAllowList.Length > 0)
 			{
 				Writer.WriteEnumArrayField("TargetConfigurationAllowList", TargetConfigurationAllowList);
 			}
-			if(TargetConfigurationDenyList != null && TargetConfigurationDenyList.Length > 0)
+			if (TargetConfigurationDenyList != null && TargetConfigurationDenyList.Length > 0)
 			{
 				Writer.WriteEnumArrayField("TargetConfigurationDenyList", TargetConfigurationDenyList);
 			}
@@ -142,17 +139,17 @@ namespace UnrealBuildTool
 			{
 				Writer.WriteEnumArrayField("TargetAllowList", TargetAllowList);
 			}
-			if(TargetDenyList != null && TargetDenyList.Length > 0)
+			if (TargetDenyList != null && TargetDenyList.Length > 0)
 			{
 				Writer.WriteEnumArrayField("TargetDenyList", TargetDenyList);
 			}
-			if(SupportedTargetPlatforms != null && SupportedTargetPlatforms.Length > 0)
+			if (SupportedTargetPlatforms != null && SupportedTargetPlatforms.Length > 0)
 			{
 				Writer.WriteStringArrayField("SupportedTargetPlatforms", SupportedTargetPlatforms.Select(x => x.ToString()).ToArray());
 			}
-			if(bHasExplicitPlatforms)
+			if (bHasExplicitPlatforms)
 			{
-				Writer.WriteValue("HasExplicitPlatforms",bHasExplicitPlatforms);
+				Writer.WriteValue("HasExplicitPlatforms", bHasExplicitPlatforms);
 			}
 			if (bEnabled && RequestedVersion != null)
 			{
@@ -184,16 +181,16 @@ namespace UnrealBuildTool
 		/// Checks the given platform names array and logs a message about any that are not known. This may simply be due to not having the platform synced in Engine/Platforms/[PlatformName].
 		/// </summary>
 		/// <param name="PlatformNames"></param>
-		private void VerifyPlatformNames( string[]? PlatformNames )
+		private void VerifyPlatformNames(string[]? PlatformNames)
 		{
 			if (PlatformNames != null)
 			{
-				foreach( string PlatformName in PlatformNames )
+				foreach (string PlatformName in PlatformNames)
 				{
 					UnrealTargetPlatform Platform;
 					if (!UnrealTargetPlatform.TryParse(PlatformName, out Platform))
 					{
-						Log.TraceLogOnce("Ignoring unknown platform '{0}' (referenced via a project's plugin descriptor for '{1}')", PlatformName, Name );
+						Log.TraceLogOnce("Ignoring unknown platform '{0}' (referenced via a project's plugin descriptor for '{1}')", PlatformName, Name);
 					}
 				}
 			}
diff --git a/Engine/Source/Programs/UnrealBuildTool/System/PrecompiledHeaderInstance.cs b/Engine/Source/Programs/UnrealBuildTool/System/PrecompiledHeaderInstance.cs
--- a/Engine/Source/Programs/UnrealBuildTool/System/PrecompiledHeaderInstance.cs
+++ b/Engine/Source/Programs/UnrealBuildTool/System/PrecompiledHeaderInstance.cs
@@ -1,10 +1,6 @@
 // Copyright Epic Games, Inc. All Rights Reserved.
 
-using System;
 using System.Collections.Generic;
-using System.Linq;
-using System.Text;
-using System.Threading.Tasks;
 using UnrealBuildBase;
 
 namespace UnrealBuildTool
@@ -52,7 +48,7 @@ namespace UnrealBuildTool
 		/// <summary>
 		/// Dictionary of definitions use in the CppCompileEnvironment
 		/// </summary>
-		public Dictionary<string,string>? DefinitionsDictionary;
+		public Dictionary<string, string>? DefinitionsDictionary;
 
 		/// <summary>
 		/// Constructor
diff --git a/Engine/Source/Programs/UnrealBuildTool/System/PrecompiledHeaderTemplate.cs b/Engine/Source/Programs/UnrealBuildTool/System/PrecompiledHeaderTemplate.cs
--- a/Engine/Source/Programs/UnrealBuildTool/System/PrecompiledHeaderTemplate.cs
+++ b/Engine/Source/Programs/UnrealBuildTool/System/PrecompiledHeaderTemplate.cs
@@ -1,10 +1,6 @@
 // Copyright Epic Games, Inc. All Rights Reserved.
 
-using System;
 using System.Collections.Generic;
-using System.Linq;
-using System.Text;
-using System.Threading.Tasks;
 using EpicGames.Core;
 using UnrealBuildBase;
 
@@ -62,11 +58,11 @@ namespace UnrealBuildTool
 		/// <returns>True if the template is compatible with the given compile environment</returns>
 		public bool IsValidFor(CppCompileEnvironment CompileEnvironment)
 		{
-			if(CompileEnvironment.bIsBuildingDLL != BaseCompileEnvironment.bIsBuildingDLL)
+			if (CompileEnvironment.bIsBuildingDLL != BaseCompileEnvironment.bIsBuildingDLL)
 			{
 				return false;
 			}
-			if(CompileEnvironment.bIsBuildingLibrary != BaseCompileEnvironment.bIsBuildingLibrary)
+			if (CompileEnvironment.bIsBuildingLibrary != BaseCompileEnvironment.bIsBuildingLibrary)
 			{
 				return false;
 			}
diff --git a/Engine/Source/Programs/UnrealBuildTool/System/PrecompiledManifest.cs b/Engine/Source/Programs/UnrealBuildTool/System/PrecompiledManifest.cs
--- a/Engine/Source/Programs/UnrealBuildTool/System/PrecompiledManifest.cs
+++ b/Engine/Source/Programs/UnrealBuildTool/System/PrecompiledManifest.cs
@@ -3,9 +3,6 @@ using System;
 using System.Collections.Generic;
 using System.Diagnostics.CodeAnalysis;
 using System.IO;
-using System.Linq;
-using System.Text;
-using System.Threading.Tasks;
 using EpicGames.Core;
 
 namespace UnrealBuildTool
@@ -39,7 +36,7 @@ namespace UnrealBuildTool
 			JsonObject RawObject = JsonObject.Read(Location);
 
 			string[] OutputFiles = RawObject.GetStringArrayField("OutputFiles");
-			foreach(string OutputFile in OutputFiles)
+			foreach (string OutputFile in OutputFiles)
 			{
 				Manifest.OutputFiles.Add(FileReference.Combine(BaseDir, OutputFile));
 			}
@@ -87,7 +84,7 @@ namespace UnrealBuildTool
 				Writer.WriteObjectStart();
 
 				string[] OutputFileStrings = new string[OutputFiles.Count];
-				for(int Idx = 0; Idx < OutputFiles.Count; Idx++)
+				for (int Idx = 0; Idx < OutputFiles.Count; Idx++)
 				{
 					OutputFileStrings[Idx] = OutputFiles[Idx].MakeRelativeTo(BaseDir);
 				}
diff --git a/Engine/Source/Programs/UnrealBuildTool/System/ProgressWriter.cs b/Engine/Source/Programs/UnrealBuildTool/System/ProgressWriter.cs
--- a/Engine/Source/Programs/UnrealBuildTool/System/ProgressWriter.cs
+++ b/Engine/Source/Programs/UnrealBuildTool/System/ProgressWriter.cs
@@ -1,10 +1,6 @@
 // Copyright Epic Games, Inc. All Rights Reserved.
 
 using System;
-using System.Collections.Generic;
-using System.Linq;
-using System.Text;
-using System.Threading.Tasks;
 using EpicGames.Core;
 using Microsoft.Extensions.Logging;
 
@@ -50,7 +46,7 @@ namespace UnrealBuildTool
 		{
 			Message = InMessage;
 			Logger = InLogger;
-			if(bInUpdateStatus)
+			if (bInUpdateStatus)
 			{
 				Status = new LogStatusScope(InMessage);
 			}
@@ -62,7 +58,7 @@ namespace UnrealBuildTool
 		/// </summary>
 		public void Dispose()
 		{
-			if(Status != null)
+			if (Status != null)
 			{
 				Status.Dispose();
 				Status = null;
@@ -85,7 +81,7 @@ namespace UnrealBuildTool
 				{
 					Logger.LogInformation("@progress '{Message}' {ProgressString}", Message, ProgressString);
 				}
-				if(Status != null)
+				if (Status != null)
 				{
 					Status.SetProgress(ProgressString);
 				}
diff --git a/Engine/Source/Programs/UnrealBuildTool/System/ProjectDescriptor.cs b/Engine/Source/Programs/UnrealBuildTool/System/ProjectDescriptor.cs
--- a/Engine/Source/Programs/UnrealBuildTool/System/ProjectDescriptor.cs
+++ b/Engine/Source/Programs/UnrealBuildTool/System/ProjectDescriptor.cs
@@ -2,9 +2,8 @@
 
 using System;
 using System.Collections.Generic;
-using System.Linq;
-using System.Text;
 using System.IO;
+using System.Linq;
 using EpicGames.Core;
 
 namespace UnrealBuildTool
@@ -35,9 +34,9 @@ namespace UnrealBuildTool
 		ProjectPluginUnification = 3,
 
 		/// <summary>
-        /// This needs to be the last line, so we can calculate the value of Latest below
+		/// This needs to be the last line, so we can calculate the value of Latest below
 		/// </summary>
-        LatestPlusOne,
+		LatestPlusOne,
 
 		/// <summary>
 		/// The latest plugin descriptor version
@@ -216,7 +215,7 @@ namespace UnrealBuildTool
 			try
 			{
 				ProjectDescriptor Descriptor = new ProjectDescriptor(RawObject, FileName.Directory, FileName);
-				if(Descriptor.Modules != null)
+				if (Descriptor.Modules != null)
 				{
 					foreach (ModuleDescriptor Module in Descriptor.Modules)
 					{
@@ -269,7 +268,7 @@ namespace UnrealBuildTool
 			Writer.WriteValue("EngineAssociation", EngineAssociation);
 			Writer.WriteValue("Category", Category);
 			Writer.WriteValue("Description", Description);
-			
+
 			if (DisableEnginePluginsByDefault)
 			{
 				Writer.WriteValue("DisableEnginePluginsByDefault", DisableEnginePluginsByDefault);
@@ -288,39 +287,39 @@ namespace UnrealBuildTool
 			PluginReferenceDescriptor.WriteArray(Writer, "Plugins", Plugins);
 
 			// Write the custom module roots
-			if(AdditionalRootDirectories.Count > 0)
+			if (AdditionalRootDirectories.Count > 0)
 			{
 				Writer.WriteStringArrayField("AdditionalRootDirectories", AdditionalRootDirectories.Select(x => x.MakeRelativeTo(BaseDir).Replace(Path.DirectorySeparatorChar, '/')));
 			}
 
 			// Write out the additional plugin directories to scan
-			if(AdditionalPluginDirectories.Count > 0)
+			if (AdditionalPluginDirectories.Count > 0)
 			{
 				Writer.WriteStringArrayField("AdditionalPluginDirectories", AdditionalPluginDirectories.Select(x => x.MakeRelativeTo(BaseDir).Replace(Path.DirectorySeparatorChar, '/')));
 			}
 
 			// Write the target platforms
-			if(TargetPlatforms != null && TargetPlatforms.Length > 0)
+			if (TargetPlatforms != null && TargetPlatforms.Length > 0)
 			{
 				Writer.WriteStringArrayField("TargetPlatforms", TargetPlatforms);
 			}
 
 			// If it's a signed sample, write the name hash
-			if(EpicSampleNameHash != 0)
+			if (EpicSampleNameHash != 0)
 			{
 				Writer.WriteValue("EpicSampleNameHash", (uint)EpicSampleNameHash);
 			}
 
 			// Write the custom build steps
-			if(InitSteps != null)
+			if (InitSteps != null)
 			{
 				InitSteps.Write(Writer, "InitSteps");
 			}
-			if(PreBuildSteps != null)
+			if (PreBuildSteps != null)
 			{
 				PreBuildSteps.Write(Writer, "PreBuildSteps");
 			}
-			if(PostBuildSteps != null)
+			if (PostBuildSteps != null)
 			{
 				PostBuildSteps.Write(Writer, "PostBuildSteps");
 			}
diff --git a/Engine/Source/Programs/UnrealBuildTool/System/RestrictedFolders.cs b/Engine/Source/Programs/UnrealBuildTool/System/RestrictedFolders.cs
--- a/Engine/Source/Programs/UnrealBuildTool/System/RestrictedFolders.cs
+++ b/Engine/Source/Programs/UnrealBuildTool/System/RestrictedFolders.cs
@@ -3,7 +3,6 @@
 using System;
 using System.Collections.Generic;
 using EpicGames.Core;
-using Microsoft.Extensions.Logging;
 
 namespace UnrealBuildTool
 {
@@ -143,7 +142,7 @@ namespace UnrealBuildTool
 		/// <returns></returns>
 		public static string[] GetNames()
 		{
-			if(Names == null)
+			if (Names == null)
 			{
 				AddConfidentialPlatforms();
 				Names = StringRegistry.GetStringNames();
@@ -157,7 +156,7 @@ namespace UnrealBuildTool
 		/// <returns>Array of restricted folder values</returns>
 		public static RestrictedFolder[] GetValues()
 		{
-			if(Values == null)
+			if (Values == null)
 			{
 				AddConfidentialPlatforms();
 				Values = Array.ConvertAll(StringRegistry.GetStringIds(), x => new RestrictedFolder(x));
diff --git a/Engine/Source/Programs/UnrealBuildTool/System/RulesAssembly.cs b/Engine/Source/Programs/UnrealBuildTool/System/RulesAssembly.cs
--- a/Engine/Source/Programs/UnrealBuildTool/System/RulesAssembly.cs
+++ b/Engine/Source/Programs/UnrealBuildTool/System/RulesAssembly.cs
@@ -3,12 +3,9 @@
 using System;
 using System.Collections.Generic;
 using System.Diagnostics.CodeAnalysis;
-using System.IO;
 using System.Linq;
 using System.Reflection;
 using System.Runtime.Serialization;
-using System.Text;
-using System.Threading.Tasks;
 using EpicGames.Core;
 using Microsoft.Extensions.Logging;
 using UnrealBuildBase;
@@ -97,7 +94,7 @@ namespace UnrealBuildTool
 		/// <summary>
 		/// Any preprocessor defines that were set when this assembly was created
 		/// </summary>
-		public IEnumerable<string>? PreprocessorDefines { get;  }
+		public IEnumerable<string>? PreprocessorDefines { get; }
 
 		/// <summary>
 		/// Constructor. Compiles a rules assembly from the given source files.
@@ -196,10 +193,10 @@ namespace UnrealBuildTool
 							Logger.LogWarning("{Location}: warning: {AttributeMessage}", Location, Attribute.Message);
 						}
 					}
-					if(CompiledType.BaseType == typeof(ModuleRules))
+					if (CompiledType.BaseType == typeof(ModuleRules))
 					{
 						ConstructorInfo? Constructor = CompiledType.GetConstructor(new Type[] { typeof(TargetInfo) });
-						if(Constructor != null)
+						if (Constructor != null)
 						{
 							FileReference? Location;
 							if (!TryGetFileNameFromType(CompiledType, out Location))
@@ -249,7 +246,7 @@ namespace UnrealBuildTool
 			BuildVersion? Version;
 			if (BuildVersion.TryRead(BuildVersion.GetDefaultFileName(), out Version))
 			{
-				for(int MinorVersion = 17; MinorVersion <= 30; MinorVersion++)
+				for (int MinorVersion = 17; MinorVersion <= 30; MinorVersion++)
 				{
 					PreprocessorDefines.Add(String.Format("UE_4_{0}_OR_LATER", MinorVersion));
 				}
@@ -284,7 +281,7 @@ namespace UnrealBuildTool
 		/// <param name="bIncludeParentAssembly">Whether to include targets in the parent assembly</param>
 		public void GetAllTargetNames(List<string> TargetNames, bool bIncludeParentAssembly)
 		{
-			if(Parent != null && bIncludeParentAssembly)
+			if (Parent != null && bIncludeParentAssembly)
 			{
 				Parent.GetAllTargetNames(TargetNames, true);
 			}
@@ -447,7 +444,7 @@ namespace UnrealBuildTool
 					// we expect only one platform group to be found in the extensions
 					if (GroupRulesObjectType != null && PlatformRulesObjectType != null)
 					{
-						throw new BuildException("Found multiple platform group overrides ({0} and {1}) for module {2} without a platform specific override. Create a platform override with the class hierarchy as needed.", 
+						throw new BuildException("Found multiple platform group overrides ({0} and {1}) for module {2} without a platform specific override. Create a platform override with the class hierarchy as needed.",
 							GroupRulesObjectType.Name, PlatformRulesObjectType.Name, ModuleName);
 					}
 					// remember the platform group if we found it, but keep searching to verify there isn't more than one
@@ -494,7 +491,7 @@ namespace UnrealBuildTool
 				RulesObject.Context = ModuleFileToContext[RulesObject.File];
 				RulesObject.Plugin = RulesObject.Context.Plugin;
 				RulesObject.bTreatAsEngineModule = bContainsEngineModules;
-				if(DefaultBuildSettings.HasValue)
+				if (DefaultBuildSettings.HasValue)
 				{
 					RulesObject.DefaultBuildSettings = DefaultBuildSettings.Value;
 				}
@@ -527,7 +524,7 @@ namespace UnrealBuildTool
 
 				// Call the constructor
 				ConstructorInfo? Constructor = RulesObjectType.GetConstructor(new Type[] { typeof(ReadOnlyTargetRules) });
-				if(Constructor == null)
+				if (Constructor == null)
 				{
 					throw new BuildException("No valid constructor found for {0}.", ModuleName);
 				}
@@ -558,7 +555,7 @@ namespace UnrealBuildTool
 			}
 			catch (Exception Ex)
 			{
-				Exception MessageEx = (Ex is TargetInvocationException && Ex.InnerException != null)? Ex.InnerException : Ex;
+				Exception MessageEx = (Ex is TargetInvocationException && Ex.InnerException != null) ? Ex.InnerException : Ex;
 				throw new BuildException(Ex, "Unable to instantiate module '{0}': {1}\n(referenced via {2})", ModuleName, MessageEx.ToString(), ReferenceChain);
 			}
 		}
@@ -596,7 +593,7 @@ namespace UnrealBuildTool
 					// we expect only one platform group to be found in the extensions
 					if (GroupRulesObjectType != null && PlatformRulesType != null)
 					{
-						throw new BuildException("Found multiple platform group overrides ({0} and {1}) for rules {2} without a platform specific override. Create a platform override with the class hierarchy as needed.", 
+						throw new BuildException("Found multiple platform group overrides ({0} and {1}) for rules {2} without a platform specific override. Create a platform override with the class hierarchy as needed.",
 							GroupRulesObjectType.Name, PlatformRulesType.Name, TypeName);
 					}
 					// remember the platform group if we found it, but keep searching to verify there isn't more than one
@@ -622,13 +619,13 @@ namespace UnrealBuildTool
 			Rules.SetOverridesForTargetType();
 
 			// Set the final value for the link type in the target rules
-			if(Rules.LinkType == TargetLinkType.Default)
+			if (Rules.LinkType == TargetLinkType.Default)
 			{
 				throw new BuildException("TargetRules.LinkType should be inferred from TargetType");
 			}
 
 			// Set the default value for whether to use the shared build environment
-			if(Rules.BuildEnvironment == TargetBuildEnvironment.Unique && Unreal.IsEngineInstalled())
+			if (Rules.BuildEnvironment == TargetBuildEnvironment.Unique && Unreal.IsEngineInstalled())
 			{
 				throw new BuildException("Targets with a unique build environment cannot be built with an installed engine.");
 			}
@@ -654,13 +651,13 @@ namespace UnrealBuildTool
 
 			// Setup utrace for Shader Compiler Worker
 			if (Rules.bShaderCompilerWorkerTrace)
-            {
+			{
 				Rules.GlobalDefinitions.Add("USE_SHADER_COMPILER_WORKER_TRACE=1");
 			}
-            else
-            {
+			else
+			{
 				Rules.GlobalDefinitions.Add("USE_SHADER_COMPILER_WORKER_TRACE=0");
-            }
+			}
 
 			// Set a macro if we allow using generated inis
 			if (!Rules.bAllowGeneratedIniWhenCooked)
@@ -776,7 +773,7 @@ namespace UnrealBuildTool
 			TargetRules? TargetRules = CreateTargetRulesInstance(TargetTypeName, new TargetInfo(TargetName, Platform, Configuration, Architectures, ProjectFile, Arguments), Logger, IsTestTarget, bSkipValidation);
 
 			if (TargetRules == null)
-            {
+			{
 				throw new BuildException("Expecting to find a type to be declared in a target rules named '{0}'.  This type must derive from the 'TargetRules' type defined by Unreal Build Tool.", TargetTypeName);
 			}
 
@@ -796,19 +793,19 @@ namespace UnrealBuildTool
 		{
 			// Create all the targets in this assembly 
 			List<string> Matches = new List<string>();
-			foreach(KeyValuePair<string, FileReference> TargetPair in TargetNameToTargetFile)
+			foreach (KeyValuePair<string, FileReference> TargetPair in TargetNameToTargetFile)
 			{
 				TargetRules? Rules = CreateTargetRulesInstance(TargetPair.Key + "Target", new TargetInfo(TargetPair.Key, Platform, Configuration, null, ProjectFile, null), Logger);
-				if(Rules != null && Rules.Type == Type)
+				if (Rules != null && Rules.Type == Type)
 				{
 					Matches.Add(TargetPair.Key);
 				}
 			}
 
 			// If we got a result, return it. If there were multiple results, fail.
-			if(Matches.Count == 0)
+			if (Matches.Count == 0)
 			{
-				if(Parent == null)
+				if (Parent == null)
 				{
 					throw new BuildException("Unable to find target of type '{0}' for project '{1}'", Type, ProjectFile);
 				}
@@ -837,7 +834,7 @@ namespace UnrealBuildTool
 						return DefaultTargetName;
 					}
 				}
-				
+
 				throw new BuildException("Found multiple targets with TargetType={0}: {1}.\nSpecify a default with a {2} entry in [/Script/BuildSettings.BuildSettings] section of your DefaultEngine.ini", Type, String.Join(", ", Matches), KeyName);
 			}
 		}
diff --git a/Engine/Source/Programs/UnrealBuildTool/System/RulesCompiler.cs b/Engine/Source/Programs/UnrealBuildTool/System/RulesCompiler.cs
--- a/Engine/Source/Programs/UnrealBuildTool/System/RulesCompiler.cs
+++ b/Engine/Source/Programs/UnrealBuildTool/System/RulesCompiler.cs
@@ -1,15 +1,14 @@
 // Copyright Epic Games, Inc. All Rights Reserved.
 
 using System;
+using System.Collections.Concurrent;
 using System.Collections.Generic;
 using System.IO;
 using System.Linq;
-using System.Reflection;
 using EpicGames.Core;
+using Microsoft.Extensions.Logging;
 using OpenTracing.Util;
 using UnrealBuildBase;
-using Microsoft.Extensions.Logging;
-using System.Collections.Concurrent;
 
 namespace UnrealBuildTool
 {
@@ -125,7 +124,7 @@ namespace UnrealBuildTool
 		private static RulesAssembly CreateEngineRulesAssemblyInternal(List<DirectoryReference> RootDirectories, string AssemblyPrefix, IReadOnlyList<PluginInfo> Plugins, bool bReadOnly, bool bSkipCompile, bool bForceCompile, RulesAssembly? Parent, ILogger Logger)
 		{
 			// Scope hierarchy
-			RulesScope Scope= new RulesScope("Engine", null);
+			RulesScope Scope = new RulesScope("Engine", null);
 			RulesScope PluginsScope = new RulesScope("Engine Plugins", Scope);
 			RulesScope ProgramsScope = new RulesScope("Engine Programs", PluginsScope);
 
@@ -467,5 +466,5 @@ namespace UnrealBuildTool
 			}
 			return null;
 		}
-    }
+	}
 }
