UnrealBuildTool: Automated code cleanup

diff --git a/Engine/Source/Programs/UnrealBuildTool/System/RulesDocumentation.cs b/Engine/Source/Programs/UnrealBuildTool/System/RulesDocumentation.cs
--- a/Engine/Source/Programs/UnrealBuildTool/System/RulesDocumentation.cs
+++ b/Engine/Source/Programs/UnrealBuildTool/System/RulesDocumentation.cs
@@ -6,12 +6,9 @@ using System.Diagnostics.CodeAnalysis;
 using System.IO;
 using System.Linq;
 using System.Reflection;
-using System.Text;
 using System.Text.RegularExpressions;
-using System.Threading.Tasks;
 using System.Xml;
 using EpicGames.Core;
-using UnrealBuildBase;
 using Microsoft.Extensions.Logging;
 
 namespace UnrealBuildTool
@@ -36,7 +33,7 @@ namespace UnrealBuildTool
 		{
 			// Get the path to the XML documentation
 			FileReference InputDocumentationFile = new FileReference(Assembly.GetExecutingAssembly().Location).ChangeExtension(".xml");
-			if(!FileReference.Exists(InputDocumentationFile))
+			if (!FileReference.Exists(InputDocumentationFile))
 			{
 				throw new BuildException("Generated assembly documentation not found at {0}.", InputDocumentationFile);
 			}
@@ -52,7 +49,7 @@ namespace UnrealBuildTool
 			// First read the fields
 			foreach (FieldInfo Field in RulesType.GetFields(BindingFlags.Instance | BindingFlags.SetProperty | BindingFlags.Public))
 			{
-				if(!Field.FieldType.IsClass || !Field.FieldType.Name.EndsWith("TargetRules"))
+				if (!Field.FieldType.IsClass || !Field.FieldType.Name.EndsWith("TargetRules"))
 				{
 					List<string>? Lines;
 					if (TryGetXmlComment(InputDocumentation, Field, Logger, out Lines))
@@ -354,19 +351,19 @@ namespace UnrealBuildTool
 			{
 				Writer.WriteLine("      <dt>{0} ({1})</dt>", Setting.Name, GetPrettyTypeName(Setting.Type));
 
-				if(Lines.Count == 1)
+				if (Lines.Count == 1)
 				{
 					Writer.WriteLine("      <dd>{0}</dd>", Lines[0]);
 				}
 				else
 				{
 					Writer.WriteLine("      <dd>");
-					for(int Idx = 0; Idx < Lines.Count; Idx++)
+					for (int Idx = 0; Idx < Lines.Count; Idx++)
 					{
-						if(Lines[Idx].StartsWith("*") || Lines[Idx].StartsWith("-"))
+						if (Lines[Idx].StartsWith("*") || Lines[Idx].StartsWith("-"))
 						{
 							Writer.WriteLine("        <ul>");
-							for(; Idx < Lines.Count && (Lines[Idx].StartsWith("*") || Lines[Idx].StartsWith("-")); Idx++)
+							for (; Idx < Lines.Count && (Lines[Idx].StartsWith("*") || Lines[Idx].StartsWith("-")); Idx++)
 							{
 								Writer.WriteLine("          <li>{0}</li>", Lines[Idx].Substring(1).TrimStart());
 							}
@@ -384,7 +381,7 @@ namespace UnrealBuildTool
 
 		static string GetPrettyTypeName(Type FieldType)
 		{
-			if(FieldType.IsGenericType)
+			if (FieldType.IsGenericType)
 			{
 				return String.Format("{0}&lt;{1}&gt;", FieldType.Name.Substring(0, FieldType.Name.IndexOf('`')), String.Join(", ", FieldType.GenericTypeArguments.Select(x => GetPrettyTypeName(x))));
 			}
diff --git a/Engine/Source/Programs/UnrealBuildTool/System/RulesScope.cs b/Engine/Source/Programs/UnrealBuildTool/System/RulesScope.cs
--- a/Engine/Source/Programs/UnrealBuildTool/System/RulesScope.cs
+++ b/Engine/Source/Programs/UnrealBuildTool/System/RulesScope.cs
@@ -1,10 +1,6 @@
 // Copyright Epic Games, Inc. All Rights Reserved.
 
 using System;
-using System.Collections.Generic;
-using System.Linq;
-using System.Text;
-using System.Threading.Tasks;
 
 namespace UnrealBuildTool
 {
@@ -41,9 +37,9 @@ namespace UnrealBuildTool
 		/// <returns>True if this scope contains the other scope</returns>
 		public bool Contains(RulesScope Other)
 		{
-			for(RulesScope? Scope = this; Scope != null; Scope = Scope.Parent)
+			for (RulesScope? Scope = this; Scope != null; Scope = Scope.Parent)
 			{
-				if(Scope == Other)
+				if (Scope == Other)
 				{
 					return true;
 				}
@@ -57,7 +53,7 @@ namespace UnrealBuildTool
 		/// <returns>String representing the hierarchy of scopes</returns>
 		public string FormatHierarchy()
 		{
-			if(Parent == null)
+			if (Parent == null)
 			{
 				return Name;
 			}
diff --git a/Engine/Source/Programs/UnrealBuildTool/System/ScopedTimer.cs b/Engine/Source/Programs/UnrealBuildTool/System/ScopedTimer.cs
--- a/Engine/Source/Programs/UnrealBuildTool/System/ScopedTimer.cs
+++ b/Engine/Source/Programs/UnrealBuildTool/System/ScopedTimer.cs
@@ -1,11 +1,7 @@
 // Copyright Epic Games, Inc. All Rights Reserved.
 
 using System;
-using System.Collections.Generic;
-using System.Linq;
 using System.Text;
-using System.Threading.Tasks;
-using EpicGames.Core;
 using Microsoft.Extensions.Logging;
 
 namespace UnrealBuildTool
diff --git a/Engine/Source/Programs/UnrealBuildTool/System/SingleInstanceMutex.cs b/Engine/Source/Programs/UnrealBuildTool/System/SingleInstanceMutex.cs
--- a/Engine/Source/Programs/UnrealBuildTool/System/SingleInstanceMutex.cs
+++ b/Engine/Source/Programs/UnrealBuildTool/System/SingleInstanceMutex.cs
@@ -1,14 +1,8 @@
 // Copyright Epic Games, Inc. All Rights Reserved.
 
 using System;
-using System.Collections.Generic;
-using System.Linq;
-using System.Reflection;
-using System.Text;
 using System.Threading;
-using System.Threading.Tasks;
 using EpicGames.Core;
-using UnrealBuildBase;
 
 namespace UnrealBuildTool
 {
diff --git a/Engine/Source/Programs/UnrealBuildTool/System/SourceFileSearch.cs b/Engine/Source/Programs/UnrealBuildTool/System/SourceFileSearch.cs
--- a/Engine/Source/Programs/UnrealBuildTool/System/SourceFileSearch.cs
+++ b/Engine/Source/Programs/UnrealBuildTool/System/SourceFileSearch.cs
@@ -1,10 +1,7 @@
 // Copyright Epic Games, Inc. All Rights Reserved.
 
-using System;
 using System.Collections.Generic;
-using System.Text;
 using System.IO;
-using System.Diagnostics;
 using System.Linq;
 using EpicGames.Core;
 using UnrealBuildBase;
@@ -14,7 +11,7 @@ namespace UnrealBuildTool
 	static class SourceFileSearch
 	{
 		// Certain file types should never be added to project files. These extensions must all be lowercase.
-		static readonly string[] DefaultExcludedFileSuffixes = new string[] 
+		static readonly string[] DefaultExcludedFileSuffixes = new string[]
 		{
 			".vcxproj",				// Visual Studio project files
 			".vcxproj.filters",		// Visual Studio filter file
@@ -33,7 +30,7 @@ namespace UnrealBuildTool
 		};
 
 		// Default directory names to exclude. Must be lowercase.
-		static readonly string[] DefaultExcludedDirectorySuffixes = new string[] 
+		static readonly string[] DefaultExcludedDirectorySuffixes = new string[]
 		{
 			Path.DirectorySeparatorChar + "intermediate",
 			Path.DirectorySeparatorChar + "source" + Path.DirectorySeparatorChar + "thirdparty",
@@ -83,7 +80,7 @@ namespace UnrealBuildTool
 
 		static void FindFilesInternal(DirectoryReference Directory, string[] ExcludedDirectorySuffixes, List<FileReference> FoundFiles, HashSet<DirectoryReference>? SearchedDirectories)
 		{
-			if(SearchedDirectories != null)
+			if (SearchedDirectories != null)
 			{
 				SearchedDirectories.Add(Directory);
 			}
@@ -117,7 +114,7 @@ namespace UnrealBuildTool
 			{
 				return false;
 			}
-			
+
 			foreach (string ExcludedSuffix in ExcludedSuffixes)
 			{
 				if (Reference.FullName.EndsWith(ExcludedSuffix, FileSystemReference.Comparison))
diff --git a/Engine/Source/Programs/UnrealBuildTool/System/SourceFileWorkingSet.cs b/Engine/Source/Programs/UnrealBuildTool/System/SourceFileWorkingSet.cs
--- a/Engine/Source/Programs/UnrealBuildTool/System/SourceFileWorkingSet.cs
+++ b/Engine/Source/Programs/UnrealBuildTool/System/SourceFileWorkingSet.cs
@@ -6,11 +6,9 @@ using System.Diagnostics;
 using System.Diagnostics.CodeAnalysis;
 using System.IO;
 using System.Linq;
-using System.Text;
-using System.Threading.Tasks;
 using EpicGames.Core;
-using UnrealBuildBase;
 using Microsoft.Extensions.Logging;
+using UnrealBuildBase;
 
 namespace UnrealBuildTool
 {
@@ -168,11 +166,11 @@ namespace UnrealBuildTool
 		{
 			if (BackgroundProcess != null)
 			{
-				if(!BackgroundProcess.WaitForExit(500))
+				if (!BackgroundProcess.WaitForExit(500))
 				{
 					Logger.LogInformation("Waiting for 'git status' command to complete");
 				}
-				if(!BackgroundProcess.WaitForExit(15000))
+				if (!BackgroundProcess.WaitForExit(15000))
 				{
 					Logger.LogInformation("Terminating git child process due to timeout");
 					try
@@ -196,7 +194,7 @@ namespace UnrealBuildTool
 		{
 			TerminateBackgroundProcess();
 
-			if(Inner != null)
+			if (Inner != null)
 			{
 				Inner.Dispose();
 			}
@@ -210,11 +208,11 @@ namespace UnrealBuildTool
 		public bool Contains(FileItem File)
 		{
 			WaitForBackgroundProcess();
-			if(Files.Contains(File.Location) || Directories.Any(x => File.Location.IsUnderDirectory(x)))
+			if (Files.Contains(File.Location) || Directories.Any(x => File.Location.IsUnderDirectory(x)))
 			{
 				return true;
 			}
-			if(Inner != null && Inner.Contains(File))
+			if (Inner != null && Inner.Contains(File))
 			{
 				return true;
 			}
@@ -334,7 +332,7 @@ namespace UnrealBuildTool
 				{
 					WorkingSet = new GitSourceFileWorkingSet(GitPath, DirectoryReference.Combine(RootDir, RepositoryPath), null, Logger);
 				}
-				else if(!TryCreateGitWorkingSet(RootDir, ProjectDirs, Logger, out WorkingSet))
+				else if (!TryCreateGitWorkingSet(RootDir, ProjectDirs, Logger, out WorkingSet))
 				{
 					WorkingSet = new GitSourceFileWorkingSet(GitPath, RootDir, null, Logger);
 				}
@@ -347,7 +345,7 @@ namespace UnrealBuildTool
 			else
 			{
 				GitSourceFileWorkingSet? WorkingSet;
-				if(TryCreateGitWorkingSet(RootDir, ProjectDirs, Logger, out WorkingSet))
+				if (TryCreateGitWorkingSet(RootDir, ProjectDirs, Logger, out WorkingSet))
 				{
 					return WorkingSet;
 				}
@@ -360,7 +358,7 @@ namespace UnrealBuildTool
 
 		static bool TryCreateGitWorkingSet(DirectoryReference RootDir, IEnumerable<DirectoryReference> ProjectDirs, ILogger Logger, [NotNullWhen(true)] out GitSourceFileWorkingSet? OutWorkingSet)
 		{
-			GitSourceFileWorkingSet? WorkingSet  = null;
+			GitSourceFileWorkingSet? WorkingSet = null;
 
 			// Create the working set for the engine directory
 			if (DirectoryReference.Exists(DirectoryReference.Combine(RootDir, ".git")) || FileReference.Exists(FileReference.Combine(RootDir, ".git")))
@@ -369,9 +367,9 @@ namespace UnrealBuildTool
 			}
 
 			// Try to create a working set for the project directory
-			foreach(DirectoryReference ProjectDir in ProjectDirs)
+			foreach (DirectoryReference ProjectDir in ProjectDirs)
 			{
-				if(WorkingSet == null || !ProjectDir.IsUnderDirectory(RootDir))
+				if (WorkingSet == null || !ProjectDir.IsUnderDirectory(RootDir))
 				{
 					if (DirectoryReference.Exists(DirectoryReference.Combine(ProjectDir, ".git")) || FileReference.Exists(FileReference.Combine(ProjectDir, ".git")))
 					{
diff --git a/Engine/Source/Programs/UnrealBuildTool/System/Subnet.cs b/Engine/Source/Programs/UnrealBuildTool/System/Subnet.cs
--- a/Engine/Source/Programs/UnrealBuildTool/System/Subnet.cs
+++ b/Engine/Source/Programs/UnrealBuildTool/System/Subnet.cs
@@ -1,12 +1,7 @@
 // Copyright Epic Games, Inc. All Rights Reserved.
 
 using System;
-using System.Collections.Generic;
-using System.Diagnostics;
-using System.Linq;
 using System.Net;
-using System.Text;
-using System.Threading.Tasks;
 
 namespace UnrealBuildTool
 {
diff --git a/Engine/Source/Programs/UnrealBuildTool/System/SupportedPlatformsAttribute.cs b/Engine/Source/Programs/UnrealBuildTool/System/SupportedPlatformsAttribute.cs
--- a/Engine/Source/Programs/UnrealBuildTool/System/SupportedPlatformsAttribute.cs
+++ b/Engine/Source/Programs/UnrealBuildTool/System/SupportedPlatformsAttribute.cs
@@ -1,11 +1,6 @@
 // Copyright Epic Games, Inc. All Rights Reserved.
 
 using System;
-using System.Collections.Generic;
-using System.Linq;
-using System.Text;
-using System.Threading.Tasks;
-using UnrealBuildBase;
 
 namespace UnrealBuildTool
 {
diff --git a/Engine/Source/Programs/UnrealBuildTool/System/SupportedTargetTypesAttribute.cs b/Engine/Source/Programs/UnrealBuildTool/System/SupportedTargetTypesAttribute.cs
--- a/Engine/Source/Programs/UnrealBuildTool/System/SupportedTargetTypesAttribute.cs
+++ b/Engine/Source/Programs/UnrealBuildTool/System/SupportedTargetTypesAttribute.cs
@@ -1,11 +1,6 @@
 // Copyright Epic Games, Inc. All Rights Reserved.
 
 using System;
-using System.Collections.Generic;
-using System.Linq;
-using System.Text;
-using System.Threading.Tasks;
-using UnrealBuildBase;
 
 namespace UnrealBuildTool
 {
diff --git a/Engine/Source/Programs/UnrealBuildTool/System/TargetMakefile.cs b/Engine/Source/Programs/UnrealBuildTool/System/TargetMakefile.cs
--- a/Engine/Source/Programs/UnrealBuildTool/System/TargetMakefile.cs
+++ b/Engine/Source/Programs/UnrealBuildTool/System/TargetMakefile.cs
@@ -1,6 +1,5 @@
 // Copyright Epic Games, Inc. All Rights Reserved.
 using System;
-using System.Collections;
 using System.Collections.Concurrent;
 using System.Collections.Generic;
 using System.IO;
@@ -18,7 +17,7 @@ namespace UnrealBuildTool
 	class TargetMakefileSourceFileInfo
 	{
 		public FileItem? SourceFileItem;
-		
+
 		/// <summary>
 		/// Hash of all the gen.cpp files being inlined in this source file
 		/// </summary>
@@ -405,9 +404,9 @@ namespace UnrealBuildTool
 		/// <param name="UbtPlugins">Collection of UBT plugins</param>
 		/// <param name="EnabledUbtPlugins">Collection of UBT plugins</param>
 		/// <param name="EnabledUhtPlugins">Collection of UBT plugins for UHT</param>
-		public TargetMakefile(string ExternalMetadata, FileReference ExecutableFile, FileReference ReceiptFile, 
+		public TargetMakefile(string ExternalMetadata, FileReference ExecutableFile, FileReference ReceiptFile,
 			DirectoryReference ProjectIntermediateDirectory, DirectoryReference ProjectIntermediateDirectoryNoArch,
-			TargetType TargetType, ConfigValueTracker ConfigValueTracker, bool bDeployAfterCompile, 
+			TargetType TargetType, ConfigValueTracker ConfigValueTracker, bool bDeployAfterCompile,
 			FileReference[]? UbtPlugins, FileReference[]? EnabledUbtPlugins, FileReference[]? EnabledUhtPlugins)
 		{
 			this.CreateTimeUtc = UnrealBuildTool.StartTimeUtc;
@@ -475,7 +474,7 @@ namespace UnrealBuildTool
 			ModuleNameToOutputItems = Reader.ReadDictionary(() => Reader.ReadString()!, () => Reader.ReadArray(() => Reader.ReadFileItem()), StringComparer.OrdinalIgnoreCase)!;
 			HotReloadModuleNames = Reader.ReadHashSet(() => Reader.ReadString(), StringComparer.OrdinalIgnoreCase)!;
 			SourceDirectories = Reader.ReadList(() => Reader.ReadDirectoryItem())!;
-			DirectoryToSourceFiles = Reader.ReadDictionary(() => Reader.ReadDirectoryItem()!, () => Reader.ReadArray(() =>  new TargetMakefileSourceFileInfo(Reader)))!;
+			DirectoryToSourceFiles = Reader.ReadDictionary(() => Reader.ReadDirectoryItem()!, () => Reader.ReadArray(() => new TargetMakefileSourceFileInfo(Reader)))!;
 			HeaderFiles = Reader.ReadHashSet(() => Reader.ReadFileItem())!;
 			WorkingSet = Reader.ReadHashSet(() => Reader.ReadFileItem())!;
 			CandidatesForWorkingSet = Reader.ReadHashSet(() => Reader.ReadFileItem())!;
@@ -542,7 +541,7 @@ namespace UnrealBuildTool
 		public void Save(FileReference Location)
 		{
 			DirectoryReference.CreateDirectory(Location.Directory);
-			using(BinaryArchiveWriter Writer = new BinaryArchiveWriter(Location))
+			using (BinaryArchiveWriter Writer = new BinaryArchiveWriter(Location))
 			{
 				Writer.WriteInt(CurrentVersion);
 #if __VPROJECT_AVAILABLE__
@@ -645,10 +644,10 @@ namespace UnrealBuildTool
 			{
 				try
 				{
-					using(BinaryArchiveReader Reader = new BinaryArchiveReader(MakefilePath))
+					using (BinaryArchiveReader Reader = new BinaryArchiveReader(MakefilePath))
 					{
 						int Version = Reader.ReadInt();
-						if(Version != CurrentVersion)
+						if (Version != CurrentVersion)
 						{
 							ReasonNotLoaded = "makefile version does not match";
 							return null;
@@ -678,7 +677,7 @@ namespace UnrealBuildTool
 			using (GlobalTracer.Instance.BuildSpan("Checking makefile validity").StartActive())
 			{
 				// Check if the arguments are different
-				if(!Enumerable.SequenceEqual(Makefile.AdditionalArguments!, Arguments))
+				if (!Enumerable.SequenceEqual(Makefile.AdditionalArguments!, Arguments))
 				{
 					Logger.LogDebug("Old command line arguments:\n", string.Join(' ', Makefile.AdditionalArguments!));
 					Logger.LogDebug("New command line arguments:\n", string.Join(' ', Arguments));
@@ -705,7 +704,7 @@ namespace UnrealBuildTool
 				}
 
 				// Check if any config settings have changed. Ini files contain build settings too.
-				if(!Makefile.ConfigValueTracker.IsValid())
+				if (!Makefile.ConfigValueTracker.IsValid())
 				{
 					ReasonNotLoaded = "config setting changed";
 					return null;
@@ -713,7 +712,7 @@ namespace UnrealBuildTool
 
 				// Get the current build metadata from the platform
 				string CurrentExternalMetadata = UEBuildPlatform.GetBuildPlatform(Platform).GetExternalBuildMetadata(ProjectFile);
-				if(String.Compare(CurrentExternalMetadata, Makefile.ExternalMetadata, StringComparison.Ordinal) != 0)
+				if (String.Compare(CurrentExternalMetadata, Makefile.ExternalMetadata, StringComparison.Ordinal) != 0)
 				{
 					Logger.LogDebug("Old metadata:\n", Makefile.ExternalMetadata);
 					Logger.LogDebug("New metadata:\n", CurrentExternalMetadata);
@@ -781,28 +780,28 @@ namespace UnrealBuildTool
 				// Check if any source files have been added or removed
 				foreach ((DirectoryItem InputDirectory, TargetMakefileSourceFileInfo[] SourceFileInfos) in Makefile.DirectoryToSourceFiles)
 				{
-					if(!InputDirectory.Exists || InputDirectory.LastWriteTimeUtc > Makefile.CreateTimeUtc)
+					if (!InputDirectory.Exists || InputDirectory.LastWriteTimeUtc > Makefile.CreateTimeUtc)
 					{
 						FileItem[] SourceFiles = UEBuildModuleCPP.GetSourceFiles(InputDirectory, Logger);
-						if(SourceFiles.Length < SourceFileInfos.Length)
+						if (SourceFiles.Length < SourceFileInfos.Length)
 						{
 							ReasonNotLoaded = "source file removed";
 							return false;
 						}
-						else if(SourceFiles.Length > SourceFileInfos.Length)
+						else if (SourceFiles.Length > SourceFileInfos.Length)
 						{
 							ReasonNotLoaded = "source file added";
 							return false;
 						}
-						else if(SourceFiles.Intersect(SourceFileInfos.Select((sfi) => sfi.SourceFileItem)).Count() != SourceFiles.Length)
+						else if (SourceFiles.Intersect(SourceFileInfos.Select((sfi) => sfi.SourceFileItem)).Count() != SourceFiles.Length)
 						{
 							ReasonNotLoaded = "source file modified";
 							return false;
 						}
 
-						foreach(DirectoryItem Directory in InputDirectory.EnumerateDirectories())
+						foreach (DirectoryItem Directory in InputDirectory.EnumerateDirectories())
 						{
-							if(!Makefile.DirectoryToSourceFiles.ContainsKey(Directory) && ContainsSourceFiles(Directory, ExcludedFolderNames, Logger))
+							if (!Makefile.DirectoryToSourceFiles.ContainsKey(Directory) && ContainsSourceFiles(Directory, ExcludedFolderNames, Logger))
 							{
 								ReasonNotLoaded = "directory added";
 								return false;
@@ -823,7 +822,7 @@ namespace UnrealBuildTool
 				}
 
 				// Check if any external dependencies have changed. These comparisons are done against the makefile creation time.
-				foreach(FileItem ExternalDependency in Makefile.ExternalDependencies)
+				foreach (FileItem ExternalDependency in Makefile.ExternalDependencies)
 				{
 					if (!ExternalDependency.Exists)
 					{
@@ -831,7 +830,7 @@ namespace UnrealBuildTool
 						ReasonNotLoaded = string.Format("{0} deleted", ExternalDependency.Location.GetFileName());
 						return false;
 					}
-					if(ExternalDependency.LastWriteTimeUtc > Makefile.CreateTimeUtc)
+					if (ExternalDependency.LastWriteTimeUtc > Makefile.CreateTimeUtc)
 					{
 						Logger.LogDebug("{File} has been modified since makefile was built.", ExternalDependency.Location);
 						ReasonNotLoaded = string.Format("{0} modified", ExternalDependency.Location.GetFileName());
@@ -860,7 +859,7 @@ namespace UnrealBuildTool
 				foreach (FileReference PluginFile in PluginsBase.EnumeratePlugins(ProjectFile))
 				{
 					FileItem PluginFileItem = FileItem.GetItemByFileReference(PluginFile);
-					if(!Makefile.PluginFiles.Contains(PluginFileItem))
+					if (!Makefile.PluginFiles.Contains(PluginFileItem))
 					{
 						Logger.LogDebug("{File} has been added", PluginFile.GetFileName());
 						ReasonNotLoaded = string.Format("{0} has been added", PluginFile.GetFileName());
@@ -872,7 +871,7 @@ namespace UnrealBuildTool
 				ConcurrentBag<FileItem> NewFilesWithMarkupBag = new ConcurrentBag<FileItem>();
 				using (ThreadPoolWorkQueue Queue = new ThreadPoolWorkQueue())
 				{
-					foreach(DirectoryItem SourceDirectory in Makefile.SourceDirectories)
+					foreach (DirectoryItem SourceDirectory in Makefile.SourceDirectories)
 					{
 						Queue.Enqueue(() => FindFilesWithMarkup(SourceDirectory, MetadataCache, ExcludedFolderNames, NewFilesWithMarkupBag, Queue));
 					}
@@ -930,19 +929,19 @@ namespace UnrealBuildTool
 		static bool ContainsSourceFiles(DirectoryItem Directory, IReadOnlySet<string> ExcludedFolderNames, ILogger Logger)
 		{
 			// Check this directory isn't ignored
-			if(!ExcludedFolderNames.Contains(Directory.Name))
+			if (!ExcludedFolderNames.Contains(Directory.Name))
 			{
 				// Check for any source files in this actual directory
 				FileItem[] SourceFiles = UEBuildModuleCPP.GetSourceFiles(Directory, Logger);
-				if(SourceFiles.Length > 0)
+				if (SourceFiles.Length > 0)
 				{
 					return true;
 				}
 
 				// Check for any source files in a subdirectory
-				foreach(DirectoryItem SubDirectory in Directory.EnumerateDirectories())
+				foreach (DirectoryItem SubDirectory in Directory.EnumerateDirectories())
 				{
-					if(ContainsSourceFiles(SubDirectory, ExcludedFolderNames, Logger))
+					if (ContainsSourceFiles(SubDirectory, ExcludedFolderNames, Logger))
 					{
 						return true;
 					}
@@ -975,7 +974,7 @@ namespace UnrealBuildTool
 			// Search through all the subfolders
 			foreach (DirectoryItem SubDirectory in Directory.EnumerateDirectories())
 			{
-				if(!ExcludedFolderNames.Contains(SubDirectory.Name))
+				if (!ExcludedFolderNames.Contains(SubDirectory.Name))
 				{
 					Queue.Enqueue(() => FindFilesWithMarkup(SubDirectory, MetadataCache, ExcludedFolderNames, FilesWithMarkup, Queue));
 				}
diff --git a/Engine/Source/Programs/UnrealBuildTool/System/TargetReceipt.cs b/Engine/Source/Programs/UnrealBuildTool/System/TargetReceipt.cs
--- a/Engine/Source/Programs/UnrealBuildTool/System/TargetReceipt.cs
+++ b/Engine/Source/Programs/UnrealBuildTool/System/TargetReceipt.cs
@@ -2,14 +2,11 @@
 
 using System;
 using System.Collections.Generic;
+using System.Diagnostics.CodeAnalysis;
 using System.IO;
-using System.Diagnostics;
 using System.Linq;
 using System.Xml.Serialization;
-using System.ComponentModel;
-using System.Runtime.Serialization;
 using EpicGames.Core;
-using System.Diagnostics.CodeAnalysis;
 using UnrealBuildBase;
 
 namespace UnrealBuildTool
@@ -430,13 +427,13 @@ namespace UnrealBuildTool
 		static FileReference ExpandPathVariables(string Path, DirectoryReference EngineDir, DirectoryReference? ProjectDir)
 		{
 			const string EnginePrefix = "$(EngineDir)";
-			if(Path.StartsWith(EnginePrefix, StringComparison.InvariantCultureIgnoreCase))
+			if (Path.StartsWith(EnginePrefix, StringComparison.InvariantCultureIgnoreCase))
 			{
 				return new FileReference(EngineDir.FullName + Path.Substring(EnginePrefix.Length));
 			}
 
 			const string ProjectPrefix = "$(ProjectDir)";
-			if(ProjectDir != null && Path.StartsWith(ProjectPrefix, StringComparison.InvariantCultureIgnoreCase))
+			if (ProjectDir != null && Path.StartsWith(ProjectPrefix, StringComparison.InvariantCultureIgnoreCase))
 			{
 				return new FileReference(ProjectDir.FullName + Path.Substring(ProjectPrefix.Length));
 			}
@@ -462,7 +459,7 @@ namespace UnrealBuildTool
 			{
 				ArchitectureSuffix = BuildArchitectures.ToString();
 			}
-		
+
 			// Build the output filename
 			if (String.IsNullOrEmpty(ArchitectureSuffix) && Configuration == UnrealTargetConfiguration.Development)
 			{
@@ -544,7 +541,7 @@ namespace UnrealBuildTool
 			FileReference? ProjectFile;
 
 			string? RelativeProjectFile;
-			if(RawObject.TryGetStringField("Project", out RelativeProjectFile))
+			if (RawObject.TryGetStringField("Project", out RelativeProjectFile))
 			{
 				ProjectFile = FileReference.Combine(Location.Directory, RelativeProjectFile);
 			}
@@ -577,12 +574,12 @@ namespace UnrealBuildTool
 
 			// Read the launch executable
 			string? Launch;
-			if(RawObject.TryGetStringField("Launch", out Launch))
+			if (RawObject.TryGetStringField("Launch", out Launch))
 			{
 				Receipt.Launch = ExpandPathVariables(Launch, EngineDir, ProjectDir);
 			}
 			string? LaunchCmd;
-			if(RawObject.TryGetStringField("LaunchCmd", out LaunchCmd))
+			if (RawObject.TryGetStringField("LaunchCmd", out LaunchCmd))
 			{
 				Receipt.LaunchCmd = ExpandPathVariables(LaunchCmd, EngineDir, ProjectDir);
 			}
@@ -619,15 +616,15 @@ namespace UnrealBuildTool
 						FileReference File = ExpandPathVariables(Path, EngineDir, ProjectDir);
 
 						StagedFileType Type;
-						if(!RuntimeDependencyObject.TryGetEnumField("Type", out Type))
+						if (!RuntimeDependencyObject.TryGetEnumField("Type", out Type))
 						{
 							// Previous format included an optional IgnoreIfMissing flag, which was only used for debug files. We can explicitly reference them as DebugNonUFS files now.
 							bool bIgnoreIfMissing;
-							if(RuntimeDependencyObject.TryGetBoolField("IgnoreIfMissing", out bIgnoreIfMissing))
+							if (RuntimeDependencyObject.TryGetBoolField("IgnoreIfMissing", out bIgnoreIfMissing))
 							{
 								bIgnoreIfMissing = false;
 							}
-							Type = bIgnoreIfMissing? StagedFileType.DebugNonUFS : StagedFileType.NonUFS;
+							Type = bIgnoreIfMissing ? StagedFileType.DebugNonUFS : StagedFileType.NonUFS;
 						}
 
 						Receipt.RuntimeDependencies.Add(File, Type);
@@ -655,15 +652,15 @@ namespace UnrealBuildTool
 
 			// Read the additional properties
 			JsonObject[]? AdditionalPropertyObjects;
-			if(RawObject.TryGetObjectArrayField("AdditionalProperties", out AdditionalPropertyObjects))
+			if (RawObject.TryGetObjectArrayField("AdditionalProperties", out AdditionalPropertyObjects))
 			{
-				foreach(JsonObject AdditionalPropertyObject in AdditionalPropertyObjects)
+				foreach (JsonObject AdditionalPropertyObject in AdditionalPropertyObjects)
 				{
 					string? Name;
-					if(AdditionalPropertyObject.TryGetStringField("Name", out Name))
+					if (AdditionalPropertyObject.TryGetStringField("Name", out Name))
 					{
 						string? Value;
-						if(AdditionalPropertyObject.TryGetStringField("Value", out Value))
+						if (AdditionalPropertyObject.TryGetStringField("Value", out Value))
 						{
 							Receipt.AdditionalProperties.Add(new ReceiptProperty(Name, Value));
 						}
@@ -738,17 +735,17 @@ namespace UnrealBuildTool
 				Writer.WriteValue("IsTestTarget", IsTestTarget);
 				Writer.WriteValue("Architecture", Architectures.ToString());
 
-				if(ProjectFile != null)
+				if (ProjectFile != null)
 				{
 					Writer.WriteValue("Project", ProjectFile.MakeRelativeTo(Location.Directory).Replace(Path.DirectorySeparatorChar, '/'));
 				}
 
-				if(Launch != null)
+				if (Launch != null)
 				{
 					Writer.WriteValue("Launch", InsertPathVariables(Launch, EngineDir, ProjectDir));
 				}
 
-				if(LaunchCmd != null)
+				if (LaunchCmd != null)
 				{
 					Writer.WriteValue("LaunchCmd", InsertPathVariables(LaunchCmd, EngineDir, ProjectDir));
 				}
diff --git a/Engine/Source/Programs/UnrealBuildTool/System/Timeline.cs b/Engine/Source/Programs/UnrealBuildTool/System/Timeline.cs
--- a/Engine/Source/Programs/UnrealBuildTool/System/Timeline.cs
+++ b/Engine/Source/Programs/UnrealBuildTool/System/Timeline.cs
@@ -4,10 +4,7 @@ using System;
 using System.Collections.Generic;
 using System.Diagnostics;
 using System.Linq;
-using System.Reflection;
-using System.Runtime.CompilerServices;
 using System.Text;
-using System.Threading.Tasks;
 using EpicGames.Core;
 using Microsoft.Extensions.Logging;
 
@@ -71,7 +68,7 @@ namespace UnrealBuildTool
 			/// </summary>
 			public void Finish()
 			{
-				if(!FinishTime.HasValue)
+				if (!FinishTime.HasValue)
 				{
 					FinishTime = Stopwatch.Elapsed;
 				}
@@ -117,7 +114,7 @@ namespace UnrealBuildTool
 		public static void AddEvent(string Name)
 		{
 			TimeSpan Time = Stopwatch.Elapsed;
-			lock(Events)
+			lock (Events)
 			{
 				Events.Add(new Event(Name, Time, Time));
 			}
@@ -197,7 +194,7 @@ namespace UnrealBuildTool
 			}
 
 			// Remove everything from the stack
-			for(; OuterEvents.Count > 0; OuterEvents.RemoveAt(OuterEvents.Count - 1))
+			for (; OuterEvents.Count > 0; OuterEvents.RemoveAt(OuterEvents.Count - 1))
 			{
 				UpdateLastEventTime(ref LastTime, OuterEvents.Last().FinishTime!.Value, MaxUnknownTime, OuterEvents, Verbosity, Logger);
 			}
@@ -238,7 +235,7 @@ namespace UnrealBuildTool
 		{
 			StringBuilder Prefix = new StringBuilder();
 
-			for(int Idx = 0; Idx < OuterEvents.Count - 1; Idx++)
+			for (int Idx = 0; Idx < OuterEvents.Count - 1; Idx++)
 			{
 				Prefix.AppendFormat(" {0,6}          ", FormatTime(StartTime - OuterEvents[Idx].StartTime));
 			}
@@ -249,7 +246,7 @@ namespace UnrealBuildTool
 			{
 				Prefix.AppendFormat("({0,6})", "???");
 			}
-			else if(FinishTime.Value == StartTime)
+			else if (FinishTime.Value == StartTime)
 			{
 				Prefix.Append(" ------ ");
 			}
diff --git a/Engine/Source/Programs/UnrealBuildTool/System/ToolMode.cs b/Engine/Source/Programs/UnrealBuildTool/System/ToolMode.cs
--- a/Engine/Source/Programs/UnrealBuildTool/System/ToolMode.cs
+++ b/Engine/Source/Programs/UnrealBuildTool/System/ToolMode.cs
@@ -1,9 +1,6 @@
 // Copyright Epic Games, Inc. All Rights Reserved.
 
 using System;
-using System.Collections.Generic;
-using System.Linq;
-using System.Text;
 using System.Threading.Tasks;
 using EpicGames.Core;
 using Microsoft.Extensions.Logging;
diff --git a/Engine/Source/Programs/UnrealBuildTool/System/Unity.cs b/Engine/Source/Programs/UnrealBuildTool/System/Unity.cs
--- a/Engine/Source/Programs/UnrealBuildTool/System/Unity.cs
+++ b/Engine/Source/Programs/UnrealBuildTool/System/Unity.cs
@@ -2,12 +2,10 @@
 
 using System;
 using System.Collections.Generic;
-using System.Text;
 using System.IO;
 using System.Linq;
 using EpicGames.Core;
 using UnrealBuildBase;
-using Microsoft.Extensions.Logging;
 
 namespace UnrealBuildTool
 {
@@ -302,7 +300,7 @@ namespace UnrealBuildTool
 				CompileEnvironment.FileInlineGenCPPMap[UnityCPPFile] = InlinedGenCPPFilesInUnity;
 
 				// Store the mapping of source files to unity files in the makefile
-				foreach(FileItem SourceFile in UnityFile.Files)
+				foreach (FileItem SourceFile in UnityFile.Files)
 				{
 					SourceFileToUnityFile[SourceFile] = UnityCPPFile;
 				}
diff --git a/Engine/Source/Programs/UnrealBuildTool/System/UnrealPluginLanguage.cs b/Engine/Source/Programs/UnrealBuildTool/System/UnrealPluginLanguage.cs
--- a/Engine/Source/Programs/UnrealBuildTool/System/UnrealPluginLanguage.cs
+++ b/Engine/Source/Programs/UnrealBuildTool/System/UnrealPluginLanguage.cs
@@ -2,16 +2,15 @@
 
 using System;
 using System.Collections.Generic;
+using System.Diagnostics.CodeAnalysis;
+using System.IO;
 using System.Linq;
 using System.Text;
 using System.Xml;
 using System.Xml.Linq;
-using System.Diagnostics;
-using System.IO;
 using EpicGames.Core;
-using UnrealBuildBase;
-using System.Diagnostics.CodeAnalysis;
 using Microsoft.Extensions.Logging;
+using UnrealBuildBase;
 
 namespace UnrealBuildTool
 {
@@ -505,7 +504,7 @@ namespace UnrealBuildTool
 
 		/** Project file reference */
 		private FileReference? ProjectFile;
-		
+
 		static private XDocument XMLDummy = XDocument.Parse("<manifest></manifest>");
 
 		private class UPLContext
@@ -616,7 +615,7 @@ namespace UnrealBuildTool
 			{
 				return false;
 			}
-	
+
 			// create a context for each architecture
 			ContextIndex++;
 			foreach (string Architecture in Architectures)
@@ -1074,7 +1073,7 @@ namespace UnrealBuildTool
 			{
 				foreach (XElement Index in Source.Elements())
 				{
-//					if (Target.Element(Index.Name) == null)
+					//					if (Target.Element(Index.Name) == null)
 					{
 						Target.Add(Index);
 					}
@@ -2426,9 +2425,9 @@ namespace UnrealBuildTool
 									{
 										int Index = StringToInt(CurrentContext, Node, IndexStr);
 										if (Index >= 0 && Index < StringList.Count)
-                                        {
+										{
 											Value = StringList.ElementAt(Index);
-                                        }
+										}
 									}
 								}
 								if (Result == "Output")
@@ -2619,7 +2618,7 @@ namespace UnrealBuildTool
 		}
 
 		public void Init(List<string> Architectures, bool bDistribution, string EngineDirectory, string BuildDirectory, string ProjectDirectory, string Configuration, bool bIsEmbedded,
-			bool bPerArchBuildDir=false, Dictionary<string, string>? ArchRemapping = null)
+			bool bPerArchBuildDir = false, Dictionary<string, string>? ArchRemapping = null)
 		{
 			GlobalContext.BoolVariables["Distribution"] = bDistribution;
 			GlobalContext.BoolVariables["IsEmbedded"] = bIsEmbedded;
@@ -2879,8 +2878,8 @@ namespace UnrealBuildTool
 		/// <param name="Logger"></param>
 		/// <param name="EngineDirectory"></param>
 		public ConfigCacheIni_UPL(UnrealTargetPlatform Platform, string BaseIniName, string? ProjectDirectory, ILogger Logger, string? EngineDirectory = null)
-			: this(Platform, BaseIniName, 
-				  (ProjectDirectory == null) ? null : new DirectoryReference(ProjectDirectory), 
+			: this(Platform, BaseIniName,
+				  (ProjectDirectory == null) ? null : new DirectoryReference(ProjectDirectory),
 				  Logger,
 				  (EngineDirectory == null) ? null : new DirectoryReference(EngineDirectory))
 		{
@@ -3247,7 +3246,7 @@ namespace UnrealBuildTool
 				int LineIndex = 1;
 				bool bMultiLine = false;
 				string SingleValue = "";
-				string  Key = "";
+				string Key = "";
 				ParseAction LastAction = ParseAction.None;
 
 				// Parse each line
@@ -3501,13 +3500,13 @@ namespace UnrealBuildTool
 			DirectoryReference? UserSettingsFolder = Unreal.UserSettingDirectory; // Match FPlatformProcess::UserSettingsDir()
 			DirectoryReference? PersonalFolder = Unreal.UserDirectory; // Match FPlatformProcess::UserDir()
 
-			if(UserSettingsFolder != null)
+			if (UserSettingsFolder != null)
 			{
 				// <AppData>/Unreal/EngineConfig/User* ini
 				yield return FileReference.Combine(UserSettingsFolder, "Unreal Engine", "Engine", "Config", "User" + BaseIniName + ".ini");
 			}
 
-			if(PersonalFolder != null)
+			if (PersonalFolder != null)
 			{
 				// <Documents>/Unreal/EngineConfig/User* ini
 				yield return FileReference.Combine(PersonalFolder, "Unreal Engine", "Engine", "Config", "User" + BaseIniName + ".ini");
diff --git a/Engine/Source/Programs/UnrealBuildTool/System/Utils.cs b/Engine/Source/Programs/UnrealBuildTool/System/Utils.cs
--- a/Engine/Source/Programs/UnrealBuildTool/System/Utils.cs
+++ b/Engine/Source/Programs/UnrealBuildTool/System/Utils.cs
@@ -2,19 +2,16 @@
 
 using System;
 using System.Collections.Generic;
-using System.IO;
 using System.Diagnostics;
+using System.IO;
+using System.Linq;
+using System.Runtime.InteropServices;
 using System.Text;
 using System.Text.RegularExpressions;
-using System.Xml;
 using System.Xml.Serialization;
-using System.Runtime.CompilerServices;
-using System.Runtime.InteropServices;
-using System.Linq;
 using EpicGames.Core;
-using UnrealBuildBase;
-using System.Collections.Concurrent;
 using Microsoft.Extensions.Logging;
+using UnrealBuildBase;
 
 namespace UnrealBuildTool
 {
@@ -290,9 +287,9 @@ namespace UnrealBuildTool
 		/// </summary>
 		static void LocalProcessOutput(DataReceivedEventArgs Args, bool bIsError, ILogger Logger)
 		{
-			if(Args != null && Args.Data != null)
+			if (Args != null && Args.Data != null)
 			{
-				if(bIsError)
+				if (bIsError)
 				{
 					Logger.LogError("{Message}", Args.Data.TrimEnd());
 				}
@@ -319,13 +316,13 @@ namespace UnrealBuildTool
 					if (!Directory.Exists(LogDir))
 					{
 						string? IniPath = UnrealBuildTool.GetRemoteIniPath();
-						if(string.IsNullOrEmpty(IniPath))
+						if (string.IsNullOrEmpty(IniPath))
 						{
 							break;
 						}
 
 						LogDir = Path.Combine(IniPath, "Saved", "Logs");
-						if(!Directory.Exists(LogDir) && !Directory.CreateDirectory(LogDir).Exists)
+						if (!Directory.Exists(LogDir) && !Directory.CreateDirectory(LogDir).Exists)
 						{
 							break;
 						}
@@ -359,7 +356,7 @@ namespace UnrealBuildTool
 				if (Args != null && Args.Data != null)
 				{
 					string data = Args.Data.TrimEnd();
-					if(string.IsNullOrEmpty(data))
+					if (string.IsNullOrEmpty(data))
 					{
 						return;
 					}
@@ -379,7 +376,7 @@ namespace UnrealBuildTool
 			LocalProcess.OutputDataReceived += Output;
 			LocalProcess.ErrorDataReceived += Output;
 			var ExitCode = RunLocalProcess(LocalProcess);
-			if(ExitCode != 0 && !string.IsNullOrEmpty(LogFilename))
+			if (ExitCode != 0 && !string.IsNullOrEmpty(LogFilename))
 			{
 				Logger.LogError("Process \'{AppName}\' failed. Details are in \'{LogFilename}\'", AppName, LogFilename);
 			}
@@ -414,7 +411,7 @@ namespace UnrealBuildTool
 		public static string RunLocalProcessAndReturnStdOut(string Command, string Args, ILogger? Logger)
 		{
 			int ExitCode;
-			return RunLocalProcessAndReturnStdOut(Command, Args, Logger, out ExitCode);	
+			return RunLocalProcessAndReturnStdOut(Command, Args, Logger, out ExitCode);
 		}
 
 		/// <summary>
@@ -434,7 +431,7 @@ namespace UnrealBuildTool
 		/// <param name="Args">Arguments to Command</param>
 		/// <param name="ExitCode">The return code from the process after it exits</param>
 		/// <param name="LogOutput">Whether to also log standard output and standard error</param>
-		public static string RunLocalProcessAndReturnStdOut(string Command, string? Args, out int ExitCode, bool LogOutput) => RunLocalProcessAndReturnStdOut(Command, Args, LogOutput? Log.Logger : null, out ExitCode);
+		public static string RunLocalProcessAndReturnStdOut(string Command, string? Args, out int ExitCode, bool LogOutput) => RunLocalProcessAndReturnStdOut(Command, Args, LogOutput ? Log.Logger : null, out ExitCode);
 
 		/// <summary>
 		/// Runs a command line process, and returns simple StdOut output.
@@ -474,7 +471,7 @@ namespace UnrealBuildTool
 				ErrorOutput = ErrorReader.ReadToEnd().Trim();
 				if (Logger != null)
 				{
-					if(FullOutput.Length > 0)
+					if (FullOutput.Length > 0)
 					{
 						Logger.LogInformation("{Output}", FullOutput);
 					}
@@ -1027,7 +1024,7 @@ namespace UnrealBuildTool
 			RelationAll = 0xffff
 		}
 
-		[DllImport("kernel32.dll", SetLastError=true)]
+		[DllImport("kernel32.dll", SetLastError = true)]
 		extern static bool GetLogicalProcessorInformationEx(LOGICAL_PROCESSOR_RELATIONSHIP RelationshipType, IntPtr Buffer, ref uint ReturnedLength);
 
 		/// <summary>
@@ -1037,13 +1034,13 @@ namespace UnrealBuildTool
 		public static int GetLogicalProcessorCount()
 		{
 			// This function uses Windows P/Invoke calls; if we're not running on Windows, just return the default.
-			if(RuntimePlatform.IsWindows)
+			if (RuntimePlatform.IsWindows)
 			{
 				const int ERROR_INSUFFICIENT_BUFFER = 122;
 
 				// Determine the required buffer size to store the processor information
 				uint ReturnLength = 0;
-				if(!GetLogicalProcessorInformationEx(LOGICAL_PROCESSOR_RELATIONSHIP.RelationGroup, IntPtr.Zero, ref ReturnLength) && Marshal.GetLastWin32Error() == ERROR_INSUFFICIENT_BUFFER)
+				if (!GetLogicalProcessorInformationEx(LOGICAL_PROCESSOR_RELATIONSHIP.RelationGroup, IntPtr.Zero, ref ReturnLength) && Marshal.GetLastWin32Error() == ERROR_INSUFFICIENT_BUFFER)
 				{
 					// Allocate a buffer for it
 					IntPtr Ptr = Marshal.AllocHGlobal((int)ReturnLength);
@@ -1052,10 +1049,10 @@ namespace UnrealBuildTool
 						if (GetLogicalProcessorInformationEx(LOGICAL_PROCESSOR_RELATIONSHIP.RelationGroup, Ptr, ref ReturnLength))
 						{
 							int Count = 0;
-							for(int Pos = 0; Pos < ReturnLength; )
+							for (int Pos = 0; Pos < ReturnLength;)
 							{
 								LOGICAL_PROCESSOR_RELATIONSHIP Type = (LOGICAL_PROCESSOR_RELATIONSHIP)Marshal.ReadInt16(Ptr, Pos);
-								if(Type == LOGICAL_PROCESSOR_RELATIONSHIP.RelationGroup)
+								if (Type == LOGICAL_PROCESSOR_RELATIONSHIP.RelationGroup)
 								{
 									// Read the values from the embedded GROUP_RELATIONSHIP structure
 									int GroupRelationshipPos = Pos + 8;
@@ -1063,7 +1060,7 @@ namespace UnrealBuildTool
 
 									// Read the processor counts from the embedded PROCESSOR_GROUP_INFO structures
 									int GroupInfoPos = GroupRelationshipPos + 24;
-									for(int GroupIdx = 0; GroupIdx < ActiveGroupCount; GroupIdx++)
+									for (int GroupIdx = 0; GroupIdx < ActiveGroupCount; GroupIdx++)
 									{
 										Count += Marshal.ReadByte(Ptr, GroupInfoPos + 1);
 										GroupInfoPos += 40 + IntPtr.Size;
@@ -1076,18 +1073,18 @@ namespace UnrealBuildTool
 					}
 					finally
 					{
-						Marshal.FreeHGlobal(Ptr);		
+						Marshal.FreeHGlobal(Ptr);
 					}
 				}
 			}
 			return Environment.ProcessorCount;
 		}
 
-		
+
 		// int sysctlbyname(const char *name, void *oldp, size_t *oldlenp, void *newp, size_t newlen); // from man page
 		[DllImport("libc")]
 		extern static int sysctlbyname(string name, out int oldp, ref UInt64 oldlenp, IntPtr newp, UInt64 newlen);
-		
+
 		/// <summary>
 		/// Gets the number of physical cores, excluding hyper threading.
 		/// </summary>
@@ -1104,7 +1101,7 @@ namespace UnrealBuildTool
 					ref ReturnLength) && Marshal.GetLastWin32Error() == ERROR_INSUFFICIENT_BUFFER)
 				{
 					// Allocate a buffer for it
-					IntPtr Ptr = Marshal.AllocHGlobal((int) ReturnLength);
+					IntPtr Ptr = Marshal.AllocHGlobal((int)ReturnLength);
 					try
 					{
 						if (GetLogicalProcessorInformationEx(LOGICAL_PROCESSOR_RELATIONSHIP.RelationProcessorCore, Ptr,
@@ -1116,7 +1113,7 @@ namespace UnrealBuildTool
 							for (int Pos = 0; Pos < ReturnLength;)
 							{
 								LOGICAL_PROCESSOR_RELATIONSHIP Type =
-									(LOGICAL_PROCESSOR_RELATIONSHIP) Marshal.ReadInt16(Ptr, Pos);
+									(LOGICAL_PROCESSOR_RELATIONSHIP)Marshal.ReadInt16(Ptr, Pos);
 								if (Type == LOGICAL_PROCESSOR_RELATIONSHIP.RelationProcessorCore)
 								{
 									Count++;
@@ -1134,7 +1131,7 @@ namespace UnrealBuildTool
 					}
 				}
 			}
-			else if(RuntimePlatform.IsMac)
+			else if (RuntimePlatform.IsMac)
 			{
 				UInt64 Size = 4;
 				if (0 == sysctlbyname("hw.physicalcpu", out int Value, ref Size, IntPtr.Zero, 0))
@@ -1175,44 +1172,69 @@ namespace UnrealBuildTool
 			return MemoryInfo.TotalAvailableMemoryBytes != 0 ? MemoryInfo.TotalAvailableMemoryBytes : -1;
 		}
 
-		
+
 		// vm_statistics64, based on the definition in <mach/vm_statistics.h>
 		[StructLayout(LayoutKind.Sequential)]
-        struct vm_statistics64 {
-          	/*natural_t*/ public int	free_count;				/* # of pages free */
-          	/*natural_t*/ public int	active_count;			/* # of pages active */
-          	/*natural_t*/ public int	inactive_count;			/* # of pages inactive */
-          	/*natural_t*/ public int	wire_count;				/* # of pages wired down */
-          	/*uint64_t */ public UInt64	zero_fill_count;		/* # of zero fill pages */
-          	/*uint64_t */ public UInt64	reactivations;			/* # of pages reactivated */
-		  	/*uint64_t */ public UInt64	pageins;				/* # of pageins */
-          	/*uint64_t */ public UInt64	pageouts;				/* # of pageouts */
-          	/*uint64_t */ public UInt64	faults;					/* # of faults */
-          	/*uint64_t */ public UInt64	cow_faults;				/* # of copy-on-writes */
-          	/*uint64_t */ public UInt64	lookups;				/* object cache lookups */
-          	/*uint64_t */ public UInt64	hits;					/* object cache hits */
-          	/*uint64_t */ public UInt64	purges;					/* # of pages purged */
-          	/*natural_t*/ public int	purgeable_count;		/* # of pages purgeable */
-          	/*
+		struct vm_statistics64
+		{
+			/*natural_t*/
+			public int free_count;              /* # of pages free */
+			/*natural_t*/
+			public int active_count;            /* # of pages active */
+			/*natural_t*/
+			public int inactive_count;          /* # of pages inactive */
+			/*natural_t*/
+			public int wire_count;              /* # of pages wired down */
+			/*uint64_t */
+			public UInt64 zero_fill_count;      /* # of zero fill pages */
+			/*uint64_t */
+			public UInt64 reactivations;            /* # of pages reactivated */
+			/*uint64_t */
+			public UInt64 pageins;              /* # of pageins */
+			/*uint64_t */
+			public UInt64 pageouts;             /* # of pageouts */
+			/*uint64_t */
+			public UInt64 faults;                   /* # of faults */
+			/*uint64_t */
+			public UInt64 cow_faults;               /* # of copy-on-writes */
+			/*uint64_t */
+			public UInt64 lookups;              /* object cache lookups */
+			/*uint64_t */
+			public UInt64 hits;                 /* object cache hits */
+			/*uint64_t */
+			public UInt64 purges;                   /* # of pages purged */
+			/*natural_t*/
+			public int purgeable_count;     /* # of pages purgeable */
+			/*
           	 * NB: speculative pages are already accounted for in "free_count",
           	 * so "speculative_count" is the number of "free" pages that are
           	 * used to hold data that was read speculatively from disk but
           	 * haven't actually been used by anyone so far.
           	 */
-          	/*natural_t*/ public int	speculative_count;		/* # of pages speculative */
-          
-          	/* added for rev1 */
-          	/*uint64_t */ public UInt64	decompressions;			/* # of pages decompressed */
-          	/*uint64_t */ public UInt64	compressions;			/* # of pages compressed */
-          	/*uint64_t */ public UInt64	swapins;				/* # of pages swapped in (via compression segments) */
-          	/*uint64_t */ public UInt64	swapouts;				/* # of pages swapped out (via compression segments) */
-          	/*natural_t*/ public int	compressor_page_count;	/* # of pages used by the compressed pager to hold all the compressed data */
-          	/*natural_t*/ public int	throttled_count;		/* # of pages throttled */
-          	/*natural_t*/ public int	external_page_count;	/* # of pages that are file-backed (non-swap) */
-          	/*natural_t*/ public int	internal_page_count;	/* # of pages that are anonymous */
-          	/*uint64_t */ public UInt64	total_uncompressed_pages_in_compressor; /* # of pages (uncompressed) held within the compressor. */
-        } // __attribute__((aligned(8))); 
-		
+			/*natural_t*/
+			public int speculative_count;       /* # of pages speculative */
+
+			/* added for rev1 */
+			/*uint64_t */
+			public UInt64 decompressions;           /* # of pages decompressed */
+			/*uint64_t */
+			public UInt64 compressions;         /* # of pages compressed */
+			/*uint64_t */
+			public UInt64 swapins;              /* # of pages swapped in (via compression segments) */
+			/*uint64_t */
+			public UInt64 swapouts;             /* # of pages swapped out (via compression segments) */
+			/*natural_t*/
+			public int compressor_page_count;   /* # of pages used by the compressed pager to hold all the compressed data */
+			/*natural_t*/
+			public int throttled_count;     /* # of pages throttled */
+			/*natural_t*/
+			public int external_page_count; /* # of pages that are file-backed (non-swap) */
+			/*natural_t*/
+			public int internal_page_count; /* # of pages that are anonymous */
+			/*uint64_t */
+			public UInt64 total_uncompressed_pages_in_compressor; /* # of pages (uncompressed) held within the compressor. */
+		} // __attribute__((aligned(8))); 
+
 		// kern_return_t host_statistics64(host_t host_priv, host_flavor_t flavor, host_info64_t host_info64_out, mach_msg_type_number_t *host_info64_outCnt); // from <mach/mach_host.h>
 		[DllImport("libc")]
 		extern static int host_statistics64(IntPtr host_priv, int flavor, out vm_statistics64 host_info64_out, ref uint host_info_count);
@@ -1220,7 +1242,7 @@ namespace UnrealBuildTool
 		// mach_port_t mach_host_self() // from <mach/mach_init.h>
 		[DllImport("libc")]
 		extern static IntPtr mach_host_self();
-		
+
 		/// <summary>
 		/// Gets the total memory bytes free, based on what is known to the garbage collector.
 		/// </summary>
@@ -1229,7 +1251,7 @@ namespace UnrealBuildTool
 		public static long GetFreeMemoryBytes()
 		{
 			long FreeMemoryBytes = -1;
-			
+
 			GCMemoryInfo MemoryInfo = GC.GetGCMemoryInfo();
 			// TotalAvailableMemoryBytes will be 0 if garbage collection has not run yet
 			if (MemoryInfo.TotalAvailableMemoryBytes != 0)
@@ -1240,13 +1262,13 @@ namespace UnrealBuildTool
 			// On Mac, MemoryInfo.MemoryLoadBytes includes memory used to cache disk-backed files ("Cached Files" in
 			// Activity Monitor), which can result in a significant over-estimate of memory pressure.
 			// We treat memory used for caching of disk-backed files as free for use in compilation tasks.
-			if (RuntimePlatform.IsMac)	
+			if (RuntimePlatform.IsMac)
 			{
 				// host_statistics64() flavor, from <mach/host_info.h>
 				int HOST_VM_INFO64 = 4;
 				// host_statistics64() count of 32bit values in output struct, from <mach/host_info.h>
 				int HOST_VM_INFO64_COUNT = Marshal.SizeOf(typeof(vm_statistics64)) / 4;
-		
+
 				vm_statistics64 VMStats;
 				uint StructSize = (uint)HOST_VM_INFO64_COUNT;
 				IntPtr Host = mach_host_self();
@@ -1258,7 +1280,7 @@ namespace UnrealBuildTool
 				{
 					PageSize = 4096; // likely result
 				}
-				
+
 				FreeMemoryBytes += (long)PageSize * (long)VMStats.external_page_count;
 			}
 			return FreeMemoryBytes;
@@ -1400,12 +1422,12 @@ namespace UnrealBuildTool
 		/// <returns>True if the steps succeeded, false otherwise</returns>
 		public static void ExecuteCustomBuildSteps(FileReference[] ScriptFiles, ILogger Logger)
 		{
-			foreach(FileReference ScriptFile in ScriptFiles)
+			foreach (FileReference ScriptFile in ScriptFiles)
 			{
 				ProcessStartInfo StartInfo = new ProcessStartInfo();
 				StartInfo.FileName = BuildHostPlatform.Current.Shell.FullName;
 
-				if(BuildHostPlatform.Current.ShellType == ShellType.Cmd)
+				if (BuildHostPlatform.Current.ShellType == ShellType.Cmd)
 				{
 					StartInfo.Arguments = String.Format("/C \"{0}\"", ScriptFile.FullName);
 				}
@@ -1415,7 +1437,7 @@ namespace UnrealBuildTool
 				}
 
 				int ReturnCode = Utils.RunLocalProcessAndLogOutput(StartInfo, Logger);
-				if(ReturnCode != 0)
+				if (ReturnCode != 0)
 				{
 					throw new BuildException("Custom build step {0} {1} terminated with exit code {2}", StartInfo.FileName, StartInfo.Arguments, ReturnCode);
 				}
@@ -1736,7 +1758,7 @@ namespace UnrealBuildTool
 					RecordWriteFileIfChanged(FileItem.Location, bNew: false, bChanged: true, Logger);
 				}
 				else
-				{ 
+				{
 					RecordWriteFileIfChanged(FileItem.Location, bNew: false, bChanged: false, Logger);
 				}
 			}
@@ -1777,9 +1799,9 @@ namespace UnrealBuildTool
 		{
 			// If the string length is equivalent to the encoded length, then no non-ASCII characters were present in the string.
 			// Don't write BOM as it messes with clang when loading response files.
-			return (Encoding.UTF8.GetByteCount(Str) != Str.Length) ?  new UTF8Encoding(false) : Encoding.ASCII;
+			return (Encoding.UTF8.GetByteCount(Str) != Str.Length) ? new UTF8Encoding(false) : Encoding.ASCII;
 		}
-		
+
 		/// <summary>
 		/// Determines the appropriate encoding for a list of strings: either ASCII or UTF-8.
 		/// If the string length is equivalent to the encoded length, then no non-ASCII characters were present in the string.
diff --git a/Engine/Source/Programs/UnrealBuildTool/System/VersionNumber.cs b/Engine/Source/Programs/UnrealBuildTool/System/VersionNumber.cs
--- a/Engine/Source/Programs/UnrealBuildTool/System/VersionNumber.cs
+++ b/Engine/Source/Programs/UnrealBuildTool/System/VersionNumber.cs
@@ -3,9 +3,7 @@
 using System;
 using System.Collections.Generic;
 using System.Diagnostics.CodeAnalysis;
-using System.Linq;
 using System.Text;
-using System.Threading.Tasks;
 
 namespace UnrealBuildTool
 {
@@ -31,7 +29,7 @@ namespace UnrealBuildTool
 		/// <param name="Components">The individual version components. At least one value must be given.</param>
 		public VersionNumber(params int[] Components)
 		{
-			if(Components.Length == 0)
+			if (Components.Length == 0)
 			{
 				throw new InvalidOperationException("Version number must have at least one component");
 			}
@@ -66,7 +64,7 @@ namespace UnrealBuildTool
 		public override int GetHashCode()
 		{
 			int Result = 5831;
-			for(int Idx = 0; Idx < Components.Length; Idx++)
+			for (int Idx = 0; Idx < Components.Length; Idx++)
 			{
 				Result = (Result * 33) + Components[Idx];
 			}
@@ -79,9 +77,9 @@ namespace UnrealBuildTool
 		/// <param name="Lhs">The first version number</param>
 		/// <param name="Rhs">The second version number</param>
 		/// <returns>True if the versions are equal.</returns>
-		public static bool operator==(VersionNumber? Lhs, VersionNumber? Rhs)
+		public static bool operator ==(VersionNumber? Lhs, VersionNumber? Rhs)
 		{
-			if(Object.ReferenceEquals(Lhs, null))
+			if (Object.ReferenceEquals(Lhs, null))
 			{
 				return Object.ReferenceEquals(Rhs, null);
 			}
@@ -97,7 +95,7 @@ namespace UnrealBuildTool
 		/// <param name="Lhs">The first version number</param>
 		/// <param name="Rhs">The second version number</param>
 		/// <returns>True if the versions are not equal.</returns>
-		public static bool operator!=(VersionNumber? Lhs, VersionNumber? Rhs)
+		public static bool operator !=(VersionNumber? Lhs, VersionNumber? Rhs)
 		{
 			return !(Lhs == Rhs);
 		}
@@ -108,7 +106,7 @@ namespace UnrealBuildTool
 		/// <param name="Lhs">The first version number</param>
 		/// <param name="Rhs">The second version number</param>
 		/// <returns>True if the first version is less than the second.</returns>
-		public static bool operator<(VersionNumber Lhs, VersionNumber Rhs)
+		public static bool operator <(VersionNumber Lhs, VersionNumber Rhs)
 		{
 			return Compare(Lhs, Rhs) < 0;
 		}
@@ -119,7 +117,7 @@ namespace UnrealBuildTool
 		/// <param name="Lhs">The first version number</param>
 		/// <param name="Rhs">The second version number</param>
 		/// <returns>True if the first version is less or equal to the second.</returns>
-		public static bool operator<=(VersionNumber Lhs, VersionNumber Rhs)
+		public static bool operator <=(VersionNumber Lhs, VersionNumber Rhs)
 		{
 			return Compare(Lhs, Rhs) <= 0;
 		}
@@ -130,7 +128,7 @@ namespace UnrealBuildTool
 		/// <param name="Lhs">The first version number</param>
 		/// <param name="Rhs">The second version number</param>
 		/// <returns>True if the first version is greater than the second.</returns>
-		public static bool operator>(VersionNumber Lhs, VersionNumber Rhs)
+		public static bool operator >(VersionNumber Lhs, VersionNumber Rhs)
 		{
 			return Compare(Lhs, Rhs) > 0;
 		}
@@ -141,7 +139,7 @@ namespace UnrealBuildTool
 		/// <param name="Lhs">The first version number</param>
 		/// <param name="Rhs">The second version number</param>
 		/// <returns>True if the first version is greater or equal to the second.</returns>
-		public static bool operator>=(VersionNumber Lhs, VersionNumber Rhs)
+		public static bool operator >=(VersionNumber Lhs, VersionNumber Rhs)
 		{
 			return Compare(Lhs, Rhs) >= 0;
 		}
@@ -153,7 +151,7 @@ namespace UnrealBuildTool
 		/// <returns>A negative value if this version is before Other, a positive value if this version is after Other, and zero otherwise.</returns>
 		public int CompareTo(VersionNumber? Other)
 		{
-			return ReferenceEquals(Other, null)? 1 : Compare(this, Other);
+			return ReferenceEquals(Other, null) ? 1 : Compare(this, Other);
 		}
 
 		/// <summary>
@@ -164,11 +162,11 @@ namespace UnrealBuildTool
 		/// <returns>A negative value if Lhs is before Rhs, a positive value if Lhs is after Rhs, and zero otherwise.</returns>
 		public static int Compare(VersionNumber Lhs, VersionNumber Rhs)
 		{
-			for(int Idx = 0;;Idx++)
+			for (int Idx = 0; ; Idx++)
 			{
-				if(Idx == Lhs.Components.Length)
+				if (Idx == Lhs.Components.Length)
 				{
-					if(Idx == Rhs.Components.Length)
+					if (Idx == Rhs.Components.Length)
 					{
 						return 0;
 					}
@@ -179,11 +177,11 @@ namespace UnrealBuildTool
 				}
 				else
 				{
-					if(Idx == Rhs.Components.Length)
+					if (Idx == Rhs.Components.Length)
 					{
 						return +1;
 					}
-					else if(Lhs.Components[Idx] != Rhs.Components[Idx])
+					else if (Lhs.Components[Idx] != Rhs.Components[Idx])
 					{
 						return Lhs.Components[Idx] - Rhs.Components[Idx];
 					}
@@ -199,7 +197,7 @@ namespace UnrealBuildTool
 		public static VersionNumber Parse(string Text)
 		{
 			List<int> Components = new List<int>();
-			foreach(string TextElement in Text.Split(Delimiters))
+			foreach (string TextElement in Text.Split(Delimiters))
 			{
 				Components.Add(int.Parse(TextElement));
 			}
@@ -215,10 +213,10 @@ namespace UnrealBuildTool
 		public static bool TryParse(string Text, [NotNullWhen(true)] out VersionNumber? OutNumber)
 		{
 			List<int> Components = new List<int>();
-			foreach(string TextElement in Text.Split(Delimiters))
+			foreach (string TextElement in Text.Split(Delimiters))
 			{
 				int Component;
-				if(!int.TryParse(TextElement, out Component))
+				if (!int.TryParse(TextElement, out Component))
 				{
 					OutNumber = null;
 					return false;
@@ -237,10 +235,10 @@ namespace UnrealBuildTool
 		public override string ToString()
 		{
 			StringBuilder Result = new StringBuilder();
-			if(Components.Length > 0)
+			if (Components.Length > 0)
 			{
 				Result.Append(Components[0]);
-				for(int Idx = 1; Idx < Components.Length; Idx++)
+				for (int Idx = 1; Idx < Components.Length; Idx++)
 				{
 					Result.Append('.');
 					Result.Append(Components[Idx]);
diff --git a/Engine/Source/Programs/UnrealBuildTool/System/VersionNumberRange.cs b/Engine/Source/Programs/UnrealBuildTool/System/VersionNumberRange.cs
--- a/Engine/Source/Programs/UnrealBuildTool/System/VersionNumberRange.cs
+++ b/Engine/Source/Programs/UnrealBuildTool/System/VersionNumberRange.cs
@@ -1,8 +1,6 @@
 // Copyright Epic Games, Inc. All Rights Reserved.
 
 using System;
-using System.Collections.Generic;
-using System.Text;
 
 namespace UnrealBuildTool
 {
diff --git a/Engine/Source/Programs/UnrealBuildTool/System/XmlConfig.cs b/Engine/Source/Programs/UnrealBuildTool/System/XmlConfig.cs
--- a/Engine/Source/Programs/UnrealBuildTool/System/XmlConfig.cs
+++ b/Engine/Source/Programs/UnrealBuildTool/System/XmlConfig.cs
@@ -7,7 +7,6 @@ using System.IO;
 using System.Linq;
 using System.Reflection;
 using System.Text;
-using System.Threading.Tasks;
 using System.Xml;
 using System.Xml.Schema;
 using System.Xml.Serialization;
@@ -76,13 +75,13 @@ namespace UnrealBuildTool
 			List<Type> ConfigTypes = FindConfigurableTypes();
 
 			// Update the cache if necessary
-			if(OverrideCacheFile != null)
+			if (OverrideCacheFile != null)
 			{
 				// Set the cache file to the overriden value
 				CacheFile = OverrideCacheFile;
 
 				// Never rebuild the cache; just try to load it.
-				if(!XmlConfigData.TryRead(CacheFile, ConfigTypes, out Values))
+				if (!XmlConfigData.TryRead(CacheFile, ConfigTypes, out Values))
 				{
 					throw new BuildException("Unable to load XML config cache ({0})", CacheFile);
 				}
@@ -96,7 +95,7 @@ namespace UnrealBuildTool
 					if (Unreal.IsEngineInstalled())
 					{
 						DirectoryReference? UserSettingsDir = Unreal.UserSettingDirectory;
-						if(UserSettingsDir != null)
+						if (UserSettingsDir != null)
 						{
 							CacheFile = FileReference.Combine(UserSettingsDir, "UnrealEngine", String.Format("XmlConfigCache-{0}.bin", Unreal.RootDirectory.FullName.Replace(":", "").Replace(Path.DirectorySeparatorChar, '+')));
 						}
@@ -107,21 +106,21 @@ namespace UnrealBuildTool
 					CacheFile = FileReference.Combine(ProjectRootDirectory, "Intermediate", "Build", "XmlConfigCache.bin");
 					Values = null;
 				}
-				
+
 				// Find all the input files
 				List<FileReference> TemporaryInputFileLocations = InputFiles.Select(x => x.Location).ToList();
-				
+
 				if (ProjectRootDirectory != null)
 				{
 					FileReference ProjectRootConfigLocation = FileReference.Combine(ProjectRootDirectory, "Saved", "UnrealBuildTool", "BuildConfiguration.xml");
-					if(!FileReference.Exists(ProjectRootConfigLocation))
+					if (!FileReference.Exists(ProjectRootConfigLocation))
 					{
 						CreateDefaultConfigFile(ProjectRootConfigLocation);
 					}
 
 					TemporaryInputFileLocations.Add(ProjectRootConfigLocation);
 				}
-				
+
 				FileReference[] InputFileLocations = TemporaryInputFileLocations.ToArray();
 
 				// Get the path to the schema
@@ -129,16 +128,16 @@ namespace UnrealBuildTool
 
 				// Try to read the existing cache from disk
 				XmlConfigData? CachedValues;
-				if(IsCacheUpToDate(CacheFile, InputFileLocations) && FileReference.Exists(SchemaFile))
+				if (IsCacheUpToDate(CacheFile, InputFileLocations) && FileReference.Exists(SchemaFile))
 				{
-					if(XmlConfigData.TryRead(CacheFile, ConfigTypes, out CachedValues) && Enumerable.SequenceEqual(InputFileLocations, CachedValues.InputFiles))
+					if (XmlConfigData.TryRead(CacheFile, ConfigTypes, out CachedValues) && Enumerable.SequenceEqual(InputFileLocations, CachedValues.InputFiles))
 					{
 						Values = CachedValues;
 					}
 				}
 
 				// If that failed, regenerate it
-				if(Values == null)
+				if (Values == null)
 				{
 					// Find all the configurable fields from the given types
 					Dictionary<string, Dictionary<string, XmlConfigData.TargetMember>> CategoryToFields = new Dictionary<string, Dictionary<string, XmlConfigData.TargetMember>>();
@@ -146,7 +145,7 @@ namespace UnrealBuildTool
 
 					// Create a schema for the config files
 					XmlSchema Schema = CreateSchema(CategoryToFields);
-					if(!Unreal.IsEngineInstalled())
+					if (!Unreal.IsEngineInstalled())
 					{
 						WriteSchema(Schema, SchemaFile);
 					}
@@ -154,9 +153,9 @@ namespace UnrealBuildTool
 					// Read all the XML files and validate them against the schema
 					Dictionary<Type, Dictionary<XmlConfigData.TargetMember, XmlConfigData.ValueInfo>> TypeToValues =
 						new Dictionary<Type, Dictionary<XmlConfigData.TargetMember, XmlConfigData.ValueInfo>>();
-					foreach(FileReference InputFile in InputFileLocations)
+					foreach (FileReference InputFile in InputFileLocations)
 					{
-						if(!TryReadFile(InputFile, CategoryToFields, TypeToValues, Schema, Logger))
+						if (!TryReadFile(InputFile, CategoryToFields, TypeToValues, Schema, Logger))
 						{
 							throw new BuildException("Failed to properly read XML file : {0}", InputFile.FullName);
 						}
@@ -167,18 +166,18 @@ namespace UnrealBuildTool
 
 					// Create the new cache
 					Values = new XmlConfigData(InputFileLocations, TypeToValues.ToDictionary(
-							x => x.Key, 
+							x => x.Key,
 							x => x.Value.Select(x => x.Value).ToArray()));
 					Values.Write(CacheFile);
 				}
 			}
 
 			// Apply all the static field values
-			foreach(KeyValuePair<Type, XmlConfigData.ValueInfo[]> TypeValuesPair in Values.TypeToValues)
+			foreach (KeyValuePair<Type, XmlConfigData.ValueInfo[]> TypeValuesPair in Values.TypeToValues)
 			{
-				foreach(XmlConfigData.ValueInfo MemberValue in TypeValuesPair.Value)
+				foreach (XmlConfigData.ValueInfo MemberValue in TypeValuesPair.Value)
 				{
-					if(MemberValue.Target.IsStatic)
+					if (MemberValue.Target.IsStatic)
 					{
 						object Value = InstanceValue(MemberValue.Value, MemberValue.Target.Type);
 						MemberValue.Target.SetValue(null, Value);
@@ -219,11 +218,11 @@ namespace UnrealBuildTool
 		/// <returns>True if the type has a field with the XmlConfigFile attribute</returns>
 		static bool HasXmlConfigFileAttribute(Type Type)
 		{
-			foreach(FieldInfo Field in Type.GetFields(BindingFlags.Instance | BindingFlags.Static | BindingFlags.Public | BindingFlags.NonPublic))
+			foreach (FieldInfo Field in Type.GetFields(BindingFlags.Instance | BindingFlags.Static | BindingFlags.Public | BindingFlags.NonPublic))
 			{
-				foreach(CustomAttributeData CustomAttribute in Field.CustomAttributes)
+				foreach (CustomAttributeData CustomAttribute in Field.CustomAttributes)
 				{
-					if(CustomAttribute.AttributeType == typeof(XmlConfigFileAttribute))
+					if (CustomAttribute.AttributeType == typeof(XmlConfigFileAttribute))
 					{
 						return true;
 					}
@@ -275,7 +274,7 @@ namespace UnrealBuildTool
 				}
 
 				ILogger Logger = Log.Logger;
-				
+
 				// Find all the input file locations
 				List<InputFile> InputFilesFound = new List<InputFile>(4);
 
@@ -366,16 +365,16 @@ namespace UnrealBuildTool
 		{
 			ILogger Logger = Log.Logger;
 			for (Type? TargetType = TargetObject.GetType(); TargetType != null; TargetType = TargetType.BaseType)
-			{ 
+			{
 				XmlConfigData.ValueInfo[]? FieldValues;
-				if(Values!.TypeToValues.TryGetValue(TargetType, out FieldValues))
+				if (Values!.TypeToValues.TryGetValue(TargetType, out FieldValues))
 				{
-					foreach(XmlConfigData.ValueInfo FieldValue in FieldValues)
+					foreach (XmlConfigData.ValueInfo FieldValue in FieldValues)
 					{
 						if (!FieldValue.Target.IsStatic)
 						{
 							XmlConfigData.TargetMember TargetToWrite = FieldValue.Target;
-							
+
 							// Check if setting has been deprecated
 							if (FieldValue.XmlConfigAttribute.Deprecated)
 							{
@@ -452,7 +451,7 @@ namespace UnrealBuildTool
 		/// <returns>New instance of the given value, if necessary</returns>
 		static object InstanceValue(object Value, Type ValueType)
 		{
-			if(ValueType == typeof(string[]))
+			if (ValueType == typeof(string[]))
 			{
 				return ((string[])Value).Clone();
 			}
@@ -473,16 +472,16 @@ namespace UnrealBuildTool
 		{
 			// Find all the config values for this type
 			XmlConfigData.ValueInfo[]? FieldValues;
-			if(!Values!.TypeToValues.TryGetValue(TargetType, out FieldValues))
+			if (!Values!.TypeToValues.TryGetValue(TargetType, out FieldValues))
 			{
 				Value = null;
 				return false;
 			}
 
 			// Find the value with the matching name
-			foreach(XmlConfigData.ValueInfo FieldValue in FieldValues)
+			foreach (XmlConfigData.ValueInfo FieldValue in FieldValues)
 			{
-				if(FieldValue.Target.MemberInfo.Name == Name)
+				if (FieldValue.Target.MemberInfo.Name == Name)
 				{
 					Value = FieldValue.Value;
 					return true;
@@ -501,9 +500,9 @@ namespace UnrealBuildTool
 		/// <param name="CategoryToFields">Dictionaries populated with category -> name -> field mappings on return</param>
 		static void FindConfigurableFields(IEnumerable<Type> ConfigTypes, Dictionary<string, Dictionary<string, XmlConfigData.TargetMember>> CategoryToFields)
 		{
-			foreach(Type ConfigType in ConfigTypes)
+			foreach (Type ConfigType in ConfigTypes)
 			{
-				foreach(FieldInfo FieldInfo in ConfigType.GetFields(BindingFlags.Instance | BindingFlags.Static | BindingFlags.GetField | BindingFlags.Public | BindingFlags.NonPublic))
+				foreach (FieldInfo FieldInfo in ConfigType.GetFields(BindingFlags.Instance | BindingFlags.Static | BindingFlags.GetField | BindingFlags.Public | BindingFlags.NonPublic))
 				{
 					ProcessConfigurableMember<FieldInfo>(ConfigType, FieldInfo, CategoryToFields, FieldInfo => new XmlConfigData.TargetField(FieldInfo));
 				}
@@ -542,7 +541,7 @@ namespace UnrealBuildTool
 		{
 			// Create elements for all the categories
 			XmlSchemaAll RootAll = new XmlSchemaAll();
-			foreach(KeyValuePair<string, Dictionary<string, XmlConfigData.TargetMember>> CategoryPair in CategoryToFields)
+			foreach (KeyValuePair<string, Dictionary<string, XmlConfigData.TargetMember>> CategoryPair in CategoryToFields)
 			{
 				string CategoryName = CategoryPair.Key;
 
@@ -607,32 +606,32 @@ namespace UnrealBuildTool
 			{
 				Element.SchemaTypeName = XmlSchemaType.GetBuiltInSimpleType(XmlTypeCode.String).QualifiedName;
 			}
-			else if(Type == typeof(bool))
+			else if (Type == typeof(bool))
 			{
 				Element.SchemaTypeName = XmlSchemaType.GetBuiltInSimpleType(XmlTypeCode.Boolean).QualifiedName;
 			}
-			else if(Type == typeof(int))
+			else if (Type == typeof(int))
 			{
 				Element.SchemaTypeName = XmlSchemaType.GetBuiltInSimpleType(XmlTypeCode.Int).QualifiedName;
 			}
-			else if(Type == typeof(float))
+			else if (Type == typeof(float))
 			{
 				Element.SchemaTypeName = XmlSchemaType.GetBuiltInSimpleType(XmlTypeCode.Float).QualifiedName;
 			}
-			else if(Type == typeof(double))
+			else if (Type == typeof(double))
 			{
 				Element.SchemaTypeName = XmlSchemaType.GetBuiltInSimpleType(XmlTypeCode.Double).QualifiedName;
 			}
-			else if(Type == typeof(FileReference))
+			else if (Type == typeof(FileReference))
 			{
 				Element.SchemaTypeName = XmlSchemaType.GetBuiltInSimpleType(XmlTypeCode.String).QualifiedName;
 			}
-			else if(Type.IsEnum)
+			else if (Type.IsEnum)
 			{
 				XmlSchemaSimpleTypeRestriction Restriction = new XmlSchemaSimpleTypeRestriction();
 				Restriction.BaseTypeName = XmlSchemaType.GetBuiltInSimpleType(XmlTypeCode.String).QualifiedName;
 
-				foreach(string EnumName in Enum.GetNames(Type))
+				foreach (string EnumName in Enum.GetNames(Type))
 				{
 					XmlSchemaEnumerationFacet Facet = new XmlSchemaEnumerationFacet();
 					Facet.Value = EnumName;
@@ -643,7 +642,7 @@ namespace UnrealBuildTool
 				EnumType.Content = Restriction;
 				Element.SchemaType = EnumType;
 			}
-			else if(Type == typeof(string[]))
+			else if (Type == typeof(string[]))
 			{
 				XmlSchemaElement ItemElement = new XmlSchemaElement();
 				ItemElement.Name = "Item";
@@ -678,14 +677,14 @@ namespace UnrealBuildTool
 			Settings.NewLineChars = Environment.NewLine;
 			Settings.OmitXmlDeclaration = true;
 
-			if(CachedSchemaSerializer == null)
+			if (CachedSchemaSerializer == null)
 			{
 				CachedSchemaSerializer = XmlSerializer.FromTypes(new Type[] { typeof(XmlSchema) })[0]!;
 			}
 
 			StringBuilder Output = new StringBuilder();
 			Output.AppendLine("<?xml version=\"1.0\" encoding=\"utf-8\"?>");
-			using(XmlWriter Writer = XmlWriter.Create(Output, Settings))
+			using (XmlWriter Writer = XmlWriter.Create(Output, Settings))
 			{
 				XmlSerializerNamespaces Namespaces = new XmlSerializerNamespaces();
 				Namespaces.Add("", "http://www.w3.org/2001/XMLSchema");
@@ -693,7 +692,7 @@ namespace UnrealBuildTool
 			}
 
 			string OutputText = Output.ToString();
-			if(!FileReference.Exists(Location) || File.ReadAllText(Location.FullName) != OutputText)
+			if (!FileReference.Exists(Location) || File.ReadAllText(Location.FullName) != OutputText)
 			{
 				DirectoryReference.CreateDirectory(Location.Directory);
 				File.WriteAllText(Location.FullName, OutputText);
@@ -731,22 +730,22 @@ namespace UnrealBuildTool
 		{
 			// Read the XML file, and validate it against the schema
 			XmlConfigFile? ConfigFile;
-			if(!XmlConfigFile.TryRead(Location, Schema, Logger, out ConfigFile))
+			if (!XmlConfigFile.TryRead(Location, Schema, Logger, out ConfigFile))
 			{
 				return false;
 			}
 
 			// Parse the document
-			foreach(XmlElement CategoryElement in ConfigFile.DocumentElement!.ChildNodes.OfType<XmlElement>())
+			foreach (XmlElement CategoryElement in ConfigFile.DocumentElement!.ChildNodes.OfType<XmlElement>())
 			{
 				Dictionary<string, XmlConfigData.TargetMember>? NameToField;
-				if(CategoryToFields.TryGetValue(CategoryElement.Name, out NameToField))
+				if (CategoryToFields.TryGetValue(CategoryElement.Name, out NameToField))
 				{
-					foreach(XmlElement KeyElement in CategoryElement.ChildNodes.OfType<XmlElement>())
+					foreach (XmlElement KeyElement in CategoryElement.ChildNodes.OfType<XmlElement>())
 					{
 						XmlConfigData.TargetMember? Field;
 						object Value;
-						if(NameToField.TryGetValue(KeyElement.Name, out Field))
+						if (NameToField.TryGetValue(KeyElement.Name, out Field))
 						{
 							if (Field.Type == typeof(string[]))
 							{
@@ -763,16 +762,16 @@ namespace UnrealBuildTool
 
 							// Add it to the set of values for the type containing this field
 							Dictionary<XmlConfigData.TargetMember, XmlConfigData.ValueInfo>? FieldToValue;
-							if(!TypeToValues.TryGetValue(Field.MemberInfo.DeclaringType!, out FieldToValue))
+							if (!TypeToValues.TryGetValue(Field.MemberInfo.DeclaringType!, out FieldToValue))
 							{
 								FieldToValue = new Dictionary<XmlConfigData.TargetMember, XmlConfigData.ValueInfo>();
 								TypeToValues.Add(Field.MemberInfo.DeclaringType!, FieldToValue);
 							}
-							
+
 							// Parse the corresponding value
 							XmlConfigData.ValueInfo FieldValue = new XmlConfigData.ValueInfo(Field, Value, Location,
 								Field.MemberInfo.GetCustomAttribute<XmlConfigFileAttribute>()!);
-							
+
 							FieldToValue[Field] = FieldValue;
 						}
 					}
@@ -791,11 +790,11 @@ namespace UnrealBuildTool
 		{
 			// ignore whitespace in all fields except for Strings which we leave unprocessed
 			string TrimmedText = Text.Trim();
-			if(FieldType == typeof(string))
+			if (FieldType == typeof(string))
 			{
 				return Text;
 			}
-			else if(FieldType == typeof(bool) || FieldType == typeof(bool?))
+			else if (FieldType == typeof(bool) || FieldType == typeof(bool?))
 			{
 				if (TrimmedText == "1" || TrimmedText.Equals("true", StringComparison.InvariantCultureIgnoreCase))
 				{
@@ -804,25 +803,25 @@ namespace UnrealBuildTool
 				else if (TrimmedText == "0" || TrimmedText.Equals("false", StringComparison.InvariantCultureIgnoreCase))
 				{
 					return false;
-				} 
-				else 
+				}
+				else
 				{
 					throw new Exception(String.Format("Unable to convert '{0}' to boolean. 'true/false/0/1' are the supported formats.", Text));
 				}
 			}
-			else if(FieldType == typeof(int))
+			else if (FieldType == typeof(int))
 			{
 				return Int32.Parse(TrimmedText);
 			}
-			else if(FieldType == typeof(float))
+			else if (FieldType == typeof(float))
 			{
 				return Single.Parse(TrimmedText, System.Globalization.CultureInfo.InvariantCulture);
 			}
-			else if(FieldType == typeof(double))
+			else if (FieldType == typeof(double))
 			{
 				return Double.Parse(TrimmedText, System.Globalization.CultureInfo.InvariantCulture);
 			}
-			else if(FieldType.IsEnum)
+			else if (FieldType.IsEnum)
 			{
 				return Enum.Parse(FieldType, TrimmedText);
 			}
@@ -846,7 +845,7 @@ namespace UnrealBuildTool
 		static bool IsCacheUpToDate(FileReference CacheFile, FileReference[] InputFiles)
 		{
 			// Always rebuild if the cache doesn't exist
-			if(!FileReference.Exists(CacheFile))
+			if (!FileReference.Exists(CacheFile))
 			{
 				return false;
 			}
@@ -855,15 +854,15 @@ namespace UnrealBuildTool
 			DateTime CacheWriteTime = File.GetLastWriteTimeUtc(CacheFile.FullName);
 
 			// Always rebuild if this executable is newer
-			if(File.GetLastWriteTimeUtc(Assembly.GetExecutingAssembly().Location) > CacheWriteTime)
+			if (File.GetLastWriteTimeUtc(Assembly.GetExecutingAssembly().Location) > CacheWriteTime)
 			{
 				return false;
 			}
 
 			// Check if any of the input files are newer than the cache
-			foreach(FileReference InputFile in InputFiles)
+			foreach (FileReference InputFile in InputFiles)
 			{
-				if(File.GetLastWriteTimeUtc(InputFile.FullName) > CacheWriteTime)
+				if (File.GetLastWriteTimeUtc(InputFile.FullName) > CacheWriteTime)
 				{
 					return false;
 				}
@@ -890,7 +889,7 @@ namespace UnrealBuildTool
 
 			// Get the path to the XML documentation
 			FileReference InputDocumentationFile = new FileReference(Assembly.GetExecutingAssembly().Location).ChangeExtension(".xml");
-			if(!FileReference.Exists(InputDocumentationFile))
+			if (!FileReference.Exists(InputDocumentationFile))
 			{
 				throw new BuildException("Generated assembly documentation not found at {0}.", InputDocumentationFile);
 			}
@@ -900,7 +899,7 @@ namespace UnrealBuildTool
 			InputDocumentation.Load(InputDocumentationFile.FullName);
 
 			// Make sure we can write to the output file
-			if(FileReference.Exists(OutputFile))
+			if (FileReference.Exists(OutputFile))
 			{
 				FileReference.MakeWriteable(OutputFile);
 			}
@@ -910,11 +909,11 @@ namespace UnrealBuildTool
 			}
 
 			// Generate the documentation file
-			if(OutputFile.HasExtension(".udn"))
+			if (OutputFile.HasExtension(".udn"))
 			{
 				WriteDocumentationUDN(OutputFile, InputDocumentation, CategoryToFields, Logger);
 			}
-			else if(OutputFile.HasExtension(".html"))
+			else if (OutputFile.HasExtension(".html"))
 			{
 				WriteDocumentationHTML(OutputFile, InputDocumentation, CategoryToFields, Logger);
 			}
@@ -956,7 +955,7 @@ namespace UnrealBuildTool
 					{
 						// Get the XML comment for this field
 						List<string>? Lines;
-						if(!RulesDocumentation.TryGetXmlComment(InputDocumentation, FieldPair.Value.MemberInfo, Logger, out Lines) || Lines.Count == 0)
+						if (!RulesDocumentation.TryGetXmlComment(InputDocumentation, FieldPair.Value.MemberInfo, Logger, out Lines) || Lines.Count == 0)
 						{
 							continue;
 						}
diff --git a/Engine/Source/Programs/UnrealBuildTool/System/XmlConfigData.cs b/Engine/Source/Programs/UnrealBuildTool/System/XmlConfigData.cs
--- a/Engine/Source/Programs/UnrealBuildTool/System/XmlConfigData.cs
+++ b/Engine/Source/Programs/UnrealBuildTool/System/XmlConfigData.cs
@@ -6,10 +6,7 @@ using System.Diagnostics.CodeAnalysis;
 using System.IO;
 using System.Linq;
 using System.Reflection;
-using System.Text;
-using System.Threading.Tasks;
 using EpicGames.Core;
-using UnrealBuildTool;
 
 namespace UnrealBuildTool
 {
@@ -36,27 +33,27 @@ namespace UnrealBuildTool
 			/// <summary>
 			/// Returns Reflection.MemberInfo describing the target class member.
 			/// </summary>
-			public abstract MemberInfo               MemberInfo { get; }
+			public abstract MemberInfo MemberInfo { get; }
 
 			/// <summary>
 			/// Returns Reflection.Type of the target class member.
 			/// </summary>
-			public abstract Type                     Type       { get; }
+			public abstract Type Type { get; }
 
 			/// <summary>
 			/// Indicates whether the target class member is static or not.
 			/// </summary>
-			public abstract bool                     IsStatic   { get; }
+			public abstract bool IsStatic { get; }
 
 			/// <summary>
 			/// Returns the value setter of the target class member.
 			/// </summary>
-			public abstract Action<object?, object?> SetValue   { get; }
+			public abstract Action<object?, object?> SetValue { get; }
 
 			/// <summary>
 			/// Returns the value getter of the target class member.
 			/// </summary>
-			public abstract Func<object?, object?>   GetValue   { get; }
+			public abstract Func<object?, object?> GetValue { get; }
 		}
 
 		/// <summary>
@@ -64,11 +61,11 @@ namespace UnrealBuildTool
 		/// </summary>
 		public class TargetField : TargetMember
 		{
-			public override MemberInfo               MemberInfo => FieldInfo;
-			public override Type                     Type       => FieldInfo.FieldType;
-			public override bool                     IsStatic   => FieldInfo.IsStatic;
-			public override Action<object?, object?> SetValue   => FieldInfo.SetValue;
-			public override Func<object?, object?>   GetValue   => FieldInfo.GetValue;
+			public override MemberInfo MemberInfo => FieldInfo;
+			public override Type Type => FieldInfo.FieldType;
+			public override bool IsStatic => FieldInfo.IsStatic;
+			public override Action<object?, object?> SetValue => FieldInfo.SetValue;
+			public override Func<object?, object?> GetValue => FieldInfo.GetValue;
 
 			private FieldInfo FieldInfo;
 
@@ -83,11 +80,11 @@ namespace UnrealBuildTool
 		/// </summary>
 		public class TargetProperty : TargetMember
 		{
-			public override MemberInfo               MemberInfo => PropertyInfo;
-			public override Type                     Type       => PropertyInfo.PropertyType;
-			public override bool                     IsStatic   => PropertyInfo.GetGetMethod()!.IsStatic;
-			public override Action<object?, object?> SetValue   => PropertyInfo.SetValue;
-			public override Func<object?, object?>   GetValue   => PropertyInfo.GetValue;
+			public override MemberInfo MemberInfo => PropertyInfo;
+			public override Type Type => PropertyInfo.PropertyType;
+			public override bool IsStatic => PropertyInfo.GetGetMethod()!.IsStatic;
+			public override Action<object?, object?> SetValue => PropertyInfo.SetValue;
+			public override Func<object?, object?> GetValue => PropertyInfo.GetValue;
 
 			private PropertyInfo PropertyInfo;
 
@@ -128,7 +125,7 @@ namespace UnrealBuildTool
 				this.XmlConfigAttribute = XmlConfigAttribute;
 			}
 		}
-		
+
 		/// <summary>
 		/// Stores a mapping from type -> member -> value, with all the config values for configurable fields.
 		/// </summary>
@@ -165,7 +162,7 @@ namespace UnrealBuildTool
 			using (BinaryReader Reader = new BinaryReader(File.Open(Location.FullName, FileMode.Open, FileAccess.Read, FileShare.Read)))
 			{
 				// Check the serialization version matches
-				if(Reader.ReadInt32() != SerializationVersion)
+				if (Reader.ReadInt32() != SerializationVersion)
 				{
 					Data = null;
 					return false;
@@ -177,14 +174,14 @@ namespace UnrealBuildTool
 				// Read the types
 				int NumTypes = Reader.ReadInt32();
 				Dictionary<Type, ValueInfo[]> TypeToValues = new Dictionary<Type, ValueInfo[]>(NumTypes);
-				for(int TypeIdx = 0; TypeIdx < NumTypes; TypeIdx++)
+				for (int TypeIdx = 0; TypeIdx < NumTypes; TypeIdx++)
 				{
 					// Read the type name
 					string TypeName = Reader.ReadString();
 
 					// Try to find it in the list of configurable types
 					Type? Type = Types.FirstOrDefault(x => x.Name == TypeName);
-					if(Type == null)
+					if (Type == null)
 					{
 						Data = null;
 						return false;
@@ -275,11 +272,11 @@ namespace UnrealBuildTool
 
 				// Write all the categories
 				Writer.Write(TypeToValues.Count);
-				foreach(KeyValuePair<Type, ValueInfo[]> TypePair in TypeToValues)
+				foreach (KeyValuePair<Type, ValueInfo[]> TypePair in TypeToValues)
 				{
 					Writer.Write(TypePair.Key.Name);
 					Writer.Write(TypePair.Value.Length);
-					foreach(ValueInfo MemberPair in TypePair.Value)
+					foreach (ValueInfo MemberPair in TypePair.Value)
 					{
 						Writer.Write(MemberPair.Target.MemberInfo.Name);
 						Writer.Write(MemberPair.Target.Type, MemberPair.Value);
diff --git a/Engine/Source/Programs/UnrealBuildTool/System/XmlConfigFile.cs b/Engine/Source/Programs/UnrealBuildTool/System/XmlConfigFile.cs
--- a/Engine/Source/Programs/UnrealBuildTool/System/XmlConfigFile.cs
+++ b/Engine/Source/Programs/UnrealBuildTool/System/XmlConfigFile.cs
@@ -1,15 +1,9 @@
 // Copyright Epic Games, Inc. All Rights Reserved.
 
-using System;
-using System.Collections.Generic;
 using System.Diagnostics.CodeAnalysis;
-using System.Linq;
-using System.Text;
-using System.Threading.Tasks;
 using System.Xml;
 using System.Xml.Schema;
 using EpicGames.Core;
-using UnrealBuildTool;
 using Microsoft.Extensions.Logging;
 
 namespace UnrealBuildTool
diff --git a/Engine/Source/Programs/UnrealBuildTool/System/XmlConfigFileAttribute.cs b/Engine/Source/Programs/UnrealBuildTool/System/XmlConfigFileAttribute.cs
--- a/Engine/Source/Programs/UnrealBuildTool/System/XmlConfigFileAttribute.cs
+++ b/Engine/Source/Programs/UnrealBuildTool/System/XmlConfigFileAttribute.cs
@@ -1,10 +1,6 @@
 // Copyright Epic Games, Inc. All Rights Reserved.
 
 using System;
-using System.Collections.Generic;
-using System.Linq;
-using System.Text;
-using System.Threading.Tasks;
 
 namespace UnrealBuildTool
 {
@@ -30,7 +26,7 @@ namespace UnrealBuildTool
 		/// A deprecated field should also be marked with the [Obsolete] attribute.
 		/// </summary>
 		public bool Deprecated = false;
-		
+
 		/// <summary>
 		/// If the attribute has been deprecated because it has been renamed, this field can be used to apply the
 		/// value used for this field to another.
diff --git a/Engine/Source/Programs/UnrealBuildTool/ToolChain/AppleToolChain.cs b/Engine/Source/Programs/UnrealBuildTool/ToolChain/AppleToolChain.cs
--- a/Engine/Source/Programs/UnrealBuildTool/ToolChain/AppleToolChain.cs
+++ b/Engine/Source/Programs/UnrealBuildTool/ToolChain/AppleToolChain.cs
@@ -1,17 +1,15 @@
 // Copyright Epic Games, Inc. All Rights Reserved.
 
 using System;
-using System.Collections;
 using System.Collections.Generic;
-using System.IO;
 using System.Diagnostics;
+using System.IO;
+using System.Linq;
 using System.Text;
 using System.Threading.Tasks;
 using EpicGames.Core;
-using System.Text.RegularExpressions;
-using UnrealBuildBase;
 using Microsoft.Extensions.Logging;
-using System.Linq;
+using UnrealBuildBase;
 
 namespace UnrealBuildTool
 {
@@ -54,7 +52,7 @@ namespace UnrealBuildTool
 				{
 					throw new BuildException($"Your Mac is set to use CommandLineTools for its build tools ({DeveloperDir}). Unreal expects Xcode as the build tools. Please install Xcode if it's not already, then do one of the following:\n" +
 						"  - Run Xcode, go to Settings, and in the Locations tab, choose your Xcode in Command Line Tools dropdown.\n" +
-						"  - In Terminal, run 'sudo xcode-select -s /Applications/Xcode.app' (or an alternate location if you installed Xcode to a non-standard location)\n" + 
+						"  - In Terminal, run 'sudo xcode-select -s /Applications/Xcode.app' (or an alternate location if you installed Xcode to a non-standard location)\n" +
 						"Either way, you will need to enter your Mac password.");
 				}
 
@@ -78,7 +76,7 @@ namespace UnrealBuildTool
 				{
 					throw new BuildException("Unable to get xcode version");
 				}
-				if (int.Parse(InstalledSdkVersion.Substring(0,2)) < 13)
+				if (int.Parse(InstalledSdkVersion.Substring(0, 2)) < 13)
 				{
 					throw new BuildException("Building for macOS, iOS and tvOS requires Xcode 13.4.1 or newer, Xcode " + InstalledSdkVersion + " detected");
 				}
@@ -463,7 +461,7 @@ namespace UnrealBuildTool
 			return CompileEnvironment;
 		}
 
-		protected string GetDsymutilPath(ILogger Logger, out string ExtraOptions, bool bIsForLTOBuild=false)
+		protected string GetDsymutilPath(ILogger Logger, out string ExtraOptions, bool bIsForLTOBuild = false)
 		{
 			FileReference DsymutilLocation = new FileReference("/usr/bin/dsymutil");
 
@@ -476,7 +474,7 @@ namespace UnrealBuildTool
 
 			// tease out the version number
 			string[] Tokens = DsymutilVersionString.Split(" ".ToCharArray());
-			
+
 			// sanity check
 			if (Tokens.Length < 4 || Tokens[3].Contains(".") == false)
 			{
diff --git a/Engine/Source/Programs/UnrealBuildTool/ToolChain/ClangToolChain.cs b/Engine/Source/Programs/UnrealBuildTool/ToolChain/ClangToolChain.cs
--- a/Engine/Source/Programs/UnrealBuildTool/ToolChain/ClangToolChain.cs
+++ b/Engine/Source/Programs/UnrealBuildTool/ToolChain/ClangToolChain.cs
@@ -1,13 +1,13 @@
 // Copyright Epic Games, Inc. All Rights Reserved.
 
-using EpicGames.Core;
 using System;
 using System.Collections.Generic;
-using System.Linq;
-using UnrealBuildBase;
-using Microsoft.Extensions.Logging;
 using System.IO;
+using System.Linq;
 using System.Text.RegularExpressions;
+using EpicGames.Core;
+using Microsoft.Extensions.Logging;
+using UnrealBuildBase;
 
 namespace UnrealBuildTool
 {
@@ -182,7 +182,7 @@ namespace UnrealBuildTool
 		protected ClangToolChainOptions Options;
 
 		// Dummy define to work around clang compilation related to the windows maximum path length limitation
-		protected static string ClangDummyDefine; 
+		protected static string ClangDummyDefine;
 		protected const int ClangCmdLineMaxSize = 32 * 1024;
 		protected const int ClangCmdlineDangerZone = 30 * 1024;
 
@@ -323,9 +323,9 @@ namespace UnrealBuildTool
 		/// </summary>
 		protected bool CompilerVersionLessThan(int Major, int Minor, int Patch) => Info.ClangVersion < new Version(Major, Minor, Patch);
 
-		protected bool IsAnalyzing(CppCompileEnvironment CompileEnvironment) => 
-				StaticAnalyzer == StaticAnalyzer.Default 
-			&&	CompileEnvironment.PrecompiledHeaderAction != PrecompiledHeaderAction.Create 
+		protected bool IsAnalyzing(CppCompileEnvironment CompileEnvironment) =>
+				StaticAnalyzer == StaticAnalyzer.Default
+			&& CompileEnvironment.PrecompiledHeaderAction != PrecompiledHeaderAction.Create
 			&& !CompileEnvironment.bDisableStaticAnalysis;
 
 		protected virtual void GetCppStandardCompileArgument(CppCompileEnvironment CompileEnvironment, List<string> Arguments)
@@ -754,7 +754,7 @@ namespace UnrealBuildTool
 			Arguments.Add("-pipe");
 
 			if (CompileEnvironment.Architecture.bIsX64)
-			{		
+			{
 				// UE5 minspec is 4.2
 				Arguments.Add("-msse4.2");
 			}
@@ -940,10 +940,10 @@ namespace UnrealBuildTool
 		}
 
 		protected virtual List<string> ExpandResponseFileContents(List<string> ResponseFileContents)
-        {
+		{
 			// This code is optimized for the scenario where ResponseFile has no variables to expand
 			List<string> NewList = ResponseFileContents;
-			for (int I=0; I!=NewList.Count; ++I)
+			for (int I = 0; I != NewList.Count; ++I)
 			{
 				string Line = ResponseFileContents[I];
 				string NewLine = Utils.ExpandVariables(Line);
