UnrealBuildTool: Automated code cleanup

diff --git a/Engine/Source/Programs/UnrealBuildTool/ToolChain/RemoteMac.cs b/Engine/Source/Programs/UnrealBuildTool/ToolChain/RemoteMac.cs
--- a/Engine/Source/Programs/UnrealBuildTool/ToolChain/RemoteMac.cs
+++ b/Engine/Source/Programs/UnrealBuildTool/ToolChain/RemoteMac.cs
@@ -1,19 +1,15 @@
 ﻿// Copyright Epic Games, Inc. All Rights Reserved.
 
 using System;
-using System.Collections;
 using System.Collections.Generic;
-using System.IO;
 using System.Diagnostics;
+using System.IO;
 using System.Linq;
-using System.Security.AccessControl;
-using System.Xml;
 using System.Text;
 using System.Text.RegularExpressions;
 using EpicGames.Core;
-using System.Security.Cryptography.X509Certificates;
-using UnrealBuildBase;
 using Microsoft.Extensions.Logging;
+using UnrealBuildBase;
 
 namespace UnrealBuildTool
 {
@@ -71,7 +67,7 @@ namespace UnrealBuildTool
 		/// <summary>
 		/// Save the specified port so that RemoteServerName is the machine address only
 		/// </summary>
-		private readonly int ServerPort = 22;	// Default ssh port
+		private readonly int ServerPort = 22;   // Default ssh port
 
 		/// <summary>
 		/// Path to Rsync
@@ -160,7 +156,7 @@ namespace UnrealBuildTool
 			ConfigHierarchy Ini = ConfigCache.ReadHierarchy(ConfigHierarchyType.Engine, EngineIniPath, UnrealTargetPlatform.IOS);
 
 			// Read the project settings if we don't have anything in the build configuration settings
-			if(String.IsNullOrEmpty(ServerName))
+			if (String.IsNullOrEmpty(ServerName))
 			{
 				// Read the server name
 				string IniServerName;
@@ -184,10 +180,10 @@ namespace UnrealBuildTool
 
 			// Split port out from the server name
 			int PortIdx = ServerName.LastIndexOf(':');
-			if(PortIdx != -1)
+			if (PortIdx != -1)
 			{
 				string Port = ServerName.Substring(PortIdx + 1);
-				if(!int.TryParse(Port, out ServerPort))
+				if (!int.TryParse(Port, out ServerPort))
 				{
 					throw new BuildException("Unable to parse port number from '{0}'", ServerName);
 				}
@@ -233,12 +229,12 @@ namespace UnrealBuildTool
 				CommandLine.AppendFormat(" \"{0}\"", Unreal.EngineDirectory);
 				CommandLine.Append("\"");
 
-				using(Process ChildProcess = Process.Start(BuildHostPlatform.Current.Shell.FullName, CommandLine.ToString()))
+				using (Process ChildProcess = Process.Start(BuildHostPlatform.Current.Shell.FullName, CommandLine.ToString()))
 				{
 					ChildProcess.WaitForExit();
 				}
 
-				if(!TryGetSshPrivateKey(out SshPrivateKey))
+				if (!TryGetSshPrivateKey(out SshPrivateKey))
 				{
 					throw new BuildException("Failed to generate SSH private key for {0}@{1}.", UserName, ServerName);
 				}
@@ -299,7 +295,7 @@ namespace UnrealBuildTool
 			// Build the list of directory mappings between the local and remote machines
 			Mappings = new List<RemoteMapping>();
 			Mappings.Add(new RemoteMapping(Unreal.EngineDirectory, GetRemotePath(Unreal.EngineDirectory)));
-			if(ProjectFile != null && !ProjectFile.IsUnderDirectory(Unreal.EngineDirectory))
+			if (ProjectFile != null && !ProjectFile.IsUnderDirectory(Unreal.EngineDirectory))
 			{
 				Mappings.Add(new RemoteMapping(ProjectFile.Directory, GetRemotePath(ProjectFile.Directory)));
 			}
@@ -345,7 +341,7 @@ namespace UnrealBuildTool
 				{
 					// MacOS Mojave includes a new version of SSH that generates keys that are incompatible with our version of SSH. Make sure the detected keys have the right signature.
 					string Text = FileReference.ReadAllText(KeyFile);
-					if(Text.Contains("---BEGIN RSA PRIVATE KEY---"))
+					if (Text.Contains("---BEGIN RSA PRIVATE KEY---"))
 					{
 						OutPrivateKey = KeyFile;
 						return true;
@@ -401,7 +397,7 @@ namespace UnrealBuildTool
 			// Translate all the arguments for the remote
 			List<string> RemoteArguments = GetRemoteArgumentsForTarget(TargetDesc, null);
 			RemoteArguments.Add("-Clean");
-		
+
 			// Upload the workspace
 			DirectoryReference TempDir = CreateTempDirectory(TargetDesc);
 			UploadWorkspace(TempDir, Logger);
@@ -410,7 +406,7 @@ namespace UnrealBuildTool
 			Logger.LogInformation("[Remote] Executing clean...");
 
 			StringBuilder BuildCommandLine = new StringBuilder("Engine/Build/BatchFiles/Mac/Build.sh");
-			foreach(string RemoteArgument in RemoteArguments)
+			foreach (string RemoteArgument in RemoteArguments)
 			{
 				BuildCommandLine.AppendFormat(" {0}", EscapeShellArgument(RemoteArgument));
 			}
@@ -480,7 +476,7 @@ namespace UnrealBuildTool
 			RemoteArguments.Add(String.Format("-SkipPreBuildTargets"));
 
 			// Handle any per-platform setup that is required
-			if(TargetDesc.Platform == UnrealTargetPlatform.IOS || TargetDesc.Platform == UnrealTargetPlatform.TVOS)
+			if (TargetDesc.Platform == UnrealTargetPlatform.IOS || TargetDesc.Platform == UnrealTargetPlatform.TVOS)
 			{
 				// Always generate a .stub
 				RemoteArguments.Add("-CreateStub");
@@ -490,14 +486,14 @@ namespace UnrealBuildTool
 
 				// Get the provisioning data for this project
 				IOSProvisioningData ProvisioningData = ((IOSPlatform)UEBuildPlatform.GetBuildPlatform(TargetDesc.Platform)).ReadProvisioningData(TargetDesc.ProjectFile, TargetDesc.AdditionalArguments.HasOption("-distribution"), IniBundleIdentifier);
-				if(ProvisioningData == null || ProvisioningData.MobileProvisionFile == null)
+				if (ProvisioningData == null || ProvisioningData.MobileProvisionFile == null)
 				{
 					throw new BuildException("Unable to find mobile provision for {0}. See log for more information.", TargetDesc.Name);
 				}
 
 				// Create a local copy of the provision
 				FileReference MobileProvisionFile = FileReference.Combine(TempDir, ProvisioningData.MobileProvisionFile.GetFileName());
-				if(FileReference.Exists(MobileProvisionFile))
+				if (FileReference.Exists(MobileProvisionFile))
 				{
 					FileReference.SetAttributes(MobileProvisionFile, FileAttributes.Normal);
 				}
@@ -511,12 +507,12 @@ namespace UnrealBuildTool
 				FileReference CertificateInfoFile = FileReference.Combine(TempDir, "Certificate.txt");
 				string CertificateInfoContents = String.Format("{0}\n{1}", ProvisioningData.MobileProvisionFile, FileReference.GetLastWriteTimeUtc(ProvisioningData.MobileProvisionFile).Ticks);
 
-				if(!FileReference.Exists(CertificateFile) || !FileReference.Exists(CertificateInfoFile) || FileReference.ReadAllText(CertificateInfoFile) != CertificateInfoContents)
+				if (!FileReference.Exists(CertificateFile) || !FileReference.Exists(CertificateInfoFile) || FileReference.ReadAllText(CertificateInfoFile) != CertificateInfoContents)
 				{
 					Logger.LogInformation("[Remote] Exporting certificate for {ProvisioningDataMobileProvisionFile}...", ProvisioningData.MobileProvisionFile);
 
 					StringBuilder Arguments = new StringBuilder("ExportCertificate");
-					if(TargetDesc.ProjectFile == null)
+					if (TargetDesc.ProjectFile == null)
 					{
 						Arguments.AppendFormat(" \"{0}\"", Unreal.EngineSourceDirectory);
 					}
@@ -526,7 +522,7 @@ namespace UnrealBuildTool
 					}
 					Arguments.AppendFormat(" -provisionfile \"{0}\"", ProvisioningData.MobileProvisionFile);
 					Arguments.AppendFormat(" -outputcertificate \"{0}\"", CertificateFile);
-					if(TargetDesc.Platform == UnrealTargetPlatform.TVOS)
+					if (TargetDesc.Platform == UnrealTargetPlatform.TVOS)
 					{
 						Arguments.Append(" -tvos");
 					}
@@ -534,7 +530,7 @@ namespace UnrealBuildTool
 					ProcessStartInfo StartInfo = new ProcessStartInfo();
 					StartInfo.FileName = FileReference.Combine(Unreal.EngineDirectory, "Binaries", "DotNET", "IOS", "IPhonePackager.exe").FullName;
 					StartInfo.Arguments = Arguments.ToString();
-					if(Utils.RunLocalProcessAndLogOutput(StartInfo, Logger) != 0)
+					if (Utils.RunLocalProcessAndLogOutput(StartInfo, Logger) != 0)
 					{
 						throw new BuildException("IphonePackager failed.");
 					}
@@ -559,15 +555,15 @@ namespace UnrealBuildTool
 			Logger.LogInformation("[Remote] Executing build");
 
 			StringBuilder BuildCommandLine = new StringBuilder("Engine/Build/BatchFiles/Mac/Build.sh");
-			foreach(string RemoteArgument in RemoteArguments)
+			foreach (string RemoteArgument in RemoteArguments)
 			{
 				BuildCommandLine.AppendFormat(" {0}", EscapeShellArgument(RemoteArgument));
 			}
 
 			int Result = Execute(GetRemotePath(Unreal.RootDirectory), BuildCommandLine.ToString(), Logger);
-			if(Result != 0)
+			if (Result != 0)
 			{
-				if(RemoteLogFile != null)
+				if (RemoteLogFile != null)
 				{
 					Logger.LogInformation("[Remote] Downloading {RemoteLogFile}", RemoteLogFile);
 					DownloadFile(RemoteLogFile, Logger);
@@ -581,7 +577,7 @@ namespace UnrealBuildTool
 
 			// Convert the manifest to local form
 			BuildManifest Manifest = Utils.ReadClass<BuildManifest>(RemoteManifestFile.FullName, Logger);
-			for(int Idx = 0; Idx < Manifest.BuildProducts.Count; Idx++)
+			for (int Idx = 0; Idx < Manifest.BuildProducts.Count; Idx++)
 			{
 				Manifest.BuildProducts[Idx] = GetLocalPath(Manifest.BuildProducts[Idx]).FullName;
 			}
@@ -597,14 +593,14 @@ namespace UnrealBuildTool
 			// Copy remote FrameworkAssets directory as it could contain resource bundles that must be packaged locally.
 			DirectoryReference BaseDir = DirectoryReference.FromFile(TargetDesc.ProjectFile) ?? Unreal.EngineDirectory;
 			DirectoryReference FrameworkAssetsDir = DirectoryReference.Combine(BaseDir, "Intermediate", TargetDesc.Platform == UnrealTargetPlatform.IOS ? "IOS" : "TVOS", "FrameworkAssets");
-			if(RemoteDirectoryExists(FrameworkAssetsDir, Logger))
+			if (RemoteDirectoryExists(FrameworkAssetsDir, Logger))
 			{
 				Logger.LogInformation("[Remote] Downloading {FrameworkAssetsDir}", FrameworkAssetsDir);
 				DownloadDirectory(FrameworkAssetsDir, Logger);
 			}
 
 			// Write out all the local manifests
-			foreach(FileReference LocalManifestFile in LocalManifestFiles)
+			foreach (FileReference LocalManifestFile in LocalManifestFiles)
 			{
 				Logger.LogInformation("[Remote] Writing {LocalManifestFile}", LocalManifestFile);
 				Utils.WriteClass<BuildManifest>(Manifest, LocalManifestFile.FullName, "", Logger);
@@ -642,7 +638,7 @@ namespace UnrealBuildTool
 			RemoteArguments.Add(String.Format("-XmlConfigCache={0}", GetRemotePath(XmlConfig.CacheFile!))); // Use the XML config cache built locally, since the remote won't have it
 
 			string? RemoteIniPath = UnrealBuildTool.GetRemoteIniPath();
-			if(!String.IsNullOrEmpty(RemoteIniPath))
+			if (!String.IsNullOrEmpty(RemoteIniPath))
 			{
 				RemoteArguments.Add(String.Format("-remoteini={0}", GetRemotePath(RemoteIniPath)));
 			}
@@ -655,7 +651,7 @@ namespace UnrealBuildTool
 			foreach (string LocalArgument in TargetDesc.AdditionalArguments)
 			{
 				int EqualsIdx = LocalArgument.IndexOf('=');
-				if(EqualsIdx == -1)
+				if (EqualsIdx == -1)
 				{
 					RemoteArguments.Add(LocalArgument);
 					continue;
@@ -664,21 +660,21 @@ namespace UnrealBuildTool
 				string Key = LocalArgument.Substring(0, EqualsIdx);
 				string Value = LocalArgument.Substring(EqualsIdx + 1);
 
-				if(Key.Equals("-Log", StringComparison.InvariantCultureIgnoreCase))
+				if (Key.Equals("-Log", StringComparison.InvariantCultureIgnoreCase))
 				{
 					// We are already writing to the local log file. The remote will produce a different log (RemoteLogFile)
 					continue;
 				}
-				if(Key.Equals("-Manifest", StringComparison.InvariantCultureIgnoreCase) && LocalManifestFiles != null)
+				if (Key.Equals("-Manifest", StringComparison.InvariantCultureIgnoreCase) && LocalManifestFiles != null)
 				{
 					LocalManifestFiles.Add(new FileReference(Value));
 					continue;
 				}
 
 				string RemoteArgument = LocalArgument;
-				foreach(RemoteMapping Mapping in Mappings)
+				foreach (RemoteMapping Mapping in Mappings)
 				{
-					if(Value.StartsWith(Mapping.LocalDirectory.FullName, StringComparison.InvariantCultureIgnoreCase))
+					if (Value.StartsWith(Mapping.LocalDirectory.FullName, StringComparison.InvariantCultureIgnoreCase))
 					{
 						RemoteArgument = String.Format("{0}={1}", Key, GetRemotePath(Value));
 						break;
@@ -709,8 +705,8 @@ namespace UnrealBuildTool
 			string RemoteOutputDir = Path.GetDirectoryName(RemoteOutputFile)!.Replace(Path.DirectorySeparatorChar, '/');
 			Execute(RemoteBaseDir, String.Format("mkdir -p {0}", EscapeShellArgument(RemoteOutputDir)), Logger);
 
-			string RemoteArguments = IOSToolChain.GetAssetCatalogArgs(Platform, RemoteInputDir, RemoteOutputDir); 
-			if(Execute(RemoteBaseDir, String.Format("/usr/bin/xcrun {0}", RemoteArguments), Logger) != 0)
+			string RemoteArguments = IOSToolChain.GetAssetCatalogArgs(Platform, RemoteInputDir, RemoteOutputDir);
+			if (Execute(RemoteBaseDir, String.Format("/usr/bin/xcrun {0}", RemoteArguments), Logger) != 0)
 			{
 				throw new BuildException("Failed to run actool.");
 			}
@@ -724,15 +720,15 @@ namespace UnrealBuildTool
 		/// <param name="ProjectFile">The project file</param>
 		/// <param name="Logger">The logger</param>
 		public void PrepareToDebug(string CookedDataDirectory, FileReference ProjectFile, ILogger Logger)
-        {
+		{
 			Logger.LogInformation("[Remote] Uploading the default uprojectdirs ...");
 			UploadFile(new FileReference("Default.uprojectdirs"), Logger);
 
 			Logger.LogInformation("[Remote] Generating an Xcode Project file...");
 			string CommandLine = EscapeShellArgument(GetRemotePath(Unreal.EngineDirectory));
 			CommandLine += "/Build/BatchFiles/Mac/GenerateProjectFiles.sh";
-			StringBuilder BuildCommandLine = new StringBuilder(CommandLine);			
-			
+			StringBuilder BuildCommandLine = new StringBuilder(CommandLine);
+
 			Execute(GetRemotePath(Unreal.RootDirectory), CommandLine, Logger);
 
 			Logger.LogInformation("[Remote] Uploading cooked data directory...");
@@ -749,11 +745,11 @@ namespace UnrealBuildTool
 		/// <param name="Logger">The logger</param>
 		/// <param name="ClientPlatform">The platfrom for the client (typically iOS or tvOS)</param>
 		public void RetrieveFilesGeneratedOnPrimaryMac(FileReference ProjectFile, ILogger Logger, string ClientPlatform)
-        {
+		{
 			Logger.LogInformation("[Remote] Downloading files generated on the Primary Remote Mac ...");
 			DownloadDirectory(DirectoryReference.Combine(Unreal.EngineDirectory, "Intermediate"), Logger);
 			DownloadDirectory(DirectoryReference.Combine(Unreal.EngineDirectory, "Saved"), Logger);
-			
+
 			DownloadDirectory(DirectoryReference.Combine(Unreal.EngineDirectory, "Binaries", ClientPlatform), Logger);
 			DownloadDirectory(DirectoryReference.Combine(ProjectFile.Directory, "Intermediate"), Logger);
 			DownloadDirectory(DirectoryReference.Combine(ProjectFile.Directory, "Binaries", ClientPlatform), Logger);
@@ -779,7 +775,7 @@ namespace UnrealBuildTool
 			Logger.LogInformation("[Remote] Generating an Xcode Project file...");
 			string CommandLine = EscapeShellArgument(GetRemotePath(Unreal.EngineDirectory));
 			CommandLine += "/Build/BatchFiles/Mac/GenerateProjectFiles.sh";
-			StringBuilder BuildCommandLine = new StringBuilder(CommandLine);	
+			StringBuilder BuildCommandLine = new StringBuilder(CommandLine);
 			Execute(GetRemotePath(Unreal.RootDirectory), CommandLine, Logger);
 
 			UploadDirectory(new DirectoryReference("UE5.xcworkspace"), Logger);
@@ -792,9 +788,9 @@ namespace UnrealBuildTool
 		/// <returns>Local filename corresponding to the remote path</returns>
 		private FileReference GetLocalPath(string RemotePath)
 		{
-			foreach(RemoteMapping Mapping in Mappings)
+			foreach (RemoteMapping Mapping in Mappings)
 			{
-				if(RemotePath.StartsWith(Mapping.RemoteDirectory, StringComparison.InvariantCultureIgnoreCase) && RemotePath.Length > Mapping.RemoteDirectory.Length && RemotePath[Mapping.RemoteDirectory.Length] == '/')
+				if (RemotePath.StartsWith(Mapping.RemoteDirectory, StringComparison.InvariantCultureIgnoreCase) && RemotePath.Length > Mapping.RemoteDirectory.Length && RemotePath[Mapping.RemoteDirectory.Length] == '/')
 				{
 					return FileReference.Combine(Mapping.LocalDirectory, RemotePath.Substring(Mapping.RemoteDirectory.Length + 1));
 				}
@@ -829,7 +825,7 @@ namespace UnrealBuildTool
 		/// <returns>Path in cygwin format</returns>
 		private static string GetLocalCygwinPath(FileSystemReference InPath)
 		{
-			if(InPath.FullName.Length < 2 || InPath.FullName[1] != ':')
+			if (InPath.FullName.Length < 2 || InPath.FullName[1] != ':')
 			{
 				throw new BuildException("Invalid local path for converting to cygwin format ({0}).", InPath);
 			}
@@ -863,7 +859,7 @@ namespace UnrealBuildTool
 			Arguments.Add("-q");
 
 			int Result = Rsync(String.Join(" ", Arguments), Logger);
-			if(Result != 0)
+			if (Result != 0)
 			{
 				throw new BuildException("Error while running Rsync (exit code {0})", Result);
 			}
@@ -886,7 +882,7 @@ namespace UnrealBuildTool
 			Arguments.Add("-q");
 
 			int Result = Rsync(String.Join(" ", Arguments), Logger);
-			if(Result != 0)
+			if (Result != 0)
 			{
 				throw new BuildException("Error while running Rsync (exit code {0})", Result);
 			}
@@ -908,7 +904,7 @@ namespace UnrealBuildTool
 			Arguments.Add("-q");
 
 			int Result = Rsync(String.Join(" ", Arguments), Logger);
-			if(Result != 0)
+			if (Result != 0)
 			{
 				throw new BuildException("Error while running Rsync (exit code {0})", Result);
 			}
@@ -925,7 +921,7 @@ namespace UnrealBuildTool
 		{
 			List<string> Arguments = new List<string>(CommonRsyncArguments);
 			Arguments.Add(String.Format("--rsync-path=\"mkdir -p {0} && rsync\"", RemoteDirectory));
-			foreach(FileReference FilterLocation in FilterLocations)
+			foreach (FileReference FilterLocation in FilterLocations)
 			{
 				Arguments.Add(String.Format("--filter=\"merge {0}\"", GetLocalCygwinPath(FilterLocation)));
 			}
@@ -934,7 +930,7 @@ namespace UnrealBuildTool
 			Arguments.Add(String.Format("\"{0}@{1}\":'{2}/'", UserName, ServerName, RemoteDirectory));
 
 			int Result = Rsync(String.Join(" ", Arguments), Logger);
-			if(Result != 0)
+			if (Result != 0)
 			{
 				throw new BuildException("Error while running Rsync (exit code {0})", Result);
 			}
@@ -950,10 +946,10 @@ namespace UnrealBuildTool
 
 			// Read the list of scripts to be uploaded
 			List<string> ScriptPaths = new List<string>();
-			foreach(string Line in FileReference.ReadAllLines(ScriptPathsFileName))
+			foreach (string Line in FileReference.ReadAllLines(ScriptPathsFileName))
 			{
 				string FileToUpload = Line.Trim();
-				if(FileToUpload.Length > 0 && FileToUpload[0] != '#')
+				if (FileToUpload.Length > 0 && FileToUpload[0] != '#')
 				{
 					ScriptPaths.Add(FileToUpload);
 				}
@@ -961,16 +957,16 @@ namespace UnrealBuildTool
 
 			// Fixup the line endings
 			List<FileReference> TargetFiles = new List<FileReference>();
-			foreach(string ScriptPath in ScriptPaths)
+			foreach (string ScriptPath in ScriptPaths)
 			{
 				FileReference SourceFile = FileReference.Combine(Unreal.EngineDirectory, ScriptPath.TrimStart('/'));
-				if(!FileReference.Exists(SourceFile))
+				if (!FileReference.Exists(SourceFile))
 				{
 					throw new BuildException("Missing script required for remote upload: {0}", SourceFile);
 				}
 
 				FileReference TargetFile = FileReference.Combine(TempDir, SourceFile.MakeRelativeTo(Unreal.EngineDirectory));
-				if(!FileReference.Exists(TargetFile) || FileReference.GetLastWriteTimeUtc(TargetFile) < FileReference.GetLastWriteTimeUtc(SourceFile))
+				if (!FileReference.Exists(TargetFile) || FileReference.GetLastWriteTimeUtc(TargetFile) < FileReference.GetLastWriteTimeUtc(SourceFile))
 				{
 					DirectoryReference.CreateDirectory(TargetFile.Directory);
 					string ScriptText = FileReference.ReadAllText(SourceFile);
@@ -981,9 +977,9 @@ namespace UnrealBuildTool
 
 			// Write a file that protects all the scripts from being overridden by the standard engine filters
 			FileReference ScriptProtectList = FileReference.Combine(TempDir, "RsyncEngineScripts-Protect.txt");
-			using(StreamWriter Writer = new StreamWriter(ScriptProtectList.FullName))
+			using (StreamWriter Writer = new StreamWriter(ScriptProtectList.FullName))
 			{
-				foreach(string ScriptPath in ScriptPaths)
+				foreach (string ScriptPath in ScriptPaths)
 				{
 					Writer.WriteLine("protect {0}", ScriptPath);
 				}
@@ -1000,13 +996,13 @@ namespace UnrealBuildTool
 			// Upload the engine files
 			List<FileReference> EngineFilters = new List<FileReference>();
 			EngineFilters.Add(ScriptProtectList);
-			if(Unreal.IsEngineInstalled())
+			if (Unreal.IsEngineInstalled())
 			{
 				EngineFilters.Add(FileReference.Combine(Unreal.EngineDirectory, "Build", "Rsync", "RsyncEngineInstalled.txt"));
 
 				// Upload MarketplaceRules.dll (if it exists) when in InstallBuild. The path is specified in RulesCompiler::CreateMarketplaceRulesAssembly()
 				DirectoryReference MarketplaceEngineDir = DirectoryReference.Combine(Unreal.WritableEngineDirectory, "Intermediate", "Build", "BuildRules");
-				if(DirectoryReference.Exists(MarketplaceEngineDir))
+				if (DirectoryReference.Exists(MarketplaceEngineDir))
 				{
 					Logger.LogInformation("[Remote] Uploading MarketplaceRules Engine files located in {MarketplaceEngineDir}", MarketplaceEngineDir);
 					UploadDirectory(MarketplaceEngineDir, Logger);
@@ -1087,7 +1083,7 @@ namespace UnrealBuildTool
 		void DownloadFile(FileReference LocalFile, ILogger Logger)
 		{
 			RemoteMapping? Mapping = Mappings.FirstOrDefault(x => LocalFile.IsUnderDirectory(x.LocalDirectory));
-			if(Mapping == null)
+			if (Mapping == null)
 			{
 				throw new BuildException("File for download '{0}' is not under any mapped directory.", LocalFile);
 			}
@@ -1098,7 +1094,7 @@ namespace UnrealBuildTool
 			Arguments.Add("-q");
 
 			int Result = Rsync(String.Join(" ", Arguments), Logger);
-			if(Result != 0)
+			if (Result != 0)
 			{
 				throw new BuildException("Unable to download '{0}' from the remote Mac (exit code {1}).", LocalFile, Result);
 			}
@@ -1112,22 +1108,22 @@ namespace UnrealBuildTool
 		void DownloadFiles(IEnumerable<FileReference> Files, ILogger Logger)
 		{
 			List<FileReference>[] FileGroups = new List<FileReference>[Mappings.Count];
-			for(int Idx = 0; Idx < Mappings.Count; Idx++)
+			for (int Idx = 0; Idx < Mappings.Count; Idx++)
 			{
 				FileGroups[Idx] = new List<FileReference>();
 			}
-			foreach(FileReference File in Files)
+			foreach (FileReference File in Files)
 			{
 				int MappingIdx = Mappings.FindIndex(x => File.IsUnderDirectory(x.LocalDirectory));
-				if(MappingIdx == -1)
+				if (MappingIdx == -1)
 				{
 					throw new BuildException("File for download '{0}' is not under the engine or project directory.", File);
 				}
 				FileGroups[MappingIdx].Add(File);
 			}
-			for(int Idx = 0; Idx < Mappings.Count; Idx++)
+			for (int Idx = 0; Idx < Mappings.Count; Idx++)
 			{
-				if(FileGroups[Idx].Count > 0)
+				if (FileGroups[Idx].Count > 0)
 				{
 					FileReference DownloadListLocation = FileReference.Combine(Unreal.EngineDirectory, "Intermediate", "Rsync", "Download.txt");
 					DirectoryReference.CreateDirectory(DownloadListLocation.Directory);
@@ -1139,7 +1135,7 @@ namespace UnrealBuildTool
 					Arguments.Add(String.Format("\"{0}/\"", GetLocalCygwinPath(Mappings[Idx].LocalDirectory)));
 
 					int Result = Rsync(String.Join(" ", Arguments), Logger);
-					if(Result != 0)
+					if (Result != 0)
 					{
 						throw new BuildException("Unable to download files from remote Mac (exit code {0})", Result);
 					}
@@ -1189,7 +1185,7 @@ namespace UnrealBuildTool
 		/// <returns>Exit code from Rsync</returns>
 		private int Rsync(string Arguments, ILogger Logger)
 		{
-			using(Process RsyncProcess = new Process())
+			using (Process RsyncProcess = new Process())
 			{
 				DataReceivedEventHandler OutputHandler = (E, Args) => { RsyncOutput(Args, false, Logger); };
 				DataReceivedEventHandler ErrorHandler = (E, Args) => { RsyncOutput(Args, true, Logger); };
@@ -1301,18 +1297,18 @@ namespace UnrealBuildTool
 		protected int ExecuteAndCaptureOutput(string Command, ILogger Logger, out StringBuilder Output)
 		{
 			StringBuilder FullCommand = new StringBuilder();
-			foreach(string CommonSshArgument in CommonSshArguments)
+			foreach (string CommonSshArgument in CommonSshArguments)
 			{
 				FullCommand.AppendFormat("{0} ", CommonSshArgument);
 			}
 			FullCommand.Append(Command.Replace("\"", "\\\""));
 
-			using(Process SSHProcess = new Process())
+			using (Process SSHProcess = new Process())
 			{
 				Output = new StringBuilder();
 
 				StringBuilder OutputLocal = Output;
-				DataReceivedEventHandler OutputHandler = (E, Args) => { if(Args.Data != null){ OutputLocal.Append(Args.Data); } };
+				DataReceivedEventHandler OutputHandler = (E, Args) => { if (Args.Data != null) { OutputLocal.Append(Args.Data); } };
 
 				SSHProcess.StartInfo.FileName = SshExe.FullName;
 				SSHProcess.StartInfo.WorkingDirectory = SshExe.Directory.FullName;
@@ -1337,7 +1333,7 @@ namespace UnrealBuildTool
 
 			// Find the matches, and early out if there are none
 			MatchCollection Matches = Regex.Matches(Text, Pattern, RegexOptions.IgnoreCase);
-			if(Matches.Count == 0)
+			if (Matches.Count == 0)
 			{
 				return Text;
 			}
@@ -1345,16 +1341,16 @@ namespace UnrealBuildTool
 			// Replace any remote paths with local ones
 			StringBuilder Result = new StringBuilder();
 			int StartIdx = 0;
-			foreach(Match? Match in Matches)
+			foreach (Match? Match in Matches)
 			{
 				// Append the text leading up to this path
 				Result.Append(Text, StartIdx, Match!.Index - StartIdx);
 
 				// Try to convert the path
 				string Path = Match.Value;
-				foreach(RemoteMapping Mapping in Mappings)
+				foreach (RemoteMapping Mapping in Mappings)
 				{
-					if(Path.StartsWith(Mapping.RemoteDirectory))
+					if (Path.StartsWith(Mapping.RemoteDirectory))
 					{
 						Path = Mapping.LocalDirectory + Path.Substring(Mapping.RemoteDirectory.Length).Replace('/', '\\');
 						break;
diff --git a/Engine/Source/Programs/UnrealBuildTool/UnrealBuildTool.cs b/Engine/Source/Programs/UnrealBuildTool/UnrealBuildTool.cs
--- a/Engine/Source/Programs/UnrealBuildTool/UnrealBuildTool.cs
+++ b/Engine/Source/Programs/UnrealBuildTool/UnrealBuildTool.cs
@@ -7,9 +7,6 @@ using System.IO;
 using System.Linq;
 using System.Reflection;
 using System.Text;
-using System.Text.Json;
-using System.Threading;
-using System.Threading.Tasks;
 using EpicGames.Core;
 using Microsoft.Extensions.Logging;
 using OpenTracing.Util;
@@ -147,7 +144,7 @@ namespace UnrealBuildTool
 		/// <returns>Location of the installed project file</returns>
 		static public FileReference? GetInstalledProjectFile()
 		{
-			if(IsProjectInstalled())
+			if (IsProjectInstalled())
 			{
 				return InstalledProjectFile;
 			}
@@ -164,11 +161,11 @@ namespace UnrealBuildTool
 		/// <returns>True if the file is part of the installed distribution, false otherwise</returns>
 		static public bool IsFileInstalled(FileReference File)
 		{
-			if(Unreal.IsEngineInstalled() && File.IsUnderDirectory(Unreal.EngineDirectory))
+			if (Unreal.IsEngineInstalled() && File.IsUnderDirectory(Unreal.EngineDirectory))
 			{
 				return true;
 			}
-			if(IsProjectInstalled() && File.IsUnderDirectory(InstalledProjectFile!.Directory))
+			if (IsProjectInstalled() && File.IsUnderDirectory(InstalledProjectFile!.Directory))
 			{
 				return true;
 			}
@@ -208,12 +205,12 @@ namespace UnrealBuildTool
 			[CommandLine(Prefix = "-h")]
 			[CommandLine(Prefix = "--help")]
 			public bool bGetHelp = false;
-			
+
 			/// <summary>
 			/// The amount of detail to write to the log
 			/// </summary>
-			[CommandLine(Prefix = "-Verbose", Value ="Verbose", Description = "Increase output verbosity")]
-			[CommandLine(Prefix = "-VeryVerbose", Value ="VeryVerbose", Description = "Increase output verbosity more")]
+			[CommandLine(Prefix = "-Verbose", Value = "Verbose", Description = "Increase output verbosity")]
+			[CommandLine(Prefix = "-VeryVerbose", Value = "VeryVerbose", Description = "Increase output verbosity more")]
 			public LogEventType LogOutputLevel = LogEventType.Log;
 
 			/// <summary>
@@ -274,35 +271,35 @@ namespace UnrealBuildTool
 			/// </summary>
 			[CommandLine("-Mode=")] // description handling is special-cased in PrintUsage()
 
-			[CommandLine("-Clean", Value="Clean", Description = "Clean build products. Equivalent to -Mode=Clean")]
-
-			[CommandLine("-ProjectFiles", Value="GenerateProjectFiles", Description = "Generate project files based on IDE preference. Equivalent to -Mode=GenerateProjectFiles")]
-			[CommandLine("-ProjectFileFormat=", Value="GenerateProjectFiles", Description = "Generate project files in specified format. May be used multiple times.")]
-			[CommandLine("-Makefile", Value="GenerateProjectFiles", Description = "Generate Linux Makefile")]
-			[CommandLine("-CMakefile", Value="GenerateProjectFiles", Description = "Generate project files for CMake")]
-			[CommandLine("-QMakefile", Value="GenerateProjectFiles", Description = "Generate project files for QMake")]
-			[CommandLine("-KDevelopfile", Value="GenerateProjectFiles", Description = "Generate project files for KDevelop")]
-			[CommandLine("-CodeliteFiles", Value="GenerateProjectFiles", Description = "Generate project files for Codelite")]
-			[CommandLine("-XCodeProjectFiles", Value="GenerateProjectFiles", Description = "Generate project files for XCode")]
-			[CommandLine("-EddieProjectFiles", Value="GenerateProjectFiles", Description = "Generate project files for Eddie")]
-			[CommandLine("-VSCode", Value="GenerateProjectFiles", Description = "Generate project files for Visual Studio Code")]
-			[CommandLine("-VSMac", Value="GenerateProjectFiles", Description = "Generate project files for Visual Studio Mac")]
-			[CommandLine("-CLion", Value="GenerateProjectFiles", Description = "Generate project files for CLion")]
-			[CommandLine("-Rider", Value="GenerateProjectFiles", Description = "Generate project files for Rider")]
-			#if __VPROJECT_AVAILABLE__
-				[CommandLine("-VProject", Value = "GenerateProjectFiles")]
-			#endif
+			[CommandLine("-Clean", Value = "Clean", Description = "Clean build products. Equivalent to -Mode=Clean")]
+
+			[CommandLine("-ProjectFiles", Value = "GenerateProjectFiles", Description = "Generate project files based on IDE preference. Equivalent to -Mode=GenerateProjectFiles")]
+			[CommandLine("-ProjectFileFormat=", Value = "GenerateProjectFiles", Description = "Generate project files in specified format. May be used multiple times.")]
+			[CommandLine("-Makefile", Value = "GenerateProjectFiles", Description = "Generate Linux Makefile")]
+			[CommandLine("-CMakefile", Value = "GenerateProjectFiles", Description = "Generate project files for CMake")]
+			[CommandLine("-QMakefile", Value = "GenerateProjectFiles", Description = "Generate project files for QMake")]
+			[CommandLine("-KDevelopfile", Value = "GenerateProjectFiles", Description = "Generate project files for KDevelop")]
+			[CommandLine("-CodeliteFiles", Value = "GenerateProjectFiles", Description = "Generate project files for Codelite")]
+			[CommandLine("-XCodeProjectFiles", Value = "GenerateProjectFiles", Description = "Generate project files for XCode")]
+			[CommandLine("-EddieProjectFiles", Value = "GenerateProjectFiles", Description = "Generate project files for Eddie")]
+			[CommandLine("-VSCode", Value = "GenerateProjectFiles", Description = "Generate project files for Visual Studio Code")]
+			[CommandLine("-VSMac", Value = "GenerateProjectFiles", Description = "Generate project files for Visual Studio Mac")]
+			[CommandLine("-CLion", Value = "GenerateProjectFiles", Description = "Generate project files for CLion")]
+			[CommandLine("-Rider", Value = "GenerateProjectFiles", Description = "Generate project files for Rider")]
+#if __VPROJECT_AVAILABLE__
+			[CommandLine("-VProject", Value = "GenerateProjectFiles")]
+#endif
 			public string? Mode = null;
 
 			// The following Log settings exists in this location because, at the time of writing, EpicGames.Core does
 			// not have access to XmlConfigFileAttribute.
-			
+
 			/// <summary>
 			/// Whether to backup an existing log file, rather than overwriting it.
 			/// </summary>
-			[XmlConfigFile(Category = "Log")] 
+			[XmlConfigFile(Category = "Log")]
 			public bool bBackupLogFiles = Log.BackupLogFiles;
-			
+
 			/// <summary>
 			/// The number of log file backups to preserve. Older backups will be deleted.
 			/// </summary>
@@ -453,16 +450,16 @@ namespace UnrealBuildTool
 
 				if (
 					// Print usage if there are zero arguments provided
-					ArgumentsArray.Length == 0 
+					ArgumentsArray.Length == 0
 
 					// Print usage if the user asks for help
-					|| Options.bGetHelp 
+					|| Options.bGetHelp
 					)
 				{
 					PrintUsage();
 					return Options.bGetHelp ? 0 : 1;
 				}
-				
+
 				// Configure the log system
 				Log.OutputLevel = Options.LogOutputLevel;
 				Log.IncludeTimestamps = Options.bLogTimestamps;
@@ -507,10 +504,10 @@ namespace UnrealBuildTool
 
 				// Get the type of the mode to execute, using a fast-path for the build mode.
 				Type? ModeType = typeof(BuildMode);
-				if(Options.Mode != null)
+				if (Options.Mode != null)
 				{
 					// Try to get the correct mode
-					if(!ModeNameToType.TryGetValue(Options.Mode, out ModeType))
+					if (!ModeNameToType.TryGetValue(Options.Mode, out ModeType))
 					{
 						List<string> ModuleNameList = ModeNameToType.Keys.ToList();
 						ModuleNameList.Sort(StringComparer.OrdinalIgnoreCase);
@@ -529,7 +526,7 @@ namespace UnrealBuildTool
 				}
 
 				// Start prefetching the contents of the engine folder
-				if((ModeOptions & ToolModeOptions.StartPrefetchingEngine) != 0)
+				if ((ModeOptions & ToolModeOptions.StartPrefetchingEngine) != 0)
 				{
 					using (GlobalTracer.Instance.BuildSpan("FileMetadataPrefetch.QueueEngineDirectory()").StartActive())
 					{
@@ -538,26 +535,26 @@ namespace UnrealBuildTool
 				}
 
 				// Read the XML configuration files
-				if((ModeOptions & ToolModeOptions.XmlConfig) != 0)
+				if ((ModeOptions & ToolModeOptions.XmlConfig) != 0)
 				{
 					using (GlobalTracer.Instance.BuildSpan("XmlConfig.ReadConfigFiles()").StartActive())
 					{
 						string XmlConfigMutexName = SingleInstanceMutex.GetUniqueMutexForPath("UnrealBuildTool_Mutex_XmlConfig", Assembly.GetExecutingAssembly().Location);
-						using(SingleInstanceMutex XmlConfigMutex = new SingleInstanceMutex(XmlConfigMutexName, true))
+						using (SingleInstanceMutex XmlConfigMutex = new SingleInstanceMutex(XmlConfigMutexName, true))
 						{
 							FileReference? XmlConfigCache = Arguments.GetFileReferenceOrDefault("-XmlConfigCache=", null);
 							XmlConfig.ReadConfigFiles(XmlConfigCache, null, Logger);
 						}
 					}
-				
+
 					XmlConfig.ApplyTo(Options);
 				}
-				
+
 				Log.BackupLogFiles = Options.bBackupLogFiles;
 				Log.LogFileBackupCount = Options.LogFileBackupCount;
 
 				// Add the log writer if requested. When building a target, we'll create the writer for the default log file later.
-				if(Options.LogFileName != null)
+				if (Options.LogFileName != null)
 				{
 					Log.AddFileWriter("LogTraceListener", Options.LogFileName);
 				}
@@ -573,7 +570,7 @@ namespace UnrealBuildTool
 						ModuleFileName = Path.GetFullPath(ModuleFileName);
 					}
 					FileReference RunFileTemp = FileReference.Combine(RunsDir, $"{Process.GetCurrentProcess().Id}_{ContentHash.MD5(Encoding.UTF8.GetBytes(ModuleFileName.ToUpperInvariant()))}");
-					File.WriteAllLines(RunFileTemp.FullName, new string [] { ModuleFileName });
+					File.WriteAllLines(RunFileTemp.FullName, new string[] { ModuleFileName });
 					RunFile = RunFileTemp;
 				}
 				catch
@@ -624,7 +621,7 @@ namespace UnrealBuildTool
 				}
 
 				// Acquire a lock for this branch
-				if((ModeOptions & ToolModeOptions.SingleInstance) != 0 && !Options.bNoMutex)
+				if ((ModeOptions & ToolModeOptions.SingleInstance) != 0 && !Options.bNoMutex)
 				{
 					using (GlobalTracer.Instance.BuildSpan("SingleInstanceMutex.Acquire()").StartActive())
 					{
@@ -634,7 +631,7 @@ namespace UnrealBuildTool
 				}
 
 				// Register all the build platforms
-				if((ModeOptions & ToolModeOptions.BuildPlatforms) != 0)
+				if ((ModeOptions & ToolModeOptions.BuildPlatforms) != 0)
 				{
 					using (GlobalTracer.Instance.BuildSpan("UEBuildPlatform.RegisterPlatforms()").StartActive())
 					{
