UnrealBuildTool: More automated code cleanup

#rnx

diff --git a/Engine/Source/Programs/UnrealBuildTool/Artifacts/ActionArtifactCache.cs b/Engine/Source/Programs/UnrealBuildTool/Artifacts/ActionArtifactCache.cs
--- a/Engine/Source/Programs/UnrealBuildTool/Artifacts/ActionArtifactCache.cs
+++ b/Engine/Source/Programs/UnrealBuildTool/Artifacts/ActionArtifactCache.cs
@@ -36,7 +36,7 @@ namespace UnrealBuildTool.Artifacts
 			switch (tree)
 			{
 				case ArtifactDirectoryTree.Absolute:
-					return string.Empty;
+					return String.Empty;
 				case ArtifactDirectoryTree.Engine:
 					if (Cache == null || Cache.EngineRoot == null)
 					{
@@ -203,7 +203,7 @@ namespace UnrealBuildTool.Artifacts
 
 			ArtifactMapping[] mappings = await ArtifactCache.QueryArtifactMappingsAsync(new IoHash[] { key }, cancellationToken);
 
-			string actionDescription = string.Empty;
+			string actionDescription = String.Empty;
 			if (LogCacheMisses)
 			{
 				actionDescription = $"{(action.CommandDescription ?? action.CommandPath.GetFileNameWithoutExtension())} {action.StatusDescription}".Trim();
diff --git a/Engine/Source/Programs/UnrealBuildTool/Configuration/ModuleRules.cs b/Engine/Source/Programs/UnrealBuildTool/Configuration/ModuleRules.cs
--- a/Engine/Source/Programs/UnrealBuildTool/Configuration/ModuleRules.cs
+++ b/Engine/Source/Programs/UnrealBuildTool/Configuration/ModuleRules.cs
@@ -643,16 +643,25 @@ namespace UnrealBuildTool
 			get
 			{
 				if (OptimizeCodeOverride.HasValue)
+				{
 					return OptimizeCodeOverride.Value;
+				}
 
 				bool? ShouldOptimizeCode = null;
 				if (Target.EnableOptimizeCodeForModules?.Contains(Name) ?? false)
+				{
 					ShouldOptimizeCode = true;
+				}
+
 				if (Target.DisableOptimizeCodeForModules?.Contains(Name) ?? false)
+				{
 					ShouldOptimizeCode = false;
+				}
 
 				if (!ShouldOptimizeCode.HasValue)
+				{
 					return CodeOptimization.Default;
+				}
 
 				return ShouldOptimizeCode.Value ? CodeOptimization.Always : CodeOptimization.Never;
 			}
@@ -964,7 +973,10 @@ namespace UnrealBuildTool
 			{
 				bool UseUnity = true;
 				if (Target.DisableUnityBuildForModules?.Contains(Name) ?? false)
+				{
 					UseUnity = false;
+				}
+
 				return bUseUnityOverride ?? UseUnity;
 			}
 		}
@@ -1035,7 +1047,7 @@ namespace UnrealBuildTool
 		/// are used, and checks that source files include their matching header first.
 		/// </summary>
 		[Obsolete("Deprecated in UE5.2 - Use IWYUSupport instead.")]
-		public bool bEnforceIWYU { set { if (!value) IWYUSupport = IWYUSupport.None; } }
+		public bool bEnforceIWYU { set { if (!value) { IWYUSupport = IWYUSupport.None; } } }
 
 		/// <summary>
 		/// Allows "include what you use" to modify the source code when run. bEnforceIWYU must be true for this variable to matter.
@@ -1662,7 +1674,7 @@ namespace UnrealBuildTool
 		/// <returns>True if it's available, false otherwise</returns>
 		protected bool IsPlatformAvailable(UnrealTargetPlatform InPlatform)
 		{
-			return UEBuildPlatform.IsPlatformAvailableForTarget(InPlatform, this.Target);
+			return UEBuildPlatform.IsPlatformAvailableForTarget(InPlatform, Target);
 		}
 
 		/// <summary>
@@ -1810,8 +1822,8 @@ namespace UnrealBuildTool
 			string TargetPlatform = Target.Platform.ToString();
 			string? Platform = null;
 			string? Architecture = null;
-			string Linkage = string.Empty;
-			string Toolset = string.Empty;
+			string Linkage = String.Empty;
+			string Toolset = String.Empty;
 			if (Target.Platform == UnrealTargetPlatform.Win64)
 			{
 				Platform = "windows";
@@ -1842,7 +1854,7 @@ namespace UnrealBuildTool
 				Platform = "osx";
 			}
 
-			if (string.IsNullOrEmpty(TargetPlatform) || string.IsNullOrEmpty(Platform) || string.IsNullOrEmpty(Architecture))
+			if (String.IsNullOrEmpty(TargetPlatform) || String.IsNullOrEmpty(Platform) || String.IsNullOrEmpty(Architecture))
 			{
 				throw new System.NotSupportedException($"Platform {Target.Platform} not currently supported by vcpkg");
 			}
@@ -1867,7 +1879,7 @@ namespace UnrealBuildTool
 				throw new DirectoryNotFoundException(VcPackageRoot);
 			}
 
-			string LibraryExtension = string.Empty;
+			string LibraryExtension = String.Empty;
 			if (Target.Platform == UnrealTargetPlatform.Win64)
 			{
 				LibraryExtension = ".lib";
diff --git a/Engine/Source/Programs/UnrealBuildTool/Configuration/ModuleRulesContext.cs b/Engine/Source/Programs/UnrealBuildTool/Configuration/ModuleRulesContext.cs
--- a/Engine/Source/Programs/UnrealBuildTool/Configuration/ModuleRulesContext.cs
+++ b/Engine/Source/Programs/UnrealBuildTool/Configuration/ModuleRulesContext.cs
@@ -56,7 +56,7 @@ namespace UnrealBuildTool
 		{
 			this.Scope = Scope;
 			this.DefaultOutputBaseDir = DefaultOutputBaseDir;
-			this.bCanUseForSharedPCH = true;
+			bCanUseForSharedPCH = true;
 		}
 
 		/// <summary>
@@ -65,14 +65,14 @@ namespace UnrealBuildTool
 		/// <param name="Other">The context to copy from</param>
 		public ModuleRulesContext(ModuleRulesContext Other)
 		{
-			this.Scope = Other.Scope;
-			this.DefaultOutputBaseDir = Other.DefaultOutputBaseDir;
-			this.Plugin = Other.Plugin;
-			this.bCanHotReload = Other.bCanHotReload;
-			this.bCanBuildDebugGame = Other.bCanBuildDebugGame;
-			this.bCanUseForSharedPCH = Other.bCanUseForSharedPCH;
-			this.bClassifyAsGameModuleForUHT = Other.bClassifyAsGameModuleForUHT;
-			this.DefaultUHTModuleType = Other.DefaultUHTModuleType;
+			Scope = Other.Scope;
+			DefaultOutputBaseDir = Other.DefaultOutputBaseDir;
+			Plugin = Other.Plugin;
+			bCanHotReload = Other.bCanHotReload;
+			bCanBuildDebugGame = Other.bCanBuildDebugGame;
+			bCanUseForSharedPCH = Other.bCanUseForSharedPCH;
+			bClassifyAsGameModuleForUHT = Other.bClassifyAsGameModuleForUHT;
+			DefaultUHTModuleType = Other.DefaultUHTModuleType;
 		}
 	}
 }
diff --git a/Engine/Source/Programs/UnrealBuildTool/Configuration/TargetDescriptor.cs b/Engine/Source/Programs/UnrealBuildTool/Configuration/TargetDescriptor.cs
--- a/Engine/Source/Programs/UnrealBuildTool/Configuration/TargetDescriptor.cs
+++ b/Engine/Source/Programs/UnrealBuildTool/Configuration/TargetDescriptor.cs
@@ -127,7 +127,7 @@ namespace UnrealBuildTool
 		public TargetDescriptor(FileReference? ProjectFile, string TargetName, UnrealTargetPlatform Platform, UnrealTargetConfiguration Configuration, UnrealArchitectures Architectures, CommandLineArguments? Arguments)
 		{
 			this.ProjectFile = ProjectFile;
-			this.Name = TargetName;
+			Name = TargetName;
 			this.Platform = Platform;
 			this.Configuration = Configuration;
 			if (Architectures == null)
@@ -150,7 +150,7 @@ namespace UnrealBuildTool
 				foreach (FileReference FileList in FileLists)
 				{
 					string[] Files = FileReference.ReadAllLines(FileList);
-					foreach (string File in Files.Where(x => !string.IsNullOrWhiteSpace(x)))
+					foreach (string File in Files.Where(x => !String.IsNullOrWhiteSpace(x)))
 					{
 						SpecificFilesToCompile.Add(FileReference.Combine(Unreal.RootDirectory, File));
 					}
@@ -190,7 +190,7 @@ namespace UnrealBuildTool
 
 		public TargetDescriptor Copy()
 		{
-			return (TargetDescriptor)this.MemberwiseClone();
+			return (TargetDescriptor)MemberwiseClone();
 		}
 
 		public static TargetDescriptor FromTargetInfo(TargetInfo Info)
diff --git a/Engine/Source/Programs/UnrealBuildTool/Configuration/TargetInfo.cs b/Engine/Source/Programs/UnrealBuildTool/Configuration/TargetInfo.cs
--- a/Engine/Source/Programs/UnrealBuildTool/Configuration/TargetInfo.cs
+++ b/Engine/Source/Programs/UnrealBuildTool/Configuration/TargetInfo.cs
@@ -79,17 +79,17 @@ namespace UnrealBuildTool
 		/// <param name="Reader">Archive to read from</param>
 		public TargetInfo(BinaryArchiveReader Reader)
 		{
-			this.Name = Reader.ReadString()!;
-			this.Platform = UnrealTargetPlatform.Parse(Reader.ReadString()!);
+			Name = Reader.ReadString()!;
+			Platform = UnrealTargetPlatform.Parse(Reader.ReadString()!);
 			string ConfigurationStr = Reader.ReadString()!;
-			this.Architectures = new UnrealArchitectures(Reader.ReadArray(() => Reader.ReadString()!)!);
-			this.ProjectFile = Reader.ReadFileReferenceOrNull();
+			Architectures = new UnrealArchitectures(Reader.ReadArray(() => Reader.ReadString()!)!);
+			ProjectFile = Reader.ReadFileReferenceOrNull();
 			string[]? ArgumentStrs = Reader.ReadArray(() => Reader.ReadString()!);
 
 			if (!UnrealTargetConfiguration.TryParse(ConfigurationStr, out Configuration))
 			{
-				throw new BuildException(string.Format("The configration name {0} is not a valid configration name. Valid names are ({1})", Name,
-					string.Join(",", Enum.GetValues(typeof(UnrealTargetConfiguration)).Cast<UnrealTargetConfiguration>().Select(x => x.ToString()))));
+				throw new BuildException(String.Format("The configration name {0} is not a valid configration name. Valid names are ({1})", Name,
+					String.Join(",", Enum.GetValues(typeof(UnrealTargetConfiguration)).Cast<UnrealTargetConfiguration>().Select(x => x.ToString()))));
 			}
 
 			Arguments = ArgumentStrs == null ? null : new CommandLineArguments(ArgumentStrs);
diff --git a/Engine/Source/Programs/UnrealBuildTool/Configuration/TargetRules.cs b/Engine/Source/Programs/UnrealBuildTool/Configuration/TargetRules.cs
--- a/Engine/Source/Programs/UnrealBuildTool/Configuration/TargetRules.cs
+++ b/Engine/Source/Programs/UnrealBuildTool/Configuration/TargetRules.cs
@@ -2380,13 +2380,13 @@ namespace UnrealBuildTool
 		/// <param name="Target">Information about the target being built</param>
 		public TargetRules(TargetInfo Target)
 		{
-			this.DefaultName = Target.Name;
-			this.Platform = Target.Platform;
-			this.Configuration = Target.Configuration;
-			this.Architectures = Target.Architectures;
-			this.ProjectFile = Target.ProjectFile;
-			this.Version = Target.Version;
-			this.WindowsPlatform = new WindowsTargetRules(this);
+			DefaultName = Target.Name;
+			Platform = Target.Platform;
+			Configuration = Target.Configuration;
+			Architectures = Target.Architectures;
+			ProjectFile = Target.ProjectFile;
+			Version = Target.Version;
+			WindowsPlatform = new WindowsTargetRules(this);
 
 			// Make sure the logger was initialized by the caller
 			if (Logger == null)
diff --git a/Engine/Source/Programs/UnrealBuildTool/Configuration/TestModuleRules.cs b/Engine/Source/Programs/UnrealBuildTool/Configuration/TestModuleRules.cs
--- a/Engine/Source/Programs/UnrealBuildTool/Configuration/TestModuleRules.cs
+++ b/Engine/Source/Programs/UnrealBuildTool/Configuration/TestModuleRules.cs
@@ -59,7 +59,7 @@ namespace UnrealBuildTool
 			this.TestedModule = TestedModule;
 
 			Name = TestedModule.Name + "Tests";
-			if (!string.IsNullOrEmpty(TestedModule.ShortName))
+			if (!String.IsNullOrEmpty(TestedModule.ShortName))
 			{
 				ShortName = TestedModule.ShortName + "Tests";
 			}
@@ -300,10 +300,10 @@ namespace UnrealBuildTool
 				XDocument XInitPlatformFile = XDocument.Load(GeneratedPropertiesPlatformFile);
 
 				// Adding per-test and per-platform tags
-				string TagsValue = TestMetadata.PlatformTags.ContainsKey(ValidPlatform) ? TestMetadata.PlatformTags[ValidPlatform] : string.Empty;
+				string TagsValue = TestMetadata.PlatformTags.ContainsKey(ValidPlatform) ? TestMetadata.PlatformTags[ValidPlatform] : String.Empty;
 				AppendOrUpdateTestFlagProperty(ref XInitPlatformFile, TestMetadata.TestName, ValidPlatform.ToString(), "Tags", TagsValue);
 
-				string ExtraCompilationArgsValue = TestMetadata.PlatformCompilationExtraArgs.ContainsKey(ValidPlatform) ? TestMetadata.PlatformCompilationExtraArgs[ValidPlatform] : string.Empty;
+				string ExtraCompilationArgsValue = TestMetadata.PlatformCompilationExtraArgs.ContainsKey(ValidPlatform) ? TestMetadata.PlatformCompilationExtraArgs[ValidPlatform] : String.Empty;
 				AppendOrUpdateTestFlagProperty(ref XInitPlatformFile, TestMetadata.TestName, ValidPlatform.ToString(), "ExtraCompilationArgs", ExtraCompilationArgsValue);
 
 				string RunSupportedValue = TestMetadata.PlatformsRunUnsupported.Contains(ValidPlatform) ? "False" : "True";
diff --git a/Engine/Source/Programs/UnrealBuildTool/Configuration/TestTargetRules.cs b/Engine/Source/Programs/UnrealBuildTool/Configuration/TestTargetRules.cs
--- a/Engine/Source/Programs/UnrealBuildTool/Configuration/TestTargetRules.cs
+++ b/Engine/Source/Programs/UnrealBuildTool/Configuration/TestTargetRules.cs
@@ -144,7 +144,7 @@ namespace UnrealBuildTool
 		{
 			SetupCommonProperties(Target);
 
-			ExeBinariesSubFolder = LaunchModuleName = Name + (ExplicitTestsTarget ? string.Empty : "Tests");
+			ExeBinariesSubFolder = LaunchModuleName = Name + (ExplicitTestsTarget ? String.Empty : "Tests");
 
 			if (ExplicitTestsTarget)
 			{
diff --git a/Engine/Source/Programs/UnrealBuildTool/Configuration/UEBuildBinary.cs b/Engine/Source/Programs/UnrealBuildTool/Configuration/UEBuildBinary.cs
--- a/Engine/Source/Programs/UnrealBuildTool/Configuration/UEBuildBinary.cs
+++ b/Engine/Source/Programs/UnrealBuildTool/Configuration/UEBuildBinary.cs
@@ -138,7 +138,7 @@ namespace UnrealBuildTool
 			)
 		{
 			this.Type = Type;
-			this.OutputDir = OutputFilePaths.First().Directory;
+			OutputDir = OutputFilePaths.First().Directory;
 			this.OutputFilePaths = new List<FileReference>(OutputFilePaths);
 			this.IntermediateDirectory = IntermediateDirectory;
 			this.bAllowExports = bAllowExports;
diff --git a/Engine/Source/Programs/UnrealBuildTool/Configuration/UEBuildFramework.cs b/Engine/Source/Programs/UnrealBuildTool/Configuration/UEBuildFramework.cs
--- a/Engine/Source/Programs/UnrealBuildTool/Configuration/UEBuildFramework.cs
+++ b/Engine/Source/Programs/UnrealBuildTool/Configuration/UEBuildFramework.cs
@@ -88,8 +88,8 @@ namespace UnrealBuildTool
 		{
 			this.Name = Name;
 			this.ZipFile = ZipFile;
-			this.ZipOutputDirectory = OutputDirectory;
-			this.FrameworkDirectory = OutputDirectory;
+			ZipOutputDirectory = OutputDirectory;
+			FrameworkDirectory = OutputDirectory;
 			this.CopyBundledAssets = CopyBundledAssets;
 			this.bCopyFramework = bCopyFramework;
 			this.bLinkFramework = bLinkFramework;
@@ -112,7 +112,7 @@ namespace UnrealBuildTool
 			this.bCopyFramework = bCopyFramework;
 			if (this.FrameworkDirectory.FullName.EndsWith(".xcframework"))
 			{
-				this.XCFrameworkVariants = LoadXCFrameworkVariants(Logger);
+				XCFrameworkVariants = LoadXCFrameworkVariants(Logger);
 			}
 		}
 
diff --git a/Engine/Source/Programs/UnrealBuildTool/Configuration/UEBuildModule.cs b/Engine/Source/Programs/UnrealBuildTool/Configuration/UEBuildModule.cs
--- a/Engine/Source/Programs/UnrealBuildTool/Configuration/UEBuildModule.cs
+++ b/Engine/Source/Programs/UnrealBuildTool/Configuration/UEBuildModule.cs
@@ -1082,9 +1082,9 @@ namespace UnrealBuildTool
 				}
 
 				// No module has confessed its guilt, so this is an error.
-				if (string.IsNullOrEmpty(GuiltyModule))
+				if (String.IsNullOrEmpty(GuiltyModule))
 				{
-					string CycleChain = string.Join(" -> ", ReferenceStack);
+					string CycleChain = String.Join(" -> ", ReferenceStack);
 					Logger.LogError("Circular dependency on {Name} detected.\n" +
 						"\tFull Route: {FullRoute}\n" +
 						"\tCycled Route: is {CycleRoot}.\n" +
@@ -1105,7 +1105,7 @@ namespace UnrealBuildTool
 				// Log dependencies if required
 				if (Logger.IsEnabled((LogLevel)LogEventType.VeryVerbose))
 				{
-					Logger.LogTrace("Module {0} dependencies:", this.Name);
+					Logger.LogTrace("Module {0} dependencies:", Name);
 					LogDependencyNameList("Public:", Rules.PublicDependencyModuleNames, Logger);
 					LogDependencyNameList("Private:", Rules.PrivateDependencyModuleNames, Logger);
 					LogDependencyNameList("Dynamic:", Rules.DynamicallyLoadedModuleNames, Logger);
@@ -1329,7 +1329,7 @@ namespace UnrealBuildTool
 		/// remain in their same relative order within the original Nodes sequence.
 		/// </summary>
 		/// <param name="NodeList">The list of nodes to sort.</param>
-		static public List<UEBuildModule> StableTopologicalSort(List<UEBuildModule> NodeList)
+		public static List<UEBuildModule> StableTopologicalSort(List<UEBuildModule> NodeList)
 		{
 			int NodeCount = NodeList.Count;
 
diff --git a/Engine/Source/Programs/UnrealBuildTool/Configuration/UEBuildModuleCPP.cs b/Engine/Source/Programs/UnrealBuildTool/Configuration/UEBuildModuleCPP.cs
--- a/Engine/Source/Programs/UnrealBuildTool/Configuration/UEBuildModuleCPP.cs
+++ b/Engine/Source/Programs/UnrealBuildTool/Configuration/UEBuildModuleCPP.cs
@@ -71,7 +71,7 @@ namespace UnrealBuildTool
 		/// <summary>
 		/// Global override to force all include paths to be always added
 		/// </summary>
-		static public bool bForceAddGeneratedCodeIncludePath;
+		public static bool bForceAddGeneratedCodeIncludePath;
 
 		/// <summary>
 		/// Paths containing *.gen.cpp files for this module.  If this is null then this module doesn't have any generated code.
@@ -197,8 +197,8 @@ namespace UnrealBuildTool
 			DirectoryReference MaybeVerseDirectory = DirectoryReference.Combine(Rules.File.Directory, "Verse");
 			if (IsValidVerseDirectory(MaybeVerseDirectory))
 			{
-				this.AssociatedVerseDirectory = MaybeVerseDirectory;
-				this.bDependsOnVerse = true;
+				AssociatedVerseDirectory = MaybeVerseDirectory;
+				bDependsOnVerse = true;
 			}
 
 			if (Rules.bValidateCircularDependencies || Rules.bTreatAsEngineModule)
@@ -454,12 +454,12 @@ namespace UnrealBuildTool
 			{
 				foreach (UnrealArch Arch in CompileEnvironment.PerArchPrecompiledHeaderFiles.Keys)
 				{
-					Logger.LogDebug("Module '{ModuleName}' uses PCH '{PCHIncludeFilename}' for Architecture '{Arch}'", this.Name, Arch, CompileEnvironment.PerArchPrecompiledHeaderFiles[Arch]);
+					Logger.LogDebug("Module '{ModuleName}' uses PCH '{PCHIncludeFilename}' for Architecture '{Arch}'", Name, Arch, CompileEnvironment.PerArchPrecompiledHeaderFiles[Arch]);
 				}
 			}
 			else if (CompileEnvironment.PrecompiledHeaderFile != null)
 			{
-				Logger.LogDebug("Module '{ModuleName}' uses PCH '{PCHIncludeFilename}'", this.Name, CompileEnvironment.PrecompiledHeaderFile);
+				Logger.LogDebug("Module '{ModuleName}' uses PCH '{PCHIncludeFilename}'", Name, CompileEnvironment.PrecompiledHeaderFile);
 			}
 
 			// Write all the definitions to a separate file
@@ -611,28 +611,28 @@ namespace UnrealBuildTool
 
 				if (Target.bForceUnityBuild)
 				{
-					Logger.LogTrace("Module '{ModuleName}' using unity build mode (bForceUnityBuild enabled for this module)", this.Name);
+					Logger.LogTrace("Module '{ModuleName}' using unity build mode (bForceUnityBuild enabled for this module)", Name);
 					bModuleUsesUnityBuild = true;
 				}
 				else if (!Rules.bUseUnity)
 				{
-					Logger.LogTrace("Module '{ModuleName}' not using unity build mode (bUseUnity disabled for this module)", this.Name);
+					Logger.LogTrace("Module '{ModuleName}' not using unity build mode (bUseUnity disabled for this module)", Name);
 					bModuleUsesUnityBuild = false;
 				}
 				else if (FileCount < MinSourceFilesForUnityBuild)
 				{
-					Logger.LogTrace("Module '{ModuleName}' not using unity build mode (module with fewer than {NumFiles} source files)", this.Name, MinSourceFilesForUnityBuild);
+					Logger.LogTrace("Module '{ModuleName}' not using unity build mode (module with fewer than {NumFiles} source files)", Name, MinSourceFilesForUnityBuild);
 					bModuleUsesUnityBuild = false;
 				}
 				else
 				{
-					Logger.LogTrace("Module '{ModuleName}' using unity build mode", this.Name);
+					Logger.LogTrace("Module '{ModuleName}' using unity build mode", Name);
 					bModuleUsesUnityBuild = true;
 				}
 			}
 			else
 			{
-				Logger.LogTrace("Module '{ModuleName}' not using unity build mode", this.Name);
+				Logger.LogTrace("Module '{ModuleName}' not using unity build mode", Name);
 			}
 
 			// Compile Generated CPP Files
@@ -773,10 +773,15 @@ namespace UnrealBuildTool
 			// Find FileItems for module's pch files
 			FileItem? PrivatePchFileItem = null;
 			if (Rules.PrivatePCHHeaderFile != null)
+			{
 				PrivatePchFileItem = FileItem.GetItemByFileReference(FileReference.Combine(ModuleDirectory, Rules.PrivatePCHHeaderFile));
+			}
+
 			FileItem? SharedPchFileItem = null;
 			if (Rules.SharedPCHHeaderFile != null)
+			{
 				SharedPchFileItem = FileItem.GetItemByFileReference(FileReference.Combine(ModuleDirectory, Rules.SharedPCHHeaderFile));
+			}
 
 			Dictionary<string, FileItem> NameToFileItem = new();
 
@@ -808,7 +813,9 @@ namespace UnrealBuildTool
 			}
 
 			if (CollidingFiles.Count != 0)
+			{
 				CompileEnvironment.CollidingNames = CollidingFiles;
+			}
 
 			return HeaderFileItems;
 		}
@@ -943,7 +950,7 @@ namespace UnrealBuildTool
 				if (Rules.Target.bChainPCHs)
 				{
 					// Create a lookup table for the definitions
-					DefinitionsDictionary = Definitions.ToHashSet().ToDictionary(x => x.Split('=')[0], x => (x.Split("=").Length >= 2) ? x.Split("=")[1] : string.Empty);
+					DefinitionsDictionary = Definitions.ToHashSet().ToDictionary(x => x.Split('=')[0], x => (x.Split("=").Length >= 2) ? x.Split("=")[1] : String.Empty);
 
 					// Find all the dependencies of this module
 					HashSet<UEBuildModule> ReferencedModules = new HashSet<UEBuildModule>();
@@ -1011,7 +1018,7 @@ namespace UnrealBuildTool
 		/// <param name="ModuleCompileEnvironment">The module compile environment</param>
 		/// <param name="CompileEnvironment">The shared PCH compile environment</param>
 		/// <returns>True if the two compile enviroments are compatible</returns>
-		static internal bool IsCompatibleForSharedPCH(CppCompileEnvironment ModuleCompileEnvironment, CppCompileEnvironment CompileEnvironment)
+		internal static bool IsCompatibleForSharedPCH(CppCompileEnvironment ModuleCompileEnvironment, CppCompileEnvironment CompileEnvironment)
 		{
 			if (ModuleCompileEnvironment.bOptimizeCode != CompileEnvironment.bOptimizeCode)
 			{
@@ -1060,7 +1067,7 @@ namespace UnrealBuildTool
 		/// <param name="CompileEnvironment">The shared PCH compile environment</param>
 		/// <param name="BaseCompileEnvironment">The base compile environment</param>
 		/// <returns>The unique suffix for the shared PCH</returns>
-		static private string GetSuffixForSharedPCH(CppCompileEnvironment CompileEnvironment, CppCompileEnvironment BaseCompileEnvironment)
+		private static string GetSuffixForSharedPCH(CppCompileEnvironment CompileEnvironment, CppCompileEnvironment BaseCompileEnvironment)
 		{
 			string Variant = "";
 			if (CompileEnvironment.bOptimizeCode != BaseCompileEnvironment.bOptimizeCode)
@@ -1192,7 +1199,7 @@ namespace UnrealBuildTool
 		/// </summary>
 		/// <param name="ModuleCompileEnvironment">The module compile environment</param>
 		/// <param name="CompileEnvironment">The shared PCH compile environment</param>
-		static private void CopySettingsForSharedPCH(CppCompileEnvironment ModuleCompileEnvironment, CppCompileEnvironment CompileEnvironment)
+		private static void CopySettingsForSharedPCH(CppCompileEnvironment ModuleCompileEnvironment, CppCompileEnvironment CompileEnvironment)
 		{
 			CompileEnvironment.bOptimizeCode = ModuleCompileEnvironment.bOptimizeCode;
 			CompileEnvironment.bUseRTTI = ModuleCompileEnvironment.bUseRTTI;
diff --git a/Engine/Source/Programs/UnrealBuildTool/Configuration/UEBuildPlatform.cs b/Engine/Source/Programs/UnrealBuildTool/Configuration/UEBuildPlatform.cs
--- a/Engine/Source/Programs/UnrealBuildTool/Configuration/UEBuildPlatform.cs
+++ b/Engine/Source/Programs/UnrealBuildTool/Configuration/UEBuildPlatform.cs
@@ -175,8 +175,8 @@ namespace UnrealBuildTool
 		/// <param name="SingleArchitecture"></param>
 		public UnrealArchitectureConfig(UnrealArch SingleArchitecture)
 		{
-			this.Mode = UnrealArchitectureMode.SingleArchitecture;
-			this.AllSupportedArchitectures = new UnrealArchitectures(SingleArchitecture);
+			Mode = UnrealArchitectureMode.SingleArchitecture;
+			AllSupportedArchitectures = new UnrealArchitectures(SingleArchitecture);
 		}
 
 		/// <summary>
@@ -187,7 +187,7 @@ namespace UnrealBuildTool
 		protected UnrealArchitectureConfig(UnrealArchitectureMode Mode, IEnumerable<UnrealArch> SupportedArchitectures)
 		{
 			this.Mode = Mode;
-			this.AllSupportedArchitectures = new UnrealArchitectures(SupportedArchitectures);
+			AllSupportedArchitectures = new UnrealArchitectures(SupportedArchitectures);
 		}
 
 
@@ -995,7 +995,7 @@ namespace UnrealBuildTool
 		/// <returns>string    The binary extension (i.e. 'exe' or 'dll')</returns>
 		public virtual string GetBinaryExtension(UEBuildBinaryType InBinaryType)
 		{
-			throw new BuildException("GetBinaryExtensiton for {0} not handled in {1}", InBinaryType.ToString(), this.ToString());
+			throw new BuildException("GetBinaryExtensiton for {0} not handled in {1}", InBinaryType.ToString(), ToString());
 		}
 
 		/// <summary>
@@ -1006,7 +1006,7 @@ namespace UnrealBuildTool
 		/// <returns>string[]    The debug info extensions (i.e. 'pdb')</returns>
 		public virtual string[] GetDebugInfoExtensions(ReadOnlyTargetRules InTarget, UEBuildBinaryType InBinaryType)
 		{
-			throw new BuildException("GetDebugInfoExtensions for {0} not handled in {1}", InBinaryType.ToString(), this.ToString());
+			throw new BuildException("GetDebugInfoExtensions for {0} not handled in {1}", InBinaryType.ToString(), ToString());
 		}
 
 		/// <summary>
@@ -1063,7 +1063,9 @@ namespace UnrealBuildTool
 			ConfigHierarchy DefaultIni = ConfigCache.ReadHierarchy(ConfigHierarchyType.Engine, (DirectoryReference?)null, Platform);
 
 			// look at all bool values
-			if (BoolKeys != null) foreach (string Key in BoolKeys)
+			if (BoolKeys != null)
+			{
+				foreach (string Key in BoolKeys)
 				{
 					bool Default = false, Project = false;
 					DefaultIni.GetBool(Section, Key, out Default);
@@ -1074,9 +1076,12 @@ namespace UnrealBuildTool
 						return false;
 					}
 				}
+			}
 
 			// look at all int values
-			if (IntKeys != null) foreach (string Key in IntKeys)
+			if (IntKeys != null)
+			{
+				foreach (string Key in IntKeys)
 				{
 					int Default = 0, Project = 0;
 					DefaultIni.GetInt32(Section, Key, out Default);
@@ -1087,9 +1092,12 @@ namespace UnrealBuildTool
 						return false;
 					}
 				}
+			}
 
 			// look for all string values
-			if (StringKeys != null) foreach (string Key in StringKeys)
+			if (StringKeys != null)
+			{
+				foreach (string Key in StringKeys)
 				{
 					string? Default = "", Project = "";
 					DefaultIni.GetString(Section, Key, out Default);
@@ -1100,6 +1108,7 @@ namespace UnrealBuildTool
 						return false;
 					}
 				}
+			}
 
 			// if we get here, we match all important settings
 			return true;
diff --git a/Engine/Source/Programs/UnrealBuildTool/Configuration/UEBuildTarget.cs b/Engine/Source/Programs/UnrealBuildTool/Configuration/UEBuildTarget.cs
--- a/Engine/Source/Programs/UnrealBuildTool/Configuration/UEBuildTarget.cs
+++ b/Engine/Source/Programs/UnrealBuildTool/Configuration/UEBuildTarget.cs
@@ -81,7 +81,7 @@ namespace UnrealBuildTool
 		/// </summary>
 		/// <param name="Name"></param>
 		/// <returns></returns>
-		static private UnrealTargetPlatform FindOrAddByName(string Name)
+		private static UnrealTargetPlatform FindOrAddByName(string Name)
 		{
 			return new UnrealTargetPlatform(GetUniqueStringRegistry().FindOrAddByName(Name));
 		}
@@ -92,7 +92,7 @@ namespace UnrealBuildTool
 		/// <param name="Alias"></param>
 		/// /// <param name="Original"></param>
 		/// <returns></returns>
-		static private UnrealTargetPlatform AddAliasByName(string Alias, UnrealTargetPlatform Original)
+		private static UnrealTargetPlatform AddAliasByName(string Alias, UnrealTargetPlatform Original)
 		{
 			GetUniqueStringRegistry().FindOrAddAlias(Alias, Original.ToString());
 			return Original;
@@ -159,7 +159,7 @@ namespace UnrealBuildTool
 		/// <param name="Name"></param>
 		/// <param name="Platform"></param>
 		/// <returns></returns>
-		static public bool TryParse(string Name, out UnrealTargetPlatform Platform)
+		public static bool TryParse(string Name, out UnrealTargetPlatform Platform)
 		{
 			if (GetUniqueStringRegistry().HasString(Name))
 			{
@@ -182,7 +182,7 @@ namespace UnrealBuildTool
 		/// </summary>
 		/// <param name="Name"></param>
 		/// <returns></returns>
-		static public UnrealTargetPlatform Parse(string Name)
+		public static UnrealTargetPlatform Parse(string Name)
 		{
 			if (GetUniqueStringRegistry().HasString(Name))
 			{
@@ -195,8 +195,8 @@ namespace UnrealBuildTool
 				return new UnrealTargetPlatform(Id);
 			}
 
-			throw new BuildException(string.Format("The platform name {0} is not a valid platform name. Valid names are ({1})", Name,
-				string.Join(",", GetUniqueStringRegistry().GetStringNames())));
+			throw new BuildException(String.Format("The platform name {0} is not a valid platform name. Valid names are ({1})", Name,
+				String.Join(",", GetUniqueStringRegistry().GetStringNames())));
 		}
 
 
@@ -289,7 +289,9 @@ namespace UnrealBuildTool
 		public override bool CanConvertFrom(ITypeDescriptorContext? context, Type sourceType)
 		{
 			if (sourceType == typeof(string))
+			{
 				return true;
+			}
 
 			return base.CanConvertFrom(context, sourceType);
 		}
@@ -297,7 +299,9 @@ namespace UnrealBuildTool
 		public override bool CanConvertTo(ITypeDescriptorContext? context, Type? destinationType)
 		{
 			if (destinationType == typeof(string))
+			{
 				return true;
+			}
 
 			return base.CanConvertTo(context, destinationType);
 		}
@@ -405,7 +409,7 @@ namespace UnrealBuildTool
 		/// </summary>
 		/// <param name="Name"></param>
 		/// <returns></returns>
-		static private UnrealPlatformGroup FindOrAddByName(string Name)
+		private static UnrealPlatformGroup FindOrAddByName(string Name)
 		{
 			return new UnrealPlatformGroup(GetUniqueStringRegistry().FindOrAddByName(Name));
 		}
@@ -629,7 +633,7 @@ namespace UnrealBuildTool
 		/// <param name="Name"></param>
 		/// <param name="bIsX64">True if X64, false if not, or null if it's Default and will ask the platform later</param>
 		/// <returns></returns>
-		static private UnrealArch FindOrAddByName(string Name, bool? bIsX64)
+		private static UnrealArch FindOrAddByName(string Name, bool? bIsX64)
 		{
 			int Id = GetUniqueStringRegistry().FindOrAddByName(Name);
 			if (bIsX64 != null)
@@ -738,7 +742,7 @@ namespace UnrealBuildTool
 		/// <param name="Name"></param>
 		/// <param name="Arch"></param>
 		/// <returns></returns>
-		static public bool TryParse(string Name, out UnrealArch Arch)
+		public static bool TryParse(string Name, out UnrealArch Arch)
 		{
 			// handle the Host string
 			if (Name.Equals("Host", StringComparison.OrdinalIgnoreCase))
@@ -770,7 +774,7 @@ namespace UnrealBuildTool
 		/// </summary>
 		/// <param name="Name"></param>
 		/// <returns></returns>
-		static public UnrealArch Parse(string Name)
+		public static UnrealArch Parse(string Name)
 		{
 			// handle the Host string
 			if (Name.Equals("Host", StringComparison.OrdinalIgnoreCase))
@@ -791,8 +795,8 @@ namespace UnrealBuildTool
 				return new UnrealArch(Id);
 			}
 
-			throw new BuildException(string.Format("The Architecture name {0} is not a valid Architecture name. Valid names are ({1})", Name,
-				string.Join(",", GetUniqueStringRegistry().GetStringNames())));
+			throw new BuildException(String.Format("The Architecture name {0} is not a valid Architecture name. Valid names are ({1})", Name,
+				String.Join(",", GetUniqueStringRegistry().GetStringNames())));
 		}
 
 
@@ -850,7 +854,9 @@ namespace UnrealBuildTool
 		public override bool CanConvertFrom(ITypeDescriptorContext? context, Type sourceType)
 		{
 			if (sourceType == typeof(string))
+			{
 				return true;
+			}
 
 			return base.CanConvertFrom(context, sourceType);
 		}
@@ -858,7 +864,9 @@ namespace UnrealBuildTool
 		public override bool CanConvertTo(ITypeDescriptorContext? context, Type? destinationType)
 		{
 			if (destinationType == typeof(string))
+			{
 				return true;
+			}
 
 			return base.CanConvertTo(context, destinationType);
 		}
@@ -914,7 +922,7 @@ namespace UnrealBuildTool
 					UnrealArch Arch;
 					if (!UnrealArch.TryParse(ArchString, out Arch) || !ArchConfig.AllSupportedArchitectures.Contains(Arch))
 					{
-						string ValidArches = string.Join("\n  ", ArchConfig.AllSupportedArchitectures.Architectures);
+						string ValidArches = String.Join("\n  ", ArchConfig.AllSupportedArchitectures.Architectures);
 						throw new BuildException($"Platform {ValidationPlatform} does not support architecture {ArchString}. Valid options are:\n  {ValidArches}");
 					}
 				}
@@ -963,7 +971,7 @@ namespace UnrealBuildTool
 		/// <returns></returns>
 		public static UnrealArchitectures? FromString(string? ArchString, UnrealTargetPlatform? ValidationPlatform)
 		{
-			if (string.IsNullOrEmpty(ArchString))
+			if (String.IsNullOrEmpty(ArchString))
 			{
 				return null;
 			}
@@ -986,7 +994,7 @@ namespace UnrealBuildTool
 			{
 				if (Architectures.Count > 1)
 				{
-					throw new BuildException($"Asking CppCompileEnvironment for a single Architecture, but it has multiple Architectures ({string.Join(", ", Architectures)}). This indicates logic needs to be updated");
+					throw new BuildException($"Asking CppCompileEnvironment for a single Architecture, but it has multiple Architectures ({String.Join(", ", Architectures)}). This indicates logic needs to be updated");
 				}
 				return Architectures[0];
 			}
@@ -1002,12 +1010,12 @@ namespace UnrealBuildTool
 		/// </summary>
 		public override string ToString()
 		{
-			return string.Join('+', Architectures);
+			return String.Join('+', Architectures);
 		}
 
 		internal string GetFolderNameForPlatform(UnrealArchitectureConfig Config)
 		{
-			return string.Join('+', Architectures.Select(x => Config.GetFolderNameForArchitecture(x)));
+			return String.Join('+', Architectures.Select(x => Config.GetFolderNameForArchitecture(x)));
 		}
 
 		/// <summary>
@@ -1345,9 +1353,9 @@ namespace UnrealBuildTool
 			bool bUsesTargetReceiptToEnablePlugins = (ThisRules.Type == TargetType.Editor && ThisRules.LinkType != TargetLinkType.Monolithic);
 			if (!bUsesTargetReceiptToEnablePlugins && (ThisRules.EnablePlugins.Count > 0 || ThisRules.DisablePlugins.Count > 0))
 			{
-				throw new BuildException(string.Format("Explicitly enabling and disabling plugins for a target is only supported when using a unique build environment (eg. for monolithic game targets). EnabledPlugins={0}, DisabledPlugins={1}",
-					string.Join(", ", ThisRules.EnablePlugins),
-					string.Join(", ", ThisRules.DisablePlugins)
+				throw new BuildException(String.Format("Explicitly enabling and disabling plugins for a target is only supported when using a unique build environment (eg. for monolithic game targets). EnabledPlugins={0}, DisabledPlugins={1}",
+					String.Join(", ", ThisRules.EnablePlugins),
+					String.Join(", ", ThisRules.DisablePlugins)
 				));
 			}
 		}
@@ -1601,7 +1609,7 @@ namespace UnrealBuildTool
 				ProjectDirectory = ProjectFile.Directory;
 				// programs are often compiled with another project passed in for context, so if it's using a _different_ uproject, use that uproject
 				// for the output location
-				if (Rules.Type != TargetType.Program || string.Compare(ProjectFile.GetFileNameWithoutAnyExtensions(), Rules.Name, true) != 0)
+				if (Rules.Type != TargetType.Program || String.Compare(ProjectFile.GetFileNameWithoutAnyExtensions(), Rules.Name, true) != 0)
 				{
 					OutputRootDirectory = ProjectDirectory;
 				}
@@ -2394,7 +2402,9 @@ namespace UnrealBuildTool
 			foreach (var Module in ModulesToGenerateHeadersFor)
 			{
 				if (Module.GeneratedCodeDirectoryUHT != null)
+				{
 					FileMetadataPrefetch.QueueDirectoryTree(Module.GeneratedCodeDirectoryUHT);
+				}
 			}
 
 			if (Rules.bUseSharedPCHs)
@@ -3072,10 +3082,10 @@ namespace UnrealBuildTool
 						IEnumerable<UEBuildModule> NonRedistDeps = AllDependencies.Where((DependantModule) =>
 							DependantModule.GetDirectDependencyModules().Intersect(NonRedistModules).Any()
 						);
-						string Message = string.Format("Non-editor build cannot depend on non-redistributable modules. {0} depends on '{1}'.", Binary.ToString(), string.Join("', '", NonRedistModules));
+						string Message = String.Format("Non-editor build cannot depend on non-redistributable modules. {0} depends on '{1}'.", Binary.ToString(), String.Join("', '", NonRedistModules));
 						if (NonRedistDeps.Any())
 						{
-							Message = string.Format("{0}\nDependant modules '{1}'", Message, string.Join("', '", NonRedistDeps));
+							Message = String.Format("{0}\nDependant modules '{1}'", Message, String.Join("', '", NonRedistDeps));
 						}
 						if (Rules.bBreakBuildOnLicenseViolation)
 						{
@@ -3738,7 +3748,7 @@ namespace UnrealBuildTool
 				ModuleRules? ModuleRules;
 				try
 				{
-					ModuleRules = RulesAssembly.CreateModuleRules(FilteredModuleName, this.Rules, "all modules option", Logger);
+					ModuleRules = RulesAssembly.CreateModuleRules(FilteredModuleName, Rules, "all modules option", Logger);
 				}
 				catch (BuildException)
 				{
diff --git a/Engine/Source/Programs/UnrealBuildTool/Executors/Experimental/FASTBuild.cs b/Engine/Source/Programs/UnrealBuildTool/Executors/Experimental/FASTBuild.cs
--- a/Engine/Source/Programs/UnrealBuildTool/Executors/Experimental/FASTBuild.cs
+++ b/Engine/Source/Programs/UnrealBuildTool/Executors/Experimental/FASTBuild.cs
@@ -76,7 +76,7 @@ namespace UnrealBuildTool
 		/// </summary>
 		ActionExecutor LocalExecutor;
 
-		public readonly static string DefaultExecutableBasePath = Path.Combine(Unreal.EngineDirectory.FullName, "Extras", "ThirdPartyNotUE", "FASTBuild");
+		public static readonly string DefaultExecutableBasePath = Path.Combine(Unreal.EngineDirectory.FullName, "Extras", "ThirdPartyNotUE", "FASTBuild");
 
 		//////////////////////////////////////////
 		// Tweakables
@@ -168,7 +168,7 @@ namespace UnrealBuildTool
 		{
 			XmlConfig.ApplyTo(this);
 
-			this.LocalExecutor = new ParallelExecutor(MaxLocalActions, bAllCores, bCompactOutput, Logger);
+			LocalExecutor = new ParallelExecutor(MaxLocalActions, bAllCores, bCompactOutput, Logger);
 		}
 
 		public override string Name => "FASTBuild";
@@ -180,10 +180,10 @@ namespace UnrealBuildTool
 
 		public static string? GetExecutablePath()
 		{
-			if (string.IsNullOrEmpty(FBuildExecutablePath))
+			if (String.IsNullOrEmpty(FBuildExecutablePath))
 			{
 				string? EnvPath = Environment.GetEnvironmentVariable("FASTBUILD_EXECUTABLE_PATH");
-				if (!string.IsNullOrEmpty(EnvPath))
+				if (!String.IsNullOrEmpty(EnvPath))
 				{
 					FBuildExecutablePath = EnvPath;
 				}
@@ -194,10 +194,10 @@ namespace UnrealBuildTool
 
 		public static string? GetCachePath()
 		{
-			if (string.IsNullOrEmpty(FBuildCachePath))
+			if (String.IsNullOrEmpty(FBuildCachePath))
 			{
 				string? EnvPath = Environment.GetEnvironmentVariable("FASTBUILD_CACHE_PATH");
-				if (!string.IsNullOrEmpty(EnvPath))
+				if (!String.IsNullOrEmpty(EnvPath))
 				{
 					FBuildCachePath = EnvPath;
 				}
@@ -208,10 +208,10 @@ namespace UnrealBuildTool
 
 		public static string? GetBrokeragePath()
 		{
-			if (string.IsNullOrEmpty(FBuildBrokeragePath))
+			if (String.IsNullOrEmpty(FBuildBrokeragePath))
 			{
 				string? EnvPath = Environment.GetEnvironmentVariable("FASTBUILD_BROKERAGE_PATH");
-				if (!string.IsNullOrEmpty(EnvPath))
+				if (!String.IsNullOrEmpty(EnvPath))
 				{
 					FBuildBrokeragePath = EnvPath;
 				}
@@ -222,10 +222,10 @@ namespace UnrealBuildTool
 
 		public static string? GetCoordinator()
 		{
-			if (string.IsNullOrEmpty(FBuildCoordinator))
+			if (String.IsNullOrEmpty(FBuildCoordinator))
 			{
 				string? EnvPath = Environment.GetEnvironmentVariable("FASTBUILD_COORDINATOR");
-				if (!string.IsNullOrEmpty(EnvPath))
+				if (!String.IsNullOrEmpty(EnvPath))
 				{
 					FBuildCoordinator = EnvPath;
 				}
@@ -239,14 +239,14 @@ namespace UnrealBuildTool
 			string? ExecutablePath = GetExecutablePath();
 			if (BuildHostPlatform.Current.Platform == UnrealTargetPlatform.Mac)
 			{
-				if (string.IsNullOrEmpty(ExecutablePath))
+				if (String.IsNullOrEmpty(ExecutablePath))
 				{
 					FBuildExecutablePath = Path.Combine(DefaultExecutableBasePath, BuildHostPlatform.Current.Platform.ToString(), "FBuild");
 				}
 			}
 			else if (BuildHostPlatform.Current.Platform == UnrealTargetPlatform.Win64)
 			{
-				if (string.IsNullOrEmpty(ExecutablePath))
+				if (String.IsNullOrEmpty(ExecutablePath))
 				{
 					FBuildExecutablePath = Path.Combine(DefaultExecutableBasePath, BuildHostPlatform.Current.Platform.ToString(), "FBuild.exe");
 				}
@@ -261,19 +261,21 @@ namespace UnrealBuildTool
 			// That's when the FASTBuild coordinator or brokerage folder is available.
 			// On Mac the latter needs the brokerage folder to be mounted, on Windows the brokerage env variable has to be set or the path specified in UBT's config
 			string? Coordinator = GetCoordinator();
-			if (string.IsNullOrEmpty(Coordinator))
+			if (String.IsNullOrEmpty(Coordinator))
 			{
 				string? BrokeragePath = GetBrokeragePath();
-				if (string.IsNullOrEmpty(BrokeragePath) || !Directory.Exists(BrokeragePath))
+				if (String.IsNullOrEmpty(BrokeragePath) || !Directory.Exists(BrokeragePath))
 				{
 					return false;
 				}
 			}
 
-			if (!string.IsNullOrEmpty(FBuildExecutablePath))
+			if (!String.IsNullOrEmpty(FBuildExecutablePath))
 			{
 				if (File.Exists(FBuildExecutablePath))
+				{
 					return true;
+				}
 
 				Logger.LogWarning("FBuildExecutablePath '{FBuildExecutablePath}' doesn't exist! Attempting to find executable in PATH.", FBuildExecutablePath);
 			}
@@ -331,13 +333,19 @@ namespace UnrealBuildTool
 		private string ActionToDependencyString(long UniqueId, string StatusDescription, string? CommandDescription = null, ActionType? ActionType = null)
 		{
 			string? ExtraInfoString = null;
-			if ((CommandDescription != null) && string.IsNullOrEmpty(CommandDescription))
+			if ((CommandDescription != null) && String.IsNullOrEmpty(CommandDescription))
+			{
 				ExtraInfoString = CommandDescription;
+			}
 			else if (ActionType != null)
+			{
 				ExtraInfoString = ActionType.Value.ToString();
+			}
 
-			if ((ExtraInfoString != null) && !string.IsNullOrEmpty(ExtraInfoString))
+			if ((ExtraInfoString != null) && !String.IsNullOrEmpty(ExtraInfoString))
+			{
 				ExtraInfoString = $" ({ExtraInfoString})";
+			}
 
 			return $"\t\t'{ActionToActionString(UniqueId)}', ;{StatusDescription}{ExtraInfoString}";
 		}
@@ -374,7 +382,7 @@ namespace UnrealBuildTool
 
 		private FBBuildType BuildType = FBBuildType.Windows;
 
-		private readonly static Tuple<string, Func<LinkedAction, string>, FBBuildType>[] BuildTypeSearchParams = new Tuple<string, Func<LinkedAction, string>, FBBuildType>[]
+		private static readonly Tuple<string, Func<LinkedAction, string>, FBBuildType>[] BuildTypeSearchParams = new Tuple<string, Func<LinkedAction, string>, FBBuildType>[]
 		{
 			Tuple.Create<string, Func<LinkedAction, string>, FBBuildType>
 			(
@@ -489,10 +497,12 @@ namespace UnrealBuildTool
 		private void AddPreBuildDependenciesText(IEnumerable<LinkedAction>? PreBuildDependencies)
 		{
 			if (PreBuildDependencies == null || !PreBuildDependencies.Any())
+			{
 				return;
+			}
 
 			AddText($"\t.PreBuildDependencies = {{\n");
-			AddText($"{string.Join("\n", PreBuildDependencies.Select(ActionToDependencyString))}\n");
+			AddText($"{String.Join("\n", PreBuildDependencies.Select(ActionToDependencyString))}\n");
 			AddText($"\t}} \n");
 		}
 
@@ -537,7 +547,9 @@ namespace UnrealBuildTool
 				for (int i = ResponseFileTokenIndex + 1; i < RawTokens.Length; ++i)
 				{
 					if (RawTokens[i - 1].Contains(".response") || RawTokens[i - 1].Contains(".rsp"))
+					{
 						break;
+					}
 
 					responseCommandline += " " + RawTokens[i];
 				}
@@ -557,12 +569,17 @@ namespace UnrealBuildTool
 
 					string[] Separators = { "\n", " ", "\r" };
 					if (File.Exists(ResponseFilePath))
+					{
 						RawTokens = ResponseFileText.Split(Separators, StringSplitOptions.RemoveEmptyEntries); //Certainly not ideal
+					}
 				}
 				catch (Exception e)
 				{
-					if (!string.IsNullOrEmpty(e.Message))
+					if (!String.IsNullOrEmpty(e.Message))
+					{
 						Logger.LogInformation("{Message}", e.Message);
+					}
+
 					Logger.LogError("Looks like a response file in: {CompilerCommandLine}, but we could not load it! {Ex}", CompilerCommandLine, e.Message);
 					ResponseFilePath = "";
 				}
@@ -578,7 +595,7 @@ namespace UnrealBuildTool
 			for (int i = 0; i < AllTokens.Count; ++i)
 			{
 				string Token = AllTokens[i];
-				if (string.IsNullOrEmpty(Token))
+				if (String.IsNullOrEmpty(Token))
 				{
 					if (ProcessedTokens.Count > 0 && QuotesOpened)
 					{
@@ -594,9 +611,13 @@ namespace UnrealBuildTool
 				for (int j = 0; j < Token.Length; ++j)
 				{
 					if (Token[j] == '\\') //Ignore escaped quotes
+					{
 						++j;
+					}
 					else if (Token[j] == '"')
+					{
 						numQuotes++;
+					}
 				}
 
 				// Defines can have escaped quotes and other strings inside them
@@ -615,7 +636,7 @@ namespace UnrealBuildTool
 						for (; i < AllTokens.Count; ++i)
 						{
 							string NextToken = AllTokens[i];
-							if (string.IsNullOrEmpty(NextToken))
+							if (String.IsNullOrEmpty(NextToken))
 							{
 								PartialToken += " ";
 							}
@@ -741,9 +762,9 @@ namespace UnrealBuildTool
 				}
 			}
 
-			ParsedCompilerOptions["OtherOptions"] = string.Join(" ", ProcessedTokens) + " ";
+			ParsedCompilerOptions["OtherOptions"] = String.Join(" ", ProcessedTokens) + " ";
 
-			if (SaveResponseFile && !string.IsNullOrEmpty(ResponseFilePath))
+			if (SaveResponseFile && !String.IsNullOrEmpty(ResponseFilePath))
 			{
 				ParsedCompilerOptions["@"] = ResponseFilePath;
 			}
@@ -753,7 +774,7 @@ namespace UnrealBuildTool
 
 		private string GetOptionValue(Dictionary<string, string> OptionsDictionary, string Key, LinkedAction Action, ILogger Logger, bool ProblemIfNotFound = false)
 		{
-			string? Value = string.Empty;
+			string? Value = String.Empty;
 			if (OptionsDictionary.TryGetValue(Key, out Value))
 			{
 				return Value.Trim(new Char[] { '\"' });
@@ -773,19 +794,27 @@ namespace UnrealBuildTool
 		{
 			object? returnValue = Microsoft.Win32.Registry.GetValue("HKEY_LOCAL_MACHINE\\SOFTWARE\\" + keyName, valueName, defaultValue);
 			if (returnValue != null)
+			{
 				return returnValue.ToString()!;
+			}
 
 			returnValue = Microsoft.Win32.Registry.GetValue("HKEY_CURRENT_USER\\SOFTWARE\\" + keyName, valueName, defaultValue);
 			if (returnValue != null)
+			{
 				return returnValue.ToString()!;
+			}
 
 			returnValue = Microsoft.Win32.Registry.GetValue("HKEY_LOCAL_MACHINE\\SOFTWARE\\Wow6432Node\\" + keyName, valueName, defaultValue);
 			if (returnValue != null)
+			{
 				return returnValue.ToString()!;
+			}
 
 			returnValue = Microsoft.Win32.Registry.GetValue("HKEY_CURRENT_USER\\SOFTWARE\\Wow6432Node\\" + keyName, valueName, defaultValue);
 			if (returnValue != null)
+			{
 				return returnValue.ToString()!;
+			}
 
 			return defaultValue.ToString()!;
 		}
@@ -806,7 +835,7 @@ namespace UnrealBuildTool
 						ToolChain: WindowsCompiler.Default,
 						Platform: UnrealTargetPlatform.Win64,
 						Architecture: UnrealArch.X64,
-						CompilerVersion: string.IsNullOrEmpty(CompilerVersion) ? null : CompilerVersion,
+						CompilerVersion: String.IsNullOrEmpty(CompilerVersion) ? null : CompilerVersion,
 						WindowsSdkVersion: null,
 						SuppliedSdkDirectoryForVersion: null,
 						bUseCPPWinRT: false,
@@ -914,18 +943,18 @@ namespace UnrealBuildTool
 				}
 				else
 				{
-					var numericDirectories = Directory.GetDirectories(VCEnv.GetToolPath().ToString()).Where(d => Path.GetFileName(d).All(char.IsDigit));
+					var numericDirectories = Directory.GetDirectories(VCEnv.GetToolPath().ToString()).Where(d => Path.GetFileName(d).All(Char.IsDigit));
 					var cluiDirectories = numericDirectories.Where(d => Directory.GetFiles(d, "clui.dll").Any());
 					if (cluiDirectories.Any())
 					{
 						cluiSubDirName = Path.GetFileName(cluiDirectories.First());
 						if (UsingCLFilter)
 						{
-							AddText(string.Format("\t\t'$CLFilterRoot$/{0}/clui.dll'\n", cluiSubDirName));
+							AddText(String.Format("\t\t'$CLFilterRoot$/{0}/clui.dll'\n", cluiSubDirName));
 						}
 						else
 						{
-							AddText(string.Format("\t\t'$Root$/{0}/clui.dll'\n", cluiSubDirName));
+							AddText(String.Format("\t\t'$Root$/{0}/clui.dll'\n", cluiSubDirName));
 						}
 						cluiDllPath = new FileReference(cluiDirectories.First() + "/clui.dll");
 					}
@@ -946,7 +975,7 @@ namespace UnrealBuildTool
 				AddText($"\t\t'$Root$/msobj{platformVersionNumber}.dll'\n");
 				AddText($"\t\t'$Root$/mspdb{platformVersionNumber}.dll'\n");
 
-				List<String> PotentialMSVCRedistPaths = new List<String>(Directory.EnumerateDirectories(string.Format("{0}/Redist/MSVC", VCEnv.GetVCInstallDirectory())));
+				List<String> PotentialMSVCRedistPaths = new List<String>(Directory.EnumerateDirectories(String.Format("{0}/Redist/MSVC", VCEnv.GetVCInstallDirectory())));
 				string? PrefferedMSVCRedistPath = null;
 				string? FinalMSVCRedistPath = "";
 
@@ -976,7 +1005,7 @@ namespace UnrealBuildTool
 
 				}
 
-				PotentialMSVCRedistPaths = new List<String>(Directory.EnumerateDirectories(string.Format("{0}/{1}", PrefferedMSVCRedistPath, VCEnv.Architecture)));
+				PotentialMSVCRedistPaths = new List<String>(Directory.EnumerateDirectories(String.Format("{0}/{1}", PrefferedMSVCRedistPath, VCEnv.Architecture)));
 
 				FinalMSVCRedistPath = PotentialMSVCRedistPaths.Find(x => x.Contains(".CRT"));
 
@@ -988,7 +1017,7 @@ namespace UnrealBuildTool
 				// if (VCEnv.Compiler == WindowsCompiler.VisualStudio2019)
 				{
 					AddText($"\t\t'$Root$/msvcp{platformVersionNumber}.dll'\n");
-					AddText(string.Format("\t\t'{0}/vccorlib{1}.dll'\n", FinalMSVCRedistPath, platformVersionNumber));
+					AddText(String.Format("\t\t'{0}/vccorlib{1}.dll'\n", FinalMSVCRedistPath, platformVersionNumber));
 					AddText($"\t\t'$Root$/tbbmalloc.dll'\n");
 
 					//AddText(string.Format("\t\t'{0}/Redist/MSVC/{1}/x64/Microsoft.VC141.CRT/vccorlib{2}.dll'\n", VCEnv.GetVCInstallDirectory(), VCEnv.ToolChainVersion, platformVersionNumber));
@@ -1016,37 +1045,49 @@ namespace UnrealBuildTool
 			if (bEnableCaching)
 			{
 				string? CachePath = GetCachePath();
-				if (!string.IsNullOrEmpty(CachePath))
+				if (!String.IsNullOrEmpty(CachePath))
+				{
 					AddText($"\t.CachePath = '{CachePath}'\n");
+				}
 			}
 
 			if (bEnableDistribution)
 			{
 				string? BrokeragePath = GetBrokeragePath();
-				if (!string.IsNullOrEmpty(BrokeragePath))
+				if (!String.IsNullOrEmpty(BrokeragePath))
+				{
 					AddText($"\t.BrokeragePath = '{BrokeragePath}'\n");
+				}
 			}
 
 			//Start Environment
 			AddText("\t.Environment = \n\t{\n");
 			if (VCEnv != null)
 			{
-				AddText(string.Format("\t\t\"PATH={0}\\Common7\\IDE\\;{1};{2}\\bin\\{3}\\x64\",\n", VCEnv.GetVCInstallDirectory(), VCEnv.GetToolPath(), VCEnv.WindowsSdkDir, VCEnv.WindowsSdkVersion));
+				AddText(String.Format("\t\t\"PATH={0}\\Common7\\IDE\\;{1};{2}\\bin\\{3}\\x64\",\n", VCEnv.GetVCInstallDirectory(), VCEnv.GetToolPath(), VCEnv.WindowsSdkDir, VCEnv.WindowsSdkVersion));
 			}
 
 			if (!IsApple())
 			{
 				if (envVars.ContainsKey("TMP"))
+				{
 					AddText($"\t\t\"TMP={envVars["TMP"]}\",\n");
+				}
 
 				if (envVars.ContainsKey("SystemRoot"))
+				{
 					AddText($"\t\t\"SystemRoot={envVars["SystemRoot"]}\",\n");
+				}
 
 				if (envVars.ContainsKey("INCLUDE"))
+				{
 					AddText($"\t\t\"INCLUDE={envVars["INCLUDE"]}\",\n");
+				}
 
 				if (envVars.ContainsKey("LIB"))
+				{
 					AddText($"\t\t\"LIB={envVars["LIB"]}\",\n");
+				}
 			}
 
 			AddText("\t}\n"); //End environment
@@ -1066,18 +1107,18 @@ namespace UnrealBuildTool
 
 			string OutputObjectFileName = GetOptionValue(ParsedCompilerOptions, IsMSVC() ? "/Fo" : "-o", Action, Logger, ProblemIfNotFound: !IsMSVC());
 
-			if (IsMSVC() && string.IsNullOrEmpty(OutputObjectFileName)) // Didn't find /Fo, try /fo
+			if (IsMSVC() && String.IsNullOrEmpty(OutputObjectFileName)) // Didn't find /Fo, try /fo
 			{
 				OutputObjectFileName = GetOptionValue(ParsedCompilerOptions, "/fo", Action, Logger, ProblemIfNotFound: true);
 			}
 
-			if (string.IsNullOrEmpty(OutputObjectFileName)) //No /Fo or /fo, we're probably in trouble.
+			if (String.IsNullOrEmpty(OutputObjectFileName)) //No /Fo or /fo, we're probably in trouble.
 			{
 				throw new Exception("We have no OutputObjectFileName. Bailing. Our Action.CommandArguments were: " + Action.CommandArguments);
 			}
 
 			string IntermediatePath = Path.GetDirectoryName(OutputObjectFileName)!;
-			if (string.IsNullOrEmpty(IntermediatePath))
+			if (String.IsNullOrEmpty(IntermediatePath))
 			{
 				throw new Exception("We have no IntermediatePath. Bailing. Our Action.CommandArguments were: " + Action.CommandArguments);
 			}
@@ -1085,7 +1126,7 @@ namespace UnrealBuildTool
 			IntermediatePath = IsApple() ? IntermediatePath.Replace("\\", "/") : IntermediatePath;
 
 			string InputFile = GetOptionValue(ParsedCompilerOptions, "InputFile", Action, Logger, ProblemIfNotFound: true);
-			if (string.IsNullOrEmpty(InputFile))
+			if (String.IsNullOrEmpty(InputFile))
 			{
 				throw new Exception("We have no InputFile. Bailing. Our Action.CommandArguments were: " + Action.CommandArguments);
 			}
@@ -1232,7 +1273,7 @@ namespace UnrealBuildTool
 
 				string JoinedActions = Actions
 					.Select(Action => ActionToDependencyString(Action))
-					.DefaultIfEmpty(string.Empty)
+					.DefaultIfEmpty(String.Empty)
 					.Aggregate((str, obj) => str + "\n" + obj);
 
 				AddText("Alias( 'all' ) \n{\n");
@@ -1303,18 +1344,18 @@ namespace UnrealBuildTool
 			FBStartInfo.RedirectStandardOutput = true;
 
 			string? Coordinator = GetCoordinator();
-			if (!string.IsNullOrEmpty(Coordinator) && !FBStartInfo.EnvironmentVariables.ContainsKey("FASTBUILD_COORDINATOR"))
+			if (!String.IsNullOrEmpty(Coordinator) && !FBStartInfo.EnvironmentVariables.ContainsKey("FASTBUILD_COORDINATOR"))
 			{
 				FBStartInfo.EnvironmentVariables.Add("FASTBUILD_COORDINATOR", Coordinator);
 			}
 			FBStartInfo.EnvironmentVariables.Remove("FASTBUILD_BROKERAGE_PATH"); // remove stale serialized value and defer to GetBrokeragePath
 			string? BrokeragePath = GetBrokeragePath();
-			if (!string.IsNullOrEmpty(BrokeragePath) && !FBStartInfo.EnvironmentVariables.ContainsKey("FASTBUILD_BROKERAGE_PATH"))
+			if (!String.IsNullOrEmpty(BrokeragePath) && !FBStartInfo.EnvironmentVariables.ContainsKey("FASTBUILD_BROKERAGE_PATH"))
 			{
 				FBStartInfo.EnvironmentVariables.Add("FASTBUILD_BROKERAGE_PATH", BrokeragePath);
 			}
 			string? CachePath = GetCachePath();
-			if (!string.IsNullOrEmpty(CachePath) && !FBStartInfo.EnvironmentVariables.ContainsKey("FASTBUILD_CACHE_PATH"))
+			if (!String.IsNullOrEmpty(CachePath) && !FBStartInfo.EnvironmentVariables.ContainsKey("FASTBUILD_CACHE_PATH"))
 			{
 				FBStartInfo.EnvironmentVariables.Add("FASTBUILD_CACHE_PATH", CachePath);
 			}
@@ -1328,7 +1369,9 @@ namespace UnrealBuildTool
 				DataReceivedEventHandler OutputEventHandler = (Sender, Args) =>
 				{
 					if (Args.Data != null)
+					{
 						Logger.LogInformation("{Output}", Args.Data);
+					}
 				};
 
 				FBProcess.OutputDataReceived += OutputEventHandler;
diff --git a/Engine/Source/Programs/UnrealBuildTool/Executors/ImmediateActionQueue.cs b/Engine/Source/Programs/UnrealBuildTool/Executors/ImmediateActionQueue.cs
--- a/Engine/Source/Programs/UnrealBuildTool/Executors/ImmediateActionQueue.cs
+++ b/Engine/Source/Programs/UnrealBuildTool/Executors/ImmediateActionQueue.cs
@@ -66,7 +66,7 @@ namespace UnrealBuildTool
 	/// <param name="UseActionWeights">If true, use the action weight as a secondary limit</param>
 	/// <param name="MaxActions">Maximum number of action actions</param>
 	/// <param name="MaxActionWeight">Maximum weight of actions</param>
-	record ImmediateActionQueueRunner(ImmediateActionQueueRunnerType Type, ActionPhase ActionPhase, Func<LinkedAction, Func<Task>?> RunAction, bool UseActionWeights = false, int MaxActions = int.MaxValue, double MaxActionWeight = int.MaxValue)
+	record ImmediateActionQueueRunner(ImmediateActionQueueRunnerType Type, ActionPhase ActionPhase, Func<LinkedAction, Func<Task>?> RunAction, bool UseActionWeights = false, int MaxActions = Int32.MaxValue, double MaxActionWeight = Int32.MaxValue)
 	{
 		/// <summary>
 		/// Current number of active actions
@@ -380,7 +380,7 @@ namespace UnrealBuildTool
 		/// <param name="maxActions">Maximum number of action actions</param>
 		/// <param name="maxActionWeight">Maximum weight of actions</param>
 		/// <returns>Created runner</returns>
-		public ImmediateActionQueueRunner CreateManualRunner(Func<LinkedAction, Func<Task>?> runAction, bool useActionWeights = false, int maxActions = int.MaxValue, double maxActionWeight = double.MaxValue)
+		public ImmediateActionQueueRunner CreateManualRunner(Func<LinkedAction, Func<Task>?> runAction, bool useActionWeights = false, int maxActions = Int32.MaxValue, double maxActionWeight = Double.MaxValue)
 		{
 			ImmediateActionQueueRunner runner = new(ImmediateActionQueueRunnerType.Manual, ActionPhase.Compile, runAction, useActionWeights, maxActions, maxActionWeight);
 			_runners.Add(runner);
@@ -498,7 +498,7 @@ namespace UnrealBuildTool
 					}
 					if (prematureDone)
 					{
-						AddCompletedActions(int.MaxValue);
+						AddCompletedActions(Int32.MaxValue);
 					}
 					return false;
 				}
@@ -523,13 +523,13 @@ namespace UnrealBuildTool
 				}
 				else
 				{
-					ExecuteResults results = new(new List<string>(), int.MaxValue, TimeSpan.Zero, TimeSpan.Zero);
+					ExecuteResults results = new(new List<string>(), Int32.MaxValue, TimeSpan.Zero, TimeSpan.Zero);
 					OnActionCompleted(action, false, results);
 				}
 			}
 			else if (ex is OperationCanceledException)
 			{
-				ExecuteResults results = new(new List<string>(), int.MaxValue, TimeSpan.Zero, TimeSpan.Zero);
+				ExecuteResults results = new(new List<string>(), Int32.MaxValue, TimeSpan.Zero, TimeSpan.Zero);
 				OnActionCompleted(action, false, results);
 			}
 			else
@@ -539,7 +539,7 @@ namespace UnrealBuildTool
 					ExceptionUtils.FormatException(ex),
 					ExceptionUtils.FormatExceptionDetails(ex),
 				};
-				ExecuteResults results = new(text, int.MaxValue, TimeSpan.Zero, TimeSpan.Zero);
+				ExecuteResults results = new(text, Int32.MaxValue, TimeSpan.Zero, TimeSpan.Zero);
 				OnActionCompleted(action, false, results);
 			}
 		}
@@ -782,7 +782,7 @@ namespace UnrealBuildTool
 			{
 				// do nothing
 			}
-			else if (count != int.MaxValue)
+			else if (count != Int32.MaxValue)
 			{
 				if (Interlocked.Add(ref _completedActions, count) == TotalActions)
 				{
@@ -841,7 +841,7 @@ namespace UnrealBuildTool
 		private void LogAction(LinkedAction action, ExecuteResults? executeTaskResult)
 		{
 			List<string>? logLines = null;
-			int exitCode = int.MaxValue;
+			int exitCode = Int32.MaxValue;
 			TimeSpan executionTime = TimeSpan.Zero;
 			TimeSpan processorTime = TimeSpan.Zero;
 			string? additionalDescription = null;
@@ -855,7 +855,7 @@ namespace UnrealBuildTool
 			}
 
 			// Write it to the log
-			string description = string.Empty;
+			string description = String.Empty;
 			if (action.bShouldOutputStatusDescription || (logLines != null && logLines.Count == 0))
 			{
 				description = $"{(action.CommandDescription ?? action.CommandPath.GetFileNameWithoutExtension())} {action.StatusDescription}".Trim();
@@ -864,7 +864,7 @@ namespace UnrealBuildTool
 			{
 				description = $"{(action.CommandDescription ?? action.CommandPath.GetFileNameWithoutExtension())} {logLines[0]}".Trim();
 			}
-			if (!string.IsNullOrEmpty(additionalDescription))
+			if (!String.IsNullOrEmpty(additionalDescription))
 			{
 				description = $"{description} {additionalDescription}";
 			}
@@ -876,7 +876,7 @@ namespace UnrealBuildTool
 				ProgressWriter.Write(completedActions, Actions.Length);
 
 				// Canceled
-				if (exitCode == int.MaxValue)
+				if (exitCode == Int32.MaxValue)
 				{
 					Logger.LogInformation("[{CompletedActions}/{TotalActions}] {Description} canceled", completedActions, totalActions, description);
 					return;
diff --git a/Engine/Source/Programs/UnrealBuildTool/Executors/SNDBS.cs b/Engine/Source/Programs/UnrealBuildTool/Executors/SNDBS.cs
--- a/Engine/Source/Programs/UnrealBuildTool/Executors/SNDBS.cs
+++ b/Engine/Source/Programs/UnrealBuildTool/Executors/SNDBS.cs
@@ -30,7 +30,7 @@ namespace UnrealBuildTool
 
 		private static string FindDbsExe(string ExeName)
 		{
-			string InstallPath = Path.Combine(ProgramFilesx86 ?? string.Empty, "SCE", "Common", "SN-DBS", "bin", ExeName);
+			string InstallPath = Path.Combine(ProgramFilesx86 ?? String.Empty, "SCE", "Common", "SN-DBS", "bin", ExeName);
 			if (File.Exists(InstallPath))
 			{
 				return InstallPath;
@@ -38,7 +38,7 @@ namespace UnrealBuildTool
 			else
 			{
 				// Legacy install location using SCE_ROOT_DIR
-				return Path.Combine(SCERoot ?? string.Empty, "Common", "SN-DBS", "bin", ExeName);
+				return Path.Combine(SCERoot ?? String.Empty, "Common", "SN-DBS", "bin", ExeName);
 			}
 		}
 
@@ -267,7 +267,7 @@ namespace UnrealBuildTool
 						}).ToList();
 					}
 
-					string CommandDescription = string.IsNullOrWhiteSpace(a.CommandDescription) ? a.ActionType.ToString() : a.CommandDescription;
+					string CommandDescription = String.IsNullOrWhiteSpace(a.CommandDescription) ? a.ActionType.ToString() : a.CommandDescription;
 					Job["echo"] = $"{ProgressMarkupPrefix}{CommandDescription}:{a.StatusDescription}";
 
 					return Job;
@@ -278,7 +278,7 @@ namespace UnrealBuildTool
 			bool bHasRewrites = GenerateSNDBSIncludeRewriteRules();
 
 			var ConfigList = TargetDescriptors.Select(Descriptor => $"{Descriptor.Name}|{Descriptor.Platform}|{Descriptor.Configuration}");
-			var ConfigDescription = string.Join(",", ConfigList);
+			var ConfigDescription = String.Join(",", ConfigList);
 
 			var StartInfo = new ProcessStartInfo(
 				SNDBSBuildExe,
@@ -400,19 +400,19 @@ namespace UnrealBuildTool
 				var Lines = new[]
 				{
 					@"pattern1=^COMPILED_PLATFORM_HEADER\(\s*([^ ,]+)\s*\)",
-					$"expansions1={string.Join("|", Platforms.Select(Name => $"{Name}/{Name}$1|{Name}$1"))}",
+					$"expansions1={String.Join("|", Platforms.Select(Name => $"{Name}/{Name}$1|{Name}$1"))}",
 
 					@"pattern2=^COMPILED_PLATFORM_HEADER_WITH_PREFIX\(\s*([^ ,]+)\s*,\s*([^ ,]+)\s*\)",
-					$"expansions2={string.Join("|", Platforms.Select(Name => $"$1/{Name}/{Name}$2|$1/{Name}$2"))}",
+					$"expansions2={String.Join("|", Platforms.Select(Name => $"$1/{Name}/{Name}$2|$1/{Name}$2"))}",
 
 					@"pattern3=^[A-Z]{5}_PLATFORM_HEADER_NAME\(\s*([^ ,]+)\s*\)",
-					$"expansions3={string.Join("|", Platforms.Select(Name => $"{Name}/{Name}$1|{Name}$1"))}",
+					$"expansions3={String.Join("|", Platforms.Select(Name => $"{Name}/{Name}$1|{Name}$1"))}",
 
 					@"pattern4=^[A-Z]{5}_PLATFORM_HEADER_NAME_WITH_PREFIX\(\s*([^ ,]+)\s*,\s*([^ ,]+)\s*\)",
-					$"expansions4={string.Join("|", Platforms.Select(Name => $"$1/{Name}/{Name}$2|$1/{Name}$2"))}"
+					$"expansions4={String.Join("|", Platforms.Select(Name => $"$1/{Name}/{Name}$2|$1/{Name}$2"))}"
 				};
 
-				File.WriteAllText(IncludeRewriteRulesFile.FullName, string.Join(Environment.NewLine, new[] { "[computed-include-rules]" }.Concat(Lines)));
+				File.WriteAllText(IncludeRewriteRulesFile.FullName, String.Join(Environment.NewLine, new[] { "[computed-include-rules]" }.Concat(Lines)));
 				return true;
 			}
 			else
diff --git a/Engine/Source/Programs/UnrealBuildTool/Executors/XGE.cs b/Engine/Source/Programs/UnrealBuildTool/Executors/XGE.cs
--- a/Engine/Source/Programs/UnrealBuildTool/Executors/XGE.cs
+++ b/Engine/Source/Programs/UnrealBuildTool/Executors/XGE.cs
@@ -69,7 +69,7 @@ namespace UnrealBuildTool
 		/// Minimum number of actions to use XGE execution.
 		/// </summary>
 		[XmlConfigFile(Category = "XGE")]
-		static public int MinActions = 2;
+		public static int MinActions = 2;
 
 		/// <summary>
 		/// Check for a concurrent XGE build and treat the XGE executor as unavailable if it's in use.
@@ -489,7 +489,7 @@ namespace UnrealBuildTool
 				// <Tool ... />
 				XmlElement ToolElement = XGETaskDocument.CreateElement("Tool");
 				ToolsElement.AppendChild(ToolElement);
-				ToolElement.SetAttribute("Name", string.Format("Tool{0}", ActionIndex));
+				ToolElement.SetAttribute("Name", String.Format("Tool{0}", ActionIndex));
 				ToolElement.SetAttribute("AllowRemote", CanExecuteRemotely.ToString());
 
 				// The XGE documentation says that 'AllowIntercept' must be set to 'true' for all tools where 'AllowRemote' is enabled
@@ -534,7 +534,7 @@ namespace UnrealBuildTool
 				}
 				ToolElement.SetAttribute(
 					"OutputFileMasks",
-					string.Join(
+					String.Join(
 						",",
 						Action.ProducedItems.Select(
 							delegate (FileItem ProducedItem) { return ProducedItem.Location.GetFileName(); }
@@ -544,11 +544,11 @@ namespace UnrealBuildTool
 
 				if (Action.ActionType == ActionType.Compile)
 				{
-					ToolElement.SetAttribute("AutoRecover", string.Join(',', CompileAutoRecover));
+					ToolElement.SetAttribute("AutoRecover", String.Join(',', CompileAutoRecover));
 				}
 				else if (Action.ActionType == ActionType.Link)
 				{
-					ToolElement.SetAttribute("AutoRecover", string.Join(',', LinkAutoRecover));
+					ToolElement.SetAttribute("AutoRecover", String.Join(',', LinkAutoRecover));
 				}
 			}
 
@@ -567,8 +567,8 @@ namespace UnrealBuildTool
 				ProjectElement.AppendChild(TaskElement);
 				TaskElement.SetAttribute("SourceFile", "");
 				TaskElement.SetAttribute("Caption", Action.StatusDescription);
-				TaskElement.SetAttribute("Name", string.Format("Action{0}", ActionIndex));
-				TaskElement.SetAttribute("Tool", string.Format("Tool{0}", ActionIndex));
+				TaskElement.SetAttribute("Name", String.Format("Action{0}", ActionIndex));
+				TaskElement.SetAttribute("Tool", String.Format("Tool{0}", ActionIndex));
 				TaskElement.SetAttribute("WorkingDir", Action.WorkingDirectory.FullName);
 				TaskElement.SetAttribute("SkipIfProjectFailed", "true");
 				TaskElement.SetAttribute("AllowRestartOnLocal", "true");
@@ -579,13 +579,13 @@ namespace UnrealBuildTool
 				{
 					if (Actions.Contains(PrerequisiteAction))
 					{
-						DependencyNames.Add(string.Format("Action{0}", Actions.IndexOf(PrerequisiteAction)));
+						DependencyNames.Add(String.Format("Action{0}", Actions.IndexOf(PrerequisiteAction)));
 					}
 				}
 
 				if (DependencyNames.Count > 0)
 				{
-					TaskElement.SetAttribute("DependsOn", string.Join(";", DependencyNames.ToArray()));
+					TaskElement.SetAttribute("DependsOn", String.Join(";", DependencyNames.ToArray()));
 				}
 			}
 
@@ -648,7 +648,7 @@ namespace UnrealBuildTool
 
 			ProcessStartInfo XGEStartInfo = new ProcessStartInfo(
 				XgConsolePath,
-				string.Format("\"{0}\" /Rebuild /NoWait {1} /NoLogo {2} /ShowAgent /ShowTime {3}",
+				String.Format("\"{0}\" /Rebuild /NoWait {1} /NoLogo {2} /ShowAgent /ShowTime {3}",
 					TaskFilePath,
 					bStopXGECompilationAfterErrors ? "/StopOnErrors" : "",
 					SilentOption,
@@ -711,8 +711,8 @@ namespace UnrealBuildTool
 			using (ProgressWriter Writer = new ProgressWriter("Compiling C++ source files...", false, Logger))
 			{
 				int NumCompletedActions = 0;
-				string ProgressText = string.Empty;
-				string CommandDescription = string.Empty;
+				string ProgressText = String.Empty;
+				string CommandDescription = String.Empty;
 				HashSet<int> ReportedActionIndices = new();
 
 				// Create a wrapper delegate that will parse the output actions
@@ -730,7 +730,7 @@ namespace UnrealBuildTool
 							if (EndOfMarkupPrefix != -1)
 							{
 								MarkupLength = EndOfMarkupPrefix + 1;
-								if (int.TryParse(Text.Substring(ProgressMarkupPrefix.Length + 1, EndOfMarkupPrefix - ProgressMarkupPrefix.Length - 1), out ActionIndex))
+								if (Int32.TryParse(Text.Substring(ProgressMarkupPrefix.Length + 1, EndOfMarkupPrefix - ProgressMarkupPrefix.Length - 1), out ActionIndex))
 								{
 									// We keep track of the actions that we have already reported so NumCompletedActions match up with NumActions
 									if (ReportedActionIndices.Add(ActionIndex))
@@ -740,13 +740,13 @@ namespace UnrealBuildTool
 								}
 							}
 							// Flush old progress text
-							if (!string.IsNullOrEmpty(ProgressText))
+							if (!String.IsNullOrEmpty(ProgressText))
 							{
 								Logger.LogInformation("[{NumCompletedActions}/{NumActions}] Complete {ProgressText}", NumCompletedActions, NumActions, ProgressText);
-								ProgressText = string.Empty;
+								ProgressText = String.Empty;
 							}
 
-							CommandDescription = ActionIndex != -1 ? Actions[ActionIndex].CommandDescription + " " : string.Empty;
+							CommandDescription = ActionIndex != -1 ? Actions[ActionIndex].CommandDescription + " " : String.Empty;
 
 							// Strip out anything that is just an XGE timer. Some programs don't output anything except the progress text.
 							Text = Args.Data.Substring(MarkupLength);
@@ -760,11 +760,11 @@ namespace UnrealBuildTool
 							Logger.LogInformation("[{NumCompletedActions}/{NumActions}] {CommandDescription}{Text}", NumCompletedActions, NumActions, CommandDescription, Text);
 							return;
 						}
-						if (!string.IsNullOrEmpty(ProgressText))
+						if (!String.IsNullOrEmpty(ProgressText))
 						{
 							Logger.LogInformation("[{NumCompletedActions}/{NumActions}] {CommandDescription}{Text} {ProgressText}", NumCompletedActions, NumActions, CommandDescription, Text, ProgressText);
-							ProgressText = string.Empty;
-							CommandDescription = string.Empty;
+							ProgressText = String.Empty;
+							CommandDescription = String.Empty;
 							return;
 						}
 						WriteToolOutput(Text);
diff --git a/Engine/Source/Programs/UnrealBuildTool/Modes/AggregateClangTimingInfo.cs b/Engine/Source/Programs/UnrealBuildTool/Modes/AggregateClangTimingInfo.cs
--- a/Engine/Source/Programs/UnrealBuildTool/Modes/AggregateClangTimingInfo.cs
+++ b/Engine/Source/Programs/UnrealBuildTool/Modes/AggregateClangTimingInfo.cs
@@ -71,26 +71,26 @@ namespace UnrealBuildTool
 			{
 				SourceFile = inputFile;
 
-				TotalExecuteCompiler = trace?.traceEvents?.FindLast(x => string.Equals(x.name, "Total ExecuteCompiler"))?.dur ?? 0;
+				TotalExecuteCompiler = trace?.traceEvents?.FindLast(x => String.Equals(x.name, "Total ExecuteCompiler"))?.dur ?? 0;
 
 				// Subset of execute compiler
-				TotalFrontend = trace?.traceEvents?.FindLast(x => string.Equals(x.name, "Total Frontend"))?.dur ?? 0;
-				TotalBackend = trace?.traceEvents?.FindLast(x => string.Equals(x.name, "Total Backend"))?.dur ?? 0;
+				TotalFrontend = trace?.traceEvents?.FindLast(x => String.Equals(x.name, "Total Frontend"))?.dur ?? 0;
+				TotalBackend = trace?.traceEvents?.FindLast(x => String.Equals(x.name, "Total Backend"))?.dur ?? 0;
 
 				// Subset of frontend
-				TotalSource = trace?.traceEvents?.FindLast(x => string.Equals(x.name, "Total Source"))?.dur ?? 0;
-				TotalInstantiateFunction = trace?.traceEvents?.FindLast(x => string.Equals(x.name, "Total InstantiateFunction"))?.dur ?? 0;
-				TotalCodeGenFunction = trace?.traceEvents?.FindLast(x => string.Equals(x.name, "Total CodeGen Function"))?.dur ?? 0;
+				TotalSource = trace?.traceEvents?.FindLast(x => String.Equals(x.name, "Total Source"))?.dur ?? 0;
+				TotalInstantiateFunction = trace?.traceEvents?.FindLast(x => String.Equals(x.name, "Total InstantiateFunction"))?.dur ?? 0;
+				TotalCodeGenFunction = trace?.traceEvents?.FindLast(x => String.Equals(x.name, "Total CodeGen Function"))?.dur ?? 0;
 
 				// Subset of backend
-				TotalModuleToFunctionPassAdaptor = trace?.traceEvents?.FindLast(x => string.Equals(x.name, "Total ModuleToFunctionPassAdaptor"))?.dur ?? 0;
-				TotalModuleInlinerWrapperPass = trace?.traceEvents?.FindLast(x => string.Equals(x.name, "Total ModuleInlinerWrapperPass"))?.dur ?? 0;
-				TotalOptModule = trace?.traceEvents?.FindLast(x => string.Equals(x.name, "Total OptModule"))?.dur ?? 0;
+				TotalModuleToFunctionPassAdaptor = trace?.traceEvents?.FindLast(x => String.Equals(x.name, "Total ModuleToFunctionPassAdaptor"))?.dur ?? 0;
+				TotalModuleInlinerWrapperPass = trace?.traceEvents?.FindLast(x => String.Equals(x.name, "Total ModuleInlinerWrapperPass"))?.dur ?? 0;
+				TotalOptModule = trace?.traceEvents?.FindLast(x => String.Equals(x.name, "Total OptModule"))?.dur ?? 0;
 
 				// Frontend entry counts
-				SourceEntries = trace?.traceEvents?.Where(x => string.Equals(x.name, "Source")).LongCount() ?? 0;
-				InstantiateFunctionEntries = trace?.traceEvents?.Where(x => string.Equals(x.name, "InstantiateFunction")).LongCount() ?? 0;
-				CodeGenFunctionEntries = trace?.traceEvents?.Where(x => string.Equals(x.name, "CodeGen Function")).LongCount() ?? 0;
+				SourceEntries = trace?.traceEvents?.Where(x => String.Equals(x.name, "Source")).LongCount() ?? 0;
+				InstantiateFunctionEntries = trace?.traceEvents?.Where(x => String.Equals(x.name, "InstantiateFunction")).LongCount() ?? 0;
+				CodeGenFunctionEntries = trace?.traceEvents?.Where(x => String.Equals(x.name, "CodeGen Function")).LongCount() ?? 0;
 
 				// Other
 				ObjectBytes = GetObjectSize();
@@ -135,8 +135,8 @@ namespace UnrealBuildTool
 				".pch",
 			};
 
-			public static string CsvHeader => string.Join(',', CsvColumns);
-			public string CsvLine => string.Join(',', CsvColumns.Select(x => GetType().GetProperty(x)!.GetValue(this)!.ToString()));
+			public static string CsvHeader => String.Join(',', CsvColumns);
+			public string CsvLine => String.Join(',', CsvColumns.Select(x => GetType().GetProperty(x)!.GetValue(this)!.ToString()));
 
 			private long GetObjectSize()
 			{
@@ -205,7 +205,7 @@ namespace UnrealBuildTool
 				Dictionary<FileReference, List<long>> Sources = new Dictionary<FileReference, List<long>>();
 				foreach (ClangTrace ClangTrace in ClangTraces.Where(x => x.traceEvents != null))
 				{
-					foreach (ClangTrace.TraceEvent Event in ClangTrace.traceEvents!.Where(x => string.Equals(x.name, "Source") && x.args?.ContainsKey("detail") == true))
+					foreach (ClangTrace.TraceEvent Event in ClangTrace.traceEvents!.Where(x => String.Equals(x.name, "Source") && x.args?.ContainsKey("detail") == true))
 					{
 						FileReference SourceFile = new FileReference(Event.args!["detail"].ToString()!);
 						if (!Sources.ContainsKey(SourceFile))
diff --git a/Engine/Source/Programs/UnrealBuildTool/Modes/AggregateParsedTimingInfo.cs b/Engine/Source/Programs/UnrealBuildTool/Modes/AggregateParsedTimingInfo.cs
--- a/Engine/Source/Programs/UnrealBuildTool/Modes/AggregateParsedTimingInfo.cs
+++ b/Engine/Source/Programs/UnrealBuildTool/Modes/AggregateParsedTimingInfo.cs
@@ -231,7 +231,7 @@ namespace UnrealBuildTool
 					{
 						TemplateParamSig.Add("...");
 					}
-					string GroupName = Child.Name.Replace(Match.Groups["Template"].Value, String.Format("<{0}>", string.Join(", ", TemplateParamSig)));
+					string GroupName = Child.Name.Replace(Match.Groups["Template"].Value, String.Format("<{0}>", String.Join(", ", TemplateParamSig)));
 
 					// See if we have a group for this template already. If not, add it.
 					if (!ChildGroups.ContainsKey(GroupName))
diff --git a/Engine/Source/Programs/UnrealBuildTool/Modes/BuildMode.cs b/Engine/Source/Programs/UnrealBuildTool/Modes/BuildMode.cs
--- a/Engine/Source/Programs/UnrealBuildTool/Modes/BuildMode.cs
+++ b/Engine/Source/Programs/UnrealBuildTool/Modes/BuildMode.cs
@@ -535,7 +535,7 @@ namespace UnrealBuildTool
 
 				// Now that we have the dependencies object, optionally create the artifact cache
 				IActionArtifactCache? actionArtifactCache = null;
-				if (!string.IsNullOrEmpty(BuildConfiguration.ArtifactDirectory) && (BuildConfiguration.bArtifactRead || BuildConfiguration.bArtifactWrites))
+				if (!String.IsNullOrEmpty(BuildConfiguration.ArtifactDirectory) && (BuildConfiguration.bArtifactRead || BuildConfiguration.bArtifactWrites))
 				{
 					DirectoryReference artifactDirectory = new(BuildConfiguration.ArtifactDirectory);
 					DirectoryReference.CreateDirectory(artifactDirectory);
diff --git a/Engine/Source/Programs/UnrealBuildTool/Modes/FixIncludePathsMode.cs b/Engine/Source/Programs/UnrealBuildTool/Modes/FixIncludePathsMode.cs
--- a/Engine/Source/Programs/UnrealBuildTool/Modes/FixIncludePathsMode.cs
+++ b/Engine/Source/Programs/UnrealBuildTool/Modes/FixIncludePathsMode.cs
@@ -226,8 +226,8 @@ namespace UnrealBuildTool
 												}
 
 												// if the include and the source file live in the same directory then it is OK to be relative
-												if (string.Equals(System.IO.Directory.GetParent(FullPath)?.FullName, System.IO.Directory.GetParent(InputFile.FullName)?.FullName, StringComparison.CurrentCultureIgnoreCase) &&
-													string.Equals(Include, System.IO.Path.GetFileName(FullPath), StringComparison.CurrentCultureIgnoreCase))
+												if (String.Equals(System.IO.Directory.GetParent(FullPath)?.FullName, System.IO.Directory.GetParent(InputFile.FullName)?.FullName, StringComparison.CurrentCultureIgnoreCase) &&
+													String.Equals(Include, System.IO.Path.GetFileName(FullPath), StringComparison.CurrentCultureIgnoreCase))
 												{
 													Logger.LogInformation("{FileName}({LineNumber}): Using '{Include}' because it is in the same directory.", InputFile.FullName, LineNumber, Include);
 													PreferredInclude = Include;
@@ -399,7 +399,7 @@ namespace UnrealBuildTool
 					return -1;
 				}
 
-				return string.Compare(x, y);
+				return String.Compare(x, y);
 			}
 		}
 
diff --git a/Engine/Source/Programs/UnrealBuildTool/Modes/GenerateClangDatabase.cs b/Engine/Source/Programs/UnrealBuildTool/Modes/GenerateClangDatabase.cs
--- a/Engine/Source/Programs/UnrealBuildTool/Modes/GenerateClangDatabase.cs
+++ b/Engine/Source/Programs/UnrealBuildTool/Modes/GenerateClangDatabase.cs
@@ -288,7 +288,7 @@ namespace UnrealBuildTool
 			return ExtraPlatformArguments;
 		}
 
-		static private string ConvertResponseFile(FileItem OriginalFileItem, IEnumerable<string> FileContents, ILogger Logger)
+		private static string ConvertResponseFile(FileItem OriginalFileItem, IEnumerable<string> FileContents, ILogger Logger)
 		{
 			List<string> NewFileContents = new List<string>(FileContents);
 			FileItem NewFileItem = FileItem.GetItemByFileReference(new FileReference(OriginalFileItem.AbsolutePath + ".gcd"));
@@ -333,7 +333,7 @@ namespace UnrealBuildTool
 			return NewFileItem.AbsolutePath;
 		}
 
-		static private string ConvertPath(string Line, string OldPath)
+		private static string ConvertPath(string Line, string OldPath)
 		{
 			OldPath = OldPath.Replace("\"", "");
 			if (OldPath[^1] == '\\' || OldPath[^1] == '/')
diff --git a/Engine/Source/Programs/UnrealBuildTool/Modes/GenerateProjectFilesMode.cs b/Engine/Source/Programs/UnrealBuildTool/Modes/GenerateProjectFilesMode.cs
--- a/Engine/Source/Programs/UnrealBuildTool/Modes/GenerateProjectFilesMode.cs
+++ b/Engine/Source/Programs/UnrealBuildTool/Modes/GenerateProjectFilesMode.cs
@@ -166,7 +166,7 @@ namespace UnrealBuildTool
 			}
 			if (BadPlatformNames.Count > 0)
 			{
-				Log.TraceInformationOnce("\nSome Platforms were skipped due to invalid SDK setup: {0}.\nSee the log file for detailed information\n\n", string.Join(", ", BadPlatformNames));
+				Log.TraceInformationOnce("\nSome Platforms were skipped due to invalid SDK setup: {0}.\nSee the log file for detailed information\n\n", String.Join(", ", BadPlatformNames));
 				Logger.LogInformation("");
 			}
 			Logger.LogDebug("---   SDK INFO END   ---");
@@ -292,7 +292,7 @@ namespace UnrealBuildTool
 
 			// We have a project file either via -project= or because there was something called .uproject in the arg list
 			// so now validate it
-			if (!string.IsNullOrEmpty(CandidateProjectPath))
+			if (!String.IsNullOrEmpty(CandidateProjectPath))
 			{
 				FileReference? CandidateProjectFile = FileReference.FindCorrectCase(new FileReference(CandidateProjectPath));
 
@@ -311,7 +311,7 @@ namespace UnrealBuildTool
 				if (CandidateProjectFile == null || !FileReference.Exists(CandidateProjectFile))
 				{
 					// if we didn't find anything then throw an error as the user explicitly provided a uproject
-					throw new Exception(string.Format("Unable to find project file based on argument {0}", CandidateProjectPath));
+					throw new Exception(String.Format("Unable to find project file based on argument {0}", CandidateProjectPath));
 				}
 
 				Logger.LogDebug("Resolved project argument {CandidateProjectPath} to {CandidateProjectFile}", CandidateProjectPath, CandidateProjectFile);
diff --git a/Engine/Source/Programs/UnrealBuildTool/Modes/IWYUMode.cs b/Engine/Source/Programs/UnrealBuildTool/Modes/IWYUMode.cs
--- a/Engine/Source/Programs/UnrealBuildTool/Modes/IWYUMode.cs
+++ b/Engine/Source/Programs/UnrealBuildTool/Modes/IWYUMode.cs
@@ -369,13 +369,13 @@ namespace UnrealBuildTool
 
 			string? ModuleToUpdateName = GetModuleToUpdateName(TargetDescriptors[0]);
 
-			if (!string.IsNullOrEmpty(ModuleToUpdateName))
+			if (!String.IsNullOrEmpty(ModuleToUpdateName))
 			{
 				foreach (string OnlyModuleName in TargetDescriptors[0].OnlyModuleNames)
 				{
-					if (string.Compare(OnlyModuleName, ModuleToUpdateName, StringComparison.OrdinalIgnoreCase) != 0)
+					if (String.Compare(OnlyModuleName, ModuleToUpdateName, StringComparison.OrdinalIgnoreCase) != 0)
 					{
-						Logger.LogError($"ModuleToUpdate '{ModuleToUpdateName}' was not in list of specified modules: {string.Join(", ", TargetDescriptors[0].OnlyModuleNames)}");
+						Logger.LogError($"ModuleToUpdate '{ModuleToUpdateName}' was not in list of specified modules: {String.Join(", ", TargetDescriptors[0].OnlyModuleNames)}");
 						return -1;
 					}
 				}
@@ -674,7 +674,10 @@ namespace UnrealBuildTool
 							if (Info.File.Contains(".gen.cpp", StringComparison.Ordinal))
 							{
 								lock (GeneratedCppInfos)
+								{
 									GeneratedCppInfos.Add(Info);
+								}
+
 								continue;
 							}
 
@@ -729,13 +732,18 @@ namespace UnrealBuildTool
 									if (Include.Full.Contains("/UHT/"))
 									{
 										lock (GeneratedHeaderInfos)
+										{
 											GeneratedHeaderInfos.Add(Include.Full);
+										}
+
 										break;
 									}
 									else if (Include.Full.Contains("/VNI/"))
 									{
 										lock (GeneratedHeaderInfos)
+										{
 											GeneratedHeaderInfos.Add(Include.Full);
+										}
 									}
 								}
 							}
@@ -752,7 +760,9 @@ namespace UnrealBuildTool
 
 				// Something went wrong parsing iwyu files.
 				if (ReadSuccess == 0)
+				{
 					return -1;
+				}
 
 				if (bWriteToc)
 				{
@@ -984,7 +994,10 @@ namespace UnrealBuildTool
 						if (SeenInclude.Printable == Include.Printable && Infos.TryGetValue(Include.Full, out IncludeInfo))
 						{
 							foreach (string I in IncludeInfo.TransitiveIncludes.Keys)
+							{
 								SeenTransitiveIncludes.Add(I);
+							}
+
 							break;
 						}
 					}
@@ -1059,11 +1072,16 @@ namespace UnrealBuildTool
 			{
 				int LastIndexOfSlash = FilePath.LastIndexOf('/');
 				if (LastIndexOfSlash == -1)
+				{
 					break;
+				}
+
 				FilePath = FilePath.Substring(0, LastIndexOfSlash);
 				UEBuildModule? Module;
 				if (PathToModule.TryGetValue(FilePath, out Module))
+				{
 					return Module;
+				}
 			}
 			return null;
 		}
@@ -1079,7 +1097,9 @@ namespace UnrealBuildTool
 				FileReference FileReference = FileReference.Combine(Dir, Include);
 				FileItem FileItem = FileItem.GetItemByFileReference(FileReference);
 				if (FileItem.Exists)
+				{
 					return FileItem.Location.FullName.Replace('\\', '/');
+				}
 			}
 
 			bool Nested = Visited != null;
@@ -1089,7 +1109,9 @@ namespace UnrealBuildTool
 				FileReference FileReference = FileReference.Combine(FileReference.FromString(From).Directory, Include);
 				FileItem FileItem = FileItem.GetItemByFileReference(FileReference);
 				if (FileItem.Exists)
+				{
 					return FileItem.Location.FullName.Replace('\\', '/');
+				}
 			}
 
 			foreach (var PublicModule in Module.Rules.PublicDependencyModuleNames.Union(Module.Rules.PublicIncludePathModuleNames).Union(Module.Rules.PrivateIncludePathModuleNames).Union(Module.Rules.PrivateDependencyModuleNames))
@@ -1104,11 +1126,15 @@ namespace UnrealBuildTool
 						Visited.Add(DependencyModule);
 					}
 					else if (!Visited.Add(DependencyModule))
+					{
 						continue;
+					}
 
 					string? Str = GetFullName(Include, From, DependencyModule, NameToModule, Visited);
 					if (Str != null)
+					{
 						return Str;
+					}
 				}
 			}
 
@@ -1258,7 +1284,9 @@ namespace UnrealBuildTool
 			Parallel.ForEach(Infos.Values, Info =>
 			{
 				if (!IsValidForUpdate(Info, ValidPaths, true))
+				{
 					return;
+				}
 
 				bool IsCpp = Info.IsCpp;
 				bool IsPrivate = IsCpp || Info.File.Contains("/Private/");
@@ -1288,15 +1316,21 @@ namespace UnrealBuildTool
 					string NameWithoutPath = Include.Printable;
 					int LastSlash = NameWithoutPath.LastIndexOf("/");
 					if (LastSlash != -1)
+					{
 						NameWithoutPath = NameWithoutPath.Substring(LastSlash + 1);
+					}
 
 					string QuotedPrintable = Include.System ? $"<{Include.Printable}>" : $"\"{Include.Printable}\"";
 
 					bool Keep = true;
 					if (Info.File == Include.Full) // Sometimes IWYU outputs include to the same file if .gen.cpp is inlined and includes file with slightly different path. just skip those
+					{
 						Keep = false;
+					}
 					else if (IsCpp && MatchingH == NameWithoutPath)
+					{
 						Keep = true;
+					}
 					else if (!bNoTransitiveIncludes)
 					{
 						foreach (var Include2 in Info.Includes)
@@ -1308,7 +1342,10 @@ namespace UnrealBuildTool
 								if (Include2.Resolved.TransitiveIncludes!.TryGetValue(Key, out TransitivePath))
 								{
 									if (ShouldLog != null)
-										TransitivelyIncluded.TryAdd(QuotedPrintable, string.Join(" -> ", Include2.Printable, TransitivePath));
+									{
+										TransitivelyIncluded.TryAdd(QuotedPrintable, String.Join(" -> ", Include2.Printable, TransitivePath));
+									}
+
 									Keep = false;
 									break;
 								}
@@ -1373,9 +1410,13 @@ namespace UnrealBuildTool
 				foreach (IWYUIncludeEntry SeenInclude in Info.IncludesSeenInFile)
 				{
 					if (SeenInclude.System)
+					{
 						SeenIncludes.Add($"<{SeenInclude.Printable}>");
+					}
 					else
+					{
 						SeenIncludes.Add($"\"{SeenInclude.Printable}\"");
+					}
 				}
 
 
@@ -1403,7 +1444,10 @@ namespace UnrealBuildTool
 					ReadOnlySpan<char> LineSpan = Line.AsSpan().Trim();
 					bool StartsWithHash = LineSpan.Length > 0 && LineSpan[0] == '#';
 					if (StartsWithHash)
+					{
 						LineSpan = LineSpan.Slice(1).TrimStart();
+					}
+
 					if (!StartsWithHash || !LineSpan.StartsWith("include"))
 					{
 						// Might be forward declaration.. 
@@ -1418,9 +1462,13 @@ namespace UnrealBuildTool
 						if (Line.Contains("IWYU pragma: "))
 						{
 							if (Line.Contains(": begin_keep"))
+							{
 								ForceKeepScope = true;
+							}
 							else if (Line.Contains(": end_keep"))
+							{
 								ForceKeepScope = false;
+							}
 						}
 
 						// File is autogenerated by some tool, don't mess with it
@@ -1430,7 +1478,9 @@ namespace UnrealBuildTool
 						}
 
 						if (LookFor_API && !Contains_API && Line.Contains("_API", StringComparison.Ordinal))
+						{
 							Contains_API = true;
+						}
 
 						continue;
 					}
@@ -1486,14 +1536,18 @@ namespace UnrealBuildTool
 						}
 
 						if (Line.Contains("IWYU pragma: ", StringComparison.Ordinal))
+						{
 							ForceKeep = true;
+						}
 					}
 
 					string Include = IncludeSpan.ToString();
 
 					// If Include is not seen it means that it is probably inside a #if/#endif with condition false. These includes we can't touch
 					if (!SeenIncludes.Contains(Include))
+					{
 						continue;
+					}
 
 					if (!ForceKeep && !ForceKeepScope && !CleanedupIncludes.Contains(Include))
 					{
@@ -1509,7 +1563,9 @@ namespace UnrealBuildTool
 				{
 					//IncludesToAdd.Add(new IWYUIncludeEntry() { Full = PlatformInfo!.File, Printable = "HAL/Platform.h", Resolved = PlatformInfo });
 					if (!LinesToRemove.Remove("#include \"HAL/Platform.h\""))
+					{
 						IncludesToAdd.Add("\"HAL/Platform.h\"");
+					}
 				}
 
 				// Nothing has changed! early out of this file
@@ -1577,10 +1633,15 @@ namespace UnrealBuildTool
 						++LineIndex;
 
 						if (!OldLine.TrimStart().StartsWith("#pragma"))
+						{
 							continue;
+						}
+
 						NewLines.Add("");
 						foreach (string Line in LinesToAdd)
+						{
 							NewLines.Add(Line);
+						}
 
 						NewLines.AddRange(ExistingLines.Skip(LineIndex));
 						break;
@@ -1602,7 +1663,10 @@ namespace UnrealBuildTool
 						ReadOnlySpan<char> OldLineSpan = OldLine.AsSpan().TrimStart();
 						bool StartsWithHash = OldLineSpan.Length > 0 && OldLineSpan[0] == '#';
 						if (StartsWithHash)
+						{
 							OldLineSpan = OldLineSpan.Slice(1).TrimStart();
+						}
+
 						bool IsInclude = StartsWithHash && OldLineSpan.StartsWith("include");
 
 						string OldLineTrimmedStart = OldLine.TrimStart();
@@ -1623,7 +1687,9 @@ namespace UnrealBuildTool
 								{
 									// This logic is a bit shaky but handle the situations where file starts with #if and ends with #endif
 									if (LastSeenIncludeBeforeCode != -1)
+									{
 										IsInFirstIncludeBlock = false;
+									}
 								}
 
 								// This need to be inside "IsInFirstIncludeBlock" check because some files have forward declares far down in the file
@@ -1666,7 +1732,9 @@ namespace UnrealBuildTool
 						}
 
 						if (IsInFirstIncludeBlock)
+						{
 							LastSeenIncludeBeforeCode = NewLines.Count - 1;
+						}
 					}
 
 					if (LinesToAdd.Count > 0)
@@ -1681,7 +1749,9 @@ namespace UnrealBuildTool
 						LastSeenIncludeBeforeCode += LinesToAdd.Count;
 
 						if (FirstForwardDeclareLineIndex != -1 && FirstForwardDeclareLineIndex > LastSeenIncludeBeforeCode)
+						{
 							FirstForwardDeclareLineIndex += LinesToAdd.Count;
+						}
 					}
 
 					if (ForwardDeclarationsToAdd.Count > 0)
@@ -1711,7 +1781,9 @@ namespace UnrealBuildTool
 				{
 					Dictionary<string, string> PrintableToFull = new();
 					foreach (var Seen in Info.IncludesSeenInFile)
+					{
 						PrintableToFull.TryAdd(Seen.Printable, Seen.Full);
+					}
 
 					// Remove the includes in LinesRemoved
 					LinesRemoved.RemoveWhere(Line =>
@@ -1730,7 +1802,10 @@ namespace UnrealBuildTool
 						IncludeSpan = IncludeSpan.Slice(1, Index);
 						string? Full;
 						if (!PrintableToFull.TryGetValue(IncludeSpan.ToString(), out Full))
+						{
 							return false;
+						}
+
 						return Info.TransitiveIncludes.ContainsKey(Full);
 					});
 
@@ -1763,12 +1838,16 @@ namespace UnrealBuildTool
 					}
 				}
 				lock (UpdatedFiles)
+				{
 					UpdatedFiles.Add(new(Info, NewLines));
+				}
 			});
 
 			// Something went wrong processing code files.
 			if (ProcessSuccess == 0)
+			{
 				return -1;
+			}
 
 			Logger.LogInformation($"Parsed {FilesParseCount} and updated {UpdatedFiles.Count} files (Found {OutOfDateCount} .iwyu files out of date)");
 
@@ -1781,7 +1860,9 @@ namespace UnrealBuildTool
 
 			// If we have been logging we can exit now since we don't want to write any files to disk
 			if (ShouldLog != null)
+			{
 				return 0;
+			}
 
 			List<System.Diagnostics.Process> P4Processes = new();
 
@@ -1820,7 +1901,9 @@ namespace UnrealBuildTool
 				foreach (var (Info, NewLines) in UpdatedFiles)
 				{
 					if (new FileInfo(Info.File).IsReadOnly)
+					{
 						ReadOnlyFileInfos.Add(Info);
+					}
 				}
 
 				if (ReadOnlyFileInfos.Count > 0)
@@ -1833,10 +1916,14 @@ namespace UnrealBuildTool
 					{
 						Logger.LogInformation($"   edit {Info.File}");
 						if (--ShowCount == 0)
+						{
 							break;
+						}
 					}
 					if (ReadOnlyFileInfos.Count > 5)
+					{
 						Logger.LogInformation($"   ... and {ReadOnlyFileInfos.Count - 5} more.");
+					}
 
 					StringBuilder P4Arguments = new();
 					int BatchSize = 10;
@@ -1872,7 +1959,10 @@ namespace UnrealBuildTool
 			foreach (var (Info, NewLines) in UpdatedFiles)
 			{
 				if (SkippedFiles.Contains(Info.Source!))
+				{
 					continue;
+				}
+
 				try
 				{
 					File.WriteAllLines(Info.File, NewLines);
@@ -1894,7 +1984,9 @@ namespace UnrealBuildTool
 			}
 
 			if (!WriteSuccess)
+			{
 				return -1;
+			}
 
 			Logger.LogInformation($"Done!");
 			return 0;
@@ -1973,14 +2065,18 @@ namespace UnrealBuildTool
 			Parallel.ForEach(Infos.Values, Info =>
 			{
 				if (!IsValidForUpdate(Info, ValidPaths, false) || Info.File.EndsWith(".h"))
+				{
 					return;
+				}
 
 				Dictionary<string, string> SeenTransitiveIncludes = new();
 				Stack<string> Stack = new();
 				CalculateTransitive(Info, Info, Stack, SeenTransitiveIncludes, new Dictionary<string, string>(), true);
 
 				lock (InfosToCompare)
+				{
 					InfosToCompare.Add(new(Info, SeenTransitiveIncludes));
+				}
 
 				foreach (var Include in SeenTransitiveIncludes.Union(Info.TransitiveIncludes))
 				{
@@ -2014,29 +2110,49 @@ namespace UnrealBuildTool
 				long OptimizedSize = 0;
 				long SeenSize = 0;
 				foreach (var Include in Kvp.Item1.TransitiveIncludes.Keys)
+				{
 					OptimizedSize += FileToSize[Include];
+				}
+
 				foreach (var Include in Kvp.Item2.Keys)
+				{
 					SeenSize += FileToSize[Include];
+				}
+
 				long Diff = SeenSize - OptimizedSize;
 				lock (ByteSizeDiffList)
+				{
 					ByteSizeDiffList.Add(new(Diff, Kvp.Item1));
+				}
 			});
 
 			ByteSizeDiffList.Sort((a, b) =>
 			{
 				if (a.Item1 == b.Item1)
+				{
 					return 0;
+				}
+
 				if (a.Item1 < b.Item1)
+				{
 					return 1;
+				}
+
 				return -1;
 			});
 
 			Func<long, string> PrettySize = (Size) =>
 			{
 				if (Size > 1024 * 1024)
+				{
 					return (((float)Size) / (1024 * 1024)).ToString("0.00") + "mb";
+				}
+
 				if (Size > 1024)
+				{
 					return (((float)Size) / (1024)).ToString("0.00") + "kb";
+				}
+
 				return Size + "b";
 			};
 
@@ -2049,7 +2165,10 @@ namespace UnrealBuildTool
 				IWYUInfo File = ByteSizeDiffList[I].Item2;
 				long OptimizedSize = 0;
 				foreach (var Include in File.TransitiveIncludes.Keys)
+				{
 					OptimizedSize += FileToSize[Include];
+				}
+
 				float Percent = 100.0f - (((float)OptimizedSize) / (OptimizedSize + Saved) * 100.0f);
 				Logger.LogInformation($"{Path.GetFileName(File.File)}   {PrettySize(OptimizedSize + Saved)} -> {PrettySize(OptimizedSize)}  (Saved {PrettySize(Saved)} or {Percent.ToString("0.0")}%)");
 			}
diff --git a/Engine/Source/Programs/UnrealBuildTool/Modes/ParseMsvcTimingInfo.cs b/Engine/Source/Programs/UnrealBuildTool/Modes/ParseMsvcTimingInfo.cs
--- a/Engine/Source/Programs/UnrealBuildTool/Modes/ParseMsvcTimingInfo.cs
+++ b/Engine/Source/Programs/UnrealBuildTool/Modes/ParseMsvcTimingInfo.cs
@@ -40,7 +40,7 @@ namespace UnrealBuildTool
 			TimingDataType CurrentType = TimingDataType.None;
 			foreach (string Line in AllLines)
 			{
-				if (string.IsNullOrWhiteSpace(Line))
+				if (String.IsNullOrWhiteSpace(Line))
 				{
 					continue;
 				}
@@ -101,7 +101,7 @@ namespace UnrealBuildTool
 			}
 
 			// Build the summary.
-			TimingData Summary = new TimingData(InputFile.FullName.Replace(".timing.txt", string.Empty), TimingDataType.Summary);
+			TimingData Summary = new TimingData(InputFile.FullName.Replace(".timing.txt", String.Empty), TimingDataType.Summary);
 			Summary.AddChild(SummarizeParsedTimingData("IncludeTimings", TimingDataType.Include, Includes));
 			Summary.AddChild(SummarizeParsedTimingData("ClassTimings", TimingDataType.Class, Classes));
 			Summary.AddChild(SummarizeParsedTimingData("FunctionTimings", TimingDataType.Function, Functions));
@@ -199,7 +199,7 @@ namespace UnrealBuildTool
 			LineDepth = TimingDataMatch.Groups["Indent"].Success ? TimingDataMatch.Groups["Indent"].Value.Count() : 0;
 
 			TimingData ParsedTimingData = new TimingData(TimingDataMatch.Groups["Name"].Value, TimingType);
-			ParsedTimingData.ExclusiveDuration = float.Parse(TimingDataMatch.Groups["Duration"].Value);
+			ParsedTimingData.ExclusiveDuration = Single.Parse(TimingDataMatch.Groups["Duration"].Value);
 
 			return ParsedTimingData;
 		}
@@ -257,7 +257,7 @@ namespace UnrealBuildTool
 
 					int Indent = Match.Groups[1].Length;
 					string FileName = Match.Groups[2].Value;
-					float Duration = float.Parse(Match.Groups[3].Value);
+					float Duration = Single.Parse(Match.Groups[3].Value);
 
 					while (Indent <= FinishTimesForIndent.Count - 1)
 					{
@@ -315,7 +315,7 @@ namespace UnrealBuildTool
 				float Time;
 				ClassNameToTime.TryGetValue(ClassName, out Time);
 
-				Time += float.Parse(Match.Groups[2].Value);
+				Time += Single.Parse(Match.Groups[2].Value);
 				ClassNameToTime[ClassName] = Time;
 			}
 
diff --git a/Engine/Source/Programs/UnrealBuildTool/Modes/ProfileUnitySizesMode.cs b/Engine/Source/Programs/UnrealBuildTool/Modes/ProfileUnitySizesMode.cs
--- a/Engine/Source/Programs/UnrealBuildTool/Modes/ProfileUnitySizesMode.cs
+++ b/Engine/Source/Programs/UnrealBuildTool/Modes/ProfileUnitySizesMode.cs
@@ -71,7 +71,7 @@ namespace UnrealBuildTool
 				if (State != null)
 				{
 					string? LogText = State.ToString();
-					if (!string.IsNullOrEmpty(LogText))
+					if (!String.IsNullOrEmpty(LogText))
 					{
 						// Console.WriteLine(LogText);
 						Match ExecutorTimingMatch = ExecutorTimingRegex.Match(LogText);
@@ -95,7 +95,7 @@ namespace UnrealBuildTool
 						Match NumFilesMatch = NumFilesRegex.Match(LogText);
 						if (NumFilesMatch.Success)
 						{
-							if (!int.TryParse(NumFilesMatch.Groups[1].Value, out TimingData.NumFiles))
+							if (!Int32.TryParse(NumFilesMatch.Groups[1].Value, out TimingData.NumFiles))
 							{
 								Console.WriteLine($"Failed to parse '{LogText}'");
 							}
